Chaos - add Simd Gauss Seidel Collision solver (currently an exact dupe of the non-simd version)
- remove Jacobi path from Gauss Seidel solvers (because we have a separate PartialJacobi solver already)
- move orginal Gauss Seidel solver into a correctly named file

diff --git a/Engine/Source/Programs/HeadlessChaos/Private/HeadlessChaosTestShapePadding.cpp b/Engine/Source/Programs/HeadlessChaos/Private/HeadlessChaosTestShapePadding.cpp
--- a/Engine/Source/Programs/HeadlessChaos/Private/HeadlessChaosTestShapePadding.cpp
+++ b/Engine/Source/Programs/HeadlessChaos/Private/HeadlessChaosTestShapePadding.cpp
@@ -7,7 +7,7 @@
 #include "Chaos/Box.h"
 #include "Chaos/CollisionResolution.h"
 #include "Chaos/CollisionResolutionTypes.h"
-#include "Chaos/Collision/SolverCollisionContainer.h"
+#include "Chaos/Collision/PBDCollisionContainerSolver.h"
 #include "Chaos/Convex.h"
 #include "Chaos/Sphere.h"
 #include "Chaos/GJK.h"
diff --git a/Engine/Source/Runtime/Experimental/Chaos/Private/Chaos/Collision/PBDCollisionConstraint.cpp b/Engine/Source/Runtime/Experimental/Chaos/Private/Chaos/Collision/PBDCollisionConstraint.cpp
--- a/Engine/Source/Runtime/Experimental/Chaos/Private/Chaos/Collision/PBDCollisionConstraint.cpp
+++ b/Engine/Source/Runtime/Experimental/Chaos/Private/Chaos/Collision/PBDCollisionConstraint.cpp
@@ -2,8 +2,8 @@
 #include "Chaos/Collision/PBDCollisionConstraint.h"
 #include "Chaos/Collision/CollisionConstraintAllocator.h"
 #include "Chaos/Collision/PBDCollisionConstraintHandle.h"
+#include "Chaos/Collision/PBDCollisionContainerSolver.h"
 #include "Chaos/Collision/PBDCollisionSolver.h"
-#include "Chaos/Collision/SolverCollisionContainer.h"
 #include "Chaos/Evolution/SolverBody.h"
 #include "Chaos/Evolution/SolverBodyContainer.h"
 #include "Chaos/Island/IslandManager.h"
diff --git a/Engine/Source/Runtime/Experimental/Chaos/Private/Chaos/Collision/PBDCollisionContainerSolver.cpp b/Engine/Source/Runtime/Experimental/Chaos/Private/Chaos/Collision/PBDCollisionContainerSolver.cpp
new file mode 100644
--- /dev/null
+++ b/Engine/Source/Runtime/Experimental/Chaos/Private/Chaos/Collision/PBDCollisionContainerSolver.cpp
@@ -0,0 +1,612 @@
+// Copyright Epic Games, Inc. All Rights Reserved.
+#include "Chaos/Collision/PBDCollisionContainerSolver.h"
+
+#include "Chaos/Collision/CollisionApplyType.h"
+#include "Chaos/Collision/PBDCollisionConstraint.h"
+#include "Chaos/CollisionResolution.h"
+#include "Chaos/DebugDrawQueue.h"
+#include "Chaos/Evolution/SolverBodyContainer.h"
+#include "Chaos/Island/IslandManager.h"
+#include "Chaos/Particle/ParticleUtilities.h"
+#include "Chaos/PBDCollisionConstraints.h"
+#include "Chaos/PBDCollisionConstraintsContact.h"
+#include "Chaos/Utilities.h"
+
+// Private includes
+
+#include "ChaosLog.h"
+
+//PRAGMA_DISABLE_OPTIMIZATION
+
+namespace Chaos
+{
+	namespace CVars
+	{
+		extern bool bChaos_PBDCollisionSolver_Position_SolveEnabled;
+		extern bool bChaos_PBDCollisionSolver_Velocity_SolveEnabled;
+
+		// If one body is more than MassRatio1 times the mass of the other, adjust the solver stiffness when the lighter body is underneath.
+		// Solver stiffness will be equal to 1 when the mass ratio is MassRatio1.
+		// Solver stiffness will be equal to 0 when the mass ratio is MassRatio2.
+		FRealSingle Chaos_PBDCollisionSolver_AutoStiffness_MassRatio1 = 0;
+		FRealSingle Chaos_PBDCollisionSolver_AutoStiffness_MassRatio2 = 0;
+		FAutoConsoleVariableRef CVarChaosPBDCollisionSolverAutoStiffnessMassRatio1(TEXT("p.Chaos.PBDCollisionSolver.AutoStiffness.MassRatio1"), Chaos_PBDCollisionSolver_AutoStiffness_MassRatio1, TEXT(""));
+		FAutoConsoleVariableRef CVarChaosPBDCollisionSolverAutoStiffnessMassRatio2(TEXT("p.Chaos.PBDCollisionSolver.AutoStiffness.MassRatio2"), Chaos_PBDCollisionSolver_AutoStiffness_MassRatio2, TEXT(""));
+
+		// Jacobi solver stiffness
+		// @todo(chaos): to be tuned
+		FRealSingle Chaos_PBDCollisionSolver_JacobiStiffness = 0.5f;
+		FAutoConsoleVariableRef CVarChaosPBDCollisionSolverJacobiStiffness(TEXT("p.Chaos.PBDCollisionSolver.JacobiStiffness"), Chaos_PBDCollisionSolver_JacobiStiffness, TEXT(""));
+
+		// Jacobi position tolerance. Position corrections below this are zeroed.
+		// @todo(chaos): to be tuned
+		FRealSingle Chaos_PBDCollisionSolver_JacobiPositionTolerance = 1.e-6f;
+		FAutoConsoleVariableRef CVarChaosPBDCollisionSolverJacobiPositionTolerance(TEXT("p.Chaos.PBDCollisionSolver.JacobiPositionTolerance"), Chaos_PBDCollisionSolver_JacobiPositionTolerance, TEXT(""));
+
+		// Jacobi rotation tolerance. Rotation corrections below this are zeroed.
+		// @todo(chaos): to be tuned
+		FRealSingle Chaos_PBDCollisionSolver_JacobiRotationTolerance = 1.e-8f;
+		FAutoConsoleVariableRef CVarChaosPBDCollisionSolverJacobiRotationTolerance(TEXT("p.Chaos.PBDCollisionSolver.JacobiRotationTolerance"), Chaos_PBDCollisionSolver_JacobiRotationTolerance, TEXT(""));
+	}
+
+
+	//////////////////////////////////////////////////////////////////////////////////////////////////
+	//////////////////////////////////////////////////////////////////////////////////////////////////
+	//////////////////////////////////////////////////////////////////////////////////////////////////
+	//////////////////////////////////////////////////////////////////////////////////////////////////
+
+	// Transform the Constraint's local-space data into world space for use by  the collision solver and also calculate tangents, errors, etc
+	FORCEINLINE_DEBUGGABLE void UpdateCollisionSolverContactPointFromConstraint(Private::FPBDCollisionSolver& Solver, const int32 SolverPointIndex, const FPBDCollisionConstraint* Constraint, const int32 ConstraintPointIndex, const FRealSingle Dt, const FSolverBody& Body0, const FSolverBody& Body1)
+	{
+		const FManifoldPoint& ManifoldPoint = Constraint->GetManifoldPoint(ConstraintPointIndex);
+
+		const FRealSingle Restitution = FRealSingle(Constraint->GetRestitution());
+		const FRealSingle RestitutionVelocityThreshold = FRealSingle(Constraint->GetRestitutionThreshold()) * Dt;
+
+		// World-space shape transforms. Manifold data is currently relative to these spaces
+		const FRigidTransform3& ShapeWorldTransform0 = Constraint->GetShapeWorldTransform0();
+		const FRigidTransform3& ShapeWorldTransform1 = Constraint->GetShapeWorldTransform1();
+
+		// World-space contact points on each shape
+		const FVec3 WorldContact0 = ShapeWorldTransform0.TransformPositionNoScale(FVec3(ManifoldPoint.ContactPoint.ShapeContactPoints[0]));
+		const FVec3 WorldContact1 = ShapeWorldTransform1.TransformPositionNoScale(FVec3(ManifoldPoint.ContactPoint.ShapeContactPoints[1]));
+		const FVec3 WorldContact = FReal(0.5) * (WorldContact0 + WorldContact1);
+		const FVec3f WorldRelativeContact0 = FVec3f(WorldContact - Body0.P());
+		const FVec3f WorldRelativeContact1 = FVec3f(WorldContact - Body1.P());
+
+		// World-space normal
+		const FVec3f WorldContactNormal = FVec3f(ShapeWorldTransform1.TransformVectorNoScale(FVec3(ManifoldPoint.ContactPoint.ShapeContactNormal)));
+
+		// World-space tangents
+		FVec3f WorldContactTangentU = FVec3f::CrossProduct(FVec3f(0, 1, 0), WorldContactNormal);
+		if (!WorldContactTangentU.Normalize(FRealSingle(UE_KINDA_SMALL_NUMBER)))
+		{
+			WorldContactTangentU = FVec3f::CrossProduct(FVec3f(1, 0, 0), WorldContactNormal);
+			WorldContactTangentU = WorldContactTangentU.GetUnsafeNormal();
+		}
+		const FVec3f WorldContactTangentV = FVec3f::CrossProduct(WorldContactNormal, WorldContactTangentU);
+
+		// Calculate contact velocity if we will need it below (restitution and/or frist-contact friction)
+		const bool bNeedContactVelocity = (!ManifoldPoint.Flags.bHasStaticFrictionAnchor) || (Restitution > FRealSingle(0));
+		FVec3f ContactVel = FVec3(0);
+		if (bNeedContactVelocity)
+		{
+			const FVec3f ContactVel0 = Body0.V() + FVec3f::CrossProduct(Body0.W(), WorldRelativeContact0);
+			const FVec3f ContactVel1 = Body1.V() + FVec3f::CrossProduct(Body1.W(), WorldRelativeContact1);
+			ContactVel = ContactVel0 - ContactVel1;
+		}
+
+		// If we have contact data from a previous tick, use it to calculate the lateral position delta we need
+		// to apply to move the contacts back to their original relative locations (i.e., to enforce static friction).
+		// Otherwise, estimate the friction correction from the contact velocity.
+		// NOTE: quadratic shapes use the velocity-based path most of the time, unless the relative motion is very small.
+		FVec3f WorldFrictionDelta = FVec3f(0);
+		if (ManifoldPoint.Flags.bHasStaticFrictionAnchor)
+		{
+			const FVec3f FrictionDelta0 = ShapeWorldTransform0.TransformPositionNoScale(FVec3(ManifoldPoint.ShapeAnchorPoints[0])) - WorldContact0;
+			const FVec3f FrictionDelta1 = ShapeWorldTransform1.TransformPositionNoScale(FVec3(ManifoldPoint.ShapeAnchorPoints[1])) - WorldContact1;
+			WorldFrictionDelta = FrictionDelta0 - FrictionDelta1;
+		}
+		else
+		{
+			// @todo(chaos): consider adding a multiplier to the initial contact friction
+			WorldFrictionDelta = ContactVel * Dt;
+		}
+
+		// The contact point error we are trying to correct in this solver
+		const FRealSingle TargetPhi = ManifoldPoint.TargetPhi;
+		const FVec3f WorldContactDelta = FVec3f(WorldContact0 - WorldContact1);
+		const FRealSingle WorldContactDeltaNormal = FVec3f::DotProduct(WorldContactDelta, WorldContactNormal) - TargetPhi;
+		const FRealSingle WorldContactDeltaTangentU = FVec3f::DotProduct(WorldContactDelta + WorldFrictionDelta, WorldContactTangentU);
+		const FRealSingle WorldContactDeltaTangentV = FVec3f::DotProduct(WorldContactDelta + WorldFrictionDelta, WorldContactTangentV);
+
+		// The target contact velocity, taking restitution into account
+		FRealSingle WorldContactTargetVelocityNormal = FRealSingle(0);
+		if (Restitution > FRealSingle(0))
+		{
+			const FRealSingle ContactVelocityNormal = FVec3f::DotProduct(ContactVel, WorldContactNormal);
+			if (ContactVelocityNormal < -RestitutionVelocityThreshold)
+			{
+				WorldContactTargetVelocityNormal = -Restitution * ContactVelocityNormal;
+			}
+		}
+
+		FWorldContactPoint& WorldContactPoint = Solver.GetWorldContactPoint(SolverPointIndex);
+		WorldContactPoint.RelativeContactPoints[0] = WorldRelativeContact0;
+		WorldContactPoint.RelativeContactPoints[1] = WorldRelativeContact1;
+		WorldContactPoint.ContactNormal = WorldContactNormal;
+		WorldContactPoint.ContactTangentU = WorldContactTangentU;
+		WorldContactPoint.ContactTangentV = WorldContactTangentV;
+		WorldContactPoint.ContactDeltaNormal = WorldContactDeltaNormal;
+		WorldContactPoint.ContactDeltaTangentU = WorldContactDeltaTangentU;
+		WorldContactPoint.ContactDeltaTangentV = WorldContactDeltaTangentV;
+		WorldContactPoint.ContactTargetVelocityNormal = WorldContactTargetVelocityNormal;
+	}
+
+	FORCEINLINE_DEBUGGABLE void UpdateCollisionSolverManifoldFromConstraint(Private::FPBDCollisionSolver& Solver, const FPBDCollisionConstraint* Constraint, const FSolverReal Dt, const int32 ConstraintPointBeginIndex, const int32 ConstraintPointEndIndex)
+	{
+		const FSolverBody& Body0 = Solver.SolverBody0().SolverBody();
+		const FSolverBody& Body1 = Solver.SolverBody1().SolverBody();
+
+		// Only calculate state for newly added contacts. Normally this is all of them, but maybe not if incremental collision is used by RBAN.
+		// Also we only add active points to the solver's manifold points list
+		for (int32 ConstraintManifoldPointIndex = ConstraintPointBeginIndex; ConstraintManifoldPointIndex < ConstraintPointEndIndex; ++ConstraintManifoldPointIndex)
+		{
+			if (!Constraint->GetManifoldPoint(ConstraintManifoldPointIndex).Flags.bDisabled)
+			{
+				const int32 SolverManifoldPointIndex = Solver.AddManifoldPoint();
+
+				// Transform the constraint contact data into world space for use by the solver
+				// We build this data directly into the solver's world-space contact data which looks a bit odd with "Init" called after but there you go
+				UpdateCollisionSolverContactPointFromConstraint(Solver, SolverManifoldPointIndex, Constraint, ConstraintManifoldPointIndex, Dt, Body0, Body1);
+
+				// Prepare the solver (set constraint-space mass etc)
+				Solver.FinalizeManifoldPoint(SolverManifoldPointIndex, Dt);
+			}
+		}
+	}
+
+	FORCEINLINE_DEBUGGABLE void UpdateCollisionSolverFromConstraint(Private::FPBDCollisionSolver& Solver, const FPBDCollisionConstraint* Constraint, const FSolverReal Dt, const FPBDCollisionSolverSettings& SolverSettings, bool& bOutPerIterationCollision)
+	{
+		// Friction values. Static and Dynamic friction are applied in the position solve for most shapes.
+		// We can also run in a mode without static friction at all. This is faster but stacking is not possible.
+		const FSolverReal StaticFriction = FSolverReal(Constraint->GetStaticFriction());
+		const FSolverReal DynamicFriction = FSolverReal(Constraint->GetDynamicFriction());
+		FSolverReal PositionStaticFriction = FSolverReal(0);
+		FSolverReal PositionDynamicFriction = FSolverReal(0);
+		FSolverReal VelocityDynamicFriction = FSolverReal(0);
+		if (SolverSettings.NumPositionFrictionIterations > 0)
+		{
+			PositionStaticFriction = StaticFriction;
+			if (!Constraint->HasQuadraticShape())
+			{
+				PositionDynamicFriction = DynamicFriction;
+			}
+			else
+			{
+				// Quadratic shapes don't use PBD dynamic friction - it has issues at slow speeds where the WxR is
+				// less than the position tolerance for friction point matching
+				// @todo(chaos): fix PBD dynamic friction on quadratic shapes
+				VelocityDynamicFriction = DynamicFriction;
+			}
+		}
+		else
+		{
+			VelocityDynamicFriction = DynamicFriction;
+		}
+
+		Solver.SetFriction(PositionStaticFriction, PositionDynamicFriction, VelocityDynamicFriction);
+
+		const FReal SolverStiffness = Constraint->GetStiffness();
+		Solver.SetStiffness(FSolverReal(SolverStiffness));
+
+		Solver.SolverBody0().SetInvMScale(Constraint->GetInvMassScale0());
+		Solver.SolverBody0().SetInvIScale(Constraint->GetInvInertiaScale0());
+		Solver.SolverBody0().SetShockPropagationScale(FReal(1));
+		Solver.SolverBody1().SetInvMScale(Constraint->GetInvMassScale1());
+		Solver.SolverBody1().SetInvIScale(Constraint->GetInvInertiaScale1());
+		Solver.SolverBody1().SetShockPropagationScale(FReal(1));
+
+		bOutPerIterationCollision = (!Constraint->GetUseManifold() || Constraint->GetUseIncrementalCollisionDetection());
+
+		UpdateCollisionSolverManifoldFromConstraint(Solver, Constraint, Dt, 0, Constraint->NumManifoldPoints());
+	}
+
+	FORCEINLINE_DEBUGGABLE void UpdateCollisionConstraintFromSolver(FPBDCollisionConstraint* Constraint, const Private::FPBDCollisionSolver& Solver, const FSolverReal Dt)
+	{
+		Constraint->ResetSolverResults();
+
+		// NOTE: We only put the non-pruned manifold points into the solver so the ManifoldPointIndex and
+		// SolverManifoldPointIndex do not necessarily match. See GatherManifoldPoints
+		int32 SolverManifoldPointIndex = 0;
+		for (int32 ManifoldPointIndex = 0; ManifoldPointIndex < Constraint->NumManifoldPoints(); ++ManifoldPointIndex)
+		{
+			FSolverVec3 NetPushOut = FSolverVec3(0);
+			FSolverVec3 NetImpulse = FSolverVec3(0);
+			FSolverReal StaticFrictionRatio = FSolverReal(0);
+
+			const Private::FPBDCollisionSolverManifoldPoint& SolverManifoldPoint = Solver.GetManifoldPoint(SolverManifoldPointIndex);
+			if (!Constraint->GetManifoldPoint(ManifoldPointIndex).Flags.bDisabled)
+			{
+				NetPushOut =
+					SolverManifoldPoint.NetPushOutNormal * SolverManifoldPoint.WorldContact.ContactNormal +
+					SolverManifoldPoint.NetPushOutTangentU * SolverManifoldPoint.WorldContact.ContactTangentU +
+					SolverManifoldPoint.NetPushOutTangentV * SolverManifoldPoint.WorldContact.ContactTangentV;
+
+				NetImpulse =
+					SolverManifoldPoint.NetImpulseNormal * SolverManifoldPoint.WorldContact.ContactNormal +
+					SolverManifoldPoint.NetImpulseTangentU * SolverManifoldPoint.WorldContact.ContactTangentU +
+					SolverManifoldPoint.NetImpulseTangentV * SolverManifoldPoint.WorldContact.ContactTangentV;
+
+				StaticFrictionRatio = SolverManifoldPoint.StaticFrictionRatio;
+
+				++SolverManifoldPointIndex;
+			}
+
+			// NOTE: We call this even for points we did not run the solver for (but with zero results)
+			Constraint->SetSolverResults(ManifoldPointIndex,
+				NetPushOut,
+				NetImpulse,
+				StaticFrictionRatio,
+				Dt);
+		}
+
+		Constraint->EndTick();
+	}
+
+	//////////////////////////////////////////////////////////////////////////////////////////////////
+	//////////////////////////////////////////////////////////////////////////////////////////////////
+	//////////////////////////////////////////////////////////////////////////////////////////////////
+	//////////////////////////////////////////////////////////////////////////////////////////////////
+
+
+	FPBDCollisionContainerSolver::FPBDCollisionContainerSolver(const FPBDCollisionConstraints& InConstraintContainer, const int32 InPriority)
+		: FConstraintContainerSolver(InPriority)
+		, ConstraintContainer(InConstraintContainer)
+		, bPerIterationCollisionDetection(false)
+	{
+#if !UE_BUILD_SHIPPING && !UE_BUILD_TEST && INTEL_ISPC
+		Private::FPBDCollisionSolverHelper::CheckISPC();
+#endif
+	}
+
+	FPBDCollisionContainerSolver::~FPBDCollisionContainerSolver()
+	{
+	}
+
+	void FPBDCollisionContainerSolver::Reset(const int32 MaxCollisions)
+	{
+		// A over-allocation policy to avoid reallocation every frame in the common case where a pile of objects is dropped
+		// and the number of contacts increases every tick.
+		if (MaxCollisions > CollisionSolvers.Max())
+		{
+			CollisionSolvers.Reset((5 * MaxCollisions) / 4);	// +25%
+		}
+
+		CollisionSolvers.Reset(MaxCollisions);
+		CollisionConstraints.Reset(CollisionSolvers.Max());
+		bCollisionConstraintPerIterationCollisionDetection.Reset(CollisionSolvers.Max());
+		CollisionSolverManifoldPoints.Reset(CollisionSolvers.Max() * Private::FPBDCollisionSolver::MaxPointsPerConstraint);
+
+		// Just set the array size for these right away - all data will be initialized later and there is no constructor
+		// We use CollisionConstraints.Num() to track how many constraints we currently have
+		CollisionSolvers.SetNumUninitialized(MaxCollisions);
+		bCollisionConstraintPerIterationCollisionDetection.SetNumUninitialized(MaxCollisions);
+	}
+
+	void FPBDCollisionContainerSolver::AddConstraints()
+	{
+		Reset(ConstraintContainer.NumConstraints());
+
+		for (FPBDCollisionConstraintHandle* ConstraintHandle : ConstraintContainer.GetConstraintHandles())
+		{
+			FPBDCollisionConstraint& Constraint = ConstraintHandle->GetContact();
+
+			AddConstraint(Constraint);
+		}
+	}
+
+	void FPBDCollisionContainerSolver::AddConstraints(const TArrayView<Private::FPBDIslandConstraint>& IslandConstraints)
+	{
+		for (Private::FPBDIslandConstraint& IslandConstraint : IslandConstraints)
+		{
+			// We will only ever be given constraints from our container (asserts in non-shipping)
+			FPBDCollisionConstraint& Constraint = IslandConstraint.GetConstraint()->AsUnsafe<FPBDCollisionConstraintHandle>()->GetContact();
+
+			AddConstraint(Constraint);
+		}
+	}
+
+	void FPBDCollisionContainerSolver::AddConstraint(FPBDCollisionConstraint& Constraint)
+	{
+		// NOTE: No need to add to CollisionSolvers or bCollisionConstraintPerIterationCollisionDetection
+		const int32 Index = CollisionConstraints.Add(&Constraint);
+
+		// Allocate the manifold points for this constraints solver
+		// NOTE: we don't know how many points we will create if collision detection is deferred or incrememntal, so just allocate space for the max allowed
+		// @todo(chaos): see if we can do better - we don't want to pay for this logic when it is only used rarely
+		const bool bDeferredCollisionDetection = ConstraintContainer.GetDetectorSettings().bDeferNarrowPhase;
+		const int32 ManifoldPointMax = (bDeferredCollisionDetection || Constraint.GetUseIncrementalManifold()) ? Private::FPBDCollisionSolver::MaxPointsPerConstraint : Constraint.NumManifoldPoints();
+		if (ManifoldPointMax > 0)
+		{
+			const int32 ManifoldPointStart = CollisionSolverManifoldPoints.Num();
+			check(CollisionSolverManifoldPoints.Max() >= ManifoldPointStart + ManifoldPointMax);	// No resize following reset
+			CollisionSolverManifoldPoints.SetNum(ManifoldPointStart + ManifoldPointMax);
+
+			CollisionSolvers[Index].SetManifoldPointsBuffer(&CollisionSolverManifoldPoints[ManifoldPointStart], ManifoldPointMax);
+		}
+		else
+		{
+			// Ideally we wouldn't have any active constraint with no manifold points, but this can
+			// currently happen for CCD objects and incremental collisions
+			CollisionSolvers[Index].Reset();
+		}
+	}
+
+
+	void FPBDCollisionContainerSolver::AddBodies(FSolverBodyContainer& SolverBodyContainer)
+	{
+		for (int32 SolverIndex = 0, SolverEndIndex = NumSolvers(); SolverIndex < SolverEndIndex; ++SolverIndex)
+		{
+			Private::FPBDCollisionSolver& CollisionSolver = CollisionSolvers[SolverIndex];
+			FPBDCollisionConstraint* Constraint = CollisionConstraints[SolverIndex];
+			check(Constraint != nullptr);
+
+			// Find the solver bodies for the particles we constrain. This will add them to the container
+			// if they aren't there already, and ensure that they are populated with the latest data.
+			FSolverBody* Body0 = SolverBodyContainer.FindOrAdd(Constraint->GetParticle0());
+			FSolverBody* Body1 = SolverBodyContainer.FindOrAdd(Constraint->GetParticle1());
+			check(Body0 != nullptr);
+			check(Body1 != nullptr);
+
+			CollisionSolver.SetSolverBodies(*Body0, *Body1);
+		}
+	}
+
+	void FPBDCollisionContainerSolver::GatherInput(const FReal Dt)
+	{
+		GatherInput(Dt, 0, NumSolvers());
+	}
+
+	void FPBDCollisionContainerSolver::CachePrefetchSolver(const int32 ConstraintIndex) const
+	{
+		if (ConstraintIndex < NumSolvers())
+		{
+			const int32 NumManifoldPoints = CollisionSolvers[ConstraintIndex].NumManifoldPoints();
+			if (NumManifoldPoints > 0)
+			{
+				FPlatformMisc::PrefetchBlock(&CollisionSolvers[ConstraintIndex].GetManifoldPoint(0), NumManifoldPoints * sizeof(Private::FPBDCollisionSolverManifoldPoint));
+			}
+			FPlatformMisc::PrefetchBlock(&CollisionSolvers[ConstraintIndex], sizeof(FPBDCollisionContainerSolver));
+		}
+	}
+
+	void FPBDCollisionContainerSolver::GatherInput(const FReal InDt, const int32 BeginIndex, const int32 EndIndex)
+	{
+		// NOTE: may be called in parallel. Should not change the container or any elements outside of [BeginIndex, EndIndex)
+
+		check(BeginIndex >= 0);
+		check(EndIndex <= NumSolvers());
+
+		const FSolverReal Dt = FSolverReal(InDt);
+
+		const int32 PrefetchCount = 4;
+		for (int32 PrefetchIndex = 0; PrefetchIndex < PrefetchCount; ++PrefetchIndex)
+		{
+			CachePrefetchSolver(PrefetchIndex);
+		}
+
+		bool bWantPerIterationCollisionDetection = false;
+		for (int32 ConstraintIndex = BeginIndex; ConstraintIndex < EndIndex; ++ConstraintIndex)
+		{
+			Private::FPBDCollisionSolver& CollisionSolver = CollisionSolvers[ConstraintIndex];
+			FPBDCollisionConstraint* Constraint = CollisionConstraints[ConstraintIndex];
+			bool& bPerIterationCollision = bCollisionConstraintPerIterationCollisionDetection[ConstraintIndex];
+
+			UpdateCollisionSolverFromConstraint(CollisionSolver, Constraint, Dt, ConstraintContainer.GetSolverSettings(), bPerIterationCollision);
+
+			// We need to run collision every iteration if we are not using manifolds, or are using incremental manifolds
+			bWantPerIterationCollisionDetection = bWantPerIterationCollisionDetection || bPerIterationCollision;
+
+			CachePrefetchSolver(ConstraintIndex + PrefetchCount);
+		}
+
+		if (bWantPerIterationCollisionDetection)
+		{
+			// We should lock here? We only ever set to true or do nothing so I think it doesn't matter if this happens on multiple threads...
+			bPerIterationCollisionDetection = true;
+		}
+	}
+
+	void FPBDCollisionContainerSolver::ScatterOutput(const FReal Dt)
+	{
+		ScatterOutput(Dt, 0, NumSolvers());
+	}
+
+	void FPBDCollisionContainerSolver::ScatterOutput(const FReal InDt, const int32 BeginIndex, const int32 EndIndex)
+	{
+		// NOTE: may be called in parallel. Should not change the container or any elements outside of [BeginIndex, EndIndex)
+
+		check(BeginIndex >= 0);
+		check(EndIndex <= NumSolvers());
+
+		const FSolverReal Dt = FSolverReal(InDt);
+
+		for (int32 ConstraintIndex = BeginIndex; ConstraintIndex < EndIndex; ++ConstraintIndex)
+		{
+			Private::FPBDCollisionSolver& CollisionSolver = CollisionSolvers[ConstraintIndex];
+			FPBDCollisionConstraint* Constraint = CollisionConstraints[ConstraintIndex];
+
+			UpdateCollisionConstraintFromSolver(Constraint, CollisionSolver, Dt);
+
+			// Reset the collision solver here as the pointers will be invalid on the next tick
+			// @todo(chaos): Pointers are always reinitalized before use next tick so this isn't strictly necessary
+			CollisionSolver.Reset();
+		}
+	}
+
+	void FPBDCollisionContainerSolver::ApplyPositionConstraints(const FReal Dt, const int32 It, const int32 NumIts)
+	{
+		SolvePositionImpl(Dt, It, NumIts, 0, NumSolvers(), ConstraintContainer.GetSolverSettings());
+	}
+
+	void FPBDCollisionContainerSolver::ApplyVelocityConstraints(const FReal Dt, const int32 It, const int32 NumIts)
+	{
+		SolveVelocityImpl(Dt, It, NumIts, 0, NumSolvers(), ConstraintContainer.GetSolverSettings());
+	}
+
+	void FPBDCollisionContainerSolver::ApplyProjectionConstraints(const FReal Dt, const int32 It, const int32 NumIts)
+	{
+		// Not supported for collisions
+	}
+
+	void FPBDCollisionContainerSolver::UpdatePositionShockPropagation(const FReal Dt, const int32 It, const int32 NumIts, const int32 BeginIndex, const int32 EndIndex, const FPBDCollisionSolverSettings& SolverSettings)
+	{
+		// If this is the first shock propagation iteration, enable it on each solver
+		const bool bEnableShockPropagation = (It == NumIts - SolverSettings.NumPositionShockPropagationIterations);
+		if (bEnableShockPropagation)
+		{
+			for (int32 SolverIndex = BeginIndex; SolverIndex < EndIndex; ++SolverIndex)
+			{
+				CollisionSolvers[SolverIndex].EnablePositionShockPropagation();
+			}
+		}
+	}
+
+	void FPBDCollisionContainerSolver::UpdateVelocityShockPropagation(const FReal Dt, const int32 It, const int32 NumIts, const int32 BeginIndex, const int32 EndIndex, const FPBDCollisionSolverSettings& SolverSettings)
+	{
+		// Set/reset the shock propagation based on current iteration. The position solve may
+		// have left the bodies with a mass scale and we want to change or reset it.
+		const bool bEnableShockPropagation = (It == NumIts - SolverSettings.NumVelocityShockPropagationIterations);
+		if (bEnableShockPropagation)
+		{
+			for (int32 SolverIndex = BeginIndex; SolverIndex < EndIndex; ++SolverIndex)
+			{
+				CollisionSolvers[SolverIndex].EnableVelocityShockPropagation();
+			}
+		}
+		else if (It == 0)
+		{
+			for (int32 SolverIndex = BeginIndex; SolverIndex < EndIndex; ++SolverIndex)
+			{
+				CollisionSolvers[SolverIndex].DisableShockPropagation();
+			}
+		}
+	}
+
+	bool FPBDCollisionContainerSolver::SolvePositionImpl(const FReal InDt, const int32 It, const int32 NumIts, const int32 BeginIndex, const int32 EndIndex, const FPBDCollisionSolverSettings& SolverSettings)
+	{
+		SCOPE_CYCLE_COUNTER(STAT_Collisions_Apply);
+		if (!CVars::bChaos_PBDCollisionSolver_Position_SolveEnabled)
+		{
+			return false;
+		}
+
+		if (EndIndex <= BeginIndex)
+		{
+			return false;
+		}
+
+		UpdatePositionShockPropagation(InDt, It, NumIts, BeginIndex, EndIndex, SolverSettings);
+
+		// We run collision detection here under two conditions (normally it is run after Integration and before the constraint solver phase):
+		// 1) When deferring collision detection until the solver phase for better joint-collision behaviour (RBAN). In this case, we only do this on the first iteration.
+		// 2) When using no manifolds or incremental manifolds, where we may add/replace manifold points every iteration.
+		const bool bRunDeferredCollisionDetection = (It == 0) && ConstraintContainer.GetDetectorSettings().bDeferNarrowPhase;
+		if (bRunDeferredCollisionDetection || bPerIterationCollisionDetection)
+		{
+			UpdateCollisions(InDt, BeginIndex, EndIndex);
+		}
+
+		// Only apply friction for the last few (tunable) iterations
+		// Adjust max pushout to attempt to make it iteration count independent
+		const FSolverReal Dt = FSolverReal(InDt);
+		const bool bApplyStaticFriction = (It >= (NumIts - SolverSettings.NumPositionFrictionIterations));
+		const FSolverReal MaxPushOut = (SolverSettings.MaxPushOutVelocity > 0) ? (FSolverReal(SolverSettings.MaxPushOutVelocity) * Dt) / FSolverReal(NumIts) : FSolverReal(0);
+
+		// Apply the position correction
+		if (bApplyStaticFriction)
+		{
+			Private::FPBDCollisionSolverHelper::SolvePositionWithFriction(MakeArrayView(&CollisionSolvers[BeginIndex], EndIndex - BeginIndex), Dt, MaxPushOut);
+		}
+		else
+		{
+			Private::FPBDCollisionSolverHelper::SolvePositionNoFriction(MakeArrayView(&CollisionSolvers[BeginIndex], EndIndex - BeginIndex), Dt, MaxPushOut);
+		}
+
+		return true;
+	}
+
+	bool FPBDCollisionContainerSolver::SolveVelocityImpl(const FReal InDt, const int32 It, const int32 NumIts, const int32 BeginIndex, const int32 EndIndex, const FPBDCollisionSolverSettings& SolverSettings)
+	{
+		SCOPE_CYCLE_COUNTER(STAT_Collisions_ApplyPushOut);
+		if (!CVars::bChaos_PBDCollisionSolver_Velocity_SolveEnabled)
+		{
+			return false;
+		}
+
+		if (EndIndex <= BeginIndex)
+		{
+			return false;
+		}
+
+		UpdateVelocityShockPropagation(InDt, It, NumIts, BeginIndex, EndIndex, SolverSettings);
+
+		const FSolverReal Dt = FSolverReal(InDt);
+		const bool bApplyDynamicFriction = (It >= NumIts - SolverSettings.NumVelocityFrictionIterations);
+
+		Private::FPBDCollisionSolverHelper::SolveVelocity(MakeArrayView(&CollisionSolvers[BeginIndex], EndIndex - BeginIndex), Dt, bApplyDynamicFriction);
+
+		return true;
+	}
+
+	void FPBDCollisionContainerSolver::UpdateCollisions(const FReal InDt, const int32 BeginIndex, const int32 EndIndex)
+	{
+		const FSolverReal Dt = FSolverReal(InDt);
+		const bool bDeferredCollisionDetection = ConstraintContainer.GetDetectorSettings().bDeferNarrowPhase;
+
+		bool bNeedsAnotherIteration = false;
+		for (int32 SolverIndex = BeginIndex; SolverIndex < EndIndex; ++SolverIndex)
+		{
+			if (bDeferredCollisionDetection || bCollisionConstraintPerIterationCollisionDetection[SolverIndex])
+			{
+				Private::FPBDCollisionSolver& CollisionSolver = CollisionSolvers[SolverIndex];
+				FPBDCollisionConstraint* Constraint = CollisionConstraints[SolverIndex];
+
+				// Run collision detection at the current transforms including any correction from previous iterations
+				const FSolverBody& Body0 = CollisionSolver.SolverBody0().SolverBody();
+				const FSolverBody& Body1 = CollisionSolver.SolverBody1().SolverBody();
+				const FRigidTransform3 CorrectedActorWorldTransform0 = FRigidTransform3(Body0.CorrectedActorP(), Body0.CorrectedActorQ());
+				const FRigidTransform3 CorrectedActorWorldTransform1 = FRigidTransform3(Body1.CorrectedActorP(), Body1.CorrectedActorQ());
+				const FRigidTransform3 CorrectedShapeWorldTransform0 = Constraint->GetShapeRelativeTransform0() * CorrectedActorWorldTransform0;
+				const FRigidTransform3 CorrectedShapeWorldTransform1 = Constraint->GetShapeRelativeTransform1() * CorrectedActorWorldTransform1;
+
+				// @todo(chaos): this is ugly - pass these to the required functions instead and remove from the constraint class
+				// This is now only needed for LevelSet collision (see UpdateLevelsetLevelsetConstraint)
+				Constraint->SetSolverBodies(&Body0, &Body1);
+
+				// Reset the manifold if we are not using manifolds (we just use the first manifold point)
+				if (!Constraint->GetUseManifold())
+				{
+					Constraint->ResetActiveManifoldContacts();
+					CollisionSolver.ResetManifold();
+				}
+
+				// We need to know how many points were added to the manifold
+				const int32 BeginPointIndex = Constraint->NumManifoldPoints();
+
+				// NOTE: We deliberately have not updated the ShapwWorldTranforms on the constraint. If we did that, we would calculate 
+				// errors incorrectly in UpdateManifoldPoints, because the solver assumes nothing has been moved as we iterate (we accumulate 
+				// corrections that will be applied later.)
+				Constraint->ResetPhi(Constraint->GetCullDistance());
+				Collisions::UpdateConstraint(*Constraint, CorrectedShapeWorldTransform0, CorrectedShapeWorldTransform1, Dt);
+
+				// Update the manifold based on the new or updated contacts
+				UpdateCollisionSolverManifoldFromConstraint(CollisionSolver, Constraint, Dt, BeginPointIndex, Constraint->NumManifoldPoints());
+
+				Constraint->SetSolverBodies(nullptr, nullptr);
+			}
+		}
+	}
+
+}
diff --git a/Engine/Source/Runtime/Experimental/Chaos/Private/Chaos/Collision/PBDCollisionContainerSolverJacobi.cpp b/Engine/Source/Runtime/Experimental/Chaos/Private/Chaos/Collision/PBDCollisionContainerSolverJacobi.cpp
--- a/Engine/Source/Runtime/Experimental/Chaos/Private/Chaos/Collision/PBDCollisionContainerSolverJacobi.cpp
+++ b/Engine/Source/Runtime/Experimental/Chaos/Private/Chaos/Collision/PBDCollisionContainerSolverJacobi.cpp
@@ -31,553 +31,557 @@ namespace Chaos
 		extern FRealSingle Chaos_PBDCollisionSolver_JacobiStiffness;
 	}
 
-	//////////////////////////////////////////////////////////////////////////////////////////////////
-	//////////////////////////////////////////////////////////////////////////////////////////////////
-	//////////////////////////////////////////////////////////////////////////////////////////////////
-	//////////////////////////////////////////////////////////////////////////////////////////////////
-
-	// Transform the Constraint's local-space data into world space for use by  the collision solver and also calculate tangents, errors, etc
-	FORCEINLINE_DEBUGGABLE void UpdateCollisionSolverContactPointFromConstraint(
-		Private::FPBDCollisionSolverJacobi& Solver, 
-		const int32 SolverPointIndex, 
-		const FPBDCollisionConstraint* Constraint, 
-		const int32 ConstraintPointIndex, 
-		const FRealSingle Dt, 
-		const FSolverBody& Body0, 
-		const FSolverBody& Body1)
+	namespace Private
 	{
-		const FManifoldPoint& ManifoldPoint = Constraint->GetManifoldPoint(ConstraintPointIndex);
 
-		const FRealSingle Restitution = FRealSingle(Constraint->GetRestitution());
-		const FRealSingle RestitutionVelocityThreshold = FRealSingle(Constraint->GetRestitutionThreshold()) * Dt;
-
-		// World-space shape transforms. Manifold data is currently relative to these spaces
-		const FRigidTransform3& ShapeWorldTransform0 = Constraint->GetShapeWorldTransform0();
-		const FRigidTransform3& ShapeWorldTransform1 = Constraint->GetShapeWorldTransform1();
+		//////////////////////////////////////////////////////////////////////////////////////////////////
+		//////////////////////////////////////////////////////////////////////////////////////////////////
+		//////////////////////////////////////////////////////////////////////////////////////////////////
+		//////////////////////////////////////////////////////////////////////////////////////////////////
+
+		// Transform the Constraint's local-space data into world space for use by  the collision solver and also calculate tangents, errors, etc
+		FORCEINLINE_DEBUGGABLE void UpdateCollisionSolverContactPointFromConstraint(
+			Private::FPBDCollisionSolverJacobi& Solver, 
+			const int32 SolverPointIndex, 
+			const FPBDCollisionConstraint* Constraint, 
+			const int32 ConstraintPointIndex, 
+			const FRealSingle Dt, 
+			const FSolverBody& Body0, 
+			const FSolverBody& Body1)
+		{
+			const FManifoldPoint& ManifoldPoint = Constraint->GetManifoldPoint(ConstraintPointIndex);
 
-		// World-space contact points on each shape
-		const FVec3 WorldContact0 = ShapeWorldTransform0.TransformPositionNoScale(FVec3(ManifoldPoint.ContactPoint.ShapeContactPoints[0]));
-		const FVec3 WorldContact1 = ShapeWorldTransform1.TransformPositionNoScale(FVec3(ManifoldPoint.ContactPoint.ShapeContactPoints[1]));
-		const FVec3 WorldContact = FReal(0.5) * (WorldContact0 + WorldContact1);
-		const FVec3f WorldRelativeContact0 = FVec3f(WorldContact - Body0.P());
-		const FVec3f WorldRelativeContact1 = FVec3f(WorldContact - Body1.P());
+			const FRealSingle Restitution = FRealSingle(Constraint->GetRestitution());
+			const FRealSingle RestitutionVelocityThreshold = FRealSingle(Constraint->GetRestitutionThreshold()) * Dt;
 
-		// World-space normal
-		const FVec3f WorldContactNormal = FVec3f(ShapeWorldTransform1.TransformVectorNoScale(FVec3(ManifoldPoint.ContactPoint.ShapeContactNormal)));
+			// World-space shape transforms. Manifold data is currently relative to these spaces
+			const FRigidTransform3& ShapeWorldTransform0 = Constraint->GetShapeWorldTransform0();
+			const FRigidTransform3& ShapeWorldTransform1 = Constraint->GetShapeWorldTransform1();
 
-		// World-space tangents
-		FVec3f WorldContactTangentU = FVec3f::CrossProduct(FVec3f(0, 1, 0), WorldContactNormal);
-		if (!WorldContactTangentU.Normalize(FRealSingle(UE_KINDA_SMALL_NUMBER)))
-		{
-			WorldContactTangentU = FVec3f::CrossProduct(FVec3f(1, 0, 0), WorldContactNormal);
-			WorldContactTangentU = WorldContactTangentU.GetUnsafeNormal();
-		}
-		const FVec3f WorldContactTangentV = FVec3f::CrossProduct(WorldContactNormal, WorldContactTangentU);
+			// World-space contact points on each shape
+			const FVec3 WorldContact0 = ShapeWorldTransform0.TransformPositionNoScale(FVec3(ManifoldPoint.ContactPoint.ShapeContactPoints[0]));
+			const FVec3 WorldContact1 = ShapeWorldTransform1.TransformPositionNoScale(FVec3(ManifoldPoint.ContactPoint.ShapeContactPoints[1]));
+			const FVec3 WorldContact = FReal(0.5) * (WorldContact0 + WorldContact1);
+			const FVec3f WorldRelativeContact0 = FVec3f(WorldContact - Body0.P());
+			const FVec3f WorldRelativeContact1 = FVec3f(WorldContact - Body1.P());
 
-		// Calculate contact velocity if we will need it below (restitution and/or frist-contact friction)
-		const bool bNeedContactVelocity = (!ManifoldPoint.Flags.bHasStaticFrictionAnchor) || (Restitution > FRealSingle(0));
-		FVec3f ContactVel = FVec3(0);
-		if (bNeedContactVelocity)
-		{
-			const FVec3f ContactVel0 = Body0.V() + FVec3f::CrossProduct(Body0.W(), WorldRelativeContact0);
-			const FVec3f ContactVel1 = Body1.V() + FVec3f::CrossProduct(Body1.W(), WorldRelativeContact1);
-			ContactVel = ContactVel0 - ContactVel1;
-		}
+			// World-space normal
+			const FVec3f WorldContactNormal = FVec3f(ShapeWorldTransform1.TransformVectorNoScale(FVec3(ManifoldPoint.ContactPoint.ShapeContactNormal)));
 
-		// If we have contact data from a previous tick, use it to calculate the lateral position delta we need
-		// to apply to move the contacts back to their original relative locations (i.e., to enforce static friction).
-		// Otherwise, estimate the friction correction from the contact velocity.
-		// NOTE: quadratic shapes use the velocity-based path most of the time, unless the relative motion is very small.
-		FVec3f WorldFrictionDelta = FVec3f(0);
-		if (ManifoldPoint.Flags.bHasStaticFrictionAnchor)
-		{
-			const FVec3f FrictionDelta0 = ShapeWorldTransform0.TransformPositionNoScale(FVec3(ManifoldPoint.ShapeAnchorPoints[0])) - WorldContact0;
-			const FVec3f FrictionDelta1 = ShapeWorldTransform1.TransformPositionNoScale(FVec3(ManifoldPoint.ShapeAnchorPoints[1])) - WorldContact1;
-			WorldFrictionDelta = FrictionDelta0 - FrictionDelta1;
-		}
-		else
-		{
-			// @todo(chaos): consider adding a multiplier to the initial contact friction
-			WorldFrictionDelta = ContactVel * Dt;
-		}
+			// World-space tangents
+			FVec3f WorldContactTangentU = FVec3f::CrossProduct(FVec3f(0, 1, 0), WorldContactNormal);
+			if (!WorldContactTangentU.Normalize(FRealSingle(UE_KINDA_SMALL_NUMBER)))
+			{
+				WorldContactTangentU = FVec3f::CrossProduct(FVec3f(1, 0, 0), WorldContactNormal);
+				WorldContactTangentU = WorldContactTangentU.GetUnsafeNormal();
+			}
+			const FVec3f WorldContactTangentV = FVec3f::CrossProduct(WorldContactNormal, WorldContactTangentU);
 
-		// The contact point error we are trying to correct in this solver
-		const FRealSingle TargetPhi = ManifoldPoint.TargetPhi;
-		const FVec3f WorldContactDelta = FVec3f(WorldContact0 - WorldContact1);
-		const FRealSingle WorldContactDeltaNormal = FVec3f::DotProduct(WorldContactDelta, WorldContactNormal) - TargetPhi;
-		const FRealSingle WorldContactDeltaTangentU = FVec3f::DotProduct(WorldContactDelta + WorldFrictionDelta, WorldContactTangentU);
-		const FRealSingle WorldContactDeltaTangentV = FVec3f::DotProduct(WorldContactDelta + WorldFrictionDelta, WorldContactTangentV);
+			// Calculate contact velocity if we will need it below (restitution and/or frist-contact friction)
+			const bool bNeedContactVelocity = (!ManifoldPoint.Flags.bHasStaticFrictionAnchor) || (Restitution > FRealSingle(0));
+			FVec3f ContactVel = FVec3(0);
+			if (bNeedContactVelocity)
+			{
+				const FVec3f ContactVel0 = Body0.V() + FVec3f::CrossProduct(Body0.W(), WorldRelativeContact0);
+				const FVec3f ContactVel1 = Body1.V() + FVec3f::CrossProduct(Body1.W(), WorldRelativeContact1);
+				ContactVel = ContactVel0 - ContactVel1;
+			}
 
-		// The target contact velocity, taking restitution into account
-		FRealSingle WorldContactTargetVelocityNormal = FRealSingle(0);
-		if (Restitution > FRealSingle(0))
-		{
-			const FRealSingle ContactVelocityNormal = FVec3f::DotProduct(ContactVel, WorldContactNormal);
-			if (ContactVelocityNormal < -RestitutionVelocityThreshold)
+			// If we have contact data from a previous tick, use it to calculate the lateral position delta we need
+			// to apply to move the contacts back to their original relative locations (i.e., to enforce static friction).
+			// Otherwise, estimate the friction correction from the contact velocity.
+			// NOTE: quadratic shapes use the velocity-based path most of the time, unless the relative motion is very small.
+			FVec3f WorldFrictionDelta = FVec3f(0);
+			if (ManifoldPoint.Flags.bHasStaticFrictionAnchor)
 			{
-				WorldContactTargetVelocityNormal = -Restitution * ContactVelocityNormal;
+				const FVec3f FrictionDelta0 = ShapeWorldTransform0.TransformPositionNoScale(FVec3(ManifoldPoint.ShapeAnchorPoints[0])) - WorldContact0;
+				const FVec3f FrictionDelta1 = ShapeWorldTransform1.TransformPositionNoScale(FVec3(ManifoldPoint.ShapeAnchorPoints[1])) - WorldContact1;
+				WorldFrictionDelta = FrictionDelta0 - FrictionDelta1;
+			}
+			else
+			{
+				// @todo(chaos): consider adding a multiplier to the initial contact friction
+				WorldFrictionDelta = ContactVel * Dt;
 			}
-		}
 
-		Solver.InitManifoldPoint(
-			SolverPointIndex,
-			WorldRelativeContact0,
-			WorldRelativeContact1,
-			WorldContactNormal,
-			WorldContactTangentU,
-			WorldContactTangentV,
-			WorldContactDeltaNormal,
-			WorldContactDeltaTangentU,
-			WorldContactDeltaTangentV,
-			WorldContactTargetVelocityNormal);
-	}
+			// The contact point error we are trying to correct in this solver
+			const FRealSingle TargetPhi = ManifoldPoint.TargetPhi;
+			const FVec3f WorldContactDelta = FVec3f(WorldContact0 - WorldContact1);
+			const FRealSingle WorldContactDeltaNormal = FVec3f::DotProduct(WorldContactDelta, WorldContactNormal) - TargetPhi;
+			const FRealSingle WorldContactDeltaTangentU = FVec3f::DotProduct(WorldContactDelta + WorldFrictionDelta, WorldContactTangentU);
+			const FRealSingle WorldContactDeltaTangentV = FVec3f::DotProduct(WorldContactDelta + WorldFrictionDelta, WorldContactTangentV);
 
-	FORCEINLINE_DEBUGGABLE void UpdateCollisionSolverManifoldFromConstraint(
-		Private::FPBDCollisionSolverJacobi& Solver, 
-		const FPBDCollisionConstraint* Constraint, 
-		const FSolverReal Dt, 
-		const int32 ConstraintPointBeginIndex, 
-		const int32 ConstraintPointEndIndex)
-	{
-		const FSolverBody& Body0 = Solver.SolverBody0().SolverBody();
-		const FSolverBody& Body1 = Solver.SolverBody1().SolverBody();
+			// The target contact velocity, taking restitution into account
+			FRealSingle WorldContactTargetVelocityNormal = FRealSingle(0);
+			if (Restitution > FRealSingle(0))
+			{
+				const FRealSingle ContactVelocityNormal = FVec3f::DotProduct(ContactVel, WorldContactNormal);
+				if (ContactVelocityNormal < -RestitutionVelocityThreshold)
+				{
+					WorldContactTargetVelocityNormal = -Restitution * ContactVelocityNormal;
+				}
+			}
+
+			Solver.InitManifoldPoint(
+				SolverPointIndex,
+				WorldRelativeContact0,
+				WorldRelativeContact1,
+				WorldContactNormal,
+				WorldContactTangentU,
+				WorldContactTangentV,
+				WorldContactDeltaNormal,
+				WorldContactDeltaTangentU,
+				WorldContactDeltaTangentV,
+				WorldContactTargetVelocityNormal);
+		}
 
-		// Only calculate state for newly added contacts. Normally this is all of them, but maybe not if incremental collision is used by RBAN.
-		// Also we only add active points to the solver's manifold points list
-		for (int32 ConstraintManifoldPointIndex = ConstraintPointBeginIndex; ConstraintManifoldPointIndex < ConstraintPointEndIndex; ++ConstraintManifoldPointIndex)
+		FORCEINLINE_DEBUGGABLE void UpdateCollisionSolverManifoldFromConstraint(
+			Private::FPBDCollisionSolverJacobi& Solver, 
+			const FPBDCollisionConstraint* Constraint, 
+			const FSolverReal Dt, 
+			const int32 ConstraintPointBeginIndex, 
+			const int32 ConstraintPointEndIndex)
 		{
-			if (!Constraint->GetManifoldPoint(ConstraintManifoldPointIndex).Flags.bDisabled)
+			const FSolverBody& Body0 = Solver.SolverBody0().SolverBody();
+			const FSolverBody& Body1 = Solver.SolverBody1().SolverBody();
+
+			// Only calculate state for newly added contacts. Normally this is all of them, but maybe not if incremental collision is used by RBAN.
+			// Also we only add active points to the solver's manifold points list
+			for (int32 ConstraintManifoldPointIndex = ConstraintPointBeginIndex; ConstraintManifoldPointIndex < ConstraintPointEndIndex; ++ConstraintManifoldPointIndex)
 			{
-				const int32 SolverManifoldPointIndex = Solver.AddManifoldPoint();
+				if (!Constraint->GetManifoldPoint(ConstraintManifoldPointIndex).Flags.bDisabled)
+				{
+					const int32 SolverManifoldPointIndex = Solver.AddManifoldPoint();
 
-				// Transform the constraint contact data into world space for use by the solver
-				// We build this data directly into the solver's world-space contact data which looks a bit odd with "Init" called after but there you go
-				UpdateCollisionSolverContactPointFromConstraint(Solver, SolverManifoldPointIndex, Constraint, ConstraintManifoldPointIndex, Dt, Body0, Body1);
+					// Transform the constraint contact data into world space for use by the solver
+					// We build this data directly into the solver's world-space contact data which looks a bit odd with "Init" called after but there you go
+					UpdateCollisionSolverContactPointFromConstraint(Solver, SolverManifoldPointIndex, Constraint, ConstraintManifoldPointIndex, Dt, Body0, Body1);
 
-				// Prepare the solver (set constraint-space mass etc)
-				Solver.FinalizeManifoldPoint(SolverManifoldPointIndex);
+					// Prepare the solver (set constraint-space mass etc)
+					Solver.FinalizeManifoldPoint(SolverManifoldPointIndex);
+				}
 			}
 		}
-	}
 
-	FORCEINLINE_DEBUGGABLE void UpdateCollisionSolverFromConstraint(
-		Private::FPBDCollisionSolverJacobi& Solver, 
-		const FPBDCollisionConstraint* Constraint, 
-		const FSolverReal Dt, 
-		const FPBDCollisionSolverSettings& SolverSettings, 
-		bool& bOutPerIterationCollision)
-	{
-		// Friction values. Static and Dynamic friction are applied in the position solve for most shapes.
-		// We can also run in a mode without static friction at all. This is faster but stacking is not possible.
-		const FSolverReal StaticFriction = FSolverReal(Constraint->GetStaticFriction());
-		const FSolverReal DynamicFriction = FSolverReal(Constraint->GetDynamicFriction());
-		FSolverReal PositionStaticFriction = FSolverReal(0);
-		FSolverReal PositionDynamicFriction = FSolverReal(0);
-		FSolverReal VelocityDynamicFriction = FSolverReal(0);
-		if (SolverSettings.NumPositionFrictionIterations > 0)
+		FORCEINLINE_DEBUGGABLE void UpdateCollisionSolverFromConstraint(
+			Private::FPBDCollisionSolverJacobi& Solver, 
+			const FPBDCollisionConstraint* Constraint, 
+			const FSolverReal Dt, 
+			const FPBDCollisionSolverSettings& SolverSettings, 
+			bool& bOutPerIterationCollision)
 		{
-			PositionStaticFriction = StaticFriction;
-			if (!Constraint->HasQuadraticShape())
+			// Friction values. Static and Dynamic friction are applied in the position solve for most shapes.
+			// We can also run in a mode without static friction at all. This is faster but stacking is not possible.
+			const FSolverReal StaticFriction = FSolverReal(Constraint->GetStaticFriction());
+			const FSolverReal DynamicFriction = FSolverReal(Constraint->GetDynamicFriction());
+			FSolverReal PositionStaticFriction = FSolverReal(0);
+			FSolverReal PositionDynamicFriction = FSolverReal(0);
+			FSolverReal VelocityDynamicFriction = FSolverReal(0);
+			if (SolverSettings.NumPositionFrictionIterations > 0)
 			{
-				PositionDynamicFriction = DynamicFriction;
+				PositionStaticFriction = StaticFriction;
+				if (!Constraint->HasQuadraticShape())
+				{
+					PositionDynamicFriction = DynamicFriction;
+				}
+				else
+				{
+					// Quadratic shapes don't use PBD dynamic friction - it has issues at slow speeds where the WxR is
+					// less than the position tolerance for friction point matching
+					// @todo(chaos): fix PBD dynamic friction on quadratic shapes
+					VelocityDynamicFriction = DynamicFriction;
+				}
 			}
 			else
 			{
-				// Quadratic shapes don't use PBD dynamic friction - it has issues at slow speeds where the WxR is
-				// less than the position tolerance for friction point matching
-				// @todo(chaos): fix PBD dynamic friction on quadratic shapes
 				VelocityDynamicFriction = DynamicFriction;
 			}
-		}
-		else
-		{
-			VelocityDynamicFriction = DynamicFriction;
-		}
-
-		Solver.SetFriction(PositionStaticFriction, PositionDynamicFriction, VelocityDynamicFriction);
 
-		const FReal SolverStiffness = Constraint->GetStiffness();
-		Solver.SetStiffness(FSolverReal(SolverStiffness));
+			Solver.SetFriction(PositionStaticFriction, PositionDynamicFriction, VelocityDynamicFriction);
 
-		Solver.SolverBody0().SetInvMScale(Constraint->GetInvMassScale0());
-		Solver.SolverBody0().SetInvIScale(Constraint->GetInvInertiaScale0());
-		Solver.SolverBody0().SetShockPropagationScale(FReal(1));
-		Solver.SolverBody1().SetInvMScale(Constraint->GetInvMassScale1());
-		Solver.SolverBody1().SetInvIScale(Constraint->GetInvInertiaScale1());
-		Solver.SolverBody1().SetShockPropagationScale(FReal(1));
+			const FReal SolverStiffness = Constraint->GetStiffness();
+			Solver.SetStiffness(FSolverReal(SolverStiffness));
 
-		bOutPerIterationCollision = (!Constraint->GetUseManifold() || Constraint->GetUseIncrementalCollisionDetection());
+			Solver.SolverBody0().SetInvMScale(Constraint->GetInvMassScale0());
+			Solver.SolverBody0().SetInvIScale(Constraint->GetInvInertiaScale0());
+			Solver.SolverBody0().SetShockPropagationScale(FReal(1));
+			Solver.SolverBody1().SetInvMScale(Constraint->GetInvMassScale1());
+			Solver.SolverBody1().SetInvIScale(Constraint->GetInvInertiaScale1());
+			Solver.SolverBody1().SetShockPropagationScale(FReal(1));
 
-		UpdateCollisionSolverManifoldFromConstraint(Solver, Constraint, Dt, 0, Constraint->NumManifoldPoints());
-	}
+			bOutPerIterationCollision = (!Constraint->GetUseManifold() || Constraint->GetUseIncrementalCollisionDetection());
 
-	FORCEINLINE_DEBUGGABLE void UpdateCollisionConstraintFromSolver(
-		FPBDCollisionConstraint* Constraint, 
-		const Private::FPBDCollisionSolverJacobi& Solver, 
-		const FSolverReal Dt)
-	{
-		Constraint->ResetSolverResults();
+			UpdateCollisionSolverManifoldFromConstraint(Solver, Constraint, Dt, 0, Constraint->NumManifoldPoints());
+		}
 
-		// NOTE: We only put the non-pruned manifold points into the solver so the ManifoldPointIndex and
-		// SolverManifoldPointIndex do not necessarily match. See GatherManifoldPoints
-		int32 SolverManifoldPointIndex = 0;
-		for (int32 ManifoldPointIndex = 0; ManifoldPointIndex < Constraint->NumManifoldPoints(); ++ManifoldPointIndex)
+		FORCEINLINE_DEBUGGABLE void UpdateCollisionConstraintFromSolver(
+			FPBDCollisionConstraint* Constraint, 
+			const Private::FPBDCollisionSolverJacobi& Solver, 
+			const FSolverReal Dt)
 		{
-			FSolverVec3 NetPushOut = FSolverVec3(0);
-			FSolverVec3 NetImpulse = FSolverVec3(0);
-			FSolverReal StaticFrictionRatio = FSolverReal(0);
+			Constraint->ResetSolverResults();
 
-			if (!Constraint->GetManifoldPoint(ManifoldPointIndex).Flags.bDisabled)
+			// NOTE: We only put the non-pruned manifold points into the solver so the ManifoldPointIndex and
+			// SolverManifoldPointIndex do not necessarily match. See GatherManifoldPoints
+			int32 SolverManifoldPointIndex = 0;
+			for (int32 ManifoldPointIndex = 0; ManifoldPointIndex < Constraint->NumManifoldPoints(); ++ManifoldPointIndex)
 			{
-				NetPushOut = Solver.GetNetPushOut(SolverManifoldPointIndex);
-
-				NetImpulse = Solver.GetNetImpulse(SolverManifoldPointIndex);
+				FSolverVec3 NetPushOut = FSolverVec3(0);
+				FSolverVec3 NetImpulse = FSolverVec3(0);
+				FSolverReal StaticFrictionRatio = FSolverReal(0);
 
-				StaticFrictionRatio = Solver.GetStaticFrictionRatio(SolverManifoldPointIndex);
+				if (!Constraint->GetManifoldPoint(ManifoldPointIndex).Flags.bDisabled)
+				{
+					NetPushOut = Solver.GetNetPushOut(SolverManifoldPointIndex);
 
-				++SolverManifoldPointIndex;
-			}
+					NetImpulse = Solver.GetNetImpulse(SolverManifoldPointIndex);
 
-			// NOTE: We call this even for points we did not run the solver for (but with zero results)
-			Constraint->SetSolverResults(ManifoldPointIndex,
-				NetPushOut,
-				NetImpulse,
-				StaticFrictionRatio,
-				Dt);
-		}
+					StaticFrictionRatio = Solver.GetStaticFrictionRatio(SolverManifoldPointIndex);
 
-		Constraint->EndTick();
-	}
+					++SolverManifoldPointIndex;
+				}
 
-	//////////////////////////////////////////////////////////////////////////////////////////////////
-	//////////////////////////////////////////////////////////////////////////////////////////////////
-	//////////////////////////////////////////////////////////////////////////////////////////////////
-	//////////////////////////////////////////////////////////////////////////////////////////////////
+				// NOTE: We call this even for points we did not run the solver for (but with zero results)
+				Constraint->SetSolverResults(ManifoldPointIndex,
+					NetPushOut,
+					NetImpulse,
+					StaticFrictionRatio,
+					Dt);
+			}
 
+			Constraint->EndTick();
+		}
 
-	FPBDCollisionContainerSolverJacobi::FPBDCollisionContainerSolverJacobi(const FPBDCollisionConstraints& InConstraintContainer, const int32 InPriority)
-		: FConstraintContainerSolver(InPriority)
-		, ConstraintContainer(InConstraintContainer)
-		, bPerIterationCollisionDetection(false)
-	{
-#if !UE_BUILD_SHIPPING && !UE_BUILD_TEST && INTEL_ISPC
-		Private::FPBDCollisionSolverJacobiHelper::CheckISPC();
-#endif
-	}
+		//////////////////////////////////////////////////////////////////////////////////////////////////
+		//////////////////////////////////////////////////////////////////////////////////////////////////
+		//////////////////////////////////////////////////////////////////////////////////////////////////
+		//////////////////////////////////////////////////////////////////////////////////////////////////
 
-	FPBDCollisionContainerSolverJacobi::~FPBDCollisionContainerSolverJacobi()
-	{
-	}
 
-	void FPBDCollisionContainerSolverJacobi::Reset(const int32 MaxCollisions)
-	{
-		// A over-allocation policy to avoid reallocation every frame in the common case where a pile of objects is dropped
-		// and the number of contacts increases every tick.
-		if (MaxCollisions > CollisionSolvers.Max())
+		FPBDCollisionContainerSolverJacobi::FPBDCollisionContainerSolverJacobi(const FPBDCollisionConstraints& InConstraintContainer, const int32 InPriority)
+			: FConstraintContainerSolver(InPriority)
+			, ConstraintContainer(InConstraintContainer)
+			, bPerIterationCollisionDetection(false)
 		{
-			CollisionSolvers.Reset((5 * MaxCollisions) / 4);	// +25%
+	#if !UE_BUILD_SHIPPING && !UE_BUILD_TEST && INTEL_ISPC
+			Private::FPBDCollisionSolverJacobiHelper::CheckISPC();
+	#endif
 		}
 
-		CollisionSolvers.Reset(MaxCollisions);
-		CollisionConstraints.Reset(CollisionSolvers.Max());
-		bCollisionConstraintPerIterationCollisionDetection.Reset(CollisionSolvers.Max());
-
-		// Just set the array size for these right away - all data will be initialized later and there is no constructor
-		// We use CollisionConstraints.Num() to track how many constraints we currently have
-		CollisionSolvers.SetNumUninitialized(MaxCollisions);
-		bCollisionConstraintPerIterationCollisionDetection.SetNumUninitialized(MaxCollisions);
-	}
-
-	void FPBDCollisionContainerSolverJacobi::AddConstraints()
-	{
-		Reset(ConstraintContainer.NumConstraints());
-
-		for (FPBDCollisionConstraintHandle* ConstraintHandle : ConstraintContainer.GetConstraintHandles())
+		FPBDCollisionContainerSolverJacobi::~FPBDCollisionContainerSolverJacobi()
 		{
-			FPBDCollisionConstraint& Constraint = ConstraintHandle->GetContact();
-
-			AddConstraint(Constraint);
 		}
-	}
 
-	void FPBDCollisionContainerSolverJacobi::AddConstraints(const TArrayView<Private::FPBDIslandConstraint>& IslandConstraints)
-	{
-		for (Private::FPBDIslandConstraint& IslandConstraint : IslandConstraints)
+		void FPBDCollisionContainerSolverJacobi::Reset(const int32 MaxCollisions)
 		{
-			// We will only ever be given constraints from our container (asserts in non-shipping)
-			FPBDCollisionConstraint& Constraint = IslandConstraint.GetConstraint()->AsUnsafe<FPBDCollisionConstraintHandle>()->GetContact();
+			// A over-allocation policy to avoid reallocation every frame in the common case where a pile of objects is dropped
+			// and the number of contacts increases every tick.
+			if (MaxCollisions > CollisionSolvers.Max())
+			{
+				CollisionSolvers.Reset((5 * MaxCollisions) / 4);	// +25%
+			}
+
+			CollisionSolvers.Reset(MaxCollisions);
+			CollisionConstraints.Reset(CollisionSolvers.Max());
+			bCollisionConstraintPerIterationCollisionDetection.Reset(CollisionSolvers.Max());
 
-			AddConstraint(Constraint);
+			// Just set the array size for these right away - all data will be initialized later and there is no constructor
+			// We use CollisionConstraints.Num() to track how many constraints we currently have
+			CollisionSolvers.SetNumUninitialized(MaxCollisions);
+			bCollisionConstraintPerIterationCollisionDetection.SetNumUninitialized(MaxCollisions);
 		}
-	}
 
-	void FPBDCollisionContainerSolverJacobi::AddConstraint(FPBDCollisionConstraint& Constraint)
-	{
-		// NOTE: No need to add to CollisionSolvers or bCollisionConstraintPerIterationCollisionDetection
-		const int32 Index = CollisionConstraints.Add(&Constraint);
+		void FPBDCollisionContainerSolverJacobi::AddConstraints()
+		{
+			Reset(ConstraintContainer.NumConstraints());
 
-		CollisionSolvers[Index].Reset();
-	}
+			for (FPBDCollisionConstraintHandle* ConstraintHandle : ConstraintContainer.GetConstraintHandles())
+			{
+				FPBDCollisionConstraint& Constraint = ConstraintHandle->GetContact();
 
+				AddConstraint(Constraint);
+			}
+		}
 
-	void FPBDCollisionContainerSolverJacobi::AddBodies(FSolverBodyContainer& SolverBodyContainer)
-	{
-		for (int32 SolverIndex = 0, SolverEndIndex = NumSolvers(); SolverIndex < SolverEndIndex; ++SolverIndex)
+		void FPBDCollisionContainerSolverJacobi::AddConstraints(const TArrayView<Private::FPBDIslandConstraint>& IslandConstraints)
 		{
-			Private::FPBDCollisionSolverJacobi& CollisionSolver = CollisionSolvers[SolverIndex];
-			FPBDCollisionConstraint* Constraint = CollisionConstraints[SolverIndex];
-			check(Constraint != nullptr);
-
-			// Find the solver bodies for the particles we constrain. This will add them to the container
-			// if they aren't there already, and ensure that they are populated with the latest data.
-			FSolverBody* Body0 = SolverBodyContainer.FindOrAdd(Constraint->GetParticle0());
-			FSolverBody* Body1 = SolverBodyContainer.FindOrAdd(Constraint->GetParticle1());
-			check(Body0 != nullptr);
-			check(Body1 != nullptr);
-
-			CollisionSolver.SetSolverBodies(*Body0, *Body1);
-		}
-	}
+			for (Private::FPBDIslandConstraint& IslandConstraint : IslandConstraints)
+			{
+				// We will only ever be given constraints from our container (asserts in non-shipping)
+				FPBDCollisionConstraint& Constraint = IslandConstraint.GetConstraint()->AsUnsafe<FPBDCollisionConstraintHandle>()->GetContact();
 
-	void FPBDCollisionContainerSolverJacobi::GatherInput(const FReal Dt)
-	{
-		GatherInput(Dt, 0, NumSolvers());
-	}
+				AddConstraint(Constraint);
+			}
+		}
 
-	void FPBDCollisionContainerSolverJacobi::CachePrefetchSolver(const int32 ConstraintIndex) const
-	{
-		if (ConstraintIndex < NumSolvers())
+		void FPBDCollisionContainerSolverJacobi::AddConstraint(FPBDCollisionConstraint& Constraint)
 		{
-			FPlatformMisc::PrefetchBlock(&CollisionSolvers[ConstraintIndex], sizeof(FPBDCollisionContainerSolverJacobi));
+			// NOTE: No need to add to CollisionSolvers or bCollisionConstraintPerIterationCollisionDetection
+			const int32 Index = CollisionConstraints.Add(&Constraint);
+
+			CollisionSolvers[Index].Reset();
 		}
-	}
 
-	void FPBDCollisionContainerSolverJacobi::GatherInput(const FReal InDt, const int32 BeginIndex, const int32 EndIndex)
-	{
-		// NOTE: may be called in parallel. Should not change the container or any elements outside of [BeginIndex, EndIndex)
 
-		check(BeginIndex >= 0);
-		check(EndIndex <= NumSolvers());
+		void FPBDCollisionContainerSolverJacobi::AddBodies(FSolverBodyContainer& SolverBodyContainer)
+		{
+			for (int32 SolverIndex = 0, SolverEndIndex = NumSolvers(); SolverIndex < SolverEndIndex; ++SolverIndex)
+			{
+				Private::FPBDCollisionSolverJacobi& CollisionSolver = CollisionSolvers[SolverIndex];
+				FPBDCollisionConstraint* Constraint = CollisionConstraints[SolverIndex];
+				check(Constraint != nullptr);
 
-		const FSolverReal Dt = FSolverReal(InDt);
+				// Find the solver bodies for the particles we constrain. This will add them to the container
+				// if they aren't there already, and ensure that they are populated with the latest data.
+				FSolverBody* Body0 = SolverBodyContainer.FindOrAdd(Constraint->GetParticle0());
+				FSolverBody* Body1 = SolverBodyContainer.FindOrAdd(Constraint->GetParticle1());
+				check(Body0 != nullptr);
+				check(Body1 != nullptr);
 
-		const int32 PrefetchCount = 4;
-		for (int32 PrefetchIndex = 0; PrefetchIndex < PrefetchCount; ++PrefetchIndex)
-		{
-			CachePrefetchSolver(PrefetchIndex);
+				CollisionSolver.SetSolverBodies(*Body0, *Body1);
+			}
 		}
 
-		bool bWantPerIterationCollisionDetection = false;
-		for (int32 ConstraintIndex = BeginIndex; ConstraintIndex < EndIndex; ++ConstraintIndex)
+		void FPBDCollisionContainerSolverJacobi::GatherInput(const FReal Dt)
 		{
-			Private::FPBDCollisionSolverJacobi& CollisionSolver = CollisionSolvers[ConstraintIndex];
-			FPBDCollisionConstraint* Constraint = CollisionConstraints[ConstraintIndex];
-			bool& bPerIterationCollision = bCollisionConstraintPerIterationCollisionDetection[ConstraintIndex];
-
-			UpdateCollisionSolverFromConstraint(CollisionSolver, Constraint, Dt, ConstraintContainer.GetSolverSettings(), bPerIterationCollision);
-
-			// We need to run collision every iteration if we are not using manifolds, or are using incremental manifolds
-			bWantPerIterationCollisionDetection = bWantPerIterationCollisionDetection || bPerIterationCollision;
-
-			CachePrefetchSolver(ConstraintIndex + PrefetchCount);
+			GatherInput(Dt, 0, NumSolvers());
 		}
 
-		if (bWantPerIterationCollisionDetection)
+		void FPBDCollisionContainerSolverJacobi::CachePrefetchSolver(const int32 ConstraintIndex) const
 		{
-			// We should lock here? We only ever set to true or do nothing so I think it doesn't matter if this happens on multiple threads...
-			bPerIterationCollisionDetection = true;
+			if (ConstraintIndex < NumSolvers())
+			{
+				FPlatformMisc::PrefetchBlock(&CollisionSolvers[ConstraintIndex], sizeof(FPBDCollisionContainerSolverJacobi));
+			}
 		}
-	}
 
-	void FPBDCollisionContainerSolverJacobi::ScatterOutput(const FReal Dt)
-	{
-		ScatterOutput(Dt, 0, NumSolvers());
-	}
-
-	void FPBDCollisionContainerSolverJacobi::ScatterOutput(const FReal InDt, const int32 BeginIndex, const int32 EndIndex)
-	{
-		// NOTE: may be called in parallel. Should not change the container or any elements outside of [BeginIndex, EndIndex)
+		void FPBDCollisionContainerSolverJacobi::GatherInput(const FReal InDt, const int32 BeginIndex, const int32 EndIndex)
+		{
+			// NOTE: may be called in parallel. Should not change the container or any elements outside of [BeginIndex, EndIndex)
 
-		check(BeginIndex >= 0);
-		check(EndIndex <= NumSolvers());
+			check(BeginIndex >= 0);
+			check(EndIndex <= NumSolvers());
 
-		const FSolverReal Dt = FSolverReal(InDt);
+			const FSolverReal Dt = FSolverReal(InDt);
 
-		for (int32 ConstraintIndex = BeginIndex; ConstraintIndex < EndIndex; ++ConstraintIndex)
-		{
-			Private::FPBDCollisionSolverJacobi& CollisionSolver = CollisionSolvers[ConstraintIndex];
-			FPBDCollisionConstraint* Constraint = CollisionConstraints[ConstraintIndex];
+			const int32 PrefetchCount = 4;
+			for (int32 PrefetchIndex = 0; PrefetchIndex < PrefetchCount; ++PrefetchIndex)
+			{
+				CachePrefetchSolver(PrefetchIndex);
+			}
 
-			UpdateCollisionConstraintFromSolver(Constraint, CollisionSolver, Dt);
+			bool bWantPerIterationCollisionDetection = false;
+			for (int32 ConstraintIndex = BeginIndex; ConstraintIndex < EndIndex; ++ConstraintIndex)
+			{
+				Private::FPBDCollisionSolverJacobi& CollisionSolver = CollisionSolvers[ConstraintIndex];
+				FPBDCollisionConstraint* Constraint = CollisionConstraints[ConstraintIndex];
+				bool& bPerIterationCollision = bCollisionConstraintPerIterationCollisionDetection[ConstraintIndex];
 
-			// Reset the collision solver here as the pointers will be invalid on the next tick
-			// @todo(chaos): Pointers are always reinitalized before use next tick so this isn't strictly necessary
-			CollisionSolver.Reset();
-		}
-	}
+				UpdateCollisionSolverFromConstraint(CollisionSolver, Constraint, Dt, ConstraintContainer.GetSolverSettings(), bPerIterationCollision);
 
-	void FPBDCollisionContainerSolverJacobi::ApplyPositionConstraints(const FReal Dt, const int32 It, const int32 NumIts)
-	{
-		SolvePositionImpl(Dt, It, NumIts, 0, NumSolvers(), ConstraintContainer.GetSolverSettings());
-	}
+				// We need to run collision every iteration if we are not using manifolds, or are using incremental manifolds
+				bWantPerIterationCollisionDetection = bWantPerIterationCollisionDetection || bPerIterationCollision;
 
-	void FPBDCollisionContainerSolverJacobi::ApplyVelocityConstraints(const FReal Dt, const int32 It, const int32 NumIts)
-	{
-		SolveVelocityImpl(Dt, It, NumIts, 0, NumSolvers(), ConstraintContainer.GetSolverSettings());
-	}
-
-	void FPBDCollisionContainerSolverJacobi::ApplyProjectionConstraints(const FReal Dt, const int32 It, const int32 NumIts)
-	{
-		// Not supported for collisions
-	}
+				CachePrefetchSolver(ConstraintIndex + PrefetchCount);
+			}
 
-	void FPBDCollisionContainerSolverJacobi::UpdatePositionShockPropagation(const FReal Dt, const int32 It, const int32 NumIts, const int32 BeginIndex, const int32 EndIndex, const FPBDCollisionSolverSettings& SolverSettings)
-	{
-		// If this is the first shock propagation iteration, enable it on each solver
-		const bool bEnableShockPropagation = (It == NumIts - SolverSettings.NumPositionShockPropagationIterations);
-		if (bEnableShockPropagation)
-		{
-			for (int32 SolverIndex = BeginIndex; SolverIndex < EndIndex; ++SolverIndex)
+			if (bWantPerIterationCollisionDetection)
 			{
-				CollisionSolvers[SolverIndex].EnablePositionShockPropagation();
+				// We should lock here? We only ever set to true or do nothing so I think it doesn't matter if this happens on multiple threads...
+				bPerIterationCollisionDetection = true;
 			}
 		}
-	}
 
-	void FPBDCollisionContainerSolverJacobi::UpdateVelocityShockPropagation(const FReal Dt, const int32 It, const int32 NumIts, const int32 BeginIndex, const int32 EndIndex, const FPBDCollisionSolverSettings& SolverSettings)
-	{
-		// Set/reset the shock propagation based on current iteration. The position solve may
-		// have left the bodies with a mass scale and we want to change or reset it.
-		const bool bEnableShockPropagation = (It == NumIts - SolverSettings.NumVelocityShockPropagationIterations);
-		if (bEnableShockPropagation)
+		void FPBDCollisionContainerSolverJacobi::ScatterOutput(const FReal Dt)
 		{
-			for (int32 SolverIndex = BeginIndex; SolverIndex < EndIndex; ++SolverIndex)
-			{
-				CollisionSolvers[SolverIndex].EnableVelocityShockPropagation();
-			}
+			ScatterOutput(Dt, 0, NumSolvers());
 		}
-		else if (It == 0)
+
+		void FPBDCollisionContainerSolverJacobi::ScatterOutput(const FReal InDt, const int32 BeginIndex, const int32 EndIndex)
 		{
-			for (int32 SolverIndex = BeginIndex; SolverIndex < EndIndex; ++SolverIndex)
+			// NOTE: may be called in parallel. Should not change the container or any elements outside of [BeginIndex, EndIndex)
+
+			check(BeginIndex >= 0);
+			check(EndIndex <= NumSolvers());
+
+			const FSolverReal Dt = FSolverReal(InDt);
+
+			for (int32 ConstraintIndex = BeginIndex; ConstraintIndex < EndIndex; ++ConstraintIndex)
 			{
-				CollisionSolvers[SolverIndex].DisableShockPropagation();
+				Private::FPBDCollisionSolverJacobi& CollisionSolver = CollisionSolvers[ConstraintIndex];
+				FPBDCollisionConstraint* Constraint = CollisionConstraints[ConstraintIndex];
+
+				UpdateCollisionConstraintFromSolver(Constraint, CollisionSolver, Dt);
+
+				// Reset the collision solver here as the pointers will be invalid on the next tick
+				// @todo(chaos): Pointers are always reinitalized before use next tick so this isn't strictly necessary
+				CollisionSolver.Reset();
 			}
 		}
-	}
 
-	bool FPBDCollisionContainerSolverJacobi::SolvePositionImpl(const FReal InDt, const int32 It, const int32 NumIts, const int32 BeginIndex, const int32 EndIndex, const FPBDCollisionSolverSettings& SolverSettings)
-	{
-		SCOPE_CYCLE_COUNTER(STAT_Collisions_Apply);
-		if (!CVars::bChaos_PBDCollisionSolver_Position_SolveEnabled)
+		void FPBDCollisionContainerSolverJacobi::ApplyPositionConstraints(const FReal Dt, const int32 It, const int32 NumIts)
 		{
-			return false;
+			SolvePositionImpl(Dt, It, NumIts, 0, NumSolvers(), ConstraintContainer.GetSolverSettings());
 		}
 
-		if (EndIndex <= BeginIndex)
+		void FPBDCollisionContainerSolverJacobi::ApplyVelocityConstraints(const FReal Dt, const int32 It, const int32 NumIts)
 		{
-			return false;
+			SolveVelocityImpl(Dt, It, NumIts, 0, NumSolvers(), ConstraintContainer.GetSolverSettings());
 		}
 
-		UpdatePositionShockPropagation(InDt, It, NumIts, BeginIndex, EndIndex, SolverSettings);
-
-		// We run collision detection here under two conditions (normally it is run after Integration and before the constraint solver phase):
-		// 1) When deferring collision detection until the solver phase for better joint-collision behaviour (RBAN). In this case, we only do this on the first iteration.
-		// 2) When using no manifolds or incremental manifolds, where we may add/replace manifold points every iteration.
-		const bool bRunDeferredCollisionDetection = (It == 0) && ConstraintContainer.GetDetectorSettings().bDeferNarrowPhase;
-		if (bRunDeferredCollisionDetection || bPerIterationCollisionDetection)
+		void FPBDCollisionContainerSolverJacobi::ApplyProjectionConstraints(const FReal Dt, const int32 It, const int32 NumIts)
 		{
-			UpdateCollisions(InDt, BeginIndex, EndIndex);
+			// Not supported for collisions
 		}
 
-		// Only apply friction for the last few (tunable) iterations
-		// Adjust max pushout to attempt to make it iteration count independent
-		const FSolverReal Dt = FSolverReal(InDt);
-		const bool bApplyStaticFriction = (It >= (NumIts - SolverSettings.NumPositionFrictionIterations));
-		const FSolverReal MaxPushOut = (SolverSettings.MaxPushOutVelocity > 0) ? (FSolverReal(SolverSettings.MaxPushOutVelocity) * Dt) / FSolverReal(NumIts) : FSolverReal(0);
-
-		// Apply the position correction
-		if (bApplyStaticFriction)
-		{
-			Private::FPBDCollisionSolverJacobiHelper::SolvePositionWithFriction(MakeArrayView(&CollisionSolvers[BeginIndex], EndIndex - BeginIndex), Dt, MaxPushOut);
-		}
-		else
+		void FPBDCollisionContainerSolverJacobi::UpdatePositionShockPropagation(const FReal Dt, const int32 It, const int32 NumIts, const int32 BeginIndex, const int32 EndIndex, const FPBDCollisionSolverSettings& SolverSettings)
 		{
-			Private::FPBDCollisionSolverJacobiHelper::SolvePositionNoFriction(MakeArrayView(&CollisionSolvers[BeginIndex], EndIndex - BeginIndex), Dt, MaxPushOut);
+			// If this is the first shock propagation iteration, enable it on each solver
+			const bool bEnableShockPropagation = (It == NumIts - SolverSettings.NumPositionShockPropagationIterations);
+			if (bEnableShockPropagation)
+			{
+				for (int32 SolverIndex = BeginIndex; SolverIndex < EndIndex; ++SolverIndex)
+				{
+					CollisionSolvers[SolverIndex].EnablePositionShockPropagation();
+				}
+			}
 		}
 
-		return true;
-	}
-
-	bool FPBDCollisionContainerSolverJacobi::SolveVelocityImpl(const FReal InDt, const int32 It, const int32 NumIts, const int32 BeginIndex, const int32 EndIndex, const FPBDCollisionSolverSettings& SolverSettings)
-	{
-		SCOPE_CYCLE_COUNTER(STAT_Collisions_ApplyPushOut);
-		if (!CVars::bChaos_PBDCollisionSolver_Velocity_SolveEnabled)
+		void FPBDCollisionContainerSolverJacobi::UpdateVelocityShockPropagation(const FReal Dt, const int32 It, const int32 NumIts, const int32 BeginIndex, const int32 EndIndex, const FPBDCollisionSolverSettings& SolverSettings)
 		{
-			return false;
+			// Set/reset the shock propagation based on current iteration. The position solve may
+			// have left the bodies with a mass scale and we want to change or reset it.
+			const bool bEnableShockPropagation = (It == NumIts - SolverSettings.NumVelocityShockPropagationIterations);
+			if (bEnableShockPropagation)
+			{
+				for (int32 SolverIndex = BeginIndex; SolverIndex < EndIndex; ++SolverIndex)
+				{
+					CollisionSolvers[SolverIndex].EnableVelocityShockPropagation();
+				}
+			}
+			else if (It == 0)
+			{
+				for (int32 SolverIndex = BeginIndex; SolverIndex < EndIndex; ++SolverIndex)
+				{
+					CollisionSolvers[SolverIndex].DisableShockPropagation();
+				}
+			}
 		}
 
-		if (EndIndex <= BeginIndex)
+		bool FPBDCollisionContainerSolverJacobi::SolvePositionImpl(const FReal InDt, const int32 It, const int32 NumIts, const int32 BeginIndex, const int32 EndIndex, const FPBDCollisionSolverSettings& SolverSettings)
 		{
-			return false;
-		}
+			SCOPE_CYCLE_COUNTER(STAT_Collisions_Apply);
+			if (!CVars::bChaos_PBDCollisionSolver_Position_SolveEnabled)
+			{
+				return false;
+			}
 
-		UpdateVelocityShockPropagation(InDt, It, NumIts, BeginIndex, EndIndex, SolverSettings);
+			if (EndIndex <= BeginIndex)
+			{
+				return false;
+			}
 
-		const FSolverReal Dt = FSolverReal(InDt);
-		const bool bApplyDynamicFriction = (It >= NumIts - SolverSettings.NumVelocityFrictionIterations);
+			UpdatePositionShockPropagation(InDt, It, NumIts, BeginIndex, EndIndex, SolverSettings);
 
-		Private::FPBDCollisionSolverJacobiHelper::SolveVelocity(MakeArrayView(&CollisionSolvers[BeginIndex], EndIndex - BeginIndex), Dt, bApplyDynamicFriction);
+			// We run collision detection here under two conditions (normally it is run after Integration and before the constraint solver phase):
+			// 1) When deferring collision detection until the solver phase for better joint-collision behaviour (RBAN). In this case, we only do this on the first iteration.
+			// 2) When using no manifolds or incremental manifolds, where we may add/replace manifold points every iteration.
+			const bool bRunDeferredCollisionDetection = (It == 0) && ConstraintContainer.GetDetectorSettings().bDeferNarrowPhase;
+			if (bRunDeferredCollisionDetection || bPerIterationCollisionDetection)
+			{
+				UpdateCollisions(InDt, BeginIndex, EndIndex);
+			}
 
-		return true;
-	}
+			// Only apply friction for the last few (tunable) iterations
+			// Adjust max pushout to attempt to make it iteration count independent
+			const FSolverReal Dt = FSolverReal(InDt);
+			const bool bApplyStaticFriction = (It >= (NumIts - SolverSettings.NumPositionFrictionIterations));
+			const FSolverReal MaxPushOut = (SolverSettings.MaxPushOutVelocity > 0) ? (FSolverReal(SolverSettings.MaxPushOutVelocity) * Dt) / FSolverReal(NumIts) : FSolverReal(0);
 
-	void FPBDCollisionContainerSolverJacobi::UpdateCollisions(const FReal InDt, const int32 BeginIndex, const int32 EndIndex)
-	{
-		const FSolverReal Dt = FSolverReal(InDt);
-		const bool bDeferredCollisionDetection = ConstraintContainer.GetDetectorSettings().bDeferNarrowPhase;
+			// Apply the position correction
+			if (bApplyStaticFriction)
+			{
+				Private::FPBDCollisionSolverJacobiHelper::SolvePositionWithFriction(MakeArrayView(&CollisionSolvers[BeginIndex], EndIndex - BeginIndex), Dt, MaxPushOut);
+			}
+			else
+			{
+				Private::FPBDCollisionSolverJacobiHelper::SolvePositionNoFriction(MakeArrayView(&CollisionSolvers[BeginIndex], EndIndex - BeginIndex), Dt, MaxPushOut);
+			}
 
-		bool bNeedsAnotherIteration = false;
-		for (int32 SolverIndex = BeginIndex; SolverIndex < EndIndex; ++SolverIndex)
+			return true;
+		}
+
+		bool FPBDCollisionContainerSolverJacobi::SolveVelocityImpl(const FReal InDt, const int32 It, const int32 NumIts, const int32 BeginIndex, const int32 EndIndex, const FPBDCollisionSolverSettings& SolverSettings)
 		{
-			if (bDeferredCollisionDetection || bCollisionConstraintPerIterationCollisionDetection[SolverIndex])
+			SCOPE_CYCLE_COUNTER(STAT_Collisions_ApplyPushOut);
+			if (!CVars::bChaos_PBDCollisionSolver_Velocity_SolveEnabled)
 			{
-				Private::FPBDCollisionSolverJacobi& CollisionSolver = CollisionSolvers[SolverIndex];
-				FPBDCollisionConstraint* Constraint = CollisionConstraints[SolverIndex];
+				return false;
+			}
 
-				// Run collision detection at the current transforms including any correction from previous iterations
-				const FSolverBody& Body0 = CollisionSolver.SolverBody0().SolverBody();
-				const FSolverBody& Body1 = CollisionSolver.SolverBody1().SolverBody();
-				const FRigidTransform3 CorrectedActorWorldTransform0 = FRigidTransform3(Body0.CorrectedActorP(), Body0.CorrectedActorQ());
-				const FRigidTransform3 CorrectedActorWorldTransform1 = FRigidTransform3(Body1.CorrectedActorP(), Body1.CorrectedActorQ());
-				const FRigidTransform3 CorrectedShapeWorldTransform0 = Constraint->GetShapeRelativeTransform0() * CorrectedActorWorldTransform0;
-				const FRigidTransform3 CorrectedShapeWorldTransform1 = Constraint->GetShapeRelativeTransform1() * CorrectedActorWorldTransform1;
+			if (EndIndex <= BeginIndex)
+			{
+				return false;
+			}
 
-				// @todo(chaos): this is ugly - pass these to the required functions instead and remove from the constraint class
-				// This is now only needed for LevelSet collision (see UpdateLevelsetLevelsetConstraint)
-				Constraint->SetSolverBodies(&Body0, &Body1);
+			UpdateVelocityShockPropagation(InDt, It, NumIts, BeginIndex, EndIndex, SolverSettings);
 
-				// Reset the manifold if we are not using manifolds (we just use the first manifold point)
-				if (!Constraint->GetUseManifold())
-				{
-					Constraint->ResetActiveManifoldContacts();
-					CollisionSolver.ResetManifold();
-				}
+			const FSolverReal Dt = FSolverReal(InDt);
+			const bool bApplyDynamicFriction = (It >= NumIts - SolverSettings.NumVelocityFrictionIterations);
 
-				// We need to know how many points were added to the manifold
-				const int32 BeginPointIndex = Constraint->NumManifoldPoints();
+			Private::FPBDCollisionSolverJacobiHelper::SolveVelocity(MakeArrayView(&CollisionSolvers[BeginIndex], EndIndex - BeginIndex), Dt, bApplyDynamicFriction);
 
-				// NOTE: We deliberately have not updated the ShapwWorldTranforms on the constraint. If we did that, we would calculate 
-				// errors incorrectly in UpdateManifoldPoints, because the solver assumes nothing has been moved as we iterate (we accumulate 
-				// corrections that will be applied later.)
-				Constraint->ResetPhi(Constraint->GetCullDistance());
-				Collisions::UpdateConstraint(*Constraint, CorrectedShapeWorldTransform0, CorrectedShapeWorldTransform1, Dt);
+			return true;
+		}
 
-				// Update the manifold based on the new or updated contacts
-				UpdateCollisionSolverManifoldFromConstraint(CollisionSolver, Constraint, Dt, BeginPointIndex, Constraint->NumManifoldPoints());
+		void FPBDCollisionContainerSolverJacobi::UpdateCollisions(const FReal InDt, const int32 BeginIndex, const int32 EndIndex)
+		{
+			const FSolverReal Dt = FSolverReal(InDt);
+			const bool bDeferredCollisionDetection = ConstraintContainer.GetDetectorSettings().bDeferNarrowPhase;
 
-				Constraint->SetSolverBodies(nullptr, nullptr);
+			bool bNeedsAnotherIteration = false;
+			for (int32 SolverIndex = BeginIndex; SolverIndex < EndIndex; ++SolverIndex)
+			{
+				if (bDeferredCollisionDetection || bCollisionConstraintPerIterationCollisionDetection[SolverIndex])
+				{
+					Private::FPBDCollisionSolverJacobi& CollisionSolver = CollisionSolvers[SolverIndex];
+					FPBDCollisionConstraint* Constraint = CollisionConstraints[SolverIndex];
+
+					// Run collision detection at the current transforms including any correction from previous iterations
+					const FSolverBody& Body0 = CollisionSolver.SolverBody0().SolverBody();
+					const FSolverBody& Body1 = CollisionSolver.SolverBody1().SolverBody();
+					const FRigidTransform3 CorrectedActorWorldTransform0 = FRigidTransform3(Body0.CorrectedActorP(), Body0.CorrectedActorQ());
+					const FRigidTransform3 CorrectedActorWorldTransform1 = FRigidTransform3(Body1.CorrectedActorP(), Body1.CorrectedActorQ());
+					const FRigidTransform3 CorrectedShapeWorldTransform0 = Constraint->GetShapeRelativeTransform0() * CorrectedActorWorldTransform0;
+					const FRigidTransform3 CorrectedShapeWorldTransform1 = Constraint->GetShapeRelativeTransform1() * CorrectedActorWorldTransform1;
+
+					// @todo(chaos): this is ugly - pass these to the required functions instead and remove from the constraint class
+					// This is now only needed for LevelSet collision (see UpdateLevelsetLevelsetConstraint)
+					Constraint->SetSolverBodies(&Body0, &Body1);
+
+					// Reset the manifold if we are not using manifolds (we just use the first manifold point)
+					if (!Constraint->GetUseManifold())
+					{
+						Constraint->ResetActiveManifoldContacts();
+						CollisionSolver.ResetManifold();
+					}
+
+					// We need to know how many points were added to the manifold
+					const int32 BeginPointIndex = Constraint->NumManifoldPoints();
+
+					// NOTE: We deliberately have not updated the ShapwWorldTranforms on the constraint. If we did that, we would calculate 
+					// errors incorrectly in UpdateManifoldPoints, because the solver assumes nothing has been moved as we iterate (we accumulate 
+					// corrections that will be applied later.)
+					Constraint->ResetPhi(Constraint->GetCullDistance());
+					Collisions::UpdateConstraint(*Constraint, CorrectedShapeWorldTransform0, CorrectedShapeWorldTransform1, Dt);
+
+					// Update the manifold based on the new or updated contacts
+					UpdateCollisionSolverManifoldFromConstraint(CollisionSolver, Constraint, Dt, BeginPointIndex, Constraint->NumManifoldPoints());
+
+					Constraint->SetSolverBodies(nullptr, nullptr);
+				}
 			}
 		}
-	}
 
-}
+	}	// namespace Private
+}	// namespace Chaos
diff --git a/Engine/Source/Runtime/Experimental/Chaos/Private/Chaos/Collision/PBDCollisionContainerSolverSimd.cpp b/Engine/Source/Runtime/Experimental/Chaos/Private/Chaos/Collision/PBDCollisionContainerSolverSimd.cpp
new file mode 100644
--- /dev/null
+++ b/Engine/Source/Runtime/Experimental/Chaos/Private/Chaos/Collision/PBDCollisionContainerSolverSimd.cpp
@@ -0,0 +1,595 @@
+// Copyright Epic Games, Inc. All Rights Reserved.
+#include "Chaos/Collision/PBDCollisionContainerSolverSimd.h"
+
+#include "Chaos/Collision/CollisionApplyType.h"
+#include "Chaos/Collision/PBDCollisionConstraint.h"
+#include "Chaos/CollisionResolution.h"
+#include "Chaos/DebugDrawQueue.h"
+#include "Chaos/Evolution/SolverBodyContainer.h"
+#include "Chaos/Island/IslandManager.h"
+#include "Chaos/Particle/ParticleUtilities.h"
+#include "Chaos/PBDCollisionConstraints.h"
+#include "Chaos/PBDCollisionConstraintsContact.h"
+#include "Chaos/Utilities.h"
+
+// Private includes
+
+#include "ChaosLog.h"
+
+//PRAGMA_DISABLE_OPTIMIZATION
+
+namespace Chaos
+{
+	namespace CVars
+	{
+		extern bool bChaos_PBDCollisionSolver_Position_SolveEnabled;
+		extern bool bChaos_PBDCollisionSolver_Velocity_SolveEnabled;
+
+		extern FRealSingle Chaos_PBDCollisionSolver_AutoStiffness_MassRatio1;
+		extern FRealSingle Chaos_PBDCollisionSolver_AutoStiffness_MassRatio2;
+	}
+
+	namespace Private
+	{
+
+		//////////////////////////////////////////////////////////////////////////////////////////////////
+		//////////////////////////////////////////////////////////////////////////////////////////////////
+		//////////////////////////////////////////////////////////////////////////////////////////////////
+		//////////////////////////////////////////////////////////////////////////////////////////////////
+
+		// Transform the Constraint's local-space data into world space for use by  the collision solver and also calculate tangents, errors, etc
+		FORCEINLINE_DEBUGGABLE void UpdateCollisionSolverContactPointFromConstraint(Private::FPBDCollisionSolverSimd& Solver, const int32 SolverPointIndex, const FPBDCollisionConstraint* Constraint, const int32 ConstraintPointIndex, const FRealSingle Dt, const FSolverBody& Body0, const FSolverBody& Body1)
+		{
+			const FManifoldPoint& ManifoldPoint = Constraint->GetManifoldPoint(ConstraintPointIndex);
+
+			const FRealSingle Restitution = FRealSingle(Constraint->GetRestitution());
+			const FRealSingle RestitutionVelocityThreshold = FRealSingle(Constraint->GetRestitutionThreshold()) * Dt;
+
+			// World-space shape transforms. Manifold data is currently relative to these spaces
+			const FRigidTransform3& ShapeWorldTransform0 = Constraint->GetShapeWorldTransform0();
+			const FRigidTransform3& ShapeWorldTransform1 = Constraint->GetShapeWorldTransform1();
+
+			// World-space contact points on each shape
+			const FVec3 WorldContact0 = ShapeWorldTransform0.TransformPositionNoScale(FVec3(ManifoldPoint.ContactPoint.ShapeContactPoints[0]));
+			const FVec3 WorldContact1 = ShapeWorldTransform1.TransformPositionNoScale(FVec3(ManifoldPoint.ContactPoint.ShapeContactPoints[1]));
+			const FVec3 WorldContact = FReal(0.5) * (WorldContact0 + WorldContact1);
+			const FVec3f WorldRelativeContact0 = FVec3f(WorldContact - Body0.P());
+			const FVec3f WorldRelativeContact1 = FVec3f(WorldContact - Body1.P());
+
+			// World-space normal
+			const FVec3f WorldContactNormal = FVec3f(ShapeWorldTransform1.TransformVectorNoScale(FVec3(ManifoldPoint.ContactPoint.ShapeContactNormal)));
+
+			// World-space tangents
+			FVec3f WorldContactTangentU = FVec3f::CrossProduct(FVec3f(0, 1, 0), WorldContactNormal);
+			if (!WorldContactTangentU.Normalize(FRealSingle(UE_KINDA_SMALL_NUMBER)))
+			{
+				WorldContactTangentU = FVec3f::CrossProduct(FVec3f(1, 0, 0), WorldContactNormal);
+				WorldContactTangentU = WorldContactTangentU.GetUnsafeNormal();
+			}
+			const FVec3f WorldContactTangentV = FVec3f::CrossProduct(WorldContactNormal, WorldContactTangentU);
+
+			// Calculate contact velocity if we will need it below (restitution and/or frist-contact friction)
+			const bool bNeedContactVelocity = (!ManifoldPoint.Flags.bHasStaticFrictionAnchor) || (Restitution > FRealSingle(0));
+			FVec3f ContactVel = FVec3(0);
+			if (bNeedContactVelocity)
+			{
+				const FVec3f ContactVel0 = Body0.V() + FVec3f::CrossProduct(Body0.W(), WorldRelativeContact0);
+				const FVec3f ContactVel1 = Body1.V() + FVec3f::CrossProduct(Body1.W(), WorldRelativeContact1);
+				ContactVel = ContactVel0 - ContactVel1;
+			}
+
+			// If we have contact data from a previous tick, use it to calculate the lateral position delta we need
+			// to apply to move the contacts back to their original relative locations (i.e., to enforce static friction).
+			// Otherwise, estimate the friction correction from the contact velocity.
+			// NOTE: quadratic shapes use the velocity-based path most of the time, unless the relative motion is very small.
+			FVec3f WorldFrictionDelta = FVec3f(0);
+			if (ManifoldPoint.Flags.bHasStaticFrictionAnchor)
+			{
+				const FVec3f FrictionDelta0 = ShapeWorldTransform0.TransformPositionNoScale(FVec3(ManifoldPoint.ShapeAnchorPoints[0])) - WorldContact0;
+				const FVec3f FrictionDelta1 = ShapeWorldTransform1.TransformPositionNoScale(FVec3(ManifoldPoint.ShapeAnchorPoints[1])) - WorldContact1;
+				WorldFrictionDelta = FrictionDelta0 - FrictionDelta1;
+			}
+			else
+			{
+				// @todo(chaos): consider adding a multiplier to the initial contact friction
+				WorldFrictionDelta = ContactVel * Dt;
+			}
+
+			// The contact point error we are trying to correct in this solver
+			const FRealSingle TargetPhi = ManifoldPoint.TargetPhi;
+			const FVec3f WorldContactDelta = FVec3f(WorldContact0 - WorldContact1);
+			const FRealSingle WorldContactDeltaNormal = FVec3f::DotProduct(WorldContactDelta, WorldContactNormal) - TargetPhi;
+			const FRealSingle WorldContactDeltaTangentU = FVec3f::DotProduct(WorldContactDelta + WorldFrictionDelta, WorldContactTangentU);
+			const FRealSingle WorldContactDeltaTangentV = FVec3f::DotProduct(WorldContactDelta + WorldFrictionDelta, WorldContactTangentV);
+
+			// The target contact velocity, taking restitution into account
+			FRealSingle WorldContactTargetVelocityNormal = FRealSingle(0);
+			if (Restitution > FRealSingle(0))
+			{
+				const FRealSingle ContactVelocityNormal = FVec3f::DotProduct(ContactVel, WorldContactNormal);
+				if (ContactVelocityNormal < -RestitutionVelocityThreshold)
+				{
+					WorldContactTargetVelocityNormal = -Restitution * ContactVelocityNormal;
+				}
+			}
+
+			FWorldContactPoint& WorldContactPoint = Solver.GetWorldContactPoint(SolverPointIndex);
+			WorldContactPoint.RelativeContactPoints[0] = WorldRelativeContact0;
+			WorldContactPoint.RelativeContactPoints[1] = WorldRelativeContact1;
+			WorldContactPoint.ContactNormal = WorldContactNormal;
+			WorldContactPoint.ContactTangentU = WorldContactTangentU;
+			WorldContactPoint.ContactTangentV = WorldContactTangentV;
+			WorldContactPoint.ContactDeltaNormal = WorldContactDeltaNormal;
+			WorldContactPoint.ContactDeltaTangentU = WorldContactDeltaTangentU;
+			WorldContactPoint.ContactDeltaTangentV = WorldContactDeltaTangentV;
+			WorldContactPoint.ContactTargetVelocityNormal = WorldContactTargetVelocityNormal;
+		}
+
+		FORCEINLINE_DEBUGGABLE void UpdateCollisionSolverManifoldFromConstraint(Private::FPBDCollisionSolverSimd& Solver, const FPBDCollisionConstraint* Constraint, const FSolverReal Dt, const int32 ConstraintPointBeginIndex, const int32 ConstraintPointEndIndex)
+		{
+			const FSolverBody& Body0 = Solver.SolverBody0().SolverBody();
+			const FSolverBody& Body1 = Solver.SolverBody1().SolverBody();
+
+			// Only calculate state for newly added contacts. Normally this is all of them, but maybe not if incremental collision is used by RBAN.
+			// Also we only add active points to the solver's manifold points list
+			for (int32 ConstraintManifoldPointIndex = ConstraintPointBeginIndex; ConstraintManifoldPointIndex < ConstraintPointEndIndex; ++ConstraintManifoldPointIndex)
+			{
+				if (!Constraint->GetManifoldPoint(ConstraintManifoldPointIndex).Flags.bDisabled)
+				{
+					const int32 SolverManifoldPointIndex = Solver.AddManifoldPoint();
+
+					// Transform the constraint contact data into world space for use by the solver
+					// We build this data directly into the solver's world-space contact data which looks a bit odd with "Init" called after but there you go
+					UpdateCollisionSolverContactPointFromConstraint(Solver, SolverManifoldPointIndex, Constraint, ConstraintManifoldPointIndex, Dt, Body0, Body1);
+
+					// Prepare the solver (set constraint-space mass etc)
+					Solver.FinalizeManifoldPoint(SolverManifoldPointIndex, Dt);
+				}
+			}
+		}
+
+		FORCEINLINE_DEBUGGABLE void UpdateCollisionSolverFromConstraint(Private::FPBDCollisionSolverSimd& Solver, const FPBDCollisionConstraint* Constraint, const FSolverReal Dt, const FPBDCollisionSolverSettings& SolverSettings, bool& bOutPerIterationCollision)
+		{
+			// Friction values. Static and Dynamic friction are applied in the position solve for most shapes.
+			// We can also run in a mode without static friction at all. This is faster but stacking is not possible.
+			const FSolverReal StaticFriction = FSolverReal(Constraint->GetStaticFriction());
+			const FSolverReal DynamicFriction = FSolverReal(Constraint->GetDynamicFriction());
+			FSolverReal PositionStaticFriction = FSolverReal(0);
+			FSolverReal PositionDynamicFriction = FSolverReal(0);
+			FSolverReal VelocityDynamicFriction = FSolverReal(0);
+			if (SolverSettings.NumPositionFrictionIterations > 0)
+			{
+				PositionStaticFriction = StaticFriction;
+				if (!Constraint->HasQuadraticShape())
+				{
+					PositionDynamicFriction = DynamicFriction;
+				}
+				else
+				{
+					// Quadratic shapes don't use PBD dynamic friction - it has issues at slow speeds where the WxR is
+					// less than the position tolerance for friction point matching
+					// @todo(chaos): fix PBD dynamic friction on quadratic shapes
+					VelocityDynamicFriction = DynamicFriction;
+				}
+			}
+			else
+			{
+				VelocityDynamicFriction = DynamicFriction;
+			}
+
+			Solver.SetFriction(PositionStaticFriction, PositionDynamicFriction, VelocityDynamicFriction);
+
+			const FReal SolverStiffness = Constraint->GetStiffness();
+			Solver.SetStiffness(FSolverReal(SolverStiffness));
+
+			Solver.SolverBody0().SetInvMScale(Constraint->GetInvMassScale0());
+			Solver.SolverBody0().SetInvIScale(Constraint->GetInvInertiaScale0());
+			Solver.SolverBody0().SetShockPropagationScale(FReal(1));
+			Solver.SolverBody1().SetInvMScale(Constraint->GetInvMassScale1());
+			Solver.SolverBody1().SetInvIScale(Constraint->GetInvInertiaScale1());
+			Solver.SolverBody1().SetShockPropagationScale(FReal(1));
+
+			bOutPerIterationCollision = (!Constraint->GetUseManifold() || Constraint->GetUseIncrementalCollisionDetection());
+
+			UpdateCollisionSolverManifoldFromConstraint(Solver, Constraint, Dt, 0, Constraint->NumManifoldPoints());
+		}
+
+		FORCEINLINE_DEBUGGABLE void UpdateCollisionConstraintFromSolver(FPBDCollisionConstraint* Constraint, const Private::FPBDCollisionSolverSimd& Solver, const FSolverReal Dt)
+		{
+			Constraint->ResetSolverResults();
+
+			// NOTE: We only put the non-pruned manifold points into the solver so the ManifoldPointIndex and
+			// SolverManifoldPointIndex do not necessarily match. See GatherManifoldPoints
+			int32 SolverManifoldPointIndex = 0;
+			for (int32 ManifoldPointIndex = 0; ManifoldPointIndex < Constraint->NumManifoldPoints(); ++ManifoldPointIndex)
+			{
+				FSolverVec3 NetPushOut = FSolverVec3(0);
+				FSolverVec3 NetImpulse = FSolverVec3(0);
+				FSolverReal StaticFrictionRatio = FSolverReal(0);
+
+				const Private::FPBDCollisionSolverManifoldPointSimd& SolverManifoldPoint = Solver.GetManifoldPoint(SolverManifoldPointIndex);
+				if (!Constraint->GetManifoldPoint(ManifoldPointIndex).Flags.bDisabled)
+				{
+					NetPushOut =
+						SolverManifoldPoint.NetPushOutNormal * SolverManifoldPoint.WorldContact.ContactNormal +
+						SolverManifoldPoint.NetPushOutTangentU * SolverManifoldPoint.WorldContact.ContactTangentU +
+						SolverManifoldPoint.NetPushOutTangentV * SolverManifoldPoint.WorldContact.ContactTangentV;
+
+					NetImpulse =
+						SolverManifoldPoint.NetImpulseNormal * SolverManifoldPoint.WorldContact.ContactNormal +
+						SolverManifoldPoint.NetImpulseTangentU * SolverManifoldPoint.WorldContact.ContactTangentU +
+						SolverManifoldPoint.NetImpulseTangentV * SolverManifoldPoint.WorldContact.ContactTangentV;
+
+					StaticFrictionRatio = SolverManifoldPoint.StaticFrictionRatio;
+
+					++SolverManifoldPointIndex;
+				}
+
+				// NOTE: We call this even for points we did not run the solver for (but with zero results)
+				Constraint->SetSolverResults(ManifoldPointIndex,
+					NetPushOut,
+					NetImpulse,
+					StaticFrictionRatio,
+					Dt);
+			}
+
+			Constraint->EndTick();
+		}
+
+		//////////////////////////////////////////////////////////////////////////////////////////////////
+		//////////////////////////////////////////////////////////////////////////////////////////////////
+		//////////////////////////////////////////////////////////////////////////////////////////////////
+		//////////////////////////////////////////////////////////////////////////////////////////////////
+
+
+		FPBDCollisionContainerSolverSimd::FPBDCollisionContainerSolverSimd(const FPBDCollisionConstraints& InConstraintContainer, const int32 InPriority)
+			: FConstraintContainerSolver(InPriority)
+			, ConstraintContainer(InConstraintContainer)
+			, bPerIterationCollisionDetection(false)
+		{
+	#if !UE_BUILD_SHIPPING && !UE_BUILD_TEST && INTEL_ISPC
+			Private::FPBDCollisionSolverHelperSimd::CheckISPC();
+	#endif
+		}
+
+		FPBDCollisionContainerSolverSimd::~FPBDCollisionContainerSolverSimd()
+		{
+		}
+
+		void FPBDCollisionContainerSolverSimd::Reset(const int32 MaxCollisions)
+		{
+			// A over-allocation policy to avoid reallocation every frame in the common case where a pile of objects is dropped
+			// and the number of contacts increases every tick.
+			if (MaxCollisions > CollisionSolvers.Max())
+			{
+				CollisionSolvers.Reset((5 * MaxCollisions) / 4);	// +25%
+			}
+
+			CollisionSolvers.Reset(MaxCollisions);
+			CollisionConstraints.Reset(CollisionSolvers.Max());
+			bCollisionConstraintPerIterationCollisionDetection.Reset(CollisionSolvers.Max());
+			CollisionSolverManifoldPoints.Reset(CollisionSolvers.Max() * Private::FPBDCollisionSolverSimd::MaxPointsPerConstraint);
+
+			// Just set the array size for these right away - all data will be initialized later and there is no constructor
+			// We use CollisionConstraints.Num() to track how many constraints we currently have
+			CollisionSolvers.SetNumUninitialized(MaxCollisions);
+			bCollisionConstraintPerIterationCollisionDetection.SetNumUninitialized(MaxCollisions);
+		}
+
+		void FPBDCollisionContainerSolverSimd::AddConstraints()
+		{
+			Reset(ConstraintContainer.NumConstraints());
+
+			for (FPBDCollisionConstraintHandle* ConstraintHandle : ConstraintContainer.GetConstraintHandles())
+			{
+				FPBDCollisionConstraint& Constraint = ConstraintHandle->GetContact();
+
+				AddConstraint(Constraint);
+			}
+		}
+
+		void FPBDCollisionContainerSolverSimd::AddConstraints(const TArrayView<Private::FPBDIslandConstraint>& IslandConstraints)
+		{
+			for (Private::FPBDIslandConstraint& IslandConstraint : IslandConstraints)
+			{
+				// We will only ever be given constraints from our container (asserts in non-shipping)
+				FPBDCollisionConstraint& Constraint = IslandConstraint.GetConstraint()->AsUnsafe<FPBDCollisionConstraintHandle>()->GetContact();
+
+				AddConstraint(Constraint);
+			}
+		}
+
+		void FPBDCollisionContainerSolverSimd::AddConstraint(FPBDCollisionConstraint& Constraint)
+		{
+			// NOTE: No need to add to CollisionSolvers or bCollisionConstraintPerIterationCollisionDetection
+			const int32 Index = CollisionConstraints.Add(&Constraint);
+
+			// Allocate the manifold points for this constraints solver
+			// NOTE: we don't know how many points we will create if collision detection is deferred or incrememntal, so just allocate space for the max allowed
+			// @todo(chaos): see if we can do better - we don't want to pay for this logic when it is only used rarely
+			const bool bDeferredCollisionDetection = ConstraintContainer.GetDetectorSettings().bDeferNarrowPhase;
+			const int32 ManifoldPointMax = (bDeferredCollisionDetection || Constraint.GetUseIncrementalManifold()) ? Private::FPBDCollisionSolverSimd::MaxPointsPerConstraint : Constraint.NumManifoldPoints();
+			if (ManifoldPointMax > 0)
+			{
+				const int32 ManifoldPointStart = CollisionSolverManifoldPoints.Num();
+				check(CollisionSolverManifoldPoints.Max() >= ManifoldPointStart + ManifoldPointMax);	// No resize following reset
+				CollisionSolverManifoldPoints.SetNum(ManifoldPointStart + ManifoldPointMax);
+
+				CollisionSolvers[Index].SetManifoldPointsBuffer(&CollisionSolverManifoldPoints[ManifoldPointStart], ManifoldPointMax);
+			}
+			else
+			{
+				// Ideally we wouldn't have any active constraint with no manifold points, but this can
+				// currently happen for CCD objects and incremental collisions
+				CollisionSolvers[Index].Reset();
+			}
+		}
+
+
+		void FPBDCollisionContainerSolverSimd::AddBodies(FSolverBodyContainer& SolverBodyContainer)
+		{
+			for (int32 SolverIndex = 0, SolverEndIndex = NumSolvers(); SolverIndex < SolverEndIndex; ++SolverIndex)
+			{
+				Private::FPBDCollisionSolverSimd& CollisionSolver = CollisionSolvers[SolverIndex];
+				FPBDCollisionConstraint* Constraint = CollisionConstraints[SolverIndex];
+				check(Constraint != nullptr);
+
+				// Find the solver bodies for the particles we constrain. This will add them to the container
+				// if they aren't there already, and ensure that they are populated with the latest data.
+				FSolverBody* Body0 = SolverBodyContainer.FindOrAdd(Constraint->GetParticle0());
+				FSolverBody* Body1 = SolverBodyContainer.FindOrAdd(Constraint->GetParticle1());
+				check(Body0 != nullptr);
+				check(Body1 != nullptr);
+
+				CollisionSolver.SetSolverBodies(*Body0, *Body1);
+			}
+		}
+
+		void FPBDCollisionContainerSolverSimd::GatherInput(const FReal Dt)
+		{
+			GatherInput(Dt, 0, NumSolvers());
+		}
+
+		void FPBDCollisionContainerSolverSimd::CachePrefetchSolver(const int32 ConstraintIndex) const
+		{
+			if (ConstraintIndex < NumSolvers())
+			{
+				const int32 NumManifoldPoints = CollisionSolvers[ConstraintIndex].NumManifoldPoints();
+				if (NumManifoldPoints > 0)
+				{
+					FPlatformMisc::PrefetchBlock(&CollisionSolvers[ConstraintIndex].GetManifoldPoint(0), NumManifoldPoints * sizeof(Private::FPBDCollisionSolverManifoldPointSimd));
+				}
+				FPlatformMisc::PrefetchBlock(&CollisionSolvers[ConstraintIndex], sizeof(FPBDCollisionContainerSolverSimd));
+			}
+		}
+
+		void FPBDCollisionContainerSolverSimd::GatherInput(const FReal InDt, const int32 BeginIndex, const int32 EndIndex)
+		{
+			// NOTE: may be called in parallel. Should not change the container or any elements outside of [BeginIndex, EndIndex)
+
+			check(BeginIndex >= 0);
+			check(EndIndex <= NumSolvers());
+
+			const FSolverReal Dt = FSolverReal(InDt);
+
+			const int32 PrefetchCount = 4;
+			for (int32 PrefetchIndex = 0; PrefetchIndex < PrefetchCount; ++PrefetchIndex)
+			{
+				CachePrefetchSolver(PrefetchIndex);
+			}
+
+			bool bWantPerIterationCollisionDetection = false;
+			for (int32 ConstraintIndex = BeginIndex; ConstraintIndex < EndIndex; ++ConstraintIndex)
+			{
+				Private::FPBDCollisionSolverSimd& CollisionSolver = CollisionSolvers[ConstraintIndex];
+				FPBDCollisionConstraint* Constraint = CollisionConstraints[ConstraintIndex];
+				bool& bPerIterationCollision = bCollisionConstraintPerIterationCollisionDetection[ConstraintIndex];
+
+				UpdateCollisionSolverFromConstraint(CollisionSolver, Constraint, Dt, ConstraintContainer.GetSolverSettings(), bPerIterationCollision);
+
+				// We need to run collision every iteration if we are not using manifolds, or are using incremental manifolds
+				bWantPerIterationCollisionDetection = bWantPerIterationCollisionDetection || bPerIterationCollision;
+
+				CachePrefetchSolver(ConstraintIndex + PrefetchCount);
+			}
+
+			if (bWantPerIterationCollisionDetection)
+			{
+				// We should lock here? We only ever set to true or do nothing so I think it doesn't matter if this happens on multiple threads...
+				bPerIterationCollisionDetection = true;
+			}
+		}
+
+		void FPBDCollisionContainerSolverSimd::ScatterOutput(const FReal Dt)
+		{
+			ScatterOutput(Dt, 0, NumSolvers());
+		}
+
+		void FPBDCollisionContainerSolverSimd::ScatterOutput(const FReal InDt, const int32 BeginIndex, const int32 EndIndex)
+		{
+			// NOTE: may be called in parallel. Should not change the container or any elements outside of [BeginIndex, EndIndex)
+
+			check(BeginIndex >= 0);
+			check(EndIndex <= NumSolvers());
+
+			const FSolverReal Dt = FSolverReal(InDt);
+
+			for (int32 ConstraintIndex = BeginIndex; ConstraintIndex < EndIndex; ++ConstraintIndex)
+			{
+				Private::FPBDCollisionSolverSimd& CollisionSolver = CollisionSolvers[ConstraintIndex];
+				FPBDCollisionConstraint* Constraint = CollisionConstraints[ConstraintIndex];
+
+				UpdateCollisionConstraintFromSolver(Constraint, CollisionSolver, Dt);
+
+				// Reset the collision solver here as the pointers will be invalid on the next tick
+				// @todo(chaos): Pointers are always reinitalized before use next tick so this isn't strictly necessary
+				CollisionSolver.Reset();
+			}
+		}
+
+		void FPBDCollisionContainerSolverSimd::ApplyPositionConstraints(const FReal Dt, const int32 It, const int32 NumIts)
+		{
+			SolvePositionImpl(Dt, It, NumIts, 0, NumSolvers(), ConstraintContainer.GetSolverSettings());
+		}
+
+		void FPBDCollisionContainerSolverSimd::ApplyVelocityConstraints(const FReal Dt, const int32 It, const int32 NumIts)
+		{
+			SolveVelocityImpl(Dt, It, NumIts, 0, NumSolvers(), ConstraintContainer.GetSolverSettings());
+		}
+
+		void FPBDCollisionContainerSolverSimd::ApplyProjectionConstraints(const FReal Dt, const int32 It, const int32 NumIts)
+		{
+			// Not supported for collisions
+		}
+
+		void FPBDCollisionContainerSolverSimd::UpdatePositionShockPropagation(const FReal Dt, const int32 It, const int32 NumIts, const int32 BeginIndex, const int32 EndIndex, const FPBDCollisionSolverSettings& SolverSettings)
+		{
+			// If this is the first shock propagation iteration, enable it on each solver
+			const bool bEnableShockPropagation = (It == NumIts - SolverSettings.NumPositionShockPropagationIterations);
+			if (bEnableShockPropagation)
+			{
+				for (int32 SolverIndex = BeginIndex; SolverIndex < EndIndex; ++SolverIndex)
+				{
+					CollisionSolvers[SolverIndex].EnablePositionShockPropagation();
+				}
+			}
+		}
+
+		void FPBDCollisionContainerSolverSimd::UpdateVelocityShockPropagation(const FReal Dt, const int32 It, const int32 NumIts, const int32 BeginIndex, const int32 EndIndex, const FPBDCollisionSolverSettings& SolverSettings)
+		{
+			// Set/reset the shock propagation based on current iteration. The position solve may
+			// have left the bodies with a mass scale and we want to change or reset it.
+			const bool bEnableShockPropagation = (It == NumIts - SolverSettings.NumVelocityShockPropagationIterations);
+			if (bEnableShockPropagation)
+			{
+				for (int32 SolverIndex = BeginIndex; SolverIndex < EndIndex; ++SolverIndex)
+				{
+					CollisionSolvers[SolverIndex].EnableVelocityShockPropagation();
+				}
+			}
+			else if (It == 0)
+			{
+				for (int32 SolverIndex = BeginIndex; SolverIndex < EndIndex; ++SolverIndex)
+				{
+					CollisionSolvers[SolverIndex].DisableShockPropagation();
+				}
+			}
+		}
+
+		bool FPBDCollisionContainerSolverSimd::SolvePositionImpl(const FReal InDt, const int32 It, const int32 NumIts, const int32 BeginIndex, const int32 EndIndex, const FPBDCollisionSolverSettings& SolverSettings)
+		{
+			SCOPE_CYCLE_COUNTER(STAT_Collisions_Apply);
+			if (!CVars::bChaos_PBDCollisionSolver_Position_SolveEnabled)
+			{
+				return false;
+			}
+
+			if (EndIndex <= BeginIndex)
+			{
+				return false;
+			}
+
+			UpdatePositionShockPropagation(InDt, It, NumIts, BeginIndex, EndIndex, SolverSettings);
+
+			// We run collision detection here under two conditions (normally it is run after Integration and before the constraint solver phase):
+			// 1) When deferring collision detection until the solver phase for better joint-collision behaviour (RBAN). In this case, we only do this on the first iteration.
+			// 2) When using no manifolds or incremental manifolds, where we may add/replace manifold points every iteration.
+			const bool bRunDeferredCollisionDetection = (It == 0) && ConstraintContainer.GetDetectorSettings().bDeferNarrowPhase;
+			if (bRunDeferredCollisionDetection || bPerIterationCollisionDetection)
+			{
+				UpdateCollisions(InDt, BeginIndex, EndIndex);
+			}
+
+			// Only apply friction for the last few (tunable) iterations
+			// Adjust max pushout to attempt to make it iteration count independent
+			const FSolverReal Dt = FSolverReal(InDt);
+			const bool bApplyStaticFriction = (It >= (NumIts - SolverSettings.NumPositionFrictionIterations));
+			const FSolverReal MaxPushOut = (SolverSettings.MaxPushOutVelocity > 0) ? (FSolverReal(SolverSettings.MaxPushOutVelocity) * Dt) / FSolverReal(NumIts) : FSolverReal(0);
+
+			// Apply the position correction
+			if (bApplyStaticFriction)
+			{
+				Private::FPBDCollisionSolverHelperSimd::SolvePositionWithFriction(MakeArrayView(&CollisionSolvers[BeginIndex], EndIndex - BeginIndex), Dt, MaxPushOut);
+			}
+			else
+			{
+				Private::FPBDCollisionSolverHelperSimd::SolvePositionNoFriction(MakeArrayView(&CollisionSolvers[BeginIndex], EndIndex - BeginIndex), Dt, MaxPushOut);
+			}
+
+			return true;
+		}
+
+		bool FPBDCollisionContainerSolverSimd::SolveVelocityImpl(const FReal InDt, const int32 It, const int32 NumIts, const int32 BeginIndex, const int32 EndIndex, const FPBDCollisionSolverSettings& SolverSettings)
+		{
+			SCOPE_CYCLE_COUNTER(STAT_Collisions_ApplyPushOut);
+			if (!CVars::bChaos_PBDCollisionSolver_Velocity_SolveEnabled)
+			{
+				return false;
+			}
+
+			if (EndIndex <= BeginIndex)
+			{
+				return false;
+			}
+
+			UpdateVelocityShockPropagation(InDt, It, NumIts, BeginIndex, EndIndex, SolverSettings);
+
+			const FSolverReal Dt = FSolverReal(InDt);
+			const bool bApplyDynamicFriction = (It >= NumIts - SolverSettings.NumVelocityFrictionIterations);
+
+			Private::FPBDCollisionSolverHelperSimd::SolveVelocity(MakeArrayView(&CollisionSolvers[BeginIndex], EndIndex - BeginIndex), Dt, bApplyDynamicFriction);
+
+			return true;
+		}
+
+		void FPBDCollisionContainerSolverSimd::UpdateCollisions(const FReal InDt, const int32 BeginIndex, const int32 EndIndex)
+		{
+			const FSolverReal Dt = FSolverReal(InDt);
+			const bool bDeferredCollisionDetection = ConstraintContainer.GetDetectorSettings().bDeferNarrowPhase;
+
+			bool bNeedsAnotherIteration = false;
+			for (int32 SolverIndex = BeginIndex; SolverIndex < EndIndex; ++SolverIndex)
+			{
+				if (bDeferredCollisionDetection || bCollisionConstraintPerIterationCollisionDetection[SolverIndex])
+				{
+					Private::FPBDCollisionSolverSimd& CollisionSolver = CollisionSolvers[SolverIndex];
+					FPBDCollisionConstraint* Constraint = CollisionConstraints[SolverIndex];
+
+					// Run collision detection at the current transforms including any correction from previous iterations
+					const FSolverBody& Body0 = CollisionSolver.SolverBody0().SolverBody();
+					const FSolverBody& Body1 = CollisionSolver.SolverBody1().SolverBody();
+					const FRigidTransform3 CorrectedActorWorldTransform0 = FRigidTransform3(Body0.CorrectedActorP(), Body0.CorrectedActorQ());
+					const FRigidTransform3 CorrectedActorWorldTransform1 = FRigidTransform3(Body1.CorrectedActorP(), Body1.CorrectedActorQ());
+					const FRigidTransform3 CorrectedShapeWorldTransform0 = Constraint->GetShapeRelativeTransform0() * CorrectedActorWorldTransform0;
+					const FRigidTransform3 CorrectedShapeWorldTransform1 = Constraint->GetShapeRelativeTransform1() * CorrectedActorWorldTransform1;
+
+					// @todo(chaos): this is ugly - pass these to the required functions instead and remove from the constraint class
+					// This is now only needed for LevelSet collision (see UpdateLevelsetLevelsetConstraint)
+					Constraint->SetSolverBodies(&Body0, &Body1);
+
+					// Reset the manifold if we are not using manifolds (we just use the first manifold point)
+					if (!Constraint->GetUseManifold())
+					{
+						Constraint->ResetActiveManifoldContacts();
+						CollisionSolver.ResetManifold();
+					}
+
+					// We need to know how many points were added to the manifold
+					const int32 BeginPointIndex = Constraint->NumManifoldPoints();
+
+					// NOTE: We deliberately have not updated the ShapwWorldTranforms on the constraint. If we did that, we would calculate 
+					// errors incorrectly in UpdateManifoldPoints, because the solver assumes nothing has been moved as we iterate (we accumulate 
+					// corrections that will be applied later.)
+					Constraint->ResetPhi(Constraint->GetCullDistance());
+					Collisions::UpdateConstraint(*Constraint, CorrectedShapeWorldTransform0, CorrectedShapeWorldTransform1, Dt);
+
+					// Update the manifold based on the new or updated contacts
+					UpdateCollisionSolverManifoldFromConstraint(CollisionSolver, Constraint, Dt, BeginPointIndex, Constraint->NumManifoldPoints());
+
+					Constraint->SetSolverBodies(nullptr, nullptr);
+				}
+			}
+		}
+
+	}	// namespace Private
+}	// namespace Chaos
diff --git a/Engine/Source/Runtime/Experimental/Chaos/Private/Chaos/Collision/PBDCollisionSolver.cpp b/Engine/Source/Runtime/Experimental/Chaos/Private/Chaos/Collision/PBDCollisionSolver.cpp
--- a/Engine/Source/Runtime/Experimental/Chaos/Private/Chaos/Collision/PBDCollisionSolver.cpp
+++ b/Engine/Source/Runtime/Experimental/Chaos/Private/Chaos/Collision/PBDCollisionSolver.cpp
@@ -3,7 +3,6 @@
 
 #include "Chaos/Collision/CollisionApplyType.h"
 #include "Chaos/Collision/PBDCollisionConstraint.h"
-#include "Chaos/Collision/SolverCollisionContainer.h"
 #include "Chaos/CollisionResolution.h"
 #include "Chaos/CollisionResolutionTypes.h"
 #include "Chaos/CollisionResolutionUtil.h"
@@ -27,16 +26,13 @@ namespace Chaos
 	namespace CVars
 	{
 		extern int32 Chaos_Collision_UseShockPropagation;
-		extern bool bChaos_PBDCollisionSolver_UseJacobiPairSolver;
 
 		//
 		// Solver Settings
 		//
 
 		bool bChaos_PBDCollisionSolver_ISPC = false;	// NOTE: WORK IN PROGRESS. Not ready to be enabled.
-		bool bChaos_PBDCollisionSolver_UseJacobiPairSolver = false;
 		FAutoConsoleVariableRef CVarChaos_PBDCollisionSolver_ISPC(TEXT("p.Chaos.PBDCollisionSolver.ISPC"), bChaos_PBDCollisionSolver_ISPC, TEXT("Use ISPC collision solver (WIP)"));
-		FAutoConsoleVariableRef CVarChaosPBDCollisionSolverUseJacobiPairSolver(TEXT("p.Chaos.PBDCollisionSolver.UseJacobi"), bChaos_PBDCollisionSolver_UseJacobiPairSolver, TEXT("Whether to use the Jacobi collision pair solver (as option in main collision solver)"));
 
 		//
 		// Position Solver Settings
@@ -259,76 +255,34 @@ namespace Chaos
 #if INTEL_ISPC
 			if (CVars::bChaos_PBDCollisionSolver_ISPC)
 			{
-				if (!CVars::bChaos_PBDCollisionSolver_UseJacobiPairSolver)
-				{
-					ispc::SolvePositionNoFriction(
-						(ispc::FPBDCollisionSolver*)CollisionSolvers.GetData(), 
-						CollisionSolvers.Num(), 
-						Dt, 
-						MaxPushOut);
-				}
-				else
-				{
-					ispc::SolvePositionNoFrictionJacobi(
-						(ispc::FPBDCollisionSolver*)CollisionSolvers.GetData(), 
-						CollisionSolvers.Num(), 
-						Dt, 
-						MaxPushOut, 
-						CVars::Chaos_PBDCollisionSolver_JacobiStiffness);
-				}
+				ispc::SolvePositionNoFriction(
+					(ispc::FPBDCollisionSolver*)CollisionSolvers.GetData(), 
+					CollisionSolvers.Num(), 
+					Dt, 
+					MaxPushOut);
 				return;
 			}
 #endif
 
-			if (!CVars::bChaos_PBDCollisionSolver_UseJacobiPairSolver)
-			{
-				for (FPBDCollisionSolver& CollisionSolver : CollisionSolvers)
-				{
-					CollisionSolver.SolvePositionNoFriction(Dt, MaxPushOut);
-				}
-			}
-			else
+			for (FPBDCollisionSolver& CollisionSolver : CollisionSolvers)
 			{
-				for (FPBDCollisionSolver& CollisionSolver : CollisionSolvers)
-				{
-					CollisionSolver.SolvePositionNoFrictionJacobi(Dt, MaxPushOut);
-				}
+				CollisionSolver.SolvePositionNoFriction(Dt, MaxPushOut);
 			}
 		}
 
 		void FPBDCollisionSolverHelper::SolvePositionWithFriction(const TArrayView<FPBDCollisionSolver>& CollisionSolvers, const FSolverReal Dt, const FSolverReal MaxPushOut)
 		{
-			if (!CVars::bChaos_PBDCollisionSolver_UseJacobiPairSolver)
+			for (FPBDCollisionSolver& CollisionSolver : CollisionSolvers)
 			{
-				for (FPBDCollisionSolver& CollisionSolver : CollisionSolvers)
-				{
-					CollisionSolver.SolvePositionWithFriction(Dt, MaxPushOut);
-				}
-			}
-			else
-			{
-				for (FPBDCollisionSolver& CollisionSolver : CollisionSolvers)
-				{
-					CollisionSolver.SolvePositionWithFrictionJacobi(Dt, MaxPushOut);
-				}
+				CollisionSolver.SolvePositionWithFriction(Dt, MaxPushOut);
 			}
 		}
 
 		void FPBDCollisionSolverHelper::SolveVelocity(const TArrayView<FPBDCollisionSolver>& CollisionSolvers, const FSolverReal Dt, const bool bApplyDynamicFriction)
 		{
-			if (!CVars::bChaos_PBDCollisionSolver_UseJacobiPairSolver)
+			for (FPBDCollisionSolver& CollisionSolver : CollisionSolvers)
 			{
-				for (FPBDCollisionSolver& CollisionSolver : CollisionSolvers)
-				{
-					CollisionSolver.SolveVelocity(Dt, bApplyDynamicFriction);
-				}
-			}
-			else
-			{
-				for (FPBDCollisionSolver& CollisionSolver : CollisionSolvers)
-				{
-					CollisionSolver.SolveVelocityJacobi(Dt, bApplyDynamicFriction);
-				}
+				CollisionSolver.SolveVelocity(Dt, bApplyDynamicFriction);
 			}
 		}
 
diff --git a/Engine/Source/Runtime/Experimental/Chaos/Private/Chaos/Collision/PBDCollisionSolverJacobi.cpp b/Engine/Source/Runtime/Experimental/Chaos/Private/Chaos/Collision/PBDCollisionSolverJacobi.cpp
--- a/Engine/Source/Runtime/Experimental/Chaos/Private/Chaos/Collision/PBDCollisionSolverJacobi.cpp
+++ b/Engine/Source/Runtime/Experimental/Chaos/Private/Chaos/Collision/PBDCollisionSolverJacobi.cpp
@@ -3,7 +3,6 @@
 
 #include "Chaos/Collision/CollisionApplyType.h"
 #include "Chaos/Collision/PBDCollisionConstraint.h"
-#include "Chaos/Collision/SolverCollisionContainer.h"
 #include "Chaos/CollisionResolution.h"
 #include "Chaos/CollisionResolutionTypes.h"
 #include "Chaos/CollisionResolutionUtil.h"
diff --git a/Engine/Source/Runtime/Experimental/Chaos/Private/Chaos/Collision/PBDCollisionSolverSimd.cpp b/Engine/Source/Runtime/Experimental/Chaos/Private/Chaos/Collision/PBDCollisionSolverSimd.cpp
new file mode 100644
--- /dev/null
+++ b/Engine/Source/Runtime/Experimental/Chaos/Private/Chaos/Collision/PBDCollisionSolverSimd.cpp
@@ -0,0 +1,235 @@
+// Copyright Epic Games, Inc. All Rights Reserved.
+#include "Chaos/Collision/PBDCollisionSolverSimd.h"
+
+#include "Chaos/Collision/CollisionApplyType.h"
+#include "Chaos/Collision/PBDCollisionConstraint.h"
+#include "Chaos/CollisionResolution.h"
+#include "Chaos/CollisionResolutionTypes.h"
+#include "Chaos/CollisionResolutionUtil.h"
+#include "Chaos/Evolution/SolverBodyContainer.h"
+#include "Chaos/Particle/ParticleUtilities.h"
+#include "Chaos/PBDCollisionConstraints.h"
+#include "Chaos/PBDCollisionConstraintsContact.h"
+#include "Chaos/Utilities.h"
+
+//PRAGMA_DISABLE_OPTIMIZATION
+
+namespace Chaos
+{
+	namespace CVars
+	{
+		extern float Chaos_PBDCollisionSolver_Position_MinInvMassScale;
+		extern float Chaos_PBDCollisionSolver_Velocity_MinInvMassScale;
+	}
+	using namespace CVars;
+
+	namespace Private
+	{
+		//////////////////////////////////////////////////////////////////////////////////////////////////
+		//////////////////////////////////////////////////////////////////////////////////////////////////
+		//////////////////////////////////////////////////////////////////////////////////////////////////
+		//////////////////////////////////////////////////////////////////////////////////////////////////
+
+
+		void FPBDCollisionSolverSimd::EnablePositionShockPropagation()
+		{
+			SetShockPropagationInvMassScale(Chaos_PBDCollisionSolver_Position_MinInvMassScale);
+		}
+
+		void FPBDCollisionSolverSimd::EnableVelocityShockPropagation()
+		{
+			SetShockPropagationInvMassScale(Chaos_PBDCollisionSolver_Velocity_MinInvMassScale);
+		}
+
+		void FPBDCollisionSolverSimd::DisableShockPropagation()
+		{
+			SetShockPropagationInvMassScale(FReal(1));
+		}
+
+		void FPBDCollisionSolverSimd::SetShockPropagationInvMassScale(const FSolverReal InvMassScale)
+		{
+			FConstraintSolverBody& Body0 = SolverBody0();
+			FConstraintSolverBody& Body1 = SolverBody1();
+
+			// Shock propagation decreases the inverse mass of bodies that are lower in the pile
+			// of objects. This significantly improves stability of heaps and stacks. Height in the pile is indictaed by the "level". 
+			// No need to set an inverse mass scale if the other body is kinematic (with inv mass of 0).
+			// Bodies at the same level do not take part in shock propagation.
+			if (Body0.IsDynamic() && Body1.IsDynamic() && (Body0.Level() != Body1.Level()))
+			{
+				// Set the inv mass scale of the "lower" body to make it heavier
+				bool bInvMassUpdated = false;
+				if (Body0.Level() < Body1.Level())
+				{
+					if (Body0.ShockPropagationScale() != InvMassScale)
+					{
+						Body0.SetShockPropagationScale(InvMassScale);
+						bInvMassUpdated = true;
+					}
+				}
+				else
+				{
+					if (Body1.ShockPropagationScale() != InvMassScale)
+					{
+						Body1.SetShockPropagationScale(InvMassScale);
+						bInvMassUpdated = true;
+					}
+				}
+
+				// If the masses changed, we need to rebuild the contact mass for each manifold point
+				if (bInvMassUpdated)
+				{
+					for (int32 PointIndex = 0; PointIndex < NumManifoldPoints(); ++PointIndex)
+					{
+						State.ManifoldPoints[PointIndex].UpdateMassNormal(Body0, Body1);
+					}
+				}
+			}
+		}
+
+		void FPBDCollisionSolverSimd::SolveVelocityAverage(const FSolverReal Dt)
+		{
+			FConstraintSolverBody& Body0 = SolverBody0();
+			FConstraintSolverBody& Body1 = SolverBody1();
+
+			// Generate a new contact point at the average of all the active contacts
+			int32 NumActiveManifoldPoints = 0;
+			FSolverVec3 RelativeContactPosition0 = FSolverVec3(0);
+			FSolverVec3 RelativeContactPosition1 = FSolverVec3(0);
+			FSolverVec3 WorldContactNormal = FSolverVec3(0);
+			FSolverReal NetPushOutNormal = FSolverReal(0);
+			FSolverReal WorldContactVelocityTargetNormal = FSolverReal(0);
+			for (int32 PointIndex = 0; PointIndex < NumManifoldPoints(); ++PointIndex)
+			{
+				FPBDCollisionSolverManifoldPointSimd& SolverManifoldPoint = State.ManifoldPoints[PointIndex];
+				if (SolverManifoldPoint.ShouldSolveVelocity())
+				{
+					RelativeContactPosition0 += SolverManifoldPoint.WorldContact.RelativeContactPoints[0];
+					RelativeContactPosition1 += SolverManifoldPoint.WorldContact.RelativeContactPoints[1];
+					WorldContactVelocityTargetNormal += SolverManifoldPoint.WorldContact.ContactTargetVelocityNormal;
+					WorldContactNormal = SolverManifoldPoint.WorldContact.ContactNormal;	// Take last value - should all be similar
+					NetPushOutNormal += SolverManifoldPoint.NetPushOutNormal;
+					++NumActiveManifoldPoints;
+				}
+			}
+
+			// Solve for velocity at the new conatct point
+			// We only do this if we have multiple active contacts
+			// NOTE: this average contact isn't really correct, especially when all contacts do not equally
+			// contribute to the pushout (which is normal even for a simple box on a plane). E.g., the WorldContactVelocityTargetNormal
+			// and NetPushOutNormal will not be right. The goal here though is just to get close to the correct result in fewer iterations
+			// than we would have if we just solved the corners of the box in sequence.
+			if (NumActiveManifoldPoints > 1)
+			{
+				const FSolverReal InvCount = FSolverReal(1) / FSolverReal(NumActiveManifoldPoints);
+
+				FPBDCollisionSolverManifoldPointSimd AverageManifoldPoint;
+				AverageManifoldPoint.WorldContact.RelativeContactPoints[0] = RelativeContactPosition0 * InvCount;
+				AverageManifoldPoint.WorldContact.RelativeContactPoints[1] = RelativeContactPosition1 * InvCount;
+				AverageManifoldPoint.WorldContact.ContactNormal = WorldContactNormal;
+				AverageManifoldPoint.WorldContact.ContactTargetVelocityNormal = WorldContactVelocityTargetNormal * InvCount;
+			
+				// Total pushout (not average) which is correct if the average point is also the centroid but otherwise probably an overestimate.
+				// This is used to limit the possibly-attractive impulse that corrects implicit velocity errors from the PBD solve, but it
+				// which might cause stickiness if overestimated. Although it should be recitifed when we solve the normal contact point velocities
+				// so I don't think it matters...we could do better though if it is a problem
+				AverageManifoldPoint.NetPushOutNormal = NetPushOutNormal;
+
+				// Calculate the contact mass (and derived properties) for this point
+				FSolverVec3 WorldContactNormalAngular0 = FSolverVec3(0);
+				FSolverVec3 WorldContactNormalAngular1 = FSolverVec3(0);
+				FSolverReal ContactMassInvNormal = FSolverReal(0);
+				if (Body0.IsDynamic())
+				{
+					const FSolverVec3 R0xN = FSolverVec3::CrossProduct(RelativeContactPosition0, WorldContactNormal);
+					WorldContactNormalAngular0 = Body0.InvI() * R0xN;
+					ContactMassInvNormal += FSolverVec3::DotProduct(R0xN, WorldContactNormalAngular0) + Body0.InvM();
+				}
+				if (Body1.IsDynamic())
+				{
+					const FSolverVec3 R1xN = FSolverVec3::CrossProduct(RelativeContactPosition1, WorldContactNormal);
+					WorldContactNormalAngular1 = Body1.InvI() * R1xN;
+					ContactMassInvNormal += FSolverVec3::DotProduct(R1xN, WorldContactNormalAngular1) + Body1.InvM();
+				}
+				AverageManifoldPoint.ContactMassNormal = (ContactMassInvNormal > FSolverReal(UE_SMALL_NUMBER)) ? FSolverReal(1) / ContactMassInvNormal : FSolverReal(0);
+				AverageManifoldPoint.WorldContactNormalAngular0 = WorldContactNormalAngular0;
+				AverageManifoldPoint.WorldContactNormalAngular1 = WorldContactNormalAngular1;
+
+				// @todo(chaos): we don't use these - maybe do the calculation without an actual manifold point object...
+				AverageManifoldPoint.WorldContact.ContactTangentU = FSolverVec3(0);
+				AverageManifoldPoint.WorldContact.ContactTangentV = FSolverVec3(0);
+				AverageManifoldPoint.WorldContact.ContactDeltaNormal = FSolverReal(0);
+				AverageManifoldPoint.WorldContact.ContactDeltaTangentU = FSolverReal(0);
+				AverageManifoldPoint.WorldContact.ContactDeltaTangentV = FSolverReal(0);
+				AverageManifoldPoint.WorldContactTangentUAngular0 = FSolverVec3(0);
+				AverageManifoldPoint.WorldContactTangentVAngular0 = FSolverVec3(0);
+				AverageManifoldPoint.WorldContactTangentUAngular1 = FSolverVec3(0);
+				AverageManifoldPoint.WorldContactTangentVAngular1 = FSolverVec3(0);
+				AverageManifoldPoint.ContactMassTangentU = FSolverReal(0);
+				AverageManifoldPoint.ContactMassTangentV = FSolverReal(0);
+				AverageManifoldPoint.NetPushOutTangentU = FSolverReal(0);
+				AverageManifoldPoint.NetPushOutTangentV = FSolverReal(0);
+				AverageManifoldPoint.NetImpulseNormal = FSolverReal(0);
+				AverageManifoldPoint.NetImpulseTangentU = FSolverReal(0);
+				AverageManifoldPoint.NetImpulseTangentV = FSolverReal(0);
+				AverageManifoldPoint.StaticFrictionRatio = FSolverReal(0);
+
+				FSolverReal ContactVelocityDeltaNormal;
+				AverageManifoldPoint.CalculateContactVelocityErrorNormal(Body0, Body1, ContactVelocityDeltaNormal);
+
+				const FSolverReal MinImpulseNormal = FMath::Min(FSolverReal(0), -AverageManifoldPoint.NetPushOutNormal / Dt);
+
+				AverageManifoldPoint.ApplyVelocityCorrectionNormal(
+					State.Stiffness,
+					ContactVelocityDeltaNormal,
+					MinImpulseNormal,
+					Body0,
+					Body1);
+
+				// Now distribute the net impulse among the active points so we don't over-correct pushout from initial overlaps
+				for (int32 PointIndex = 0; PointIndex < NumManifoldPoints(); ++PointIndex)
+				{
+					FPBDCollisionSolverManifoldPointSimd& SolverManifoldPoint = State.ManifoldPoints[PointIndex];
+					if (SolverManifoldPoint.ShouldSolveVelocity())
+					{
+						SolverManifoldPoint.NetImpulseNormal += AverageManifoldPoint.NetImpulseNormal * InvCount;
+					}
+				}
+			}
+		}
+
+		//////////////////////////////////////////////////////////////////////////////////////////////////
+		//////////////////////////////////////////////////////////////////////////////////////////////////
+		//////////////////////////////////////////////////////////////////////////////////////////////////
+		//////////////////////////////////////////////////////////////////////////////////////////////////
+
+		void FPBDCollisionSolverHelperSimd::SolvePositionNoFriction(const TArrayView<FPBDCollisionSolverSimd>& CollisionSolvers, const FSolverReal Dt, const FSolverReal MaxPushOut)
+		{
+			for (FPBDCollisionSolverSimd& CollisionSolver : CollisionSolvers)
+			{
+				CollisionSolver.SolvePositionNoFriction(Dt, MaxPushOut);
+			}
+		}
+
+		void FPBDCollisionSolverHelperSimd::SolvePositionWithFriction(const TArrayView<FPBDCollisionSolverSimd>& CollisionSolvers, const FSolverReal Dt, const FSolverReal MaxPushOut)
+		{
+			for (FPBDCollisionSolverSimd& CollisionSolver : CollisionSolvers)
+			{
+				CollisionSolver.SolvePositionWithFriction(Dt, MaxPushOut);
+			}
+		}
+
+		void FPBDCollisionSolverHelperSimd::SolveVelocity(const TArrayView<FPBDCollisionSolverSimd>& CollisionSolvers, const FSolverReal Dt, const bool bApplyDynamicFriction)
+		{
+			for (FPBDCollisionSolverSimd& CollisionSolver : CollisionSolvers)
+			{
+				CollisionSolver.SolveVelocity(Dt, bApplyDynamicFriction);
+			}
+		}
+
+		void FPBDCollisionSolverHelperSimd::CheckISPC()
+		{
+		}
+
+	}	// namespace Private
+}	// namespace Chaos
\ No newline at end of file
diff --git a/Engine/Source/Runtime/Experimental/Chaos/Private/Chaos/Collision/SolverCollisionContainer.cpp b/Engine/Source/Runtime/Experimental/Chaos/Private/Chaos/Collision/SolverCollisionContainer.cpp
--- a/Engine/Source/Runtime/Experimental/Chaos/Private/Chaos/Collision/SolverCollisionContainer.cpp
+++ b/Engine/Source/Runtime/Experimental/Chaos/Private/Chaos/Collision/SolverCollisionContainer.cpp
@@ -1,612 +1,5 @@
 // Copyright Epic Games, Inc. All Rights Reserved.
 #include "Chaos/Collision/SolverCollisionContainer.h"
 
-#include "Chaos/Collision/CollisionApplyType.h"
-#include "Chaos/Collision/PBDCollisionConstraint.h"
-#include "Chaos/CollisionResolution.h"
-#include "Chaos/DebugDrawQueue.h"
-#include "Chaos/Evolution/SolverBodyContainer.h"
-#include "Chaos/Island/IslandManager.h"
-#include "Chaos/Particle/ParticleUtilities.h"
-#include "Chaos/PBDCollisionConstraints.h"
-#include "Chaos/PBDCollisionConstraintsContact.h"
-#include "Chaos/Utilities.h"
-
-// Private includes
-
-#include "ChaosLog.h"
-
-//PRAGMA_DISABLE_OPTIMIZATION
-
-namespace Chaos
-{
-	namespace CVars
-	{
-		extern bool bChaos_PBDCollisionSolver_Position_SolveEnabled;
-		extern bool bChaos_PBDCollisionSolver_Velocity_SolveEnabled;
-
-		// If one body is more than MassRatio1 times the mass of the other, adjust the solver stiffness when the lighter body is underneath.
-		// Solver stiffness will be equal to 1 when the mass ratio is MassRatio1.
-		// Solver stiffness will be equal to 0 when the mass ratio is MassRatio2.
-		FRealSingle Chaos_PBDCollisionSolver_AutoStiffness_MassRatio1 = 0;
-		FRealSingle Chaos_PBDCollisionSolver_AutoStiffness_MassRatio2 = 0;
-		FAutoConsoleVariableRef CVarChaosPBDCollisionSolverAutoStiffnessMassRatio1(TEXT("p.Chaos.PBDCollisionSolver.AutoStiffness.MassRatio1"), Chaos_PBDCollisionSolver_AutoStiffness_MassRatio1, TEXT(""));
-		FAutoConsoleVariableRef CVarChaosPBDCollisionSolverAutoStiffnessMassRatio2(TEXT("p.Chaos.PBDCollisionSolver.AutoStiffness.MassRatio2"), Chaos_PBDCollisionSolver_AutoStiffness_MassRatio2, TEXT(""));
-
-		// Jacobi solver stiffness
-		// @todo(chaos): to be tuned
-		FRealSingle Chaos_PBDCollisionSolver_JacobiStiffness = 0.5f;
-		FAutoConsoleVariableRef CVarChaosPBDCollisionSolverJacobiStiffness(TEXT("p.Chaos.PBDCollisionSolver.JacobiStiffness"), Chaos_PBDCollisionSolver_JacobiStiffness, TEXT(""));
-
-		// Jacobi position tolerance. Position corrections below this are zeroed.
-		// @todo(chaos): to be tuned
-		FRealSingle Chaos_PBDCollisionSolver_JacobiPositionTolerance = 1.e-6f;
-		FAutoConsoleVariableRef CVarChaosPBDCollisionSolverJacobiPositionTolerance(TEXT("p.Chaos.PBDCollisionSolver.JacobiPositionTolerance"), Chaos_PBDCollisionSolver_JacobiPositionTolerance, TEXT(""));
-		
-		// Jacobi rotation tolerance. Rotation corrections below this are zeroed.
-		// @todo(chaos): to be tuned
-		FRealSingle Chaos_PBDCollisionSolver_JacobiRotationTolerance = 1.e-8f;
-		FAutoConsoleVariableRef CVarChaosPBDCollisionSolverJacobiRotationTolerance(TEXT("p.Chaos.PBDCollisionSolver.JacobiRotationTolerance"), Chaos_PBDCollisionSolver_JacobiRotationTolerance, TEXT(""));
-	}
-
-
-	//////////////////////////////////////////////////////////////////////////////////////////////////
-	//////////////////////////////////////////////////////////////////////////////////////////////////
-	//////////////////////////////////////////////////////////////////////////////////////////////////
-	//////////////////////////////////////////////////////////////////////////////////////////////////
-
-	// Transform the Constraint's local-space data into world space for use by  the collision solver and also calculate tangents, errors, etc
-	FORCEINLINE_DEBUGGABLE void UpdateCollisionSolverContactPointFromConstraint(Private::FPBDCollisionSolver& Solver, const int32 SolverPointIndex, const FPBDCollisionConstraint* Constraint, const int32 ConstraintPointIndex, const FRealSingle Dt, const FSolverBody& Body0, const FSolverBody& Body1)
-	{
-		const FManifoldPoint& ManifoldPoint = Constraint->GetManifoldPoint(ConstraintPointIndex);
-
-		const FRealSingle Restitution = FRealSingle(Constraint->GetRestitution());
-		const FRealSingle RestitutionVelocityThreshold = FRealSingle(Constraint->GetRestitutionThreshold()) * Dt;
-
-		// World-space shape transforms. Manifold data is currently relative to these spaces
-		const FRigidTransform3& ShapeWorldTransform0 = Constraint->GetShapeWorldTransform0();
-		const FRigidTransform3& ShapeWorldTransform1 = Constraint->GetShapeWorldTransform1();
-
-		// World-space contact points on each shape
-		const FVec3 WorldContact0 = ShapeWorldTransform0.TransformPositionNoScale(FVec3(ManifoldPoint.ContactPoint.ShapeContactPoints[0]));
-		const FVec3 WorldContact1 = ShapeWorldTransform1.TransformPositionNoScale(FVec3(ManifoldPoint.ContactPoint.ShapeContactPoints[1]));
-		const FVec3 WorldContact = FReal(0.5) * (WorldContact0 + WorldContact1);
-		const FVec3f WorldRelativeContact0 = FVec3f(WorldContact - Body0.P());
-		const FVec3f WorldRelativeContact1 = FVec3f(WorldContact - Body1.P());
-
-		// World-space normal
-		const FVec3f WorldContactNormal = FVec3f(ShapeWorldTransform1.TransformVectorNoScale(FVec3(ManifoldPoint.ContactPoint.ShapeContactNormal)));
-
-		// World-space tangents
-		FVec3f WorldContactTangentU = FVec3f::CrossProduct(FVec3f(0, 1, 0), WorldContactNormal);
-		if (!WorldContactTangentU.Normalize(FRealSingle(UE_KINDA_SMALL_NUMBER)))
-		{
-			WorldContactTangentU = FVec3f::CrossProduct(FVec3f(1, 0, 0), WorldContactNormal);
-			WorldContactTangentU = WorldContactTangentU.GetUnsafeNormal();
-		}
-		const FVec3f WorldContactTangentV = FVec3f::CrossProduct(WorldContactNormal, WorldContactTangentU);
-
-		// Calculate contact velocity if we will need it below (restitution and/or frist-contact friction)
-		const bool bNeedContactVelocity = (!ManifoldPoint.Flags.bHasStaticFrictionAnchor) || (Restitution > FRealSingle(0));
-		FVec3f ContactVel = FVec3(0);
-		if (bNeedContactVelocity)
-		{
-			const FVec3f ContactVel0 = Body0.V() + FVec3f::CrossProduct(Body0.W(), WorldRelativeContact0);
-			const FVec3f ContactVel1 = Body1.V() + FVec3f::CrossProduct(Body1.W(), WorldRelativeContact1);
-			ContactVel = ContactVel0 - ContactVel1;
-		}
-
-		// If we have contact data from a previous tick, use it to calculate the lateral position delta we need
-		// to apply to move the contacts back to their original relative locations (i.e., to enforce static friction).
-		// Otherwise, estimate the friction correction from the contact velocity.
-		// NOTE: quadratic shapes use the velocity-based path most of the time, unless the relative motion is very small.
-		FVec3f WorldFrictionDelta = FVec3f(0);
-		if (ManifoldPoint.Flags.bHasStaticFrictionAnchor)
-		{
-			const FVec3f FrictionDelta0 = ShapeWorldTransform0.TransformPositionNoScale(FVec3(ManifoldPoint.ShapeAnchorPoints[0])) - WorldContact0;
-			const FVec3f FrictionDelta1 = ShapeWorldTransform1.TransformPositionNoScale(FVec3(ManifoldPoint.ShapeAnchorPoints[1])) - WorldContact1;
-			WorldFrictionDelta = FrictionDelta0 - FrictionDelta1;
-		}
-		else
-		{
-			// @todo(chaos): consider adding a multiplier to the initial contact friction
-			WorldFrictionDelta = ContactVel * Dt;
-		}
-
-		// The contact point error we are trying to correct in this solver
-		const FRealSingle TargetPhi = ManifoldPoint.TargetPhi;
-		const FVec3f WorldContactDelta = FVec3f(WorldContact0 - WorldContact1);
-		const FRealSingle WorldContactDeltaNormal = FVec3f::DotProduct(WorldContactDelta, WorldContactNormal) - TargetPhi;
-		const FRealSingle WorldContactDeltaTangentU = FVec3f::DotProduct(WorldContactDelta + WorldFrictionDelta, WorldContactTangentU);
-		const FRealSingle WorldContactDeltaTangentV = FVec3f::DotProduct(WorldContactDelta + WorldFrictionDelta, WorldContactTangentV);
-
-		// The target contact velocity, taking restitution into account
-		FRealSingle WorldContactTargetVelocityNormal = FRealSingle(0);
-		if (Restitution > FRealSingle(0))
-		{
-			const FRealSingle ContactVelocityNormal = FVec3f::DotProduct(ContactVel, WorldContactNormal);
-			if (ContactVelocityNormal < -RestitutionVelocityThreshold)
-			{
-				WorldContactTargetVelocityNormal = -Restitution * ContactVelocityNormal;
-			}
-		}
-
-		FWorldContactPoint& WorldContactPoint = Solver.GetWorldContactPoint(SolverPointIndex);
-		WorldContactPoint.RelativeContactPoints[0] = WorldRelativeContact0;
-		WorldContactPoint.RelativeContactPoints[1] = WorldRelativeContact1;
-		WorldContactPoint.ContactNormal = WorldContactNormal;
-		WorldContactPoint.ContactTangentU = WorldContactTangentU;
-		WorldContactPoint.ContactTangentV = WorldContactTangentV;
-		WorldContactPoint.ContactDeltaNormal = WorldContactDeltaNormal;
-		WorldContactPoint.ContactDeltaTangentU = WorldContactDeltaTangentU;
-		WorldContactPoint.ContactDeltaTangentV = WorldContactDeltaTangentV;
-		WorldContactPoint.ContactTargetVelocityNormal = WorldContactTargetVelocityNormal;
-	}
-
-	FORCEINLINE_DEBUGGABLE void UpdateCollisionSolverManifoldFromConstraint(Private::FPBDCollisionSolver& Solver, const FPBDCollisionConstraint* Constraint, const FSolverReal Dt, const int32 ConstraintPointBeginIndex, const int32 ConstraintPointEndIndex)
-	{
-		const FSolverBody& Body0 = Solver.SolverBody0().SolverBody();
-		const FSolverBody& Body1 = Solver.SolverBody1().SolverBody();
-
-		// Only calculate state for newly added contacts. Normally this is all of them, but maybe not if incremental collision is used by RBAN.
-		// Also we only add active points to the solver's manifold points list
-		for (int32 ConstraintManifoldPointIndex = ConstraintPointBeginIndex; ConstraintManifoldPointIndex < ConstraintPointEndIndex; ++ConstraintManifoldPointIndex)
-		{
-			if (!Constraint->GetManifoldPoint(ConstraintManifoldPointIndex).Flags.bDisabled)
-			{
-				const int32 SolverManifoldPointIndex = Solver.AddManifoldPoint();
-
-				// Transform the constraint contact data into world space for use by the solver
-				// We build this data directly into the solver's world-space contact data which looks a bit odd with "Init" called after but there you go
-				UpdateCollisionSolverContactPointFromConstraint(Solver, SolverManifoldPointIndex, Constraint, ConstraintManifoldPointIndex, Dt, Body0, Body1);
-
-				// Prepare the solver (set constraint-space mass etc)
-				Solver.FinalizeManifoldPoint(SolverManifoldPointIndex, Dt);
-			}
-		}
-	}
-
-	FORCEINLINE_DEBUGGABLE void UpdateCollisionSolverFromConstraint(Private::FPBDCollisionSolver& Solver, const FPBDCollisionConstraint* Constraint, const FSolverReal Dt, const FPBDCollisionSolverSettings& SolverSettings, bool& bOutPerIterationCollision)
-	{
-		// Friction values. Static and Dynamic friction are applied in the position solve for most shapes.
-		// We can also run in a mode without static friction at all. This is faster but stacking is not possible.
-		const FSolverReal StaticFriction = FSolverReal(Constraint->GetStaticFriction());
-		const FSolverReal DynamicFriction = FSolverReal(Constraint->GetDynamicFriction());
-		FSolverReal PositionStaticFriction = FSolverReal(0);
-		FSolverReal PositionDynamicFriction = FSolverReal(0);
-		FSolverReal VelocityDynamicFriction = FSolverReal(0);
-		if (SolverSettings.NumPositionFrictionIterations > 0)
-		{
-			PositionStaticFriction = StaticFriction;
-			if (!Constraint->HasQuadraticShape())
-			{
-				PositionDynamicFriction = DynamicFriction;
-			}
-			else
-			{
-				// Quadratic shapes don't use PBD dynamic friction - it has issues at slow speeds where the WxR is
-				// less than the position tolerance for friction point matching
-				// @todo(chaos): fix PBD dynamic friction on quadratic shapes
-				VelocityDynamicFriction = DynamicFriction;
-			}
-		}
-		else
-		{
-			VelocityDynamicFriction = DynamicFriction;
-		}
-
-		Solver.SetFriction(PositionStaticFriction, PositionDynamicFriction, VelocityDynamicFriction);
-
-		const FReal SolverStiffness = Constraint->GetStiffness();
-		Solver.SetStiffness(FSolverReal(SolverStiffness));
-
-		Solver.SolverBody0().SetInvMScale(Constraint->GetInvMassScale0());
-		Solver.SolverBody0().SetInvIScale(Constraint->GetInvInertiaScale0());
-		Solver.SolverBody0().SetShockPropagationScale(FReal(1));
-		Solver.SolverBody1().SetInvMScale(Constraint->GetInvMassScale1());
-		Solver.SolverBody1().SetInvIScale(Constraint->GetInvInertiaScale1());
-		Solver.SolverBody1().SetShockPropagationScale(FReal(1));
-
-		bOutPerIterationCollision = (!Constraint->GetUseManifold() || Constraint->GetUseIncrementalCollisionDetection());
-
-		UpdateCollisionSolverManifoldFromConstraint(Solver, Constraint, Dt, 0, Constraint->NumManifoldPoints());
-	}
-
-	FORCEINLINE_DEBUGGABLE void UpdateCollisionConstraintFromSolver(FPBDCollisionConstraint* Constraint, const Private::FPBDCollisionSolver& Solver, const FSolverReal Dt)
-	{
-		Constraint->ResetSolverResults();
-
-		// NOTE: We only put the non-pruned manifold points into the solver so the ManifoldPointIndex and
-		// SolverManifoldPointIndex do not necessarily match. See GatherManifoldPoints
-		int32 SolverManifoldPointIndex = 0;
-		for (int32 ManifoldPointIndex = 0; ManifoldPointIndex < Constraint->NumManifoldPoints(); ++ManifoldPointIndex)
-		{
-			FSolverVec3 NetPushOut = FSolverVec3(0);
-			FSolverVec3 NetImpulse = FSolverVec3(0);
-			FSolverReal StaticFrictionRatio = FSolverReal(0);
-
-			const Private::FPBDCollisionSolverManifoldPoint& SolverManifoldPoint = Solver.GetManifoldPoint(SolverManifoldPointIndex);
-			if (!Constraint->GetManifoldPoint(ManifoldPointIndex).Flags.bDisabled)
-			{
-				NetPushOut =
-					SolverManifoldPoint.NetPushOutNormal * SolverManifoldPoint.WorldContact.ContactNormal +
-					SolverManifoldPoint.NetPushOutTangentU * SolverManifoldPoint.WorldContact.ContactTangentU +
-					SolverManifoldPoint.NetPushOutTangentV * SolverManifoldPoint.WorldContact.ContactTangentV;
-
-				NetImpulse =
-					SolverManifoldPoint.NetImpulseNormal * SolverManifoldPoint.WorldContact.ContactNormal +
-					SolverManifoldPoint.NetImpulseTangentU * SolverManifoldPoint.WorldContact.ContactTangentU +
-					SolverManifoldPoint.NetImpulseTangentV * SolverManifoldPoint.WorldContact.ContactTangentV;
-
-				StaticFrictionRatio = SolverManifoldPoint.StaticFrictionRatio;
-
-				++SolverManifoldPointIndex;
-			}
-
-			// NOTE: We call this even for points we did not run the solver for (but with zero results)
-			Constraint->SetSolverResults(ManifoldPointIndex,
-				NetPushOut,
-				NetImpulse,
-				StaticFrictionRatio,
-				Dt);
-		}
-
-		Constraint->EndTick();
-	}
-
-	//////////////////////////////////////////////////////////////////////////////////////////////////
-	//////////////////////////////////////////////////////////////////////////////////////////////////
-	//////////////////////////////////////////////////////////////////////////////////////////////////
-	//////////////////////////////////////////////////////////////////////////////////////////////////
-
-
-	FPBDCollisionContainerSolver::FPBDCollisionContainerSolver(const FPBDCollisionConstraints& InConstraintContainer, const int32 InPriority)
-		: FConstraintContainerSolver(InPriority)
-		, ConstraintContainer(InConstraintContainer)
-		, bPerIterationCollisionDetection(false)
-	{
-#if !UE_BUILD_SHIPPING && !UE_BUILD_TEST && INTEL_ISPC
-		Private::FPBDCollisionSolverHelper::CheckISPC();
-#endif
-	}
-
-	FPBDCollisionContainerSolver::~FPBDCollisionContainerSolver()
-	{
-	}
-
-	void FPBDCollisionContainerSolver::Reset(const int32 MaxCollisions)
-	{
-		// A over-allocation policy to avoid reallocation every frame in the common case where a pile of objects is dropped
-		// and the number of contacts increases every tick.
-		if (MaxCollisions > CollisionSolvers.Max())
-		{
-			CollisionSolvers.Reset((5 * MaxCollisions) / 4);	// +25%
-		}
-
-		CollisionSolvers.Reset(MaxCollisions);
-		CollisionConstraints.Reset(CollisionSolvers.Max());
-		bCollisionConstraintPerIterationCollisionDetection.Reset(CollisionSolvers.Max());
-		CollisionSolverManifoldPoints.Reset(CollisionSolvers.Max() * Private::FPBDCollisionSolver::MaxPointsPerConstraint);
-
-		// Just set the array size for these right away - all data will be initialized later and there is no constructor
-		// We use CollisionConstraints.Num() to track how many constraints we currently have
-		CollisionSolvers.SetNumUninitialized(MaxCollisions);
-		bCollisionConstraintPerIterationCollisionDetection.SetNumUninitialized(MaxCollisions);
-	}
-
-	void FPBDCollisionContainerSolver::AddConstraints()
-	{
-		Reset(ConstraintContainer.NumConstraints());
-
-		for (FPBDCollisionConstraintHandle* ConstraintHandle : ConstraintContainer.GetConstraintHandles())
-		{
-			FPBDCollisionConstraint& Constraint = ConstraintHandle->GetContact();
-
-			AddConstraint(Constraint);
-		}
-	}
-
-	void FPBDCollisionContainerSolver::AddConstraints(const TArrayView<Private::FPBDIslandConstraint>& IslandConstraints)
-	{
-		for (Private::FPBDIslandConstraint& IslandConstraint : IslandConstraints)
-		{
-			// We will only ever be given constraints from our container (asserts in non-shipping)
-			FPBDCollisionConstraint& Constraint = IslandConstraint.GetConstraint()->AsUnsafe<FPBDCollisionConstraintHandle>()->GetContact();
-
-			AddConstraint(Constraint);
-		}
-	}
-
-	void FPBDCollisionContainerSolver::AddConstraint(FPBDCollisionConstraint& Constraint)
-	{
-		// NOTE: No need to add to CollisionSolvers or bCollisionConstraintPerIterationCollisionDetection
-		const int32 Index = CollisionConstraints.Add(&Constraint);
-
-		// Allocate the manifold points for this constraints solver
-		// NOTE: we don't know how many points we will create if collision detection is deferred or incrememntal, so just allocate space for the max allowed
-		// @todo(chaos): see if we can do better - we don't want to pay for this logic when it is only used rarely
-		const bool bDeferredCollisionDetection = ConstraintContainer.GetDetectorSettings().bDeferNarrowPhase;
-		const int32 ManifoldPointMax = (bDeferredCollisionDetection || Constraint.GetUseIncrementalManifold()) ? Private::FPBDCollisionSolver::MaxPointsPerConstraint : Constraint.NumManifoldPoints();
-		if (ManifoldPointMax > 0)
-		{
-			const int32 ManifoldPointStart = CollisionSolverManifoldPoints.Num();
-			check(CollisionSolverManifoldPoints.Max() >= ManifoldPointStart + ManifoldPointMax);	// No resize following reset
-			CollisionSolverManifoldPoints.SetNum(ManifoldPointStart + ManifoldPointMax);
-
-			CollisionSolvers[Index].SetManifoldPointsBuffer(&CollisionSolverManifoldPoints[ManifoldPointStart], ManifoldPointMax);
-		}
-		else
-		{
-			// Ideally we wouldn't have any active constraint with no manifold points, but this can
-			// currently happen for CCD objects and incremental collisions
-			CollisionSolvers[Index].Reset();
-		}
-	}
-
-
-	void FPBDCollisionContainerSolver::AddBodies(FSolverBodyContainer& SolverBodyContainer)
-	{
-		for (int32 SolverIndex = 0, SolverEndIndex = NumSolvers(); SolverIndex < SolverEndIndex; ++SolverIndex)
-		{
-			Private::FPBDCollisionSolver& CollisionSolver = CollisionSolvers[SolverIndex];
-			FPBDCollisionConstraint* Constraint = CollisionConstraints[SolverIndex];
-			check(Constraint != nullptr);
-
-			// Find the solver bodies for the particles we constrain. This will add them to the container
-			// if they aren't there already, and ensure that they are populated with the latest data.
-			FSolverBody* Body0 = SolverBodyContainer.FindOrAdd(Constraint->GetParticle0());
-			FSolverBody* Body1 = SolverBodyContainer.FindOrAdd(Constraint->GetParticle1());
-			check(Body0 != nullptr);
-			check(Body1 != nullptr);
-
-			CollisionSolver.SetSolverBodies(*Body0, *Body1);
-		}
-	}
-
-	void FPBDCollisionContainerSolver::GatherInput(const FReal Dt)
-	{
-		GatherInput(Dt, 0, NumSolvers());
-	}
-
-	void FPBDCollisionContainerSolver::CachePrefetchSolver(const int32 ConstraintIndex) const
-	{
-		if (ConstraintIndex < NumSolvers())
-		{
-			const int32 NumManifoldPoints = CollisionSolvers[ConstraintIndex].NumManifoldPoints();
-			if (NumManifoldPoints > 0)
-			{
-				FPlatformMisc::PrefetchBlock(&CollisionSolvers[ConstraintIndex].GetManifoldPoint(0), NumManifoldPoints * sizeof(Private::FPBDCollisionSolverManifoldPoint));
-			}
-			FPlatformMisc::PrefetchBlock(&CollisionSolvers[ConstraintIndex], sizeof(FPBDCollisionContainerSolver));
-		}
-	}
-
-	void FPBDCollisionContainerSolver::GatherInput(const FReal InDt, const int32 BeginIndex, const int32 EndIndex)
-	{
-		// NOTE: may be called in parallel. Should not change the container or any elements outside of [BeginIndex, EndIndex)
-
-		check(BeginIndex >= 0);
-		check(EndIndex <= NumSolvers());
-
-		const FSolverReal Dt = FSolverReal(InDt);
-
-		const int32 PrefetchCount = 4;
-		for (int32 PrefetchIndex = 0; PrefetchIndex < PrefetchCount; ++PrefetchIndex)
-		{
-			CachePrefetchSolver(PrefetchIndex);
-		}
-
-		bool bWantPerIterationCollisionDetection = false;
-		for (int32 ConstraintIndex = BeginIndex; ConstraintIndex < EndIndex; ++ConstraintIndex)
-		{
-			Private::FPBDCollisionSolver& CollisionSolver = CollisionSolvers[ConstraintIndex];
-			FPBDCollisionConstraint* Constraint = CollisionConstraints[ConstraintIndex];
-			bool& bPerIterationCollision = bCollisionConstraintPerIterationCollisionDetection[ConstraintIndex];
-
-			UpdateCollisionSolverFromConstraint(CollisionSolver, Constraint, Dt, ConstraintContainer.GetSolverSettings(), bPerIterationCollision);
-
-			// We need to run collision every iteration if we are not using manifolds, or are using incremental manifolds
-			bWantPerIterationCollisionDetection = bWantPerIterationCollisionDetection || bPerIterationCollision;
-
-			CachePrefetchSolver(ConstraintIndex + PrefetchCount);
-		}
-
-		if (bWantPerIterationCollisionDetection)
-		{
-			// We should lock here? We only ever set to true or do nothing so I think it doesn't matter if this happens on multiple threads...
-			bPerIterationCollisionDetection = true;
-		}
-	}
-
-	void FPBDCollisionContainerSolver::ScatterOutput(const FReal Dt)
-	{
-		ScatterOutput(Dt, 0, NumSolvers());
-	}
-
-	void FPBDCollisionContainerSolver::ScatterOutput(const FReal InDt, const int32 BeginIndex, const int32 EndIndex)
-	{
-		// NOTE: may be called in parallel. Should not change the container or any elements outside of [BeginIndex, EndIndex)
-
-		check(BeginIndex >= 0);
-		check(EndIndex <= NumSolvers());
-
-		const FSolverReal Dt = FSolverReal(InDt);
-
-		for (int32 ConstraintIndex = BeginIndex; ConstraintIndex < EndIndex; ++ConstraintIndex)
-		{
-			Private::FPBDCollisionSolver& CollisionSolver = CollisionSolvers[ConstraintIndex];
-			FPBDCollisionConstraint* Constraint = CollisionConstraints[ConstraintIndex];
-
-			UpdateCollisionConstraintFromSolver(Constraint, CollisionSolver, Dt);
-
-			// Reset the collision solver here as the pointers will be invalid on the next tick
-			// @todo(chaos): Pointers are always reinitalized before use next tick so this isn't strictly necessary
-			CollisionSolver.Reset();
-		}
-	}
-
-	void FPBDCollisionContainerSolver::ApplyPositionConstraints(const FReal Dt, const int32 It, const int32 NumIts)
-	{
-		SolvePositionImpl(Dt, It, NumIts, 0, NumSolvers(), ConstraintContainer.GetSolverSettings());
-	}
-
-	void FPBDCollisionContainerSolver::ApplyVelocityConstraints(const FReal Dt, const int32 It, const int32 NumIts)
-	{
-		SolveVelocityImpl(Dt, It, NumIts, 0, NumSolvers(), ConstraintContainer.GetSolverSettings());
-	}
-
-	void FPBDCollisionContainerSolver::ApplyProjectionConstraints(const FReal Dt, const int32 It, const int32 NumIts)
-	{
-		// Not supported for collisions
-	}
-
-	void FPBDCollisionContainerSolver::UpdatePositionShockPropagation(const FReal Dt, const int32 It, const int32 NumIts, const int32 BeginIndex, const int32 EndIndex, const FPBDCollisionSolverSettings& SolverSettings)
-	{
-		// If this is the first shock propagation iteration, enable it on each solver
-		const bool bEnableShockPropagation = (It == NumIts - SolverSettings.NumPositionShockPropagationIterations);
-		if (bEnableShockPropagation)
-		{
-			for (int32 SolverIndex = BeginIndex; SolverIndex < EndIndex; ++SolverIndex)
-			{
-				CollisionSolvers[SolverIndex].EnablePositionShockPropagation();
-			}
-		}
-	}
-
-	void FPBDCollisionContainerSolver::UpdateVelocityShockPropagation(const FReal Dt, const int32 It, const int32 NumIts, const int32 BeginIndex, const int32 EndIndex, const FPBDCollisionSolverSettings& SolverSettings)
-	{
-		// Set/reset the shock propagation based on current iteration. The position solve may
-		// have left the bodies with a mass scale and we want to change or reset it.
-		const bool bEnableShockPropagation = (It == NumIts - SolverSettings.NumVelocityShockPropagationIterations);
-		if (bEnableShockPropagation)
-		{
-			for (int32 SolverIndex = BeginIndex; SolverIndex < EndIndex; ++SolverIndex)
-			{
-				CollisionSolvers[SolverIndex].EnableVelocityShockPropagation();
-			}
-		}
-		else if (It == 0)
-		{
-			for (int32 SolverIndex = BeginIndex; SolverIndex < EndIndex; ++SolverIndex)
-			{
-				CollisionSolvers[SolverIndex].DisableShockPropagation();
-			}
-		}
-	}
-
-	bool FPBDCollisionContainerSolver::SolvePositionImpl(const FReal InDt, const int32 It, const int32 NumIts, const int32 BeginIndex, const int32 EndIndex, const FPBDCollisionSolverSettings& SolverSettings)
-	{
-		SCOPE_CYCLE_COUNTER(STAT_Collisions_Apply);
-		if (!CVars::bChaos_PBDCollisionSolver_Position_SolveEnabled)
-		{
-			return false;
-		}
-
-		if (EndIndex <= BeginIndex)
-		{
-			return false;
-		}
-
-		UpdatePositionShockPropagation(InDt, It, NumIts, BeginIndex, EndIndex, SolverSettings);
-
-		// We run collision detection here under two conditions (normally it is run after Integration and before the constraint solver phase):
-		// 1) When deferring collision detection until the solver phase for better joint-collision behaviour (RBAN). In this case, we only do this on the first iteration.
-		// 2) When using no manifolds or incremental manifolds, where we may add/replace manifold points every iteration.
-		const bool bRunDeferredCollisionDetection = (It == 0) && ConstraintContainer.GetDetectorSettings().bDeferNarrowPhase;
-		if (bRunDeferredCollisionDetection || bPerIterationCollisionDetection)
-		{
-			UpdateCollisions(InDt, BeginIndex, EndIndex);
-		}
-
-		// Only apply friction for the last few (tunable) iterations
-		// Adjust max pushout to attempt to make it iteration count independent
-		const FSolverReal Dt = FSolverReal(InDt);
-		const bool bApplyStaticFriction = (It >= (NumIts - SolverSettings.NumPositionFrictionIterations));
-		const FSolverReal MaxPushOut = (SolverSettings.MaxPushOutVelocity > 0) ? (FSolverReal(SolverSettings.MaxPushOutVelocity) * Dt) / FSolverReal(NumIts) : FSolverReal(0);
-
-		// Apply the position correction
-		if (bApplyStaticFriction)
-		{
-			Private::FPBDCollisionSolverHelper::SolvePositionWithFriction(MakeArrayView(&CollisionSolvers[BeginIndex], EndIndex - BeginIndex), Dt, MaxPushOut);
-		}
-		else
-		{
-			Private::FPBDCollisionSolverHelper::SolvePositionNoFriction(MakeArrayView(&CollisionSolvers[BeginIndex], EndIndex - BeginIndex), Dt, MaxPushOut);
-		}
-
-		return true;
-	}
-
-	bool FPBDCollisionContainerSolver::SolveVelocityImpl(const FReal InDt, const int32 It, const int32 NumIts, const int32 BeginIndex, const int32 EndIndex, const FPBDCollisionSolverSettings& SolverSettings)
-	{
-		SCOPE_CYCLE_COUNTER(STAT_Collisions_ApplyPushOut);
-		if (!CVars::bChaos_PBDCollisionSolver_Velocity_SolveEnabled)
-		{
-			return false;
-		}
-
-		if (EndIndex <= BeginIndex)
-		{
-			return false;
-		}
-
-		UpdateVelocityShockPropagation(InDt, It, NumIts, BeginIndex, EndIndex, SolverSettings);
-
-		const FSolverReal Dt = FSolverReal(InDt);
-		const bool bApplyDynamicFriction = (It >= NumIts - SolverSettings.NumVelocityFrictionIterations);
-
-		Private::FPBDCollisionSolverHelper::SolveVelocity(MakeArrayView(&CollisionSolvers[BeginIndex], EndIndex - BeginIndex), Dt, bApplyDynamicFriction);
-
-		return true;
-	}
-
-	void FPBDCollisionContainerSolver::UpdateCollisions(const FReal InDt, const int32 BeginIndex, const int32 EndIndex)
-	{
-		const FSolverReal Dt = FSolverReal(InDt);
-		const bool bDeferredCollisionDetection = ConstraintContainer.GetDetectorSettings().bDeferNarrowPhase;
-
-		bool bNeedsAnotherIteration = false;
-		for (int32 SolverIndex = BeginIndex; SolverIndex < EndIndex; ++SolverIndex)
-		{
-			if (bDeferredCollisionDetection || bCollisionConstraintPerIterationCollisionDetection[SolverIndex])
-			{
-				Private::FPBDCollisionSolver& CollisionSolver = CollisionSolvers[SolverIndex];
-				FPBDCollisionConstraint* Constraint = CollisionConstraints[SolverIndex];
-
-				// Run collision detection at the current transforms including any correction from previous iterations
-				const FSolverBody& Body0 = CollisionSolver.SolverBody0().SolverBody();
-				const FSolverBody& Body1 = CollisionSolver.SolverBody1().SolverBody();
-				const FRigidTransform3 CorrectedActorWorldTransform0 = FRigidTransform3(Body0.CorrectedActorP(), Body0.CorrectedActorQ());
-				const FRigidTransform3 CorrectedActorWorldTransform1 = FRigidTransform3(Body1.CorrectedActorP(), Body1.CorrectedActorQ());
-				const FRigidTransform3 CorrectedShapeWorldTransform0 = Constraint->GetShapeRelativeTransform0() * CorrectedActorWorldTransform0;
-				const FRigidTransform3 CorrectedShapeWorldTransform1 = Constraint->GetShapeRelativeTransform1() * CorrectedActorWorldTransform1;
-
-				// @todo(chaos): this is ugly - pass these to the required functions instead and remove from the constraint class
-				// This is now only needed for LevelSet collision (see UpdateLevelsetLevelsetConstraint)
-				Constraint->SetSolverBodies(&Body0, &Body1);
-
-				// Reset the manifold if we are not using manifolds (we just use the first manifold point)
-				if (!Constraint->GetUseManifold())
-				{
-					Constraint->ResetActiveManifoldContacts();
-					CollisionSolver.ResetManifold();
-				}
-
-				// We need to know how many points were added to the manifold
-				const int32 BeginPointIndex = Constraint->NumManifoldPoints();
-
-				// NOTE: We deliberately have not updated the ShapwWorldTranforms on the constraint. If we did that, we would calculate 
-				// errors incorrectly in UpdateManifoldPoints, because the solver assumes nothing has been moved as we iterate (we accumulate 
-				// corrections that will be applied later.)
-				Constraint->ResetPhi(Constraint->GetCullDistance());
-				Collisions::UpdateConstraint(*Constraint, CorrectedShapeWorldTransform0, CorrectedShapeWorldTransform1, Dt);
-
-				// Update the manifold based on the new or updated contacts
-				UpdateCollisionSolverManifoldFromConstraint(CollisionSolver, Constraint, Dt, BeginPointIndex, Constraint->NumManifoldPoints());
-
-				Constraint->SetSolverBodies(nullptr, nullptr);
-			}
-		}
-	}
-
-}
+// @todo(chaos): remove this file
+// Moved to PBDCollisionContainerSolver.cpp
diff --git a/Engine/Source/Runtime/Experimental/Chaos/Private/Chaos/PBDCollisionConstraints.cpp b/Engine/Source/Runtime/Experimental/Chaos/Private/Chaos/PBDCollisionConstraints.cpp
--- a/Engine/Source/Runtime/Experimental/Chaos/Private/Chaos/PBDCollisionConstraints.cpp
+++ b/Engine/Source/Runtime/Experimental/Chaos/Private/Chaos/PBDCollisionConstraints.cpp
@@ -8,9 +8,9 @@
 #include "Chaos/PBDCollisionConstraintsContact.h"
 #include "Chaos/CollisionResolution.h"
 #include "Chaos/Collision/CollisionPruning.h"
-#include "Chaos/Collision/PBDCollisionSolver.h"
+#include "Chaos/Collision/PBDCollisionContainerSolver.h"
 #include "Chaos/Collision/PBDCollisionContainerSolverJacobi.h"
-#include "Chaos/Collision/SolverCollisionContainer.h"
+#include "Chaos/Collision/PBDCollisionContainerSolverSimd.h"
 #include "Chaos/Defines.h"
 #include "Chaos/Evolution/SolverBodyContainer.h"
 #include "Chaos/GeometryQueries.h"
@@ -135,19 +135,24 @@ namespace Chaos
 
 	TUniquePtr<FConstraintContainerSolver> FPBDCollisionConstraints::CreateSceneSolver(const int32 Priority)
 	{
-		if (CollisionSolverType == Private::ECollisionSolverType::PartialJacobi)
-		{
-			return MakeUnique<FPBDCollisionContainerSolverJacobi>(*this, Priority);
-		}
-		else
-		{
-			return MakeUnique<FPBDCollisionContainerSolver>(*this, Priority);
-		}
+		// RBAN always uses Gauss Seidel solver for now
+		return MakeUnique<FPBDCollisionContainerSolver>(*this, Priority);
 	}
 
 	TUniquePtr<FConstraintContainerSolver> FPBDCollisionConstraints::CreateGroupSolver(const int32 Priority)
 	{
-		return CreateSceneSolver(Priority);
+		switch (CollisionSolverType)
+		{
+		case Private::ECollisionSolverType::GaussSeidel:
+			return MakeUnique<FPBDCollisionContainerSolver>(*this, Priority);
+		case Private::ECollisionSolverType::GaussSeidelSimd:
+			return MakeUnique<Private::FPBDCollisionContainerSolverSimd>(*this, Priority);
+		case Private::ECollisionSolverType::PartialJacobi:
+			return MakeUnique<Private::FPBDCollisionContainerSolverJacobi>(*this, Priority);
+		}
+
+		check(false);
+		return nullptr;
 	}
 
 	void FPBDCollisionConstraints::DisableHandles()
diff --git a/Engine/Source/Runtime/Experimental/Chaos/Private/Chaos/PBDRigidsEvolutionGBF.cpp b/Engine/Source/Runtime/Experimental/Chaos/Private/Chaos/PBDRigidsEvolutionGBF.cpp
--- a/Engine/Source/Runtime/Experimental/Chaos/Private/Chaos/PBDRigidsEvolutionGBF.cpp
+++ b/Engine/Source/Runtime/Experimental/Chaos/Private/Chaos/PBDRigidsEvolutionGBF.cpp
@@ -59,9 +59,9 @@ namespace Chaos
 		bool bChaosCollisionCCDUseTightBoundingBox = true;
 		FAutoConsoleVariableRef  CVarChaosCollisionCCDUseTightBoundingBox(TEXT("p.Chaos.Collision.CCD.UseTightBoundingBox"), bChaosCollisionCCDUseTightBoundingBox , TEXT(""));
 
-		// Collision Solver Type: Full Gaus Seidel or Partial Jacobi (see ECollisionSolverType)
+		// Collision Solver Type (see ECollisionSolverType)
 		int32 ChaosSolverCollisionSolverType = -1;
-		FAutoConsoleVariableRef CVarChaosSolverCollisionSolverType(TEXT("p.Chaos.Solver.Collision.SolverType"), ChaosSolverCollisionSolverType, TEXT("-1: Use default (Full Gauss Seidel); 0: Full Gauss Seidel; 1: Partial Jacobi"));
+		FAutoConsoleVariableRef CVarChaosSolverCollisionSolverType(TEXT("p.Chaos.Solver.Collision.SolverType"), ChaosSolverCollisionSolverType, TEXT("-1: Use default (Gauss Seidel); 0: Gauss Seidel; 1: Gauss Seidel SOA 2: Partial Jacobi"));
 
 		int32 ChaosSolverCollisionPriority = 0;
 		FAutoConsoleVariableRef CVarChaosSolverCollisionPriority(TEXT("p.Chaos.Solver.Collision.Priority"), ChaosSolverCollisionPriority, TEXT("Set constraint priority. Larger values are evaluated later [def:0]"));
@@ -310,7 +310,17 @@ void FPBDRigidsEvolutionGBF::UpdateCollisionSolverType()
 	// This is not intended to be performant - it is to allow switching for behaviour and performance comparisons only, without restarting the game.
 	if (ChaosSolverCollisionSolverType >= 0)
 	{
-		const Private::ECollisionSolverType CollisionSolverType = (ChaosSolverCollisionSolverType != 0) ? Private::ECollisionSolverType::PartialJacobi : Private::ECollisionSolverType::GaussSeidel;
+		// Map the cvar to an enum value. Invalid values are assumed to mean Gauss Seidel
+		Private::ECollisionSolverType CollisionSolverType = Private::ECollisionSolverType::GaussSeidel;
+		if (ChaosSolverCollisionSolverType == 1)
+		{
+			CollisionSolverType = Private::ECollisionSolverType::GaussSeidelSimd;
+		}
+		else if (ChaosSolverCollisionSolverType == 2)
+		{
+			CollisionSolverType = Private::ECollisionSolverType::PartialJacobi;
+		}
+
 		if (CollisionSolverType != CollisionConstraints.GetSolverType())
 		{
 			IslandGroupManager.RemoveConstraintContainer(CollisionConstraints);
diff --git a/Engine/Source/Runtime/Experimental/Chaos/Public/Chaos/Collision/PBDCollisionContainerSolver.h b/Engine/Source/Runtime/Experimental/Chaos/Public/Chaos/Collision/PBDCollisionContainerSolver.h
new file mode 100644
--- /dev/null
+++ b/Engine/Source/Runtime/Experimental/Chaos/Public/Chaos/Collision/PBDCollisionContainerSolver.h
@@ -0,0 +1,70 @@
+// Copyright Epic Games, Inc. All Rights Reserved.
+#pragma once
+
+#include "Chaos/CollisionResolutionTypes.h"
+#include "Chaos/Collision/CollisionApplyType.h"
+#include "Chaos/Collision/PBDCollisionSolver.h"
+#include "Chaos/Collision/PBDCollisionSolverSettings.h"
+#include "Chaos/Evolution/SolverConstraintContainer.h"
+
+namespace Chaos
+{
+	namespace Private
+	{
+		class FPBDIsland;
+		class FPBDIslandConstraint;
+	}
+	class FPBDCollisionConstraint;
+	class FPBDCollisionConstraints;
+	class FSolverBody;
+	class FSolverBodyContainer;
+
+	/**
+	 * The solver for a set of collision constraints. This collects all the data required to solve a set of collision
+	 * constraints into a contiguous, ordered buffer.
+	*/
+	class FPBDCollisionContainerSolver : public FConstraintContainerSolver
+	{
+	public:
+		UE_NONCOPYABLE(FPBDCollisionContainerSolver);
+
+		FPBDCollisionContainerSolver(const FPBDCollisionConstraints& InConstraintContainer, const int32 InPriority);
+		~FPBDCollisionContainerSolver();
+
+		int32 NumSolvers() const { return CollisionConstraints.Num(); }
+
+		virtual void Reset(const int32 InMaxCollisions) override final;
+
+		virtual int32 GetNumConstraints() const override final { return NumSolvers(); }
+
+		//
+		// IslandGroup API
+		//
+		virtual void AddConstraints() override final;
+		virtual void AddConstraints(const TArrayView<Private::FPBDIslandConstraint>& ConstraintHandles) override final;
+		virtual void AddBodies(FSolverBodyContainer& SolverBodyContainer) override final;
+		virtual void GatherInput(const FReal Dt) override final;
+		virtual void GatherInput(const FReal Dt, const int32 BeginIndex, const int32 EndIndex) override final;
+		virtual void ScatterOutput(const FReal Dt) override final;
+		virtual void ScatterOutput(const FReal Dt, const int32 BeginIndex, const int32 EndIndex) override final;
+		virtual void ApplyPositionConstraints(const FReal Dt, const int32 It, const int32 NumIts) override final;
+		virtual void ApplyVelocityConstraints(const FReal Dt, const int32 It, const int32 NumIts) override final;
+		virtual void ApplyProjectionConstraints(const FReal Dt, const int32 It, const int32 NumIts) override final;
+
+	private:
+		void CachePrefetchSolver(const int32 ConstraintIndex) const;
+		void AddConstraint(FPBDCollisionConstraint& Constraint);
+		void UpdatePositionShockPropagation(const FReal Dt, const int32 It, const int32 NumIts, const int32 BeginIndex, const int32 EndIndex, const FPBDCollisionSolverSettings& SolverSettings);
+		void UpdateVelocityShockPropagation(const FReal Dt, const int32 It, const int32 NumIts, const int32 BeginIndex, const int32 EndIndex, const FPBDCollisionSolverSettings& SolverSettings);
+		bool SolvePositionImpl(const FReal Dt, const int32 It, const int32 NumIts, const int32 BeginIndex, const int32 EndIndex, const FPBDCollisionSolverSettings& SolverSettings);
+		bool SolveVelocityImpl(const FReal Dt, const int32 It, const int32 NumIts, const int32 BeginIndex, const int32 EndIndex, const FPBDCollisionSolverSettings& SolverSettings);
+		void UpdateCollisions(const FReal InDt, const int32 BeginIndex, const int32 EndIndex);
+
+		const FPBDCollisionConstraints& ConstraintContainer;
+		TArray<Private::FPBDCollisionSolver> CollisionSolvers;
+		TArray<Private::FPBDCollisionSolverManifoldPoint> CollisionSolverManifoldPoints;
+		TArray<FPBDCollisionConstraint*> CollisionConstraints;
+		TArray<bool> bCollisionConstraintPerIterationCollisionDetection;
+		bool bPerIterationCollisionDetection;
+	};
+}
diff --git a/Engine/Source/Runtime/Experimental/Chaos/Public/Chaos/Collision/PBDCollisionContainerSolverJacobi.h b/Engine/Source/Runtime/Experimental/Chaos/Public/Chaos/Collision/PBDCollisionContainerSolverJacobi.h
--- a/Engine/Source/Runtime/Experimental/Chaos/Public/Chaos/Collision/PBDCollisionContainerSolverJacobi.h
+++ b/Engine/Source/Runtime/Experimental/Chaos/Public/Chaos/Collision/PBDCollisionContainerSolverJacobi.h
@@ -9,63 +9,65 @@
 
 namespace Chaos
 {
-	namespace Private
-	{
-		class FPBDIsland;
-		class FPBDIslandConstraint;
-	}
 	class FPBDCollisionConstraint;
 	class FPBDCollisionConstraints;
 	class FSolverBody;
 	class FSolverBodyContainer;
 
-	/**
-	 * The solver for a set of collision constraints. This collects all the data required to solve a set of collision
-	 * constraints into a contiguous, ordered buffer.
-	 * 
-	 * This version runs a Gauss-Seidel outer loop over manifolds, and a Jacobi loop over contacts in each manifold.
-	*/
-	class FPBDCollisionContainerSolverJacobi : public FConstraintContainerSolver
+	namespace Private
 	{
-	public:
-		UE_NONCOPYABLE(FPBDCollisionContainerSolverJacobi);
+		class FPBDIsland;
+		class FPBDIslandConstraint;
+
+		/**
+		 * The solver for a set of collision constraints. This collects all the data required to solve a set of collision
+		 * constraints into a contiguous, ordered buffer.
+		 * 
+		 * This version runs a Gauss-Seidel outer loop over manifolds, and a Jacobi loop over contacts in each manifold.
+		*/
+		class FPBDCollisionContainerSolverJacobi : public FConstraintContainerSolver
+		{
+		public:
+			UE_NONCOPYABLE(FPBDCollisionContainerSolverJacobi);
+
+			FPBDCollisionContainerSolverJacobi(const FPBDCollisionConstraints& InConstraintContainer, const int32 InPriority);
+			~FPBDCollisionContainerSolverJacobi();
 
-		FPBDCollisionContainerSolverJacobi(const FPBDCollisionConstraints& InConstraintContainer, const int32 InPriority);
-		~FPBDCollisionContainerSolverJacobi();
+			int32 NumSolvers() const { return CollisionConstraints.Num(); }
 
-		int32 NumSolvers() const { return CollisionConstraints.Num(); }
+			virtual void Reset(const int32 InMaxCollisions) override final;
 
-		virtual void Reset(const int32 InMaxCollisions) override final;
+			virtual int32 GetNumConstraints() const override final { return NumSolvers(); }
 
-		virtual int32 GetNumConstraints() const override final { return NumSolvers(); }
+			//
+			// IslandGroup API
+			//
+			virtual void AddConstraints() override final;
+			virtual void AddConstraints(const TArrayView<Private::FPBDIslandConstraint>& ConstraintHandles) override final;
+			virtual void AddBodies(FSolverBodyContainer& SolverBodyContainer) override final;
+			virtual void GatherInput(const FReal Dt) override final;
+			virtual void GatherInput(const FReal Dt, const int32 BeginIndex, const int32 EndIndex) override final;
+			virtual void ScatterOutput(const FReal Dt) override final;
+			virtual void ScatterOutput(const FReal Dt, const int32 BeginIndex, const int32 EndIndex) override final;
+			virtual void ApplyPositionConstraints(const FReal Dt, const int32 It, const int32 NumIts) override final;
+			virtual void ApplyVelocityConstraints(const FReal Dt, const int32 It, const int32 NumIts) override final;
+			virtual void ApplyProjectionConstraints(const FReal Dt, const int32 It, const int32 NumIts) override final;
 
-		//
-		// IslandGroup API
-		//
-		virtual void AddConstraints() override final;
-		virtual void AddConstraints(const TArrayView<Private::FPBDIslandConstraint>& ConstraintHandles) override final;
-		virtual void AddBodies(FSolverBodyContainer& SolverBodyContainer) override final;
-		virtual void GatherInput(const FReal Dt) override final;
-		virtual void GatherInput(const FReal Dt, const int32 BeginIndex, const int32 EndIndex) override final;
-		virtual void ScatterOutput(const FReal Dt) override final;
-		virtual void ScatterOutput(const FReal Dt, const int32 BeginIndex, const int32 EndIndex) override final;
-		virtual void ApplyPositionConstraints(const FReal Dt, const int32 It, const int32 NumIts) override final;
-		virtual void ApplyVelocityConstraints(const FReal Dt, const int32 It, const int32 NumIts) override final;
-		virtual void ApplyProjectionConstraints(const FReal Dt, const int32 It, const int32 NumIts) override final;
+		private:
+			void CachePrefetchSolver(const int32 ConstraintIndex) const;
+			void AddConstraint(FPBDCollisionConstraint& Constraint);
+			void UpdatePositionShockPropagation(const FReal Dt, const int32 It, const int32 NumIts, const int32 BeginIndex, const int32 EndIndex, const FPBDCollisionSolverSettings& SolverSettings);
+			void UpdateVelocityShockPropagation(const FReal Dt, const int32 It, const int32 NumIts, const int32 BeginIndex, const int32 EndIndex, const FPBDCollisionSolverSettings& SolverSettings);
+			bool SolvePositionImpl(const FReal Dt, const int32 It, const int32 NumIts, const int32 BeginIndex, const int32 EndIndex, const FPBDCollisionSolverSettings& SolverSettings);
+			bool SolveVelocityImpl(const FReal Dt, const int32 It, const int32 NumIts, const int32 BeginIndex, const int32 EndIndex, const FPBDCollisionSolverSettings& SolverSettings);
+			void UpdateCollisions(const FReal InDt, const int32 BeginIndex, const int32 EndIndex);
 
-	private:
-		void CachePrefetchSolver(const int32 ConstraintIndex) const;
-		void AddConstraint(FPBDCollisionConstraint& Constraint);
-		void UpdatePositionShockPropagation(const FReal Dt, const int32 It, const int32 NumIts, const int32 BeginIndex, const int32 EndIndex, const FPBDCollisionSolverSettings& SolverSettings);
-		void UpdateVelocityShockPropagation(const FReal Dt, const int32 It, const int32 NumIts, const int32 BeginIndex, const int32 EndIndex, const FPBDCollisionSolverSettings& SolverSettings);
-		bool SolvePositionImpl(const FReal Dt, const int32 It, const int32 NumIts, const int32 BeginIndex, const int32 EndIndex, const FPBDCollisionSolverSettings& SolverSettings);
-		bool SolveVelocityImpl(const FReal Dt, const int32 It, const int32 NumIts, const int32 BeginIndex, const int32 EndIndex, const FPBDCollisionSolverSettings& SolverSettings);
-		void UpdateCollisions(const FReal InDt, const int32 BeginIndex, const int32 EndIndex);
+			const FPBDCollisionConstraints& ConstraintContainer;
+			TArray<Private::FPBDCollisionSolverJacobi> CollisionSolvers;
+			TArray<FPBDCollisionConstraint*> CollisionConstraints;
+			TArray<bool> bCollisionConstraintPerIterationCollisionDetection;
+			bool bPerIterationCollisionDetection;
+		};
 
-		const FPBDCollisionConstraints& ConstraintContainer;
-		TArray<Private::FPBDCollisionSolverJacobi> CollisionSolvers;
-		TArray<FPBDCollisionConstraint*> CollisionConstraints;
-		TArray<bool> bCollisionConstraintPerIterationCollisionDetection;
-		bool bPerIterationCollisionDetection;
-	};
-}
+	}	// namespace Private
+}	// namespace Chaos
diff --git a/Engine/Source/Runtime/Experimental/Chaos/Public/Chaos/Collision/PBDCollisionContainerSolverSimd.h b/Engine/Source/Runtime/Experimental/Chaos/Public/Chaos/Collision/PBDCollisionContainerSolverSimd.h
new file mode 100644
--- /dev/null
+++ b/Engine/Source/Runtime/Experimental/Chaos/Public/Chaos/Collision/PBDCollisionContainerSolverSimd.h
@@ -0,0 +1,74 @@
+// Copyright Epic Games, Inc. All Rights Reserved.
+#pragma once
+
+#include "Chaos/CollisionResolutionTypes.h"
+#include "Chaos/Collision/CollisionApplyType.h"
+#include "Chaos/Collision/PBDCollisionSolverSimd.h"
+#include "Chaos/Collision/PBDCollisionSolverSettings.h"
+#include "Chaos/Evolution/SolverConstraintContainer.h"
+
+namespace Chaos
+{
+	class FPBDCollisionConstraint;
+	class FPBDCollisionConstraints;
+	class FSolverBody;
+	class FSolverBodyContainer;
+
+	namespace Private
+	{
+		class FPBDIsland;
+		class FPBDIslandConstraint;
+
+		/**
+		 * The solver for a set of collision constraints. This collects all the data required to solve a set of collision
+		 * constraints into a contiguous, ordered buffer.
+		 * 
+		 * This version runs a Gauss-Seidel outer loop over manifolds, and a Jacobi loop over contacts in each manifold.
+		*/
+		class FPBDCollisionContainerSolverSimd : public FConstraintContainerSolver
+		{
+		public:
+			UE_NONCOPYABLE(FPBDCollisionContainerSolverSimd);
+
+			FPBDCollisionContainerSolverSimd(const FPBDCollisionConstraints& InConstraintContainer, const int32 InPriority);
+			~FPBDCollisionContainerSolverSimd();
+
+			int32 NumSolvers() const { return CollisionConstraints.Num(); }
+
+			virtual void Reset(const int32 InMaxCollisions) override final;
+
+			virtual int32 GetNumConstraints() const override final { return NumSolvers(); }
+
+			//
+			// IslandGroup API
+			//
+			virtual void AddConstraints() override final;
+			virtual void AddConstraints(const TArrayView<Private::FPBDIslandConstraint>& ConstraintHandles) override final;
+			virtual void AddBodies(FSolverBodyContainer& SolverBodyContainer) override final;
+			virtual void GatherInput(const FReal Dt) override final;
+			virtual void GatherInput(const FReal Dt, const int32 BeginIndex, const int32 EndIndex) override final;
+			virtual void ScatterOutput(const FReal Dt) override final;
+			virtual void ScatterOutput(const FReal Dt, const int32 BeginIndex, const int32 EndIndex) override final;
+			virtual void ApplyPositionConstraints(const FReal Dt, const int32 It, const int32 NumIts) override final;
+			virtual void ApplyVelocityConstraints(const FReal Dt, const int32 It, const int32 NumIts) override final;
+			virtual void ApplyProjectionConstraints(const FReal Dt, const int32 It, const int32 NumIts) override final;
+
+		private:
+			void CachePrefetchSolver(const int32 ConstraintIndex) const;
+			void AddConstraint(FPBDCollisionConstraint& Constraint);
+			void UpdatePositionShockPropagation(const FReal Dt, const int32 It, const int32 NumIts, const int32 BeginIndex, const int32 EndIndex, const FPBDCollisionSolverSettings& SolverSettings);
+			void UpdateVelocityShockPropagation(const FReal Dt, const int32 It, const int32 NumIts, const int32 BeginIndex, const int32 EndIndex, const FPBDCollisionSolverSettings& SolverSettings);
+			bool SolvePositionImpl(const FReal Dt, const int32 It, const int32 NumIts, const int32 BeginIndex, const int32 EndIndex, const FPBDCollisionSolverSettings& SolverSettings);
+			bool SolveVelocityImpl(const FReal Dt, const int32 It, const int32 NumIts, const int32 BeginIndex, const int32 EndIndex, const FPBDCollisionSolverSettings& SolverSettings);
+			void UpdateCollisions(const FReal InDt, const int32 BeginIndex, const int32 EndIndex);
+
+			const FPBDCollisionConstraints& ConstraintContainer;
+			TArray<Private::FPBDCollisionSolverSimd> CollisionSolvers;
+			TArray<Private::FPBDCollisionSolverManifoldPointSimd> CollisionSolverManifoldPoints;
+			TArray<FPBDCollisionConstraint*> CollisionConstraints;
+			TArray<bool> bCollisionConstraintPerIterationCollisionDetection;
+			bool bPerIterationCollisionDetection;
+		};
+
+	}	// namespace Private
+}	// namespace Chaos
diff --git a/Engine/Source/Runtime/Experimental/Chaos/Public/Chaos/Collision/PBDCollisionSolver.h b/Engine/Source/Runtime/Experimental/Chaos/Public/Chaos/Collision/PBDCollisionSolver.h
--- a/Engine/Source/Runtime/Experimental/Chaos/Public/Chaos/Collision/PBDCollisionSolver.h
+++ b/Engine/Source/Runtime/Experimental/Chaos/Public/Chaos/Collision/PBDCollisionSolver.h
@@ -27,9 +27,6 @@ namespace Chaos
 		extern bool bChaos_PBDCollisionSolver_Velocity_AveragePointEnabled;
 		extern bool bChaos_PBDCollisionSolver_Velocity_FrictionEnabled;
 		extern float Chaos_PBDCollisionSolver_Position_StaticFrictionStiffness;
-		extern float Chaos_PBDCollisionSolver_JacobiStiffness;
-		extern float Chaos_PBDCollisionSolver_JacobiPositionTolerance;
-		extern float Chaos_PBDCollisionSolver_JacobiRotationTolerance;
 	}
 
 	namespace Private
@@ -272,16 +269,13 @@ namespace Chaos
 			 * @return true if we need to run more iterations, false if we did not apply any correction
 			*/
 			bool SolvePositionWithFriction(const FSolverReal Dt, const FSolverReal MaxPushOut);
-			bool SolvePositionWithFrictionJacobi(const FSolverReal Dt, const FSolverReal MaxPushOut);
 			bool SolvePositionNoFriction(const FSolverReal Dt, const FSolverReal MaxPushOut);
-			bool SolvePositionNoFrictionJacobi(const FSolverReal Dt, const FSolverReal MaxPushOut);
 
 			/**
 			 * @brief Calculate and apply the velocity correction for this iteration
 			 * @return true if we need to run more iterations, false if we did not apply any correction
 			*/
 			bool SolveVelocity(const FSolverReal Dt, const bool bApplyDynamicFriction);
-			bool SolveVelocityJacobi(const FSolverReal Dt, const bool bApplyDynamicFriction);
 
 		private:
 			/**
@@ -551,7 +545,6 @@ namespace Chaos
 
 			// Clamp the total impulse to be positive along the normal. We can apply a net negative impulse, 
 			// but only to correct the velocity that was added by pushout (in which case MinImpulseNormal will be negative).
-			const FSolverReal NetImpulseNormal = ManifoldPoint.NetImpulseNormal + ImpulseNormal;
 			if ((ManifoldPoint.NetImpulseNormal + ImpulseNormal) < MinImpulseNormal)
 			{
 				// We are trying to apply a net negative impulse larger than one to counteract the effective pushout impulse
@@ -604,10 +597,9 @@ namespace Chaos
 			FSolverReal ImpulseNormal = -(Stiffness * ManifoldPoint.ContactMassNormal) * ContactVelocityDeltaNormal;
 
 			// See comments in ApplyVelocityCorrection
-			const FSolverReal NetImpulseNormal = ManifoldPoint.NetImpulseNormal + ImpulseNormal;
-			if (NetImpulseNormal < MinImpulseNormal)
+			if (ManifoldPoint.NetImpulseNormal + ImpulseNormal < MinImpulseNormal)
 			{
-				ImpulseNormal = ImpulseNormal - NetImpulseNormal + MinImpulseNormal;
+				ImpulseNormal = MinImpulseNormal - ManifoldPoint.NetImpulseNormal;
 			}
 
 			ManifoldPoint.NetImpulseNormal += ImpulseNormal;
@@ -929,234 +921,6 @@ namespace Chaos
 			return false;
 		}
 
-		FORCEINLINE_DEBUGGABLE bool FPBDCollisionSolver::SolvePositionNoFrictionJacobi(const FSolverReal Dt, const FSolverReal MaxPushOut)
-		{
-			// SolverBody decorator used to add mass scaling
-			FConstraintSolverBody& Body0 = SolverBody0();
-			FConstraintSolverBody& Body1 = SolverBody1();
-
-			// Apply the position correction so that all contacts have zero separation
-			FSolverVec3 DX0 = FSolverVec3(0);
-			FSolverVec3 DX1 = FSolverVec3(0);
-			FSolverVec3 DR0 = FSolverVec3(0);
-			FSolverVec3 DR1 = FSolverVec3(0);
-
-			const FSolverReal Stiffness = CVars::Chaos_PBDCollisionSolver_JacobiStiffness * State.Stiffness;
-
-			for (int32 PointIndex = 0; PointIndex < NumManifoldPoints(); ++PointIndex)
-			{
-				FPBDCollisionSolverManifoldPoint& SolverManifoldPoint = State.ManifoldPoints[PointIndex];
-
-				FSolverReal ContactDeltaNormal;
-				SolverManifoldPoint.CalculateContactPositionErrorNormal(Body0.SolverBody(), Body1.SolverBody(), MaxPushOut, ContactDeltaNormal);
-
-				const bool bProcessManifoldPoint = (ContactDeltaNormal < FSolverReal(0)) || (SolverManifoldPoint.NetPushOutNormal > FSolverReal(UE_SMALL_NUMBER));
-				if (bProcessManifoldPoint)
-				{
-					FSolverReal PushOutNormal = FSolverReal(0);
-					CalculatePositionCorrectionNormal(
-						Stiffness,
-						ContactDeltaNormal,
-						SolverManifoldPoint.ContactMassNormal,
-						SolverManifoldPoint.NetPushOutNormal,
-						PushOutNormal);						// Out
-
-					SolverManifoldPoint.NetPushOutNormal += PushOutNormal;
-
-					// Update the particle state based on the pushout
-					if (Body0.IsDynamic())
-					{
-						DX0 += (Body0.InvM() * PushOutNormal) * SolverManifoldPoint.WorldContact.ContactNormal;
-						DR0 += SolverManifoldPoint.WorldContactNormalAngular0 * PushOutNormal;
-					}
-					if (Body1.IsDynamic())
-					{
-						DX1 -= (Body1.InvM() * PushOutNormal) * SolverManifoldPoint.WorldContact.ContactNormal;
-						DR1 -= SolverManifoldPoint.WorldContactNormalAngular1 * PushOutNormal;
-					}
-				}
-			}
-
-			if (Body0.IsDynamic())
-			{
-				if (AboveThresholdElseZero(DX0, CVars::Chaos_PBDCollisionSolver_JacobiPositionTolerance))
-				{
-					Body0.ApplyPositionDelta(DX0);
-				}
-				if (AboveThresholdElseZero(DR0, CVars::Chaos_PBDCollisionSolver_JacobiRotationTolerance))
-				{
-					Body0.ApplyRotationDelta(DR0);
-				}
-			}
-			if (Body1.IsDynamic())
-			{
-				if (AboveThresholdElseZero(DX1, CVars::Chaos_PBDCollisionSolver_JacobiPositionTolerance))
-				{
-					Body1.ApplyPositionDelta(DX1);
-				}
-				if (AboveThresholdElseZero(DR1, CVars::Chaos_PBDCollisionSolver_JacobiRotationTolerance))
-				{
-					Body1.ApplyRotationDelta(DR1);
-				}
-			}
-
-			return false;
-		}
-
-		FORCEINLINE_DEBUGGABLE bool FPBDCollisionSolver::SolvePositionWithFrictionJacobi(const FSolverReal Dt, const FSolverReal MaxPushOut)
-		{
-			// SolverBody decorator used to add mass scaling
-			FConstraintSolverBody& Body0 = SolverBody0();
-			FConstraintSolverBody& Body1 = SolverBody1();
-
-			// Accumulate net pushout for friction limits below
-			bool bApplyFriction[MaxPointsPerConstraint] = { false, };
-			int32 NumFrictionContacts = 0;
-			FSolverReal TotalPushOutNormal = FSolverReal(0);
-
-			FSolverVec3 DX0 = FSolverVec3(0);
-			FSolverVec3 DX1 = FSolverVec3(0);
-			FSolverVec3 DR0 = FSolverVec3(0);
-			FSolverVec3 DR1 = FSolverVec3(0);
-
-			const FSolverReal Stiffness = CVars::Chaos_PBDCollisionSolver_JacobiStiffness * State.Stiffness;
-
-			// Apply the position correction along the normal and determine if we want to run friction on each point
-			for (int32 PointIndex = 0; PointIndex < NumManifoldPoints(); ++PointIndex)
-			{
-				FPBDCollisionSolverManifoldPoint& SolverManifoldPoint = State.ManifoldPoints[PointIndex];
-
-				FSolverReal ContactDeltaNormal;
-				SolverManifoldPoint.CalculateContactPositionErrorNormal(Body0.SolverBody(), Body1.SolverBody(), MaxPushOut, ContactDeltaNormal);
-
-				// Apply a normal correction if we still have penetration or if we are now separated but have previously applied a correction that we may want to undo
-				const bool bProcessManifoldPoint = (ContactDeltaNormal < FSolverReal(0)) || (SolverManifoldPoint.NetPushOutNormal > FSolverReal(UE_SMALL_NUMBER));
-				if (bProcessManifoldPoint)
-				{
-					FSolverReal PushOutNormal = FSolverReal(0);
-					CalculatePositionCorrectionNormal(
-						Stiffness,
-						ContactDeltaNormal,
-						SolverManifoldPoint.ContactMassNormal,
-						SolverManifoldPoint.NetPushOutNormal,
-						PushOutNormal);						// Out
-
-					SolverManifoldPoint.NetPushOutNormal += PushOutNormal;
-
-					// Update the particle state based on the pushout
-					if (Body0.IsDynamic())
-					{
-						DX0 += (Body0.InvM() * PushOutNormal) * SolverManifoldPoint.WorldContact.ContactNormal;
-						DR0 += SolverManifoldPoint.WorldContactNormalAngular0 * PushOutNormal;
-					}
-					if (Body1.IsDynamic())
-					{
-						DX1 -= (Body1.InvM() * PushOutNormal) * SolverManifoldPoint.WorldContact.ContactNormal;
-						DR1 -= SolverManifoldPoint.WorldContactNormalAngular1 * PushOutNormal;
-					}
-
-					TotalPushOutNormal += SolverManifoldPoint.NetPushOutNormal;
-				}
-
-				// Friction gets updated for any point with a net normal correction or where we have previously had a normal correction and 
-				// already applied friction (in which case we may need to zero it)
-				if ((SolverManifoldPoint.NetPushOutNormal != 0) || (SolverManifoldPoint.NetPushOutTangentU != 0) || (SolverManifoldPoint.NetPushOutTangentV != 0))
-				{
-					bApplyFriction[PointIndex] = true;
-					++NumFrictionContacts;
-				}
-			}
-
-			// Apply the tangential position correction if required
-			if (NumFrictionContacts > 0)
-			{
-				// We clip the tangential correction at each contact to the friction cone, but we use to average impulse
-				// among all contacts as the clipping limit. This is not really correct but it is much more stable to 
-				// differences in contacts from tick to tick
-				const FSolverReal FrictionMaxPushOut = TotalPushOutNormal / FSolverReal(NumFrictionContacts);
-				const FSolverReal FrictionStiffness = Stiffness * CVars::Chaos_PBDCollisionSolver_Position_StaticFrictionStiffness;
-
-				for (int32 PointIndex = 0; PointIndex < NumManifoldPoints(); ++PointIndex)
-				{
-					if (bApplyFriction[PointIndex])
-					{
-						FPBDCollisionSolverManifoldPoint& SolverManifoldPoint = State.ManifoldPoints[PointIndex];
-
-						FSolverReal ContactDeltaTangentU, ContactDeltaTangentV;
-						SolverManifoldPoint.CalculateContactPositionErrorTangential(Body0.SolverBody(), Body1.SolverBody(), ContactDeltaTangentU, ContactDeltaTangentV);
-
-						FSolverReal PushOutTangentU = FSolverReal(0);
-						FSolverReal PushOutTangentV = FSolverReal(0);
-
-						CalculatePositionCorrectionTangent(
-							FrictionStiffness,
-							ContactDeltaTangentU,
-							SolverManifoldPoint.ContactMassTangentU,
-							SolverManifoldPoint.NetPushOutTangentU,
-							PushOutTangentU);					// Out
-
-						CalculatePositionCorrectionTangent(
-							FrictionStiffness,
-							ContactDeltaTangentV,
-							SolverManifoldPoint.ContactMassTangentV,
-							SolverManifoldPoint.NetPushOutTangentV,
-							PushOutTangentV);					// Out
-
-						SolverManifoldPoint.NetPushOutTangentU += PushOutTangentU;
-						SolverManifoldPoint.NetPushOutTangentV += PushOutTangentV;
-
-						ApplyFrictionCone(
-							State.StaticFriction,
-							State.DynamicFriction,
-							FrictionMaxPushOut,
-							PushOutTangentU,									// InOut
-							PushOutTangentV,									// InOut
-							SolverManifoldPoint.NetPushOutTangentU,				// InOut
-							SolverManifoldPoint.NetPushOutTangentV,				// InOut
-							SolverManifoldPoint.StaticFrictionRatio);			// Out
-
-						// Update the particle state based on the pushout
-						const FSolverVec3 PushOut = PushOutTangentU * SolverManifoldPoint.WorldContact.ContactTangentU + PushOutTangentV * SolverManifoldPoint.WorldContact.ContactTangentV;
-						if (Body0.IsDynamic())
-						{
-							DX0 += Body0.InvM() * PushOut;
-							DR0 += SolverManifoldPoint.WorldContactTangentUAngular0 * PushOutTangentU + SolverManifoldPoint.WorldContactTangentVAngular0 * PushOutTangentV;
-						}
-						if (Body1.IsDynamic())
-						{
-							DX1 -= Body1.InvM() * PushOut;
-							DR1 -= SolverManifoldPoint.WorldContactTangentUAngular1 * PushOutTangentU + SolverManifoldPoint.WorldContactTangentVAngular1 * PushOutTangentV;
-						}
-					}
-				}
-
-				if (Body0.IsDynamic())
-				{
-					if (AboveThresholdElseZero(DX0, CVars::Chaos_PBDCollisionSolver_JacobiPositionTolerance))
-					{
-						Body0.ApplyPositionDelta(DX0);
-					}
-					if (AboveThresholdElseZero(DR0, CVars::Chaos_PBDCollisionSolver_JacobiRotationTolerance))
-					{
-						Body0.ApplyRotationDelta(DR0);
-					}
-				}
-				if (Body1.IsDynamic())
-				{
-					if (AboveThresholdElseZero(DX1, CVars::Chaos_PBDCollisionSolver_JacobiPositionTolerance))
-					{
-						Body1.ApplyPositionDelta(DX1);
-					}
-					if (AboveThresholdElseZero(DR1, CVars::Chaos_PBDCollisionSolver_JacobiRotationTolerance))
-					{
-						Body1.ApplyRotationDelta(DR1);
-					}
-				}
-			}
-
-			return false;
-		}
-
 		FORCEINLINE_DEBUGGABLE bool FPBDCollisionSolver::SolveVelocity(const FSolverReal Dt, const bool bApplyDynamicFriction)
 		{
 			// Apply restitution at the average contact point
@@ -1222,79 +986,6 @@ namespace Chaos
 			return true;
 		}
 
-
-		FORCEINLINE_DEBUGGABLE bool FPBDCollisionSolver::SolveVelocityJacobi(const FSolverReal Dt, const bool bApplyDynamicFriction)
-		{
-			FConstraintSolverBody& Body0 = SolverBody0();
-			FConstraintSolverBody& Body1 = SolverBody1();
-
-			FSolverVec3 DV0 = FSolverVec3(0);
-			FSolverVec3 DV1 = FSolverVec3(0);
-			FSolverVec3 DW0 = FSolverVec3(0);
-			FSolverVec3 DW1 = FSolverVec3(0);
-
-			const FSolverReal Stiffness = CVars::Chaos_PBDCollisionSolver_JacobiStiffness * State.Stiffness;
-
-			for (int32 PointIndex = 0; PointIndex < NumManifoldPoints(); ++PointIndex)
-			{
-				FPBDCollisionSolverManifoldPoint& SolverManifoldPoint = State.ManifoldPoints[PointIndex];
-
-				if (SolverManifoldPoint.ShouldSolveVelocity())
-				{
-					const FSolverReal MinImpulseNormal = FMath::Min(FSolverReal(0), -SolverManifoldPoint.NetPushOutNormal / Dt);
-
-					FSolverReal ContactVelocityDeltaNormal;
-					SolverManifoldPoint.CalculateContactVelocityErrorNormal(Body0, Body1, ContactVelocityDeltaNormal);
-
-					FSolverReal ImpulseNormal = -(Stiffness * SolverManifoldPoint.ContactMassNormal) * ContactVelocityDeltaNormal;
-
-					// See comments in ApplyVelocityCorrection
-					const FSolverReal NetImpulseNormal = SolverManifoldPoint.NetImpulseNormal + ImpulseNormal;
-					if (NetImpulseNormal < MinImpulseNormal)
-					{
-						ImpulseNormal = ImpulseNormal - NetImpulseNormal + MinImpulseNormal;
-					}
-
-					SolverManifoldPoint.NetImpulseNormal += ImpulseNormal;
-
-					FSolverVec3 Impulse = ImpulseNormal * SolverManifoldPoint.WorldContact.ContactNormal;
-					if (Body0.IsDynamic())
-					{
-						DV0 += Body0.InvM() * Impulse;
-						DW0 += SolverManifoldPoint.WorldContactNormalAngular0 * ImpulseNormal;
-					}
-					if (Body1.IsDynamic())
-					{
-						DV1 -= Body1.InvM() * Impulse;
-						DW1 -= SolverManifoldPoint.WorldContactNormalAngular1 * ImpulseNormal;
-					}
-				}
-			}
-
-			const FSolverReal PositionFloor = (CVars::Chaos_PBDCollisionSolver_JacobiPositionTolerance / FMath::Square(Dt));
-			const FSolverReal RotationFloor = (CVars::Chaos_PBDCollisionSolver_JacobiRotationTolerance / FMath::Square(Dt));
-
-			if (Body0.IsDynamic())
-			{
-				AboveThresholdElseZero(DV0, PositionFloor);
-				AboveThresholdElseZero(DW0, RotationFloor);
-
-				Body0.ApplyVelocityDelta(DV0, DW0);
-			}
-			if (Body1.IsDynamic())
-			{
-				AboveThresholdElseZero(DV1, PositionFloor);
-				AboveThresholdElseZero(DW1, RotationFloor);
-
-				Body1.ApplyVelocityDelta(DV1, DW1);
-			}
-
-			// Early-out support for the velocity solve is not currently very important because we
-			// only run one iteration in the velocity solve phase.
-			// @todo(chaos): support early-out in velocity solve if necessary
-			return true;
-		}
-
 		/**
 		 * A helper for solving arrays of constraints
 		 */
diff --git a/Engine/Source/Runtime/Experimental/Chaos/Public/Chaos/Collision/PBDCollisionSolverSimd.h b/Engine/Source/Runtime/Experimental/Chaos/Public/Chaos/Collision/PBDCollisionSolverSimd.h
new file mode 100644
--- /dev/null
+++ b/Engine/Source/Runtime/Experimental/Chaos/Public/Chaos/Collision/PBDCollisionSolverSimd.h
@@ -0,0 +1,961 @@
+// Copyright Epic Games, Inc. All Rights Reserved.
+#pragma once
+
+#include "Chaos/Core.h"
+#include "Chaos/Collision/ContactPoint.h"
+#include "Chaos/Defines.h"
+#include "Chaos/Evolution/SolverBody.h"
+
+// Set to 0 to use a linearized error calculation, and set to 1 to use a non-linear error calculation in collision detection. 
+// In principle nonlinear is more accurate when large rotation corrections occur, but this is not too important for collisions because 
+// when the bodies settle the corrections are small. The linearized version is significantly faster than the non-linear version because 
+// the non-linear version requires a quaternion multiply and renormalization whereas the linear version is just a cross product.
+#define CHAOS_NONLINEAR_COLLISIONS_ENABLED 0
+
+namespace Chaos
+{
+	class FManifoldPoint;
+	class FPBDCollisionConstraint;
+
+	namespace Private
+	{
+		class FPBDCollisionSolverSimd;
+	}
+
+	namespace CVars
+	{
+		extern bool bChaos_PBDCollisionSolver_Velocity_AveragePointEnabled;
+		extern bool bChaos_PBDCollisionSolver_Velocity_FrictionEnabled;
+		extern float Chaos_PBDCollisionSolver_Position_StaticFrictionStiffness;
+		extern float Chaos_PBDCollisionSolver_JacobiStiffness;
+		extern float Chaos_PBDCollisionSolver_JacobiPositionTolerance;
+		extern float Chaos_PBDCollisionSolver_JacobiRotationTolerance;
+	}
+
+	namespace Private
+	{
+		/**
+		 * @brief A single contact point in a FPBDCollisionSolverSimd
+		 * @note Internal Chaos class subject to API changes, renaming or removal
+		*/
+		class FPBDCollisionSolverManifoldPointSimd
+		{
+		public:
+			/**
+			 * @brief Initialize the geometric data for the contact
+			*/
+			void InitContact(
+				const FSolverReal Dt,
+				const FConstraintSolverBody& Body0,
+				const FConstraintSolverBody& Body1);
+
+			/**
+			 * @brief Initialize the material related properties of the contact
+			*/
+			//void InitMaterial(
+			//	const FConstraintSolverBody& Body0,
+			//	const FConstraintSolverBody& Body1,
+			//	const FSolverReal InRestitution,
+			//	const FSolverReal InRestitutionVelocityThreshold);
+
+			/**
+			 * @brief Update the cached mass properties based on the current body transforms
+			*/
+			void UpdateMass(const FConstraintSolverBody& Body0, const FConstraintSolverBody& Body1);
+
+			/**
+			 * @brief Update the contact mass for the normal correction
+			 * This is used by shock propagation.
+			*/
+			void UpdateMassNormal(const FConstraintSolverBody& Body0, const FConstraintSolverBody& Body1);
+
+			/**
+			 * @brief Calculate the position error at the current transforms
+			 * @param MaxPushOut a limit on the position error for this iteration to prevent initial-penetration explosion (a common PBD problem)
+			*/
+			void CalculateContactPositionErrorNormal(const FConstraintSolverBody& Body0, const FConstraintSolverBody& Body1, const FSolverReal MaxPushOut, FSolverReal& OutContactDeltaNormal) const;
+			void CalculateContactPositionErrorTangential(const FConstraintSolverBody& Body0, const FConstraintSolverBody& Body1, FSolverReal& OutContactDeltaTanget0, FSolverReal& OutContactDeltaTangent1) const;
+
+			/**
+			 * @brief Calculate the velocity error at the current transforms
+			*/
+			void CalculateContactVelocityErrorNormal(const FConstraintSolverBody& Body0, const FConstraintSolverBody& Body1, FSolverReal& OutContactVelocityDeltaNormal) const;
+			void CalculateContactVelocityError(const FConstraintSolverBody& Body0, const FConstraintSolverBody& Body1, const FSolverReal DynamicFriction, const FSolverReal Dt, FSolverReal& OutContactVelocityDeltaNormal, FSolverReal& OutContactVelocityDeltaTangent0, FSolverReal& OutContactVelocityDeltaTangent1) const;
+
+
+			FORCEINLINE_DEBUGGABLE void CalculatePositionCorrectionNormal(
+				const FSolverReal Stiffness,
+				const FSolverReal ContactDeltaNormal,
+				FSolverReal& OutPushOutNormal)
+			{
+				const FSolverReal PushOutNormal = -Stiffness * ContactDeltaNormal * ContactMassNormal;
+
+				// The total pushout so far this sub-step
+				// Unilateral constraint: Net-negative impulses not allowed (negative incremental impulses are allowed as long as the net is positive)
+				if ((NetPushOutNormal + PushOutNormal) > FSolverReal(0))
+				{
+					OutPushOutNormal = PushOutNormal;
+				}
+				else
+				{
+					OutPushOutNormal = -NetPushOutNormal;
+				}
+			}
+
+			FORCEINLINE_DEBUGGABLE void CalculatePositionCorrectionTangents(
+				const FSolverReal Stiffness,
+				const FSolverReal ContactDeltaTangentU,
+				const FSolverReal ContactDeltaTangentV,
+				FSolverReal& OutPushOutTangentU,
+				FSolverReal& OutPushOutTangentV)
+			{
+				// Bilateral constraint - negative values allowed (unlike the normal correction)
+				OutPushOutTangentU = -Stiffness * ContactMassTangentU * ContactDeltaTangentU;
+				OutPushOutTangentV = -Stiffness * ContactMassTangentV * ContactDeltaTangentV;
+			}
+
+
+			FORCEINLINE_DEBUGGABLE void ApplyFrictionCone(
+				const FSolverReal StaticFriction,
+				const FSolverReal DynamicFriction,
+				const FSolverReal MaxFrictionPushOut,
+				FSolverReal& InOutPushOutTangentU,
+				FSolverReal& InOutPushOutTangentV)
+			{
+				// Assume we stay in the friction cone...
+				StaticFrictionRatio = FSolverReal(1);
+
+				if (MaxFrictionPushOut < FSolverReal(UE_KINDA_SMALL_NUMBER))
+				{
+					// Note: we have already added the current iteration's PushOut to the NetPushOut but it has not been applied to the body
+					// so we must subtract it again to calculate the actual pushout we want to undo (i.e., the net pushout that has been applied 
+					// to the body so far from previous iterations)
+					InOutPushOutTangentU = -(NetPushOutTangentU - InOutPushOutTangentU);
+					InOutPushOutTangentV = -(NetPushOutTangentV - InOutPushOutTangentV);
+					NetPushOutTangentU = FSolverReal(0);
+					NetPushOutTangentV = FSolverReal(0);
+					StaticFrictionRatio = FSolverReal(0);
+				}
+				else
+				{
+					// If we exceed the static friction cone, clip to the dynamic friction cone
+					const FSolverReal MaxStaticPushOutTangentSq = FMath::Square(StaticFriction * MaxFrictionPushOut);
+					const FSolverReal NetPushOutTangentSq = FMath::Square(NetPushOutTangentU) + FMath::Square(NetPushOutTangentV);
+					if (NetPushOutTangentSq > MaxStaticPushOutTangentSq)
+					{
+						const FSolverReal MaxDynamicPushOutTangent = DynamicFriction * MaxFrictionPushOut;
+						const FSolverReal FrictionMultiplier = MaxDynamicPushOutTangent * FMath::InvSqrt(NetPushOutTangentSq);
+						const FSolverReal NewNetPushOutTangentU = FrictionMultiplier * NetPushOutTangentU;
+						const FSolverReal NewNetPushOutTangentV = FrictionMultiplier * NetPushOutTangentV;
+						InOutPushOutTangentU = NewNetPushOutTangentU - (NetPushOutTangentU - InOutPushOutTangentU);
+						InOutPushOutTangentV = NewNetPushOutTangentV - (NetPushOutTangentV - InOutPushOutTangentV);
+						NetPushOutTangentU = NewNetPushOutTangentU;
+						NetPushOutTangentV = NewNetPushOutTangentV;
+						StaticFrictionRatio = FrictionMultiplier;
+					}
+				}
+			}
+
+			FORCEINLINE_DEBUGGABLE void ApplyPositionCorrectionTangential(
+				const FSolverReal Stiffness,
+				const FSolverReal StaticFriction,
+				const FSolverReal DynamicFriction,
+				const FSolverReal MaxFrictionPushOut,
+				const FSolverReal ContactDeltaTangentU,
+				const FSolverReal ContactDeltaTangentV,
+				FConstraintSolverBody& Body0,
+				FConstraintSolverBody& Body1)
+			{
+				FSolverReal PushOutTangentU = FSolverReal(0);
+				FSolverReal PushOutTangentV = FSolverReal(0);
+
+				CalculatePositionCorrectionTangents(
+					Stiffness,
+					ContactDeltaTangentU,
+					ContactDeltaTangentV,
+					PushOutTangentU,					// Out
+					PushOutTangentV);					// Out
+
+				NetPushOutTangentU += PushOutTangentU;
+				NetPushOutTangentV += PushOutTangentV;
+
+				ApplyFrictionCone(
+					StaticFriction,
+					DynamicFriction,
+					MaxFrictionPushOut,
+					PushOutTangentU,					// InOut
+					PushOutTangentV);					// InOut
+
+				// Update the particle state based on the pushout
+				const FVec3 PushOut = PushOutTangentU * WorldContact.ContactTangentU + PushOutTangentV * WorldContact.ContactTangentV;
+				if (Body0.IsDynamic())
+				{
+					const FSolverVec3 DX0 = Body0.InvM() * PushOut;
+					const FSolverVec3 DR0 = WorldContactTangentUAngular0 * PushOutTangentU + WorldContactTangentVAngular0 * PushOutTangentV;
+					Body0.ApplyPositionDelta(DX0);
+					Body0.ApplyRotationDelta(DR0);
+				}
+				if (Body1.IsDynamic())
+				{
+					const FSolverVec3 DX1 = -Body1.InvM() * PushOut;
+					const FSolverVec3 DR1 = WorldContactTangentUAngular1 * -PushOutTangentU + WorldContactTangentVAngular1 * -PushOutTangentV;
+					Body1.ApplyPositionDelta(DX1);
+					Body1.ApplyRotationDelta(DR1);
+				}
+			}
+
+			FORCEINLINE_DEBUGGABLE void ApplyPositionCorrectionNormal(
+				const FSolverReal Stiffness,
+				const FSolverReal ContactDeltaNormal,
+				FConstraintSolverBody& Body0,
+				FConstraintSolverBody& Body1)
+			{
+				FSolverReal PushOutNormal = FSolverReal(0);
+
+				CalculatePositionCorrectionNormal(
+					Stiffness,
+					ContactDeltaNormal,
+					PushOutNormal);						// Out
+
+				NetPushOutNormal += PushOutNormal;
+
+				// Update the particle state based on the pushout
+				if (Body0.IsDynamic())
+				{
+					const FSolverVec3 DX0 = (Body0.InvM() * PushOutNormal) * WorldContact.ContactNormal;
+					const FSolverVec3 DR0 = WorldContactNormalAngular0 * PushOutNormal;
+					Body0.ApplyPositionDelta(DX0);
+					Body0.ApplyRotationDelta(DR0);
+				}
+				if (Body1.IsDynamic())
+				{
+					const FSolverVec3 DX1 = (Body1.InvM() * -PushOutNormal) * WorldContact.ContactNormal;
+					const FSolverVec3 DR1 = WorldContactNormalAngular1 * -PushOutNormal;
+					Body1.ApplyPositionDelta(DX1);
+					Body1.ApplyRotationDelta(DR1);
+				}
+			}
+
+
+			FORCEINLINE_DEBUGGABLE void ApplyVelocityCorrection(
+				const FSolverReal Stiffness,
+				const FSolverReal Dt,
+				const FSolverReal DynamicFriction,
+				const FSolverReal ContactVelocityDeltaNormal,
+				const FSolverReal ContactVelocityDeltaTangent0,
+				const FSolverReal ContactVelocityDeltaTangent1,
+				const FSolverReal MinImpulseNormal,
+				FConstraintSolverBody& Body0,
+				FConstraintSolverBody& Body1)
+			{
+				FSolverReal ImpulseNormal = -Stiffness * ContactMassNormal * ContactVelocityDeltaNormal;
+				FSolverReal ImpulseTangentU = -Stiffness * ContactMassTangentU * ContactVelocityDeltaTangent0;
+				FSolverReal ImpulseTangentV = -Stiffness * ContactMassTangentV * ContactVelocityDeltaTangent1;
+
+				// Clamp the total impulse to be positive along the normal. We can apply a net negative impulse, 
+				// but only to correct the velocity that was added by pushout (in which case MinImpulseNormal will be negative).
+				if ((NetImpulseNormal + ImpulseNormal) < MinImpulseNormal)
+				{
+					// We are trying to apply a net negative impulse larger than one to counteract the effective pushout impulse
+					// so clamp the net impulse to be equal to minus the pushout impulse along the normal.
+					ImpulseNormal = MinImpulseNormal - NetImpulseNormal;
+				}
+
+				// Clamp the tangential impulses to the friction cone
+				if ((DynamicFriction > 0) && (Dt > 0))
+				{
+					const FSolverReal MaxImpulseTangent = FMath::Max(FSolverReal(0), DynamicFriction * (NetImpulseNormal + ImpulseNormal + NetPushOutNormal / Dt));
+					const FSolverReal MaxImpulseTangentSq = FMath::Square(MaxImpulseTangent);
+					const FSolverReal ImpulseTangentSq = FMath::Square(ImpulseTangentU) + FMath::Square(ImpulseTangentV);
+					if (ImpulseTangentSq > (MaxImpulseTangentSq + UE_SMALL_NUMBER))
+					{
+						const FSolverReal ImpulseTangentScale = MaxImpulseTangent * FMath::InvSqrt(ImpulseTangentSq);
+						ImpulseTangentU *= ImpulseTangentScale;
+						ImpulseTangentV *= ImpulseTangentScale;
+					}
+				}
+
+				NetImpulseNormal += ImpulseNormal;
+				NetImpulseTangentU += ImpulseTangentU;
+				NetImpulseTangentV += ImpulseTangentV;
+
+				// Apply the velocity deltas from the impulse
+				const FSolverVec3 Impulse = ImpulseNormal * WorldContact.ContactNormal + ImpulseTangentU * WorldContact.ContactTangentU + ImpulseTangentV * WorldContact.ContactTangentV;
+				if (Body0.IsDynamic())
+				{
+					const FSolverVec3 DV0 = Body0.InvM() * Impulse;
+					const FSolverVec3 DW0 = WorldContactNormalAngular0 * ImpulseNormal + WorldContactTangentUAngular0 * ImpulseTangentU + WorldContactTangentVAngular0 * ImpulseTangentV;
+					Body0.ApplyVelocityDelta(DV0, DW0);
+				}
+				if (Body1.IsDynamic())
+				{
+					const FSolverVec3 DV1 = -Body1.InvM() * Impulse;
+					const FSolverVec3 DW1 = WorldContactNormalAngular1 * -ImpulseNormal + WorldContactTangentUAngular1 * -ImpulseTangentU + WorldContactTangentVAngular1 * -ImpulseTangentV;
+					Body1.ApplyVelocityDelta(DV1, DW1);
+				}
+			}
+
+			FORCEINLINE_DEBUGGABLE void ApplyVelocityCorrectionNormal(
+				const FSolverReal Stiffness,
+				const FSolverReal ContactVelocityDeltaNormal,
+				const FSolverReal MinImpulseNormal,
+				FConstraintSolverBody& Body0,
+				FConstraintSolverBody& Body1)
+			{
+				FSolverReal ImpulseNormal = -(Stiffness * ContactMassNormal) * ContactVelocityDeltaNormal;
+
+				// See comments in ApplyVelocityCorrection
+				if (NetImpulseNormal + ImpulseNormal < MinImpulseNormal)
+				{
+					ImpulseNormal = MinImpulseNormal - NetImpulseNormal;
+				}
+
+				NetImpulseNormal += ImpulseNormal;
+
+				// Calculate the velocity deltas from the impulse
+				FSolverVec3 Impulse = ImpulseNormal * WorldContact.ContactNormal;
+				if (Body0.IsDynamic())
+				{
+					const FSolverVec3 DV0 = Body0.InvM() * Impulse;
+					const FSolverVec3 DW0 = WorldContactNormalAngular0 * ImpulseNormal;
+					Body0.ApplyVelocityDelta(DV0, DW0);
+				}
+				if (Body1.IsDynamic())
+				{
+					const FSolverVec3 DV1 = -Body1.InvM() * Impulse;
+					const FSolverVec3 DW1 = WorldContactNormalAngular1 * -ImpulseNormal;
+					Body1.ApplyVelocityDelta(DV1, DW1);
+				}
+			}
+
+			// @todo(chaos): make private
+		public:
+			friend class FPBDCollisionSolverSimd;
+
+			/**
+			 * @brief Whether we need to solve velocity for this manifold point (only if we were penetrating or applied a pushout)
+			*/
+			bool ShouldSolveVelocity() const;
+
+			// World-space contact data
+			FWorldContactPoint WorldContact;
+
+			// I^-1.(R x A) for each body where A is each axis (Normal, TangentU, TangentV)
+			FSolverVec3 WorldContactNormalAngular0;
+			FSolverVec3 WorldContactTangentUAngular0;
+			FSolverVec3 WorldContactTangentVAngular0;
+			FSolverVec3 WorldContactNormalAngular1;
+			FSolverVec3 WorldContactTangentUAngular1;
+			FSolverVec3 WorldContactTangentVAngular1;
+
+			// Contact mass (for non-friction)
+			FSolverReal ContactMassNormal;
+			FSolverReal ContactMassTangentU;
+			FSolverReal ContactMassTangentV;
+
+			// Solver outputs
+			FSolverReal NetPushOutNormal;
+			FSolverReal NetPushOutTangentU;
+			FSolverReal NetPushOutTangentV;
+			FSolverReal NetImpulseNormal;
+			FSolverReal NetImpulseTangentU;
+			FSolverReal NetImpulseTangentV;
+
+			// A measure of how much we exceeded the static friction threshold.
+			// Equal to (NormalPushOut / TangentialPushOut) before clamping to the friction cone.
+			// Used to move the static friction anchors to the edge of the cone in Scatter.
+			FSolverReal StaticFrictionRatio;
+		};
+
+		/**
+		 * @brief
+		 * @todo(chaos): Make this solver operate on a single contact point rather than all points in a manifold.
+		 * This would be beneficial if we have many contacts with less than 4 points in the manifold. However this
+		 * is dificult to do while we are still supporting non-manifold collisions.
+		*/
+		class FPBDCollisionSolverSimd
+		{
+		public:
+			static const int32 MaxConstrainedBodies = 2;
+			static const int32 MaxPointsPerConstraint = 4;
+
+			// Create a solver that is initialized to safe defaults
+			static FPBDCollisionSolverSimd MakeInitialized()
+			{
+				FPBDCollisionSolverSimd Solver;
+				Solver.State.Init();
+				return Solver;
+			}
+
+			// Create a solver with no initialization
+			static FPBDCollisionSolverSimd MakeUninitialized()
+			{
+				return FPBDCollisionSolverSimd();
+			}
+
+			// NOTE: Does not initialize any properties. See MakeInitialized
+			FPBDCollisionSolverSimd() {}
+
+			/** Reset the state of the collision solver */
+			void Reset()
+			{
+				State.SolverBodies[0].Reset();
+				State.SolverBodies[1].Reset();
+				State.ManifoldPoints = nullptr;
+				State.NumManifoldPoints = 0;
+				State.MaxManifoldPoints = 0;
+			}
+
+			void ResetManifold()
+			{
+				State.NumManifoldPoints = 0;
+			}
+
+			FSolverReal StaticFriction() const { return State.StaticFriction; }
+			FSolverReal DynamicFriction() const { return State.DynamicFriction; }
+			FSolverReal VelocityFriction() const { return State.VelocityFriction; }
+
+			void SetFriction(const FSolverReal InStaticFriction, const FSolverReal InDynamicFriction, const FSolverReal InVelocityFriction)
+			{
+				State.StaticFriction = InStaticFriction;
+				State.DynamicFriction = InDynamicFriction;
+				State.VelocityFriction = InVelocityFriction;
+			}
+
+			void SetStiffness(const FSolverReal InStiffness)
+			{
+				State.Stiffness = InStiffness;
+			}
+
+			void SetSolverBodies(FSolverBody& SolverBody0, FSolverBody& SolverBody1)
+			{
+				State.SolverBodies[0].SetSolverBody(SolverBody0);
+				State.SolverBodies[1].SetSolverBody(SolverBody1);
+			}
+
+			int32 NumManifoldPoints() const
+			{
+				return State.NumManifoldPoints;
+			}
+
+			int32 AddManifoldPoint()
+			{
+				check(State.NumManifoldPoints < State.MaxManifoldPoints);
+				return State.NumManifoldPoints++;
+			}
+
+			void SetManifoldPointsBuffer(FPBDCollisionSolverManifoldPointSimd* InManifoldPoints, const int32 InMaxManifoldPoints)
+			{
+				State.ManifoldPoints = InManifoldPoints;
+				State.MaxManifoldPoints = InMaxManifoldPoints;
+				State.NumManifoldPoints = 0;
+			}
+
+			const FPBDCollisionSolverManifoldPointSimd& GetManifoldPoint(const int32 ManifoldPointIndex) const
+			{
+				check(ManifoldPointIndex < NumManifoldPoints());
+				return State.ManifoldPoints[ManifoldPointIndex];
+			}
+
+			FWorldContactPoint& GetWorldContactPoint(const int32 ManifoldPointIndex)
+			{
+				return State.ManifoldPoints[ManifoldPointIndex].WorldContact;
+			}
+
+			/**
+			 * Set up a manifold point (also calls FinalizeManifoldPoint)
+			*/
+			void SetManifoldPoint(
+				const int32 ManifoldPoiontIndex,
+				const FSolverReal Dt,
+				const FSolverVec3& InRelativeContactPosition0,
+				const FSolverVec3& InRelativeContactPosition1,
+				const FSolverVec3& InWorldContactNormal,
+				const FSolverVec3& InWorldContactTangentU,
+				const FSolverVec3& InWorldContactTangentV,
+				const FSolverReal InWorldContactDeltaNormal,
+				const FSolverReal InWorldContactDeltaTangentU,
+				const FSolverReal InWorldContactDeltaTangentV,
+				const FSolverReal InWorldContactVelocityTargetNormal);
+
+			/** 
+			 * Finish manifold point setup.
+			 * NOTE: Can only be called after the WorldContact has been set up (e.g., see SetManifoldPoint)
+			*/
+			void FinalizeManifoldPoint(const int32 ManifoldPoiontIndex, const FSolverReal Dt);
+
+			/**
+			 * @brief Get the first (decorated) solver body
+			 * The decorator add a possible mass scale
+			*/
+			FConstraintSolverBody& SolverBody0() { return State.SolverBodies[0]; }
+			const FConstraintSolverBody& SolverBody0() const { return State.SolverBodies[0]; }
+
+			/**
+			 * @brief Get the second (decorated) solver body
+			 * The decorator add a possible mass scale
+			*/
+			FConstraintSolverBody& SolverBody1() { return State.SolverBodies[1]; }
+			const FConstraintSolverBody& SolverBody1() const { return State.SolverBodies[1]; }
+
+			/**
+			 * @brief Set up the mass scaling for shock propagation, using the position-phase mass scale
+			*/
+			void EnablePositionShockPropagation();
+
+			/**
+			 * @brief Set up the mass scaling for shock propagation, using the velocity-phase mass scale
+			*/
+			void EnableVelocityShockPropagation();
+
+			/**
+			 * @brief Disable mass scaling
+			*/
+			void DisableShockPropagation();
+
+			/**
+			 * @brief Calculate and apply the position correction for this iteration
+			 * @return true if we need to run more iterations, false if we did not apply any correction
+			*/
+			bool SolvePositionWithFriction(const FSolverReal Dt, const FSolverReal MaxPushOut);
+			bool SolvePositionNoFriction(const FSolverReal Dt, const FSolverReal MaxPushOut);
+
+			/**
+			 * @brief Calculate and apply the velocity correction for this iteration
+			 * @return true if we need to run more iterations, false if we did not apply any correction
+			*/
+			bool SolveVelocity(const FSolverReal Dt, const bool bApplyDynamicFriction);
+
+		private:
+			/**
+			 * @brief Apply the inverse mass scale the body with the lower level
+			 * @param InvMassScale
+			*/
+			void SetShockPropagationInvMassScale(const FSolverReal InvMassScale);
+
+			/**
+			 * @brief Run a velocity solve on the average point from all the points that received a position impulse
+			 * This is used to enforce Restitution constraints without introducing rotation artefacts without
+			 * adding more velocity iterations.
+			 * This will only perform work if there is more than one active contact.
+			*/
+			void SolveVelocityAverage(const FSolverReal Dt);
+
+			struct FState
+			{
+				FState()
+				{
+				}
+
+				void Init()
+				{
+					SolverBodies[0].Init();
+					SolverBodies[1].Init();
+					StaticFriction = 0;
+					DynamicFriction = 0;
+					VelocityFriction = 0;
+					Stiffness = 1;
+					ManifoldPoints = nullptr;
+					NumManifoldPoints = 0;
+					MaxManifoldPoints = 0;
+				}
+
+				// Static Friction in the position-solve phase
+				FSolverReal StaticFriction;
+
+				// Dynamic Friction in the position-solve phase
+				FSolverReal DynamicFriction;
+
+				// Dynamic Friction in the velocity-solve phase
+				FSolverReal VelocityFriction;
+
+				// Solver stiffness (scales all pushout and impulses)
+				FSolverReal Stiffness;
+
+				// Bodies and contacts
+				FConstraintSolverBody SolverBodies[MaxConstrainedBodies];
+				FPBDCollisionSolverManifoldPointSimd* ManifoldPoints;
+				int32 NumManifoldPoints;
+				int32 MaxManifoldPoints;
+			};
+
+			FState State;
+		};
+
+		//////////////////////////////////////////////////////////////////////////////////////////////////
+		//////////////////////////////////////////////////////////////////////////////////////////////////
+		//////////////////////////////////////////////////////////////////////////////////////////////////
+		//////////////////////////////////////////////////////////////////////////////////////////////////
+
+		FORCEINLINE_DEBUGGABLE void FPBDCollisionSolverManifoldPointSimd::InitContact(
+			const FSolverReal Dt,
+			const FConstraintSolverBody& Body0,
+			const FConstraintSolverBody& Body1)
+		{
+			NetPushOutNormal = FSolverReal(0);
+			NetPushOutTangentU = FSolverReal(0);
+			NetPushOutTangentV = FSolverReal(0);
+			NetImpulseNormal = FSolverReal(0);
+			NetImpulseTangentU = FSolverReal(0);
+			NetImpulseTangentV = FSolverReal(0);
+			StaticFrictionRatio = FSolverReal(0);
+
+			UpdateMass(Body0, Body1);
+		}
+
+		FORCEINLINE_DEBUGGABLE void FPBDCollisionSolverManifoldPointSimd::UpdateMass(const FConstraintSolverBody& Body0, const FConstraintSolverBody& Body1)
+		{
+			FSolverReal ContactMassInvNormal = FSolverReal(0);
+			FSolverReal ContactMassInvTangentU = FSolverReal(0);
+			FSolverReal ContactMassInvTangentV = FSolverReal(0);
+
+			// These are not used if not initialized below so no need to clear
+			//WorldContactNormalAngular0 = FSolverVec3(0);
+			//WorldContactTangentUAngular0 = FSolverVec3(0);
+			//WorldContactTangentVAngular0 = FSolverVec3(0);
+			//WorldContactNormalAngular1 = FSolverVec3(0);
+			//WorldContactTangentUAngular1 = FSolverVec3(0);
+			//WorldContactTangentVAngular1 = FSolverVec3(0);
+
+			if (Body0.IsDynamic())
+			{
+				const FSolverVec3 R0xN = FSolverVec3::CrossProduct(WorldContact.RelativeContactPoints[0], WorldContact.ContactNormal);
+				const FSolverVec3 R0xU = FSolverVec3::CrossProduct(WorldContact.RelativeContactPoints[0], WorldContact.ContactTangentU);
+				const FSolverVec3 R0xV = FSolverVec3::CrossProduct(WorldContact.RelativeContactPoints[0], WorldContact.ContactTangentV);
+
+				const FSolverMatrix33 InvI0 = Body0.InvI();
+
+				WorldContactNormalAngular0 = InvI0 * R0xN;
+				WorldContactTangentUAngular0 = InvI0 * R0xU;
+				WorldContactTangentVAngular0 = InvI0 * R0xV;
+
+				ContactMassInvNormal += FSolverVec3::DotProduct(R0xN, WorldContactNormalAngular0) + Body0.InvM();
+				ContactMassInvTangentU += FSolverVec3::DotProduct(R0xU, WorldContactTangentUAngular0) + Body0.InvM();
+				ContactMassInvTangentV += FSolverVec3::DotProduct(R0xV, WorldContactTangentVAngular0) + Body0.InvM();
+			}
+			if (Body1.IsDynamic())
+			{
+				const FSolverVec3 R1xN = FSolverVec3::CrossProduct(WorldContact.RelativeContactPoints[1], WorldContact.ContactNormal);
+				const FSolverVec3 R1xU = FSolverVec3::CrossProduct(WorldContact.RelativeContactPoints[1], WorldContact.ContactTangentU);
+				const FSolverVec3 R1xV = FSolverVec3::CrossProduct(WorldContact.RelativeContactPoints[1], WorldContact.ContactTangentV);
+
+				const FSolverMatrix33 InvI1 = Body1.InvI();
+
+				WorldContactNormalAngular1 = InvI1 * R1xN;
+				WorldContactTangentUAngular1 = InvI1 * R1xU;
+				WorldContactTangentVAngular1 = InvI1 * R1xV;
+
+				ContactMassInvNormal += FSolverVec3::DotProduct(R1xN, WorldContactNormalAngular1) + Body1.InvM();
+				ContactMassInvTangentU += FSolverVec3::DotProduct(R1xU, WorldContactTangentUAngular1) + Body1.InvM();
+				ContactMassInvTangentV += FSolverVec3::DotProduct(R1xV, WorldContactTangentVAngular1) + Body1.InvM();
+			}
+
+			ContactMassNormal = (ContactMassInvNormal > FSolverReal(UE_SMALL_NUMBER)) ? FSolverReal(1) / ContactMassInvNormal : FSolverReal(0);
+			ContactMassTangentU = (ContactMassInvTangentU > FSolverReal(UE_SMALL_NUMBER)) ? FSolverReal(1) / ContactMassInvTangentU : FSolverReal(0);
+			ContactMassTangentV = (ContactMassInvTangentV > FSolverReal(UE_SMALL_NUMBER)) ? FSolverReal(1) / ContactMassInvTangentV : FSolverReal(0);
+		}
+
+		FORCEINLINE_DEBUGGABLE void FPBDCollisionSolverManifoldPointSimd::UpdateMassNormal(const FConstraintSolverBody& Body0, const FConstraintSolverBody& Body1)
+		{
+			FSolverReal ContactMassInvNormal = FSolverReal(0);
+			if (Body0.IsDynamic())
+			{
+				const FSolverVec3 R0xN = FSolverVec3::CrossProduct(WorldContact.RelativeContactPoints[0], WorldContact.ContactNormal);
+				const FSolverMatrix33 InvI0 = Body0.InvI();
+				WorldContactNormalAngular0 = InvI0 * R0xN;
+				ContactMassInvNormal += FSolverVec3::DotProduct(R0xN, WorldContactNormalAngular0) + Body0.InvM();
+			}
+			if (Body1.IsDynamic())
+			{
+				const FSolverVec3 R1xN = FSolverVec3::CrossProduct(WorldContact.RelativeContactPoints[1], WorldContact.ContactNormal);
+				const FSolverMatrix33 InvI1 = Body1.InvI();
+				WorldContactNormalAngular1 = InvI1 * R1xN;
+				ContactMassInvNormal += FSolverVec3::DotProduct(R1xN, WorldContactNormalAngular1) + Body1.InvM();
+			}
+			ContactMassNormal = (ContactMassInvNormal > FSolverReal(UE_SMALL_NUMBER)) ? FSolverReal(1) / ContactMassInvNormal : FSolverReal(0);
+		}
+
+		FORCEINLINE_DEBUGGABLE void FPBDCollisionSolverManifoldPointSimd::CalculateContactPositionErrorNormal(const FConstraintSolverBody& Body0, const FConstraintSolverBody& Body1, const FSolverReal MaxPushOut, FSolverReal& OutContactDeltaNormal) const
+		{
+			// Linear version: calculate the contact delta assuming linear motion after applying a positional impulse at the contact point. There will be an error that depends on the size of the rotation.
+			const FSolverVec3 ContactDelta0 = Body0.DP() + FSolverVec3::CrossProduct(Body0.DQ(), WorldContact.RelativeContactPoints[0]);
+			const FSolverVec3 ContactDelta1 = Body1.DP() + FSolverVec3::CrossProduct(Body1.DQ(), WorldContact.RelativeContactPoints[1]);
+			const FSolverVec3 ContactDelta = ContactDelta0 - ContactDelta1;
+			OutContactDeltaNormal = WorldContact.ContactDeltaNormal + FSolverVec3::DotProduct(ContactDelta, WorldContact.ContactNormal);
+
+			// NOTE: OutContactDeltaNormal is negative for penetration
+			// NOTE: MaxPushOut == 0 disables the pushout limits
+			if ((MaxPushOut > 0) && (OutContactDeltaNormal < -MaxPushOut))
+			{
+				OutContactDeltaNormal = -MaxPushOut;
+			}
+		}
+
+		FORCEINLINE_DEBUGGABLE void FPBDCollisionSolverManifoldPointSimd::CalculateContactPositionErrorTangential(const FConstraintSolverBody& Body0, const FConstraintSolverBody& Body1, FSolverReal& OutContactDeltaTangentU, FSolverReal& OutContactDeltaTangentV) const
+		{
+			// Linear version: calculate the contact delta assuming linear motion after applying a positional impulse at the contact point. There will be an error that depends on the size of the rotation.
+			const FSolverVec3 ContactDelta0 = Body0.DP() + FSolverVec3::CrossProduct(Body0.DQ(), WorldContact.RelativeContactPoints[0]);
+			const FSolverVec3 ContactDelta1 = Body1.DP() + FSolverVec3::CrossProduct(Body1.DQ(), WorldContact.RelativeContactPoints[1]);
+			const FSolverVec3 ContactDelta = ContactDelta0 - ContactDelta1;
+			OutContactDeltaTangentU = WorldContact.ContactDeltaTangentU + FSolverVec3::DotProduct(ContactDelta, WorldContact.ContactTangentU);
+			OutContactDeltaTangentV = WorldContact.ContactDeltaTangentV + FSolverVec3::DotProduct(ContactDelta, WorldContact.ContactTangentV);
+		}
+
+		FORCEINLINE_DEBUGGABLE void FPBDCollisionSolverManifoldPointSimd::CalculateContactVelocityError(const FConstraintSolverBody& Body0, const FConstraintSolverBody& Body1, const FSolverReal DynamicFriction, const FSolverReal Dt, FSolverReal& OutContactVelocityDeltaNormal, FSolverReal& OutContactVelocityDeltaTangent0, FSolverReal& OutContactVelocityDeltaTangent1) const
+		{
+			const FSolverVec3 ContactVelocity0 = Body0.V() + FSolverVec3::CrossProduct(Body0.W(), WorldContact.RelativeContactPoints[0]);
+			const FSolverVec3 ContactVelocity1 = Body1.V() + FSolverVec3::CrossProduct(Body1.W(), WorldContact.RelativeContactPoints[1]);
+			const FSolverVec3 ContactVelocity = ContactVelocity0 - ContactVelocity1;
+			const FSolverReal ContactVelocityNormal = FSolverVec3::DotProduct(ContactVelocity, WorldContact.ContactNormal);
+			const FSolverReal ContactVelocityTangent0 = FSolverVec3::DotProduct(ContactVelocity, WorldContact.ContactTangentU);
+			const FSolverReal ContactVelocityTangent1 = FSolverVec3::DotProduct(ContactVelocity, WorldContact.ContactTangentV);
+
+			OutContactVelocityDeltaNormal = (ContactVelocityNormal - WorldContact.ContactTargetVelocityNormal);
+			OutContactVelocityDeltaTangent0 = ContactVelocityTangent0;
+			OutContactVelocityDeltaTangent1 = ContactVelocityTangent1;
+		}
+
+		FORCEINLINE_DEBUGGABLE void FPBDCollisionSolverManifoldPointSimd::CalculateContactVelocityErrorNormal(const FConstraintSolverBody& Body0, const FConstraintSolverBody& Body1, FSolverReal& OutContactVelocityDeltaNormal) const
+		{
+			const FSolverVec3 ContactVelocity0 = Body0.V() + FSolverVec3::CrossProduct(Body0.W(), WorldContact.RelativeContactPoints[0]);
+			const FSolverVec3 ContactVelocity1 = Body1.V() + FSolverVec3::CrossProduct(Body1.W(), WorldContact.RelativeContactPoints[1]);
+			const FSolverVec3 ContactVelocity = ContactVelocity0 - ContactVelocity1;
+			const FSolverReal ContactVelocityNormal = FSolverVec3::DotProduct(ContactVelocity, WorldContact.ContactNormal);
+
+			// Add up the errors in the velocity (current velocity - desired velocity)
+			OutContactVelocityDeltaNormal = (ContactVelocityNormal - WorldContact.ContactTargetVelocityNormal);
+		}
+
+		FORCEINLINE_DEBUGGABLE bool FPBDCollisionSolverManifoldPointSimd::ShouldSolveVelocity() const
+		{
+			// We ensure positive separating velocity for close contacts even if they didn't receive a pushout
+			return (NetPushOutNormal > FSolverReal(0)) || (WorldContact.ContactDeltaNormal < FSolverReal(0));
+		}
+
+		//////////////////////////////////////////////////////////////////////////////////////////////////
+		//////////////////////////////////////////////////////////////////////////////////////////////////
+		//////////////////////////////////////////////////////////////////////////////////////////////////
+		//////////////////////////////////////////////////////////////////////////////////////////////////
+
+		FORCEINLINE_DEBUGGABLE  void FPBDCollisionSolverSimd::FinalizeManifoldPoint(const int32 ManifoldPoiontIndex, const FSolverReal Dt)
+		{
+			State.ManifoldPoints[ManifoldPoiontIndex].InitContact(
+				FSolverReal(Dt),
+				State.SolverBodies[0],
+				State.SolverBodies[1]);
+		}
+
+		FORCEINLINE_DEBUGGABLE void FPBDCollisionSolverSimd::SetManifoldPoint(
+			const int32 ManifoldPoiontIndex,
+			const FSolverReal Dt,
+			const FSolverVec3& InRelativeContactPosition0,
+			const FSolverVec3& InRelativeContactPosition1,
+			const FSolverVec3& InWorldContactNormal,
+			const FSolverVec3& InWorldContactTangentU,
+			const FSolverVec3& InWorldContactTangentV,
+			const FSolverReal InWorldContactDeltaNormal,
+			const FSolverReal InWorldContactDeltaTangentU,
+			const FSolverReal InWorldContactDeltaTangentV,
+			const FSolverReal InWorldContactVelocityTargetNormal)
+		{
+			FWorldContactPoint& WorldContactPoint = State.ManifoldPoints[ManifoldPoiontIndex].WorldContact;
+			WorldContactPoint.RelativeContactPoints[0] = InRelativeContactPosition0;
+			WorldContactPoint.RelativeContactPoints[1] = InRelativeContactPosition1;
+			WorldContactPoint.ContactNormal = InWorldContactNormal;
+			WorldContactPoint.ContactTangentU = InWorldContactTangentU;
+			WorldContactPoint.ContactTangentV = InWorldContactTangentV;
+			WorldContactPoint.ContactDeltaNormal = InWorldContactDeltaNormal;
+			WorldContactPoint.ContactDeltaTangentU = InWorldContactDeltaTangentU;
+			WorldContactPoint.ContactDeltaTangentV = InWorldContactDeltaTangentV;
+			WorldContactPoint.ContactTargetVelocityNormal = InWorldContactVelocityTargetNormal;
+
+			State.ManifoldPoints[ManifoldPoiontIndex].InitContact(
+				FSolverReal(Dt),
+				State.SolverBodies[0],
+				State.SolverBodies[1]);
+		}
+
+		FORCEINLINE_DEBUGGABLE bool FPBDCollisionSolverSimd::SolvePositionWithFriction(const FSolverReal Dt, const FSolverReal MaxPushOut)
+		{
+			// SolverBody decorator used to add mass scaling
+			FConstraintSolverBody& Body0 = SolverBody0();
+			FConstraintSolverBody& Body1 = SolverBody1();
+
+			// Accumulate net pushout for friction limits below
+			bool bApplyFriction[MaxPointsPerConstraint] = { false, };
+			int32 NumFrictionContacts = 0;
+			FSolverReal TotalPushOutNormal = FSolverReal(0);
+
+			// Apply the position correction along the normal and determine if we want to run friction on each point
+			for (int32 PointIndex = 0; PointIndex < NumManifoldPoints(); ++PointIndex)
+			{
+				FPBDCollisionSolverManifoldPointSimd& SolverManifoldPoint = State.ManifoldPoints[PointIndex];
+
+				FSolverReal ContactDeltaNormal;
+				SolverManifoldPoint.CalculateContactPositionErrorNormal(Body0.SolverBody(), Body1.SolverBody(), MaxPushOut, ContactDeltaNormal);
+
+				// Apply a normal correction if we still have penetration or if we are now separated but have previously applied a correction that we may want to undo
+				const bool bProcessManifoldPoint = (ContactDeltaNormal < FSolverReal(0)) || (SolverManifoldPoint.NetPushOutNormal > FSolverReal(UE_SMALL_NUMBER));
+				if (bProcessManifoldPoint)
+				{
+					SolverManifoldPoint.ApplyPositionCorrectionNormal(
+						State.Stiffness,
+						ContactDeltaNormal,
+						Body0,
+						Body1);
+
+					TotalPushOutNormal += SolverManifoldPoint.NetPushOutNormal;
+				}
+
+				// Friction gets updated for any point with a net normal correction or where we have previously had a normal correction and 
+				// already applied friction (in which case we may need to zero it)
+				if ((SolverManifoldPoint.NetPushOutNormal != 0) || (SolverManifoldPoint.NetPushOutTangentU != 0) || (SolverManifoldPoint.NetPushOutTangentV != 0))
+				{
+					bApplyFriction[PointIndex] = true;
+					++NumFrictionContacts;
+				}
+			}
+
+			// Apply the tangential position correction if required
+			if (NumFrictionContacts > 0)
+			{
+				// We clip the tangential correction at each contact to the friction cone, but we use to average impulse
+				// among all contacts as the clipping limit. This is not really correct but it is much more stable to 
+				// differences in contacts from tick to tick
+				// @todo(chaos): try a decaying maximum per contact point rather than an average (again - we had that once!)
+				const FSolverReal FrictionMaxPushOut = TotalPushOutNormal / FSolverReal(NumFrictionContacts);
+				const FSolverReal FrictionStiffness = State.Stiffness * CVars::Chaos_PBDCollisionSolver_Position_StaticFrictionStiffness;
+
+				for (int32 PointIndex = 0; PointIndex < NumManifoldPoints(); ++PointIndex)
+				{
+					if (bApplyFriction[PointIndex])
+					{
+						FPBDCollisionSolverManifoldPointSimd& SolverManifoldPoint = State.ManifoldPoints[PointIndex];
+
+						FSolverReal ContactDeltaTangentU, ContactDeltaTangentV;
+						SolverManifoldPoint.CalculateContactPositionErrorTangential(Body0.SolverBody(), Body1.SolverBody(), ContactDeltaTangentU, ContactDeltaTangentV);
+
+						SolverManifoldPoint.ApplyPositionCorrectionTangential(
+							FrictionStiffness,
+							State.StaticFriction,
+							State.DynamicFriction,
+							FrictionMaxPushOut,
+							ContactDeltaTangentU,
+							ContactDeltaTangentV,
+							Body0,
+							Body1);
+					}
+				}
+			}
+
+			return false;
+		}
+
+		FORCEINLINE_DEBUGGABLE bool FPBDCollisionSolverSimd::SolvePositionNoFriction(const FSolverReal Dt, const FSolverReal MaxPushOut)
+		{
+			// SolverBody decorator used to add mass scaling
+			FConstraintSolverBody& Body0 = SolverBody0();
+			FConstraintSolverBody& Body1 = SolverBody1();
+
+			// Apply the position correction so that all contacts have zero separation
+			for (int32 PointIndex = 0; PointIndex < NumManifoldPoints(); ++PointIndex)
+			{
+				FPBDCollisionSolverManifoldPointSimd& SolverManifoldPoint = State.ManifoldPoints[PointIndex];
+
+				FSolverReal ContactDeltaNormal;
+				SolverManifoldPoint.CalculateContactPositionErrorNormal(Body0.SolverBody(), Body1.SolverBody(), MaxPushOut, ContactDeltaNormal);
+
+				const bool bProcessManifoldPoint = (ContactDeltaNormal < FSolverReal(0)) || (SolverManifoldPoint.NetPushOutNormal > FSolverReal(UE_SMALL_NUMBER));
+				if (bProcessManifoldPoint)
+				{
+					SolverManifoldPoint.ApplyPositionCorrectionNormal(
+						State.Stiffness,
+						ContactDeltaNormal,
+						Body0,
+						Body1);
+				}
+			}
+
+			return false;
+		}
+
+		FORCEINLINE_DEBUGGABLE bool FPBDCollisionSolverSimd::SolveVelocity(const FSolverReal Dt, const bool bApplyDynamicFriction)
+		{
+			// Apply restitution at the average contact point
+			// This means we don't need to run as many iterations to get stable bouncing
+			// It also helps with zero restitution to counter any velocioty added by the PBD solve
+			const bool bSolveAverageContact = (NumManifoldPoints() > 1) && CVars::bChaos_PBDCollisionSolver_Velocity_AveragePointEnabled;
+			if (bSolveAverageContact)
+			{
+				SolveVelocityAverage(Dt);
+			}
+
+			FConstraintSolverBody& Body0 = SolverBody0();
+			FConstraintSolverBody& Body1 = SolverBody1();
+
+			// NOTE: this dynamic friction implementation is iteration-count sensitive
+			// @todo(chaos): fix iteration count dependence of dynamic friction
+			const FSolverReal DynamicFriction = (bApplyDynamicFriction && (Dt > 0) && CVars::bChaos_PBDCollisionSolver_Velocity_FrictionEnabled) ? State.VelocityFriction : FSolverReal(0);
+
+			for (int32 PointIndex = 0; PointIndex < NumManifoldPoints(); ++PointIndex)
+			{
+				FPBDCollisionSolverManifoldPointSimd& SolverManifoldPoint = State.ManifoldPoints[PointIndex];
+
+				if (SolverManifoldPoint.ShouldSolveVelocity())
+				{
+					const FSolverReal MinImpulseNormal = FMath::Min(FSolverReal(0), -SolverManifoldPoint.NetPushOutNormal / Dt);
+
+					if (DynamicFriction > 0)
+					{
+						FSolverReal ContactVelocityDeltaNormal, ContactVelocityDeltaTangentU, ContactVelocityDeltaTangentV;
+						SolverManifoldPoint.CalculateContactVelocityError(Body0, Body1, DynamicFriction, Dt, ContactVelocityDeltaNormal, ContactVelocityDeltaTangentU, ContactVelocityDeltaTangentV);
+
+						SolverManifoldPoint.ApplyVelocityCorrection(
+							State.Stiffness,
+							Dt,
+							DynamicFriction,
+							ContactVelocityDeltaNormal,
+							ContactVelocityDeltaTangentU,
+							ContactVelocityDeltaTangentV,
+							MinImpulseNormal,
+							Body0,
+							Body1);
+					}
+					else
+					{
+						FSolverReal ContactVelocityDeltaNormal;
+						SolverManifoldPoint.CalculateContactVelocityErrorNormal(Body0, Body1, ContactVelocityDeltaNormal);
+
+						SolverManifoldPoint.ApplyVelocityCorrectionNormal(
+							State.Stiffness,
+							ContactVelocityDeltaNormal,
+							MinImpulseNormal,
+							Body0,
+							Body1);
+					}
+				}
+			}
+
+			// Early-out support for the velocity solve is not currently very important because we
+			// only run one iteration in the velocity solve phase.
+			// @todo(chaos): support early-out in velocity solve if necessary
+			return true;
+		}
+
+
+		/**
+		 * A helper for solving arrays of constraints
+		 */
+		class FPBDCollisionSolverHelperSimd
+		{
+		public:
+			static void SolvePositionNoFriction(const TArrayView<FPBDCollisionSolverSimd>& CollisionSolvers, const FSolverReal Dt, const FSolverReal MaxPushOut);
+			static void SolvePositionWithFriction(const TArrayView<FPBDCollisionSolverSimd>& CollisionSolvers, const FSolverReal Dt, const FSolverReal MaxPushOut);
+			static void SolveVelocity(const TArrayView<FPBDCollisionSolverSimd>& CollisionSolvers, const FSolverReal Dt, const bool bApplyDynamicFriction);
+
+			static void CheckISPC();
+		};
+
+	}	// namespace Private
+}	// namespace Chaos
+
diff --git a/Engine/Source/Runtime/Experimental/Chaos/Public/Chaos/Collision/SolverCollisionContainer.h b/Engine/Source/Runtime/Experimental/Chaos/Public/Chaos/Collision/SolverCollisionContainer.h
--- a/Engine/Source/Runtime/Experimental/Chaos/Public/Chaos/Collision/SolverCollisionContainer.h
+++ b/Engine/Source/Runtime/Experimental/Chaos/Public/Chaos/Collision/SolverCollisionContainer.h
@@ -1,70 +1,7 @@
 // Copyright Epic Games, Inc. All Rights Reserved.
 #pragma once
 
-#include "Chaos/CollisionResolutionTypes.h"
-#include "Chaos/Collision/CollisionApplyType.h"
-#include "Chaos/Collision/PBDCollisionSolver.h"
-#include "Chaos/Collision/PBDCollisionSolverSettings.h"
-#include "Chaos/Evolution/SolverConstraintContainer.h"
+// @todo(chaos): Deprecate this header
+// Moved to PBDCollisionContainerSolver.h
 
-namespace Chaos
-{
-	namespace Private
-	{
-		class FPBDIsland;
-		class FPBDIslandConstraint;
-	}
-	class FPBDCollisionConstraint;
-	class FPBDCollisionConstraints;
-	class FSolverBody;
-	class FSolverBodyContainer;
-
-	/**
-	 * The solver for a set of collision constraints. This collects all the data required to solve a set of collision
-	 * constraints into a contiguous, ordered buffer.
-	*/
-	class FPBDCollisionContainerSolver : public FConstraintContainerSolver
-	{
-	public:
-		UE_NONCOPYABLE(FPBDCollisionContainerSolver);
-
-		FPBDCollisionContainerSolver(const FPBDCollisionConstraints& InConstraintContainer, const int32 InPriority);
-		~FPBDCollisionContainerSolver();
-
-		int32 NumSolvers() const { return CollisionConstraints.Num(); }
-
-		virtual void Reset(const int32 InMaxCollisions) override final;
-
-		virtual int32 GetNumConstraints() const override final { return NumSolvers(); }
-
-		//
-		// IslandGroup API
-		//
-		virtual void AddConstraints() override final;
-		virtual void AddConstraints(const TArrayView<Private::FPBDIslandConstraint>& ConstraintHandles) override final;
-		virtual void AddBodies(FSolverBodyContainer& SolverBodyContainer) override final;
-		virtual void GatherInput(const FReal Dt) override final;
-		virtual void GatherInput(const FReal Dt, const int32 BeginIndex, const int32 EndIndex) override final;
-		virtual void ScatterOutput(const FReal Dt) override final;
-		virtual void ScatterOutput(const FReal Dt, const int32 BeginIndex, const int32 EndIndex) override final;
-		virtual void ApplyPositionConstraints(const FReal Dt, const int32 It, const int32 NumIts) override final;
-		virtual void ApplyVelocityConstraints(const FReal Dt, const int32 It, const int32 NumIts) override final;
-		virtual void ApplyProjectionConstraints(const FReal Dt, const int32 It, const int32 NumIts) override final;
-
-	private:
-		void CachePrefetchSolver(const int32 ConstraintIndex) const;
-		void AddConstraint(FPBDCollisionConstraint& Constraint);
-		void UpdatePositionShockPropagation(const FReal Dt, const int32 It, const int32 NumIts, const int32 BeginIndex, const int32 EndIndex, const FPBDCollisionSolverSettings& SolverSettings);
-		void UpdateVelocityShockPropagation(const FReal Dt, const int32 It, const int32 NumIts, const int32 BeginIndex, const int32 EndIndex, const FPBDCollisionSolverSettings& SolverSettings);
-		bool SolvePositionImpl(const FReal Dt, const int32 It, const int32 NumIts, const int32 BeginIndex, const int32 EndIndex, const FPBDCollisionSolverSettings& SolverSettings);
-		bool SolveVelocityImpl(const FReal Dt, const int32 It, const int32 NumIts, const int32 BeginIndex, const int32 EndIndex, const FPBDCollisionSolverSettings& SolverSettings);
-		void UpdateCollisions(const FReal InDt, const int32 BeginIndex, const int32 EndIndex);
-
-		const FPBDCollisionConstraints& ConstraintContainer;
-		TArray<Private::FPBDCollisionSolver> CollisionSolvers;
-		TArray<Private::FPBDCollisionSolverManifoldPoint> CollisionSolverManifoldPoints;
-		TArray<FPBDCollisionConstraint*> CollisionConstraints;
-		TArray<bool> bCollisionConstraintPerIterationCollisionDetection;
-		bool bPerIterationCollisionDetection;
-	};
-}
+#include "Chaos/Collision/PBDCollisionContainerSolver.h"
diff --git a/Engine/Source/Runtime/Experimental/Chaos/Public/Chaos/PBDCollisionConstraints.h b/Engine/Source/Runtime/Experimental/Chaos/Public/Chaos/PBDCollisionConstraints.h
--- a/Engine/Source/Runtime/Experimental/Chaos/Public/Chaos/PBDCollisionConstraints.h
+++ b/Engine/Source/Runtime/Experimental/Chaos/Public/Chaos/PBDCollisionConstraints.h
@@ -35,6 +35,7 @@ namespace Private
 	enum class ECollisionSolverType
 	{
 		GaussSeidel,
+		GaussSeidelSimd,
 		PartialJacobi,
 	};
 }
