Geometry collection : fix anchoring by box functino MaxLevel not using the right attribute

diff --git a/Engine/Source/Runtime/Experimental/Chaos/Private/PhysicsProxy/GeometryCollectionPhysicsProxy.cpp b/Engine/Source/Runtime/Experimental/Chaos/Private/PhysicsProxy/GeometryCollectionPhysicsProxy.cpp
--- a/Engine/Source/Runtime/Experimental/Chaos/Private/PhysicsProxy/GeometryCollectionPhysicsProxy.cpp
+++ b/Engine/Source/Runtime/Experimental/Chaos/Private/PhysicsProxy/GeometryCollectionPhysicsProxy.cpp
@@ -28,7 +28,7 @@
 #include "GeometryCollection/GeometryCollectionSizeSpecificUtility.h"
 #include "GeometryCollection/GeometryCollectionSimulationTypes.h"
 #include "GeometryCollection/ManagedArrayCollection.h"
-#include "GeometryCollection/Facades/CollectionHierarchyFacade.h"
+#include "GeometryCollection/ManagedArrayAccessor.h"
 #include "Modules/ModuleManager.h"
 #include "Chaos/PullPhysicsDataImp.h"
 #include "Chaos/PBDRigidsEvolution.h"
@@ -1831,8 +1831,10 @@ void FGeometryCollectionPhysicsProxy::SetAnchoredByTransformedBox_External(const
 				using namespace Chaos;
 				TSet<FPBDRigidClusteredParticleHandle*> TopParentHandles;
 
-				Facades::FCollectionHierarchyFacade HierarchyFacade(PhysicsThreadCollection);
-				const int32 MaxLevelToCheck = (HierarchyFacade.HasLevelAttribute()) ? MaxLevel : INDEX_NONE;
+				static const FName InitialLevelAttributeName = "InitialLevel";
+
+				const TManagedArrayAccessor<int32> InitialLevelAttribute(PhysicsThreadCollection, InitialLevelAttributeName, FGeometryCollection::TransformGroup);
+				const int32 MaxLevelToCheck = (InitialLevelAttribute.IsValid()) ? MaxLevel : INDEX_NONE;
 
 				FPBDRigidsEvolution* Evolution = RBDSolver->GetEvolution();
 				for (int32 TransformIndex = 0; TransformIndex < SolverParticleHandles.Num(); TransformIndex++)
@@ -1842,7 +1844,7 @@ void FGeometryCollectionPhysicsProxy::SetAnchoredByTransformedBox_External(const
 						bool bCheckBounds = true;
 						if (MaxLevelToCheck > INDEX_NONE)
 						{
-							const int32 InitialLevel = HierarchyFacade.GetInitialLevel(TransformIndex);
+							const int32 InitialLevel = InitialLevelAttribute.Get()[TransformIndex];
 							bCheckBounds = (InitialLevel <= MaxLevelToCheck);
 						}
 						if (bCheckBounds)
