Fix compilation error with disabled intrinsic

diff --git a/Engine/Source/Runtime/Experimental/Chaos/Public/Chaos/VectorUtility.h b/Engine/Source/Runtime/Experimental/Chaos/Public/Chaos/VectorUtility.h
--- a/Engine/Source/Runtime/Experimental/Chaos/Public/Chaos/VectorUtility.h
+++ b/Engine/Source/Runtime/Experimental/Chaos/Public/Chaos/VectorUtility.h
@@ -12,7 +12,7 @@
 
 FORCEINLINE VectorRegister4Float VectorCast4IntTo4Float(const VectorRegister4Int& V)
 {
-#if !defined(_MSC_VER)|| PLATFORM_ENABLE_VECTORINTRINSICS_NEON
+#if (!defined(_MSC_VER) || PLATFORM_ENABLE_VECTORINTRINSICS_NEON) && PLATFORM_ENABLE_VECTORINTRINSICS
 	return VectorRegister4Float(V);
 #elif PLATFORM_ENABLE_VECTORINTRINSICS
 	return _mm_castsi128_ps(V);
@@ -30,7 +30,7 @@ FORCEINLINE VectorRegister4Float VectorCast4IntTo4Float(const VectorRegister4Int
  */
 FORCEINLINE VectorRegister4Int VectorCast4FloatTo4Int(const VectorRegister4Float& V)
 {
-#if !defined(_MSC_VER) || PLATFORM_ENABLE_VECTORINTRINSICS_NEON
+#if (!defined(_MSC_VER) || PLATFORM_ENABLE_VECTORINTRINSICS_NEON) && PLATFORM_ENABLE_VECTORINTRINSICS
 	return VectorRegister4Int(V);
 #elif PLATFORM_ENABLE_VECTORINTRINSICS
 	return _mm_castps_si128(V);
