Geometrty collection : make bImportCollisionFromSource writeable for blueprints

diff --git a/Engine/Source/Runtime/Experimental/GeometryCollectionEngine/Public/GeometryCollection/GeometryCollectionObject.h b/Engine/Source/Runtime/Experimental/GeometryCollectionEngine/Public/GeometryCollection/GeometryCollectionObject.h
--- a/Engine/Source/Runtime/Experimental/GeometryCollectionEngine/Public/GeometryCollection/GeometryCollectionObject.h
+++ b/Engine/Source/Runtime/Experimental/GeometryCollectionEngine/Public/GeometryCollection/GeometryCollectionObject.h
@@ -601,7 +601,7 @@ public:
 	/**
 	* whether to import collision from the source asset
 	*/
-	UPROPERTY(EditAnywhere, BlueprintReadOnly, Category = "Collisions")
+	UPROPERTY(EditAnywhere, BlueprintReadWrite, Category = "Collisions")
 	bool bImportCollisionFromSource;
 	
 #if WITH_EDITORONLY_DATA
