make the geometry collection go net dormant if there are no clusters to update

diff --git a/Engine/Source/Runtime/Experimental/GeometryCollectionEngine/Private/GeometryCollection/GeometryCollectionComponent.cpp b/Engine/Source/Runtime/Experimental/GeometryCollectionEngine/Private/GeometryCollection/GeometryCollectionComponent.cpp
--- a/Engine/Source/Runtime/Experimental/GeometryCollectionEngine/Private/GeometryCollection/GeometryCollectionComponent.cpp
+++ b/Engine/Source/Runtime/Experimental/GeometryCollectionEngine/Private/GeometryCollection/GeometryCollectionComponent.cpp
@@ -1693,7 +1693,7 @@ void UGeometryCollectionComponent::UpdateRepData()
 		if (bClustersChanged)
 		{
 			RepData.Clusters = MoveTemp(Clusters);
-			
+
 			if (Owner->GetWorld() && Owner->GetWorld()->GetPhysicsScene())
 			{
 				RepData.ServerFrame = Owner->GetWorld()->GetPhysicsScene()->ReplicationCache.ServerFrame;
@@ -1707,10 +1707,15 @@ void UGeometryCollectionComponent::UpdateRepData()
 			if(Owner->NetDormancy != DORM_Awake)
 			{
 				//If net dormancy is Initial it must be for perf reasons, but since a cluster changed we need to replicate down
-				//TODO: set back to dormant when sim goes to sleep
 				Owner->SetNetDormancy(DORM_Awake);
 			}
 		}
+		else if (!bFirstUpdate && Owner->NetDormancy == DORM_Awake)
+		{
+			// Clusters are no longer changing so we shouldn't need to replicate anymore.
+			// TODO: Wait a few frames?
+			Owner->SetNetDormancy(DORM_DormantAll);
+		}
 	}
 }
 
