Make sure CollisionChannel change is propagated to the PhysicProxy in GeometryCollectionComponent

diff --git a/Engine/Source/Runtime/Experimental/GeometryCollectionEngine/Private/GeometryCollection/GeometryCollectionComponent.cpp b/Engine/Source/Runtime/Experimental/GeometryCollectionEngine/Private/GeometryCollection/GeometryCollectionComponent.cpp
--- a/Engine/Source/Runtime/Experimental/GeometryCollectionEngine/Private/GeometryCollection/GeometryCollectionComponent.cpp
+++ b/Engine/Source/Runtime/Experimental/GeometryCollectionEngine/Private/GeometryCollection/GeometryCollectionComponent.cpp
@@ -2859,6 +2859,19 @@ void UGeometryCollectionComponent::SendRenderDynamicData_Concurrent()
 	}
 }
 
+void UGeometryCollectionComponent::SetCollisionObjectType(ECollisionChannel Channel)
+{
+	Super::SetCollisionObjectType(Channel);
+
+	BuildInitialFilterData();
+
+	// Update filters stored on proxy
+	if (PhysicsProxy)
+	{
+		PhysicsProxy->UpdateFilterData_External(InitialSimFilter, InitialQueryFilter);
+	}
+}
+
 void UGeometryCollectionComponent::OnActorEnableCollisionChanged()
 {
 	// Update filters on BI
diff --git a/Engine/Source/Runtime/Experimental/GeometryCollectionEngine/Public/GeometryCollection/GeometryCollectionComponent.h b/Engine/Source/Runtime/Experimental/GeometryCollectionEngine/Public/GeometryCollection/GeometryCollectionComponent.h
--- a/Engine/Source/Runtime/Experimental/GeometryCollectionEngine/Public/GeometryCollection/GeometryCollectionComponent.h
+++ b/Engine/Source/Runtime/Experimental/GeometryCollectionEngine/Public/GeometryCollection/GeometryCollectionComponent.h
@@ -444,6 +444,7 @@ public:
 	virtual void CreateRenderState_Concurrent(FRegisterComponentContext* Context) override;
 	virtual void SendRenderDynamicData_Concurrent() override;
 	FORCEINLINE void SetRenderStateDirty() { bRenderStateDirty = true; }
+	virtual void SetCollisionObjectType(ECollisionChannel Channel) override;
 	virtual void OnActorEnableCollisionChanged() override;
 	virtual void BeginPlay() override;
 	virtual void EndPlay(const EEndPlayReason::Type ReasonEnd) override;
