ChaosDeformable : Protect access to the flesh input buffer.
- Protect kinematic update from missing vertex attribute.

diff --git a/Engine/Source/Runtime/Experimental/Chaos/Private/Chaos/Deformable/ChaosDeformableSolver.cpp b/Engine/Source/Runtime/Experimental/Chaos/Private/Chaos/Deformable/ChaosDeformableSolver.cpp
--- a/Engine/Source/Runtime/Experimental/Chaos/Private/Chaos/Deformable/ChaosDeformableSolver.cpp
+++ b/Engine/Source/Runtime/Experimental/Chaos/Private/Chaos/Deformable/ChaosDeformableSolver.cpp
@@ -342,75 +342,81 @@ namespace Chaos::Softs
 						FTransform GlobalTransform = Proxy->GetInitialTransform();
 						const FIntVector2& Range = Proxy->GetSolverParticleRange();
 						const FManagedArrayCollection& Rest = Proxy->GetRestCollection();
-						const TManagedArray<FVector3f>& Vertex = Rest.GetAttribute<FVector3f>("Vertex", FGeometryCollection::VerticesGroup);
-
-						// @todo(chaos) : reduce conversions
-						auto ChaosVert = [](FVector3f V) { return Chaos::FVec3(V.X, V.Y, V.Z); };
-						auto ChaosVertfloat = [](FVector3f V) { return Chaos::TVector<FSolverReal, 3>(V.X, V.Y, V.Z); };
-						auto SolverParticleToObjectVertexIndex = [&](int32 SolverParticleIndex) {return SolverParticleIndex - Range[0]; };
-
-						FFleshThreadingProxy::FFleshInputBuffer* FleshInputBuffer = nullptr;
-						if (this->CurrentInputPackage->ObjectMap.Contains(Owner))
+						if (Rest.FindAttributeTyped<FVector3f>("Vertex", FGeometryCollection::VerticesGroup))
 						{
-							FleshInputBuffer = this->CurrentInputPackage->ObjectMap[Owner]->As<FFleshThreadingProxy::FFleshInputBuffer>();
-							GlobalTransform = FleshInputBuffer->GlobalTransform;
-						}
+							const TManagedArray<FVector3f>& Vertex = Rest.GetAttribute<FVector3f>("Vertex", FGeometryCollection::VerticesGroup);
 
-						typedef GeometryCollection::Facades::FVertexBoneWeightsFacade FWeightsFacade;
-						bool bParticleTouched = false;
-						FWeightsFacade WeightsFacade(Rest);
-						if (WeightsFacade.IsValid())
-						{
-							int32 NumObjectVertices = Rest.NumElements(FGeometryCollection::VerticesGroup);
-							int32 ObjectVertexIndex = SolverParticleToObjectVertexIndex(Index);
-							if (ensure(0 <= ObjectVertexIndex && ObjectVertexIndex < NumObjectVertices))
+							// @todo(chaos) : reduce conversions
+							auto ChaosVert = [](FVector3f V) { return Chaos::FVec3(V.X, V.Y, V.Z); };
+							auto ChaosVertfloat = [](FVector3f V) { return Chaos::TVector<FSolverReal, 3>(V.X, V.Y, V.Z); };
+							auto SolverParticleToObjectVertexIndex = [&](int32 SolverParticleIndex) {return SolverParticleIndex - Range[0]; };
+
+							FFleshThreadingProxy::FFleshInputBuffer* FleshInputBuffer = nullptr;
+							if (this->CurrentInputPackage->ObjectMap.Contains(Owner))
 							{
-								if(FleshInputBuffer)
+								FleshInputBuffer = this->CurrentInputPackage->ObjectMap[Owner]->As<FFleshThreadingProxy::FFleshInputBuffer>();
+								if (FleshInputBuffer)
 								{
-									bParticleTouched = true;
-									TArray<int32> BoneIndices = WeightsFacade.GetBoneIndices()[ObjectVertexIndex];
-									TArray<float> BoneWeights = WeightsFacade.GetBoneWeights()[ObjectVertexIndex];
+									GlobalTransform = FleshInputBuffer->GlobalTransform;
+								}
+							}
 
-									FFleshThreadingProxy::FFleshInputBuffer* PreviousFleshBuffer = nullptr;
-									if (this->PreviousInputPackage && this->PreviousInputPackage->ObjectMap.Contains(Owner))
+							typedef GeometryCollection::Facades::FVertexBoneWeightsFacade FWeightsFacade;
+							bool bParticleTouched = false;
+							FWeightsFacade WeightsFacade(Rest);
+							if (WeightsFacade.IsValid())
+							{
+								int32 NumObjectVertices = Rest.NumElements(FGeometryCollection::VerticesGroup);
+								int32 ObjectVertexIndex = SolverParticleToObjectVertexIndex(Index);
+								if (ensure(0 <= ObjectVertexIndex && ObjectVertexIndex < NumObjectVertices))
+								{
+									if (FleshInputBuffer)
 									{
-										PreviousFleshBuffer = this->PreviousInputPackage->ObjectMap[Owner]->As<FFleshThreadingProxy::FFleshInputBuffer>();
-									}
+										bParticleTouched = true;
+										TArray<int32> BoneIndices = WeightsFacade.GetBoneIndices()[ObjectVertexIndex];
+										TArray<float> BoneWeights = WeightsFacade.GetBoneWeights()[ObjectVertexIndex];
 
-									MParticles.X(Index) = Chaos::TVector<FSolverReal, 3>((FSolverReal)0.);
-									TVector<FSolverReal, 3> TargetPos((FSolverReal)0.);
-									FSolverReal CurrentRatio = FSolverReal(this->Iteration) / FSolverReal(this->Property.NumSolverSubSteps);
-
-									const TManagedArray<FTransform>& RestTransforms = Rest.GetAttribute<FTransform>(FTransformCollection::TransformAttribute, FTransformCollection::TransformGroup);
-
-									for (int32 i = BoneIndices.Num()-1; 0<=i; i--)
-									{ 
-										// @todo(flesh) : Add the pre-cached component space rest transforms to the rest collection. 
-										// see  UFleshComponent::NewDeformableData for how its pulled from the SkeletalMesh
-										FVec3 LocalPoint = FleshInputBuffer->RestTransforms[BoneIndices[i]].InverseTransformPosition(ChaosVert(Vertex[Index - Range[0]]));
-										FVec3 ComponentPointAtT = FleshInputBuffer->Transforms[BoneIndices[i]].TransformPosition(LocalPoint);
-										if (PreviousFleshBuffer)
+										FFleshThreadingProxy::FFleshInputBuffer* PreviousFleshBuffer = nullptr;
+										if (this->PreviousInputPackage && this->PreviousInputPackage->ObjectMap.Contains(Owner))
 										{
-											FTransform BonePreviousTransform = PreviousFleshBuffer->Transforms[BoneIndices[i]];
-											ComponentPointAtT = ComponentPointAtT * CurrentRatio + BonePreviousTransform.TransformPosition(LocalPoint) * ((FSolverReal)1. - CurrentRatio);
+											PreviousFleshBuffer = this->PreviousInputPackage->ObjectMap[Owner]->As<FFleshThreadingProxy::FFleshInputBuffer>();
 										}
 
-										MParticles.X(Index) = GlobalTransform.TransformPosition(ComponentPointAtT);
+										MParticles.X(Index) = Chaos::TVector<FSolverReal, 3>((FSolverReal)0.);
+										TVector<FSolverReal, 3> TargetPos((FSolverReal)0.);
+										FSolverReal CurrentRatio = FSolverReal(this->Iteration) / FSolverReal(this->Property.NumSolverSubSteps);
+
+										const TManagedArray<FTransform>& RestTransforms = Rest.GetAttribute<FTransform>(FTransformCollection::TransformAttribute, FTransformCollection::TransformGroup);
 
-										// @todo(flesh): Add non rigid skinning weights. 
-										// Currently this just grabs the last joint, ideally
-										// we would be using the weighted average for the vertex
-										// positions.
-										break;
+										for (int32 i = BoneIndices.Num() - 1; 0 <= i; i--)
+										{
+											// @todo(flesh) : Add the pre-cached component space rest transforms to the rest collection. 
+											// see  UFleshComponent::NewDeformableData for how its pulled from the SkeletalMesh
+											FVec3 LocalPoint = FleshInputBuffer->RestTransforms[BoneIndices[i]].InverseTransformPosition(ChaosVert(Vertex[Index - Range[0]]));
+											FVec3 ComponentPointAtT = FleshInputBuffer->Transforms[BoneIndices[i]].TransformPosition(LocalPoint);
+											if (PreviousFleshBuffer)
+											{
+												FTransform BonePreviousTransform = PreviousFleshBuffer->Transforms[BoneIndices[i]];
+												ComponentPointAtT = ComponentPointAtT * CurrentRatio + BonePreviousTransform.TransformPosition(LocalPoint) * ((FSolverReal)1. - CurrentRatio);
+											}
+
+											MParticles.X(Index) = GlobalTransform.TransformPosition(ComponentPointAtT);
+
+											// @todo(flesh): Add non rigid skinning weights. 
+											// Currently this just grabs the last joint, ideally
+											// we would be using the weighted average for the vertex
+											// positions.
+											break;
+										}
+										MParticles.PAndInvM(Index).P = MParticles.X(Index);
 									}
-									MParticles.PAndInvM(Index).P = MParticles.X(Index);
 								}
 							}
-						}
-						if (!bParticleTouched)
-						{
-							MParticles.X(Index) = GlobalTransform.TransformPosition(ChaosVert(Vertex[Index - Range[0]]));
-							MParticles.PAndInvM(Index).P = MParticles.X(Index);
+							if (!bParticleTouched)
+							{
+								MParticles.X(Index) = GlobalTransform.TransformPosition(ChaosVert(Vertex[Index - Range[0]]));
+								MParticles.PAndInvM(Index).P = MParticles.X(Index);
+							}
 						}
 					}
 				}
