Geometry collection : when reseting a geometry collection from another, make sure to only recreate convex hulls if necessary

diff --git a/Engine/Source/Runtime/Experimental/GeometryCollectionEngine/Private/GeometryCollection/GeometryCollectionObject.cpp b/Engine/Source/Runtime/Experimental/GeometryCollectionEngine/Private/GeometryCollection/GeometryCollectionObject.cpp
--- a/Engine/Source/Runtime/Experimental/GeometryCollectionEngine/Private/GeometryCollection/GeometryCollectionObject.cpp
+++ b/Engine/Source/Runtime/Experimental/GeometryCollectionEngine/Private/GeometryCollection/GeometryCollectionObject.cpp
@@ -299,6 +299,15 @@ void UGeometryCollection::UpdateGeometryDependentProperties()
 #endif
 }
 
+void UGeometryCollection::UpdateConvexGeometryIfMissing()
+{
+	const bool bConvexAttributeMissing = !GeometryCollection->HasAttribute("ConvexHull", "Convex");
+	if (GeometryCollection && bConvexAttributeMissing)
+	{
+		UpdateConvexGeometry();
+	}
+}
+
 void UGeometryCollection::UpdateConvexGeometry()
 {
 #if WITH_EDITOR
@@ -402,7 +411,7 @@ void UGeometryCollection::ResetFrom(const FManagedArrayCollection& InCollection,
 		InCollection.CopyTo(GeometryCollection.Get());
 
 		// todo(Chaos) : we could certainly run a "dependent attribute update method here instead of having to known about convex specifically 
-		UpdateConvexGeometry();
+		UpdateConvexGeometryIfMissing();
 				
 		Materials.Append(InMaterials);
 		InitializeMaterials(bHasInternalMaterials);
diff --git a/Engine/Source/Runtime/Experimental/GeometryCollectionEngine/Public/GeometryCollection/GeometryCollectionObject.h b/Engine/Source/Runtime/Experimental/GeometryCollectionEngine/Public/GeometryCollection/GeometryCollectionObject.h
--- a/Engine/Source/Runtime/Experimental/GeometryCollectionEngine/Public/GeometryCollection/GeometryCollectionObject.h
+++ b/Engine/Source/Runtime/Experimental/GeometryCollectionEngine/Public/GeometryCollection/GeometryCollectionObject.h
@@ -682,6 +682,7 @@ public:
 	* Update the convex geometry on the collection.
 	*/
 	void UpdateConvexGeometry();
+	void UpdateConvexGeometryIfMissing();
 	
 	/*
 	 * Update properties that depend on the geometry and clustering: Proximity, Convex Hulls, Volume and Size data.
