Iris - Don't destroy subobjects that we didn't new and actors we're not allowed to.

This fixes various crashes related to deleting default subobjects and such.

diff --git a/Engine/Plugins/Runtime/ReplicationSystemTestPlugin/Source/Private/Tests/ReplicationSystem/ReplicatedTestObject.cpp b/Engine/Plugins/Runtime/ReplicationSystemTestPlugin/Source/Private/Tests/ReplicationSystem/ReplicatedTestObject.cpp
--- a/Engine/Plugins/Runtime/ReplicationSystemTestPlugin/Source/Private/Tests/ReplicationSystem/ReplicatedTestObject.cpp
+++ b/Engine/Plugins/Runtime/ReplicationSystemTestPlugin/Source/Private/Tests/ReplicationSystem/ReplicatedTestObject.cpp
@@ -135,7 +135,7 @@ UObjectReplicationBridge::FCreationHeader* UReplicatedTestObjectBridge::ReadCrea
 	return Header.Release();
 }
 
-UObject* UReplicatedTestObjectBridge::BeginInstantiateFromRemote(FNetRefHandle SubObjectOwnerHandle, const UE::Net::FNetObjectResolveContext& ResolveContext, const FCreationHeader* InHeader)
+FObjectReplicationBridgeInstantiateResult UReplicatedTestObjectBridge::BeginInstantiateFromRemote(FNetRefHandle SubObjectOwnerHandle, const UE::Net::FNetObjectResolveContext& ResolveContext, const FCreationHeader* InHeader)
 {
 	const FReplicationTestObjectCreationHeader* Header = static_cast<const FReplicationTestObjectCreationHeader*>(InHeader);
 
@@ -163,7 +163,10 @@ UObject* UReplicatedTestObjectBridge::BeginInstantiateFromRemote(FNetRefHandle S
 		CreatedObjectsOnNode->Add(TStrongObjectPtr<UObject>(CreatedObject));
 	}
 	
-	return CreatedObject;
+	FObjectReplicationBridgeInstantiateResult InstantiateResult;
+	InstantiateResult.Object = CreatedObject;
+	InstantiateResult.Flags |= EReplicationBridgeCreateNetRefHandleResultFlags::AllowDestroyInstanceFromRemote;
+	return InstantiateResult;
 }
 
 void UReplicatedTestObjectBridge::EndInstantiateFromRemote(FNetRefHandle Handle)
@@ -173,21 +176,34 @@ void UReplicatedTestObjectBridge::EndInstantiateFromRemote(FNetRefHandle Handle)
 	Instance->NetRefHandle = Handle;
 }
 
-void UReplicatedTestObjectBridge::DestroyInstanceFromRemote(UObject* Instance, bool bTearOff)
+void UReplicatedTestObjectBridge::DestroyInstanceFromRemote(UObject* Instance, EReplicationBridgeDestroyInstanceReason DestroyReason, EReplicationBridgeDestroyInstanceFlags DestroyFlags)
 {
-	if (Instance && !bTearOff)
+	if (!Instance)
 	{
-		// Remove the object from the created objects on the node
-		if (CreatedObjectsOnNode)
+		return;
+	}
+
+	if (DestroyReason == EReplicationBridgeDestroyInstanceReason::Destroy)
+	{
+		if (EnumHasAnyFlags(DestroyFlags, EReplicationBridgeDestroyInstanceFlags::AllowDestroyInstanceFromRemote))
 		{
-			CreatedObjectsOnNode->Remove(TStrongObjectPtr<UObject>(Instance));
+			// Remove the object from the created objects on the node
+			if (CreatedObjectsOnNode)
+			{
+				CreatedObjectsOnNode->Remove(TStrongObjectPtr<UObject>(Instance));
+			}
+
+			Instance->PreDestroyFromReplication();
+			Instance->MarkAsGarbage();
 		}
-
-		Instance->PreDestroyFromReplication();
-		Instance->MarkAsGarbage();
 	}
 }
 
+bool UReplicatedTestObjectBridge::IsAllowedToDestroyInstance(const UObject* Instance) const
+{
+	return true;
+}
+
 void UReplicatedTestObjectBridge::SetPollFramePeriod(UReplicatedTestObject* Instance, uint8 FramePeriod)
 {
 	using namespace UE::Net;
diff --git a/Engine/Plugins/Runtime/ReplicationSystemTestPlugin/Source/Private/Tests/ReplicationSystem/ReplicatedTestObject.h b/Engine/Plugins/Runtime/ReplicationSystemTestPlugin/Source/Private/Tests/ReplicationSystem/ReplicatedTestObject.h
--- a/Engine/Plugins/Runtime/ReplicationSystemTestPlugin/Source/Private/Tests/ReplicationSystem/ReplicatedTestObject.h
+++ b/Engine/Plugins/Runtime/ReplicationSystemTestPlugin/Source/Private/Tests/ReplicationSystem/ReplicatedTestObject.h
@@ -431,9 +431,10 @@ protected:
 	virtual bool WriteCreationHeader(UE::Net::FNetSerializationContext& Context, FNetRefHandle Handle) override;
 	virtual FCreationHeader* ReadCreationHeader(UE::Net::FNetSerializationContext& Context) override;
 
-	virtual UObject* BeginInstantiateFromRemote(FNetRefHandle SubObjectOwnerHandle, const UE::Net::FNetObjectResolveContext& ResolveContext, const FCreationHeader* InHeader) override;
+	virtual FObjectReplicationBridgeInstantiateResult BeginInstantiateFromRemote(FNetRefHandle SubObjectOwnerHandle, const UE::Net::FNetObjectResolveContext& ResolveContext, const FCreationHeader* InHeader) override;
 	virtual void EndInstantiateFromRemote(FNetRefHandle Handle) override;
-	virtual void DestroyInstanceFromRemote(UObject* Instance, bool bTearOff);
+	virtual void DestroyInstanceFromRemote(UObject* Instance, EReplicationBridgeDestroyInstanceReason DestroyReason, EReplicationBridgeDestroyInstanceFlags DestroyFlags) override;
+	virtual bool IsAllowedToDestroyInstance(const UObject* Instance) const override;
 
 	TArray<TStrongObjectPtr<UObject>>* CreatedObjectsOnNode;
 };
diff --git a/Engine/Source/Runtime/Engine/Private/Net/Iris/ReplicationSystem/ActorReplicationBridge.cpp b/Engine/Source/Runtime/Engine/Private/Net/Iris/ReplicationSystem/ActorReplicationBridge.cpp
--- a/Engine/Source/Runtime/Engine/Private/Net/Iris/ReplicationSystem/ActorReplicationBridge.cpp
+++ b/Engine/Source/Runtime/Engine/Private/Net/Iris/ReplicationSystem/ActorReplicationBridge.cpp
@@ -541,12 +541,14 @@ UObjectReplicationBridge::FCreationHeader* UActorReplicationBridge::ReadCreation
 	return nullptr;
 }
 
-UObject* UActorReplicationBridge::BeginInstantiateFromRemote(FNetRefHandle SubObjectOwnerNetHandle, const UE::Net::FNetObjectResolveContext& ResolveContext, const UObjectReplicationBridge::FCreationHeader* InHeader)
+FObjectReplicationBridgeInstantiateResult UActorReplicationBridge::BeginInstantiateFromRemote(FNetRefHandle SubObjectOwnerNetHandle, const UE::Net::FNetObjectResolveContext& ResolveContext, const UObjectReplicationBridge::FCreationHeader* InHeader)
 {
 	using namespace UE::Net::Private;
 
 	IRIS_PROFILER_SCOPE(ActorReplicationBridge_OnBeginInstantiateFromRemote);
 
+	FObjectReplicationBridgeInstantiateResult InstantiateResult;
+
 	const FActorReplicationBridgeCreationHeader* BridgeHeader = static_cast<const FActorReplicationBridgeCreationHeader*>(InHeader);
 
 	if (BridgeHeader->bIsActor)
@@ -609,7 +611,12 @@ UObject* UActorReplicationBridge::BeginInstantiateFromRemote(FNetRefHandle SubOb
 					UE_LOG_ACTORREPLICATIONBRIDGE(Verbose, TEXT("OnBeginInstantiateFromRemote Spawned Actor %s"), *Actor->GetPathName());
 				}
 
-				return Actor;
+				InstantiateResult.Object = Actor;
+				if (NetDriver->ShouldClientDestroyActor(Actor))
+				{
+					InstantiateResult.Flags |= EReplicationBridgeCreateNetRefHandleResultFlags::AllowDestroyInstanceFromRemote;
+				}
+				return InstantiateResult;
 			}
 			else
 			{
@@ -622,7 +629,12 @@ UObject* UActorReplicationBridge::BeginInstantiateFromRemote(FNetRefHandle SubOb
 			if (UObject* Object = ResolveObjectReference(Header->ObjectReference, ResolveContext))
 			{
 				UE_LOG_ACTORREPLICATIONBRIDGE(Verbose, TEXT("OnBeginInstantiateFromRemote Found static Actor using path %s"), ToCStr(Object->GetPathName()));
-				return Object;
+				InstantiateResult.Object = Object;
+				if (NetDriver->ShouldClientDestroyActor(Cast<AActor>(Object)))
+				{
+					InstantiateResult.Flags |= EReplicationBridgeCreateNetRefHandleResultFlags::AllowDestroyInstanceFromRemote;
+				}
+				return InstantiateResult;
 			}
 			else
 			{
@@ -647,28 +659,39 @@ UObject* UActorReplicationBridge::BeginInstantiateFromRemote(FNetRefHandle SubOb
 
 			if (Header->bIsNameStableForNetworking)
 			{
-				// Resolve by finding object relative to owner
+				// Resolve by finding object relative to owner. We do not allow this object to be destroyed.
 				SubObj = ResolveObjectReference(Header->ObjectReference, ResolveContext);
+
+				if (!SubObj)
+				{
+					UE_LOG_ACTORREPLICATIONBRIDGE(Warning, TEXT("BeginInstantiateFromRemote Failed to find subobjectReference for dynamic SubObject %s, Owner %s (%s)"), *Header->ObjectReference.ToString(), *SubObjectOwnerNetHandle.ToString(), *GetPathNameSafe(GetReplicatedObject(SubObjectOwnerNetHandle)));
+				}
 			}
 			else
 			{
+				UObject* ObjOuter = Owner;
+
 				// We need to spawn the subobject
 				UObject* SubObjClassObj = ResolveObjectReference(Header->ObjectClassReference, ResolveContext);
 				UClass * SubObjClass = Cast<UClass>(SubObjClassObj);
 
 				// Try to spawn SubObject
- 				SubObj = NewObject< UObject >(Owner, SubObjClass);
+ 				SubObj = NewObject<UObject>(Owner, SubObjClass);
+
+				// Sanity check some things
+				checkf(SubObj != nullptr, TEXT("UActorReplicationBridge::BeginInstantiateFromRemote: Subobject is NULL after instantiating. Class: %s, Actor %s"), *GetNameSafe(SubObjClass), *Owner->GetName());
+				checkf(SubObj->IsIn(ObjOuter), TEXT("UActorReplicationBridge::BeginInstantiateFromRemote: Subobject is not in Outer. SubObject: %s, Actor: %s Outer: %s"), *SubObj->GetName(), *Owner->GetName(), *ObjOuter->GetName());
+				checkf(Cast<AActor>(SubObj) == nullptr, TEXT("UActorReplicationBridge::BeginInstantiateFromRemote: Subobject is an Actor. SubObject: %s, Actor: %s"), *SubObj->GetName(), *Owner->GetName());
 
 				// Notify actor that we created a component from replication
 				Owner->OnSubobjectCreatedFromReplication(SubObj);
-			}
 
-			// Sanity check some things
-			check(SubObj != nullptr);
-			check(SubObj->IsIn( Owner ));
-			check(Cast< AActor >( SubObj ) == nullptr);
+				// Created objects may be destroyed.
+				InstantiateResult.Flags |= EReplicationBridgeCreateNetRefHandleResultFlags::AllowDestroyInstanceFromRemote;
+			}
 			
-			return SubObj;
+			InstantiateResult.Object = SubObj;
+			return InstantiateResult;
 		}
 		else
 		{
@@ -676,7 +699,9 @@ UObject* UActorReplicationBridge::BeginInstantiateFromRemote(FNetRefHandle SubOb
 			if (UObject* Object = ResolveObjectReference(Header->ObjectReference, ResolveContext))
 			{
 				UE_LOG_ACTORREPLICATIONBRIDGE(Verbose, TEXT("BeginInstantiateFromRemote Found static SubObject using path %s"), ToCStr(Object->GetPathName()));
-				return Object;
+				// We do not allow this object to be destroyed.
+				InstantiateResult.Object = Object;
+				return InstantiateResult;
 			}
 			else
 			{
@@ -685,7 +710,7 @@ UObject* UActorReplicationBridge::BeginInstantiateFromRemote(FNetRefHandle SubOb
 		}
 	}
 
-	return nullptr;
+	return InstantiateResult;
 }
 
 bool UActorReplicationBridge::OnInstantiatedFromRemote(UObject* Instance, const UObjectReplicationBridge::FCreationHeader* InHeader, uint32 ConnectionId) const
@@ -731,23 +756,35 @@ void UActorReplicationBridge::EndInstantiateFromRemote(FNetRefHandle Handle)
 	}
 }
 
-void UActorReplicationBridge::DestroyInstanceFromRemote(UObject* Instance, bool bTearOff)
+void UActorReplicationBridge::DestroyInstanceFromRemote(UObject* Instance, EReplicationBridgeDestroyInstanceReason DestroyReason, EReplicationBridgeDestroyInstanceFlags DestroyFlags)
 {
+	if (DestroyReason == EReplicationBridgeDestroyInstanceReason::DoNotDestroy)
+	{
+		return;
+	}
+
 	if (AActor* Actor = Cast<AActor>(Instance))
 	{
-		if (bTearOff && !NetDriver->ShouldClientDestroyTearOffActors())
+		if ((DestroyReason == EReplicationBridgeDestroyInstanceReason::TearOff) && !NetDriver->ShouldClientDestroyTearOffActors())
 		{
-			Actor->SetRole(ROLE_Authority);
-			Actor->SetReplicates(false);
+			if (Actor->GetRemoteRole() == ROLE_Authority)
+			{
+				Actor->SetRole(ROLE_Authority);
+				Actor->SetReplicates(false);
 
-			if (Actor->GetWorld() != nullptr && !IsEngineExitRequested())
+				if (Actor->GetWorld() != nullptr && !IsEngineExitRequested())
+				{
+					Actor->TornOff();
+				}
+
+				NetDriver->NotifyActorTornOff(Actor);
+			}
+			else
 			{
-				Actor->TornOff();
+				UE_LOG_ACTORREPLICATIONBRIDGE(Warning, TEXT("Trying to tear off actor which doesn't support it: %s"), ToCStr(GetNameSafe(Actor)));
 			}
-
-			NetDriver->NotifyActorTornOff(Actor);
 		}
-		else
+		else if (EnumHasAnyFlags(DestroyFlags, EReplicationBridgeDestroyInstanceFlags::AllowDestroyInstanceFromRemote))
 		{
 			// Any subobjects have already been detached from ReplicationBridge
 			Actor->PreDestroyFromReplication();
@@ -756,8 +793,13 @@ void UActorReplicationBridge::DestroyInstanceFromRemote(UObject* Instance, bool
 	}
 	else
 	{
-		// If the SubObject is being torn-off, it is up to owning actor to clean it up properly
-		if (bTearOff && !NetDriver->ShouldClientDestroyTearOffActors())
+		// If the SubObject is being torn off it is up to owning actor to clean it up properly
+		if (DestroyReason == EReplicationBridgeDestroyInstanceReason::TearOff)
+		{
+			return;
+		}
+
+		if (!EnumHasAnyFlags(DestroyFlags, EReplicationBridgeDestroyInstanceFlags::AllowDestroyInstanceFromRemote))
 		{
 			return;
 		}
@@ -992,6 +1034,16 @@ bool UActorReplicationBridge::ObjectLevelHasFinishedLoading(UObject* Object) con
 	return true;
 }
 
+bool UActorReplicationBridge::IsAllowedToDestroyInstance(const UObject* Instance) const
+{
+	if (AActor* Actor = const_cast<AActor*>(Cast<AActor>(Instance)))
+	{
+		return !NetDriver || NetDriver->ShouldClientDestroyActor(Actor);
+	}
+
+	return true;
+}
+
 UActorReplicationBridge* UActorReplicationBridge::Create(UNetDriver* NetDriver)
 {
 	UActorReplicationBridge* Bridge = NewObject<UActorReplicationBridge>(GetTransientPackage(), UActorReplicationBridge::StaticClass());
diff --git a/Engine/Source/Runtime/Engine/Public/Net/Iris/ReplicationSystem/ActorReplicationBridge.h b/Engine/Source/Runtime/Engine/Public/Net/Iris/ReplicationSystem/ActorReplicationBridge.h
--- a/Engine/Source/Runtime/Engine/Public/Net/Iris/ReplicationSystem/ActorReplicationBridge.h
+++ b/Engine/Source/Runtime/Engine/Public/Net/Iris/ReplicationSystem/ActorReplicationBridge.h
@@ -84,13 +84,14 @@ protected:
 	virtual void Initialize(UReplicationSystem* ReplicationSystem) override;
 	virtual bool WriteCreationHeader(UE::Net::FNetSerializationContext& Context, FNetRefHandle Handle) override;
 	virtual FCreationHeader* ReadCreationHeader(UE::Net::FNetSerializationContext& Context) override;
-	virtual UObject* BeginInstantiateFromRemote(FNetRefHandle SubObjectOwnerNetHandle, const UE::Net::FNetObjectResolveContext& ResolveContext, const FCreationHeader* InHeader) override;
+	virtual FObjectReplicationBridgeInstantiateResult BeginInstantiateFromRemote(FNetRefHandle SubObjectOwnerNetHandle, const UE::Net::FNetObjectResolveContext& ResolveContext, const FCreationHeader* InHeader) override;
 	virtual bool OnInstantiatedFromRemote(UObject* Instance, const FCreationHeader* InHeader, uint32 ConnectionId) const override;
 	virtual void EndInstantiateFromRemote(FNetRefHandle Handle) override;
-	virtual void DestroyInstanceFromRemote(UObject* Instance, bool bTearOff) override;
+	virtual void DestroyInstanceFromRemote(UObject* Instance, EReplicationBridgeDestroyInstanceReason DestroyReason, EReplicationBridgeDestroyInstanceFlags DestroyFlags) override;
 	virtual void GetInitialDependencies(FNetRefHandle Handle, FNetDependencyInfoArray& OutDependencies) const override;
 	virtual bool RemapPathForPIE(uint32 ConnectionId, FString& Path, bool bReading) const override;
 	virtual bool ObjectLevelHasFinishedLoading(UObject* Object) const override;
+	virtual bool IsAllowedToDestroyInstance(const UObject* Instance) const override;
 
 private:
 	void GetActorCreationHeader(const AActor* Actor, UE::Net::Private::FActorCreationHeader& Header) const;
diff --git a/Engine/Source/Runtime/Experimental/Iris/Core/Private/Iris/ReplicationSystem/NetRefHandleManager.h b/Engine/Source/Runtime/Experimental/Iris/Core/Private/Iris/ReplicationSystem/NetRefHandleManager.h
--- a/Engine/Source/Runtime/Experimental/Iris/Core/Private/Iris/ReplicationSystem/NetRefHandleManager.h
+++ b/Engine/Source/Runtime/Experimental/Iris/Core/Private/Iris/ReplicationSystem/NetRefHandleManager.h
@@ -95,14 +95,16 @@ public:
 				uint32 bDestroySubObjectWithOwner : 1U;
 				uint32 bIsDependentObject : 1U;
 				uint32 bHasDependentObjects : 1U;
-				uint32 Padding : 27U;
+				uint32 bAllowDestroyInstanceFromRemote : 1U;
+				// Padding. Adjust when adding or removing flags.
+				uint32 Padding : 26U;
 			};
 		};
 	
 		// Returns true if this is a SubObject
 		bool IsSubObject() const { return SubObjectRootIndex && !!bDestroySubObjectWithOwner; }
 
-		// Returns true if this is a DepedentObject
+		// Returns true if this is a DependentObject
 		bool IsDependentObject() const { return bIsDependentObject; }
 	};
 
diff --git a/Engine/Source/Runtime/Experimental/Iris/Core/Private/Iris/ReplicationSystem/ObjectReplicationBridge.cpp b/Engine/Source/Runtime/Experimental/Iris/Core/Private/Iris/ReplicationSystem/ObjectReplicationBridge.cpp
--- a/Engine/Source/Runtime/Experimental/Iris/Core/Private/Iris/ReplicationSystem/ObjectReplicationBridge.cpp
+++ b/Engine/Source/Runtime/Experimental/Iris/Core/Private/Iris/ReplicationSystem/ObjectReplicationBridge.cpp
@@ -472,17 +472,35 @@ void UObjectReplicationBridge::EndReplication(UObject* Instance, EEndReplication
 	UReplicationBridge::EndReplication(Handle, EndReplicationFlags, Parameters);
 }
 
-void UObjectReplicationBridge::DetachInstanceFromRemote(FNetRefHandle Handle, bool bTearOff, bool bShouldDestroyInstance)
+void UObjectReplicationBridge::DetachInstanceFromRemote(FNetRefHandle Handle, EReplicationBridgeDestroyInstanceReason DestroyReason, EReplicationBridgeDestroyInstanceFlags DestroyFlags)
 {
-	UE_LOG_OBJECTREPLICATIONBRIDGE(Verbose, TEXT("OnDetachInstanceFromRemote %s TearOff: %u bShouldDestroyInstance: %u"), *Handle.ToString(), (uint32)bTearOff, (uint32)bShouldDestroyInstance);
-	UnregisterRemoteInstance(Handle, bTearOff, bShouldDestroyInstance);
+	UE_LOG_OBJECTREPLICATIONBRIDGE(Verbose, TEXT("DetachInstanceFromRemote %s DestroyReason: %s DestroyFlags: %u"), *Handle.ToString(), ToCStr(LexToString(DestroyReason)), unsigned(DestroyFlags));
+	
+	UObject* Instance = GetObjectFromReferenceHandle(Handle);
+	
+	UnregisterInstance(Handle);
+
+	// Destroy instance if requested
+	if (Instance && DestroyReason != EReplicationBridgeDestroyInstanceReason::DoNotDestroy)
+	{
+		DestroyInstanceFromRemote(Instance, DestroyReason, DestroyFlags);
+	}
+
+	// $IRIS TODO: Cleanup any pending creation data if we have not yet instantiated the instance.
+}
+
+void UObjectReplicationBridge::DetachInstance(FNetRefHandle RefHandle)
+{
+	UnregisterInstance(RefHandle);
 }
 
-void UObjectReplicationBridge::DetachInstance(FNetRefHandle Handle)
+void UObjectReplicationBridge::UnregisterInstance(FNetRefHandle RefHandle)
 {
-	constexpr bool bTearOff = false;
-	constexpr bool bShouldDestroyInstance = false;
-	UnregisterRemoteInstance(Handle, bTearOff, bShouldDestroyInstance);
+	if (RefHandle.IsDynamic())
+	{
+		UObject* Instance = GetObjectFromReferenceHandle(RefHandle);
+		GetObjectReferenceCache()->RemoveReference(RefHandle, Instance);
+	}
 }
 
 void UObjectReplicationBridge::RegisterRemoteInstance(FNetRefHandle RefHandle, UObject* Instance, const UE::Net::FReplicationProtocol* Protocol, UE::Net::FReplicationInstanceProtocol* InstanceProtocol, const FCreationHeader* Header, uint32 ConnectionId)
@@ -500,27 +518,7 @@ void UObjectReplicationBridge::RegisterRemoteInstance(FNetRefHandle RefHandle, U
 	UE_LOG_OBJECTREPLICATIONBRIDGE(Verbose, TEXT("RegisterRemoteInstance %s %s with ProtocolId:0x%" UINT64_x_FMT), *RefHandle.ToString(), ToCStr(Instance->GetName()), Protocol->ProtocolIdentifier);
 }
 
-void UObjectReplicationBridge::UnregisterRemoteInstance(FNetRefHandle RefHandle, bool bTearOff, bool bShouldDestroyInstance)
-{
-	// Lookup the instance and remove it	
-	UObject* Instance = GetObjectFromReferenceHandle(RefHandle);
-	
-	// Try to remove any references to dynamic objects
-	if (RefHandle.IsDynamic())
-	{
-		GetObjectReferenceCache()->RemoveReference(RefHandle, Instance);
-	}
-
-	// Destroy instance if we should	
-	if ((bTearOff || bShouldDestroyInstance) && Instance)
-	{
-		DestroyInstanceFromRemote(Instance, bTearOff);
-	}
-
-	// $TODO: Cleanup any pending creation data if we have not yet instantiated the instance
-}
-	
-UE::Net::FNetRefHandle UObjectReplicationBridge::CreateNetRefHandleFromRemote(FNetRefHandle SubObjectOwnerNetHandle, FNetRefHandle WantedNetHandle, FReplicationBridgeSerializationContext& Context)
+FReplicationBridgeCreateNetRefHandleResult UObjectReplicationBridge::CreateNetRefHandleFromRemote(FNetRefHandle SubObjectOwnerNetHandle, FNetRefHandle WantedNetHandle, FReplicationBridgeSerializationContext& Context)
 {
 	LLM_SCOPE_BYTAG(IrisState);
 
@@ -535,7 +533,7 @@ UE::Net::FNetRefHandle UObjectReplicationBridge::CreateNetRefHandleFromRemote(FN
 	TUniquePtr<FCreationHeader> Header(ReadCreationHeader(Context.SerializationContext));
 	if (Context.SerializationContext.HasErrorOrOverflow())
 	{
-		return FNetRefHandle();
+		return FReplicationBridgeCreateNetRefHandleResult();
 	}
 
 	// Currently remote objects can only receive replicated data
@@ -548,14 +546,19 @@ UE::Net::FNetRefHandle UObjectReplicationBridge::CreateNetRefHandleFromRemote(FN
 	const bool bVerifyExistingProtocol = true;
 #endif
 
+	FReplicationBridgeCreateNetRefHandleResult CreateResult;
+
 	// Currently we need to always instantiate remote objects, moving forward we want to make this optional so that can be deferred until it is time to apply received state data.
 	// https://jira.it.epicgames.com/browse/UE-127369	
-	UObject* InstancePtr = BeginInstantiateFromRemote(SubObjectOwnerNetHandle, Context.SerializationContext.GetInternalContext()->ResolveContext, Header.Get());
+	FObjectReplicationBridgeInstantiateResult InstantiateResult = BeginInstantiateFromRemote(SubObjectOwnerNetHandle, Context.SerializationContext.GetInternalContext()->ResolveContext, Header.Get());
+	UObject* InstancePtr = InstantiateResult.Object;
 	if (!ensureAlwaysMsgf(InstancePtr, TEXT("Failed to instantiate Handle: %s"), *WantedNetHandle.ToString()))
 	{
-		return FNetRefHandle();
+		return CreateResult;
 	}
 
+	CreateResult.Flags |= InstantiateResult.Flags;
+
 	// Register all fragments
 	CallRegisterReplicationFragments(InstancePtr, FragmentRegistrationContext, EFragmentRegistrationFlags::None);
 
@@ -578,7 +581,7 @@ UE::Net::FNetRefHandle UObjectReplicationBridge::CreateNetRefHandleFromRemote(FN
 	{
 		if (!ensureAlways(!bVerifyExistingProtocol || ProtocolManager->ValidateReplicationProtocol(ReplicationProtocol, RegisteredFragments)))
 		{
-			return FNetRefHandle();
+			return CreateResult;
 		}
 	}
 
@@ -586,6 +589,7 @@ UE::Net::FNetRefHandle UObjectReplicationBridge::CreateNetRefHandleFromRemote(FN
 	{
 		// Create NetHandle
 		FNetRefHandle Handle = InternalCreateNetObjectFromRemote(WantedNetHandle, ReplicationProtocol);
+		CreateResult.NetRefHandle = Handle;
 		if (Handle.IsValid())
 		{
 			RegisterRemoteInstance(Handle, InstancePtr, ReplicationProtocol, InstanceProtocol.Get(), Header.Get(), Context.ConnectionId);
@@ -595,12 +599,10 @@ UE::Net::FNetRefHandle UObjectReplicationBridge::CreateNetRefHandleFromRemote(FN
 
 			// Now it is safe to issue OnActorChannelOpen callback
 			ensureAlwaysMsgf(OnInstantiatedFromRemote(InstancePtr, Header.Get(), Context.ConnectionId), TEXT("Failed to invoke OnInstantiatedFromRemote for Instance named %s %s"), *InstancePtr->GetName(), *Handle.ToString());
-
-			return Handle;
 		}
 	}
 
-	return FNetRefHandle();
+	return CreateResult;
 }
 
 void UObjectReplicationBridge::PostApplyInitialState(FNetRefHandle Handle)
diff --git a/Engine/Source/Runtime/Experimental/Iris/Core/Private/Iris/ReplicationSystem/ReplicationBridge.cpp b/Engine/Source/Runtime/Experimental/Iris/Core/Private/Iris/ReplicationSystem/ReplicationBridge.cpp
--- a/Engine/Source/Runtime/Experimental/Iris/Core/Private/Iris/ReplicationSystem/ReplicationBridge.cpp
+++ b/Engine/Source/Runtime/Experimental/Iris/Core/Private/Iris/ReplicationSystem/ReplicationBridge.cpp
@@ -1,7 +1,6 @@
 // Copyright Epic Games, Inc. All Rights Reserved.
 
 #include "Iris/ReplicationSystem/ReplicationBridge.h"
-#include "CoreTypes.h"
 #include "Containers/ArrayView.h"
 #include "Iris/IrisConfigInternal.h"
 #include "Iris/Core/IrisLog.h"
@@ -48,9 +47,9 @@ bool UReplicationBridge::WriteNetRefHandleCreationInfo(FReplicationBridgeSeriali
 	return true;
 }
 
-UE::Net::FNetRefHandle UReplicationBridge::CreateNetRefHandleFromRemote(FNetRefHandle SubObjectOwnerNetHandle, FNetRefHandle WantedNetHandle, FReplicationBridgeSerializationContext& Context)
+FReplicationBridgeCreateNetRefHandleResult UReplicationBridge::CreateNetRefHandleFromRemote(FNetRefHandle SubObjectOwnerNetHandle, FNetRefHandle WantedNetHandle, FReplicationBridgeSerializationContext& Context)
 {
-	return FNetRefHandle();
+	return FReplicationBridgeCreateNetRefHandleResult();
 };
 
 void UReplicationBridge::PreSendUpdateSingleHandle(FNetRefHandle Handle)
@@ -69,7 +68,7 @@ void UReplicationBridge::PostApplyInitialState(FNetRefHandle Handle)
 {
 }
 
-void UReplicationBridge::DetachInstanceFromRemote(FNetRefHandle Handle, bool bTearOff, bool bShouldDestroyInstance)
+void UReplicationBridge::DetachInstanceFromRemote(FNetRefHandle Handle, EReplicationBridgeDestroyInstanceReason DestroyReason, EReplicationBridgeDestroyInstanceFlags DestroyFlags)
 {
 }
 
@@ -106,7 +105,7 @@ void UReplicationBridge::CallGetInitialDependencies(FNetRefHandle Handle, FNetDe
 	}
 }
 
-void UReplicationBridge::DetachSubObjectInstancesFromRemote(FNetRefHandle OwnerHandle, bool bTearOff, bool bShouldDestroyInstance)
+void UReplicationBridge::DetachSubObjectInstancesFromRemote(FNetRefHandle OwnerHandle, EReplicationBridgeDestroyInstanceReason DestroyReason, EReplicationBridgeDestroyInstanceFlags DestroyFlags)
 {
 	using namespace UE::Net::Private;
 
@@ -117,14 +116,20 @@ void UReplicationBridge::DetachSubObjectInstancesFromRemote(FNetRefHandle OwnerH
 		{
 			FNetRefHandleManager::FReplicatedObjectData& SubObjectData = NetRefHandleManager->GetReplicatedObjectDataNoCheck(SubObjectInternalIndex);
 			const FNetRefHandle SubObjectHandle = SubObjectData.RefHandle;
-			SubObjectData.bTearOff = bTearOff;
+			SubObjectData.bTearOff = (DestroyReason == EReplicationBridgeDestroyInstanceReason::TearOff);
 
-			CallDetachInstanceFromRemote(SubObjectHandle, bTearOff, bShouldDestroyInstance);
+			EReplicationBridgeDestroyInstanceFlags SubObjectDestroyFlags = DestroyFlags;
+			// The subobject is allowed to be destroyed if both the owner and the subobject allows it.
+			if (!SubObjectData.bAllowDestroyInstanceFromRemote)
+			{
+				EnumRemoveFlags(SubObjectDestroyFlags, EReplicationBridgeDestroyInstanceFlags::AllowDestroyInstanceFromRemote);
+			}
+			CallDetachInstanceFromRemote(SubObjectHandle, DestroyReason, SubObjectDestroyFlags);
 		}
 	}
 }
 
-void UReplicationBridge::DestroyNetObjectFromRemote(FNetRefHandle Handle, bool bTearOff, bool bDestroyInstance)
+void UReplicationBridge::DestroyNetObjectFromRemote(FNetRefHandle Handle, EReplicationBridgeDestroyInstanceReason DestroyReason, EReplicationBridgeDestroyInstanceFlags DestroyFlags)
 {
 	using namespace UE::Net::Private;
 
@@ -132,14 +137,14 @@ void UReplicationBridge::DestroyNetObjectFromRemote(FNetRefHandle Handle, bool b
 	{
 		FInternalNetRefIndex OwnerInternalIndex = NetRefHandleManager->GetInternalIndex(Handle);
 		FNetRefHandleManager::FReplicatedObjectData& ObjectData = NetRefHandleManager->GetReplicatedObjectDataNoCheck(OwnerInternalIndex);
-		ObjectData.bTearOff = bTearOff;
+		ObjectData.bTearOff = (DestroyReason == EReplicationBridgeDestroyInstanceReason::TearOff);
 
 		// if the a subobject owner is to be destroyed we want to detach all subobjects before doing so to ensure we execute expected callbacks
 		// We keep tracking them internally
-		DetachSubObjectInstancesFromRemote(Handle, bTearOff, bDestroyInstance);
+		DetachSubObjectInstancesFromRemote(Handle, DestroyReason, DestroyFlags);
 		
 		// Allow derived bridges to cleanup any instance info they have stored
-		CallDetachInstanceFromRemote(Handle, bTearOff, bDestroyInstance);
+		CallDetachInstanceFromRemote(Handle, DestroyReason, DestroyFlags);
 
 		// Detach instance protocol
 		InternalDetachInstanceFromNetRefHandle(Handle);
@@ -174,9 +179,9 @@ void UReplicationBridge::CallDetachInstance(FNetRefHandle Handle)
 	DetachInstance(Handle);
 }
 
-void UReplicationBridge::CallDetachInstanceFromRemote(FNetRefHandle Handle, bool bTearOff, bool bShouldDestroyInstance)
+void UReplicationBridge::CallDetachInstanceFromRemote(FNetRefHandle Handle, EReplicationBridgeDestroyInstanceReason DestroyReason, EReplicationBridgeDestroyInstanceFlags DestroyFlags)
 {
-	DetachInstanceFromRemote(Handle, bTearOff, bShouldDestroyInstance);
+	DetachInstanceFromRemote(Handle, DestroyReason, DestroyFlags);
 }
 
 void UReplicationBridge::CallPruneStaleObjects()
@@ -212,21 +217,24 @@ bool UReplicationBridge::CallWriteNetRefHandleCreationInfo(FReplicationBridgeSer
 	}
 }
 
-UE::Net::FNetRefHandle UReplicationBridge::CallCreateNetRefHandleFromRemote(FNetRefHandle SubObjectOwnerHandle, FNetRefHandle WantedNetHandle, FReplicationBridgeSerializationContext& Context)
+FReplicationBridgeCreateNetRefHandleResult UReplicationBridge::CallCreateNetRefHandleFromRemote(FNetRefHandle SubObjectOwnerHandle, FNetRefHandle WantedNetHandle, FReplicationBridgeSerializationContext& Context)
 {
-	using namespace UE::Net;
-
 	check(!Context.bIsDestructionInfo);
 
-	FNetRefHandle ResultingHandle = CreateNetRefHandleFromRemote(SubObjectOwnerHandle, WantedNetHandle, Context);
+	FReplicationBridgeCreateNetRefHandleResult CreateResult = CreateNetRefHandleFromRemote(SubObjectOwnerHandle, WantedNetHandle, Context);
 
 	// Track subobjects on clients
-	if (ResultingHandle.IsValid() && SubObjectOwnerHandle.IsValid())
+	if (CreateResult.NetRefHandle.IsValid() && SubObjectOwnerHandle.IsValid())
 	{
-		NetRefHandleManager->AddSubObject(SubObjectOwnerHandle, ResultingHandle);
+		NetRefHandleManager->AddSubObject(SubObjectOwnerHandle, CreateResult.NetRefHandle);
 	}
 
-	return ResultingHandle;
+	return CreateResult;
+}
+
+bool UReplicationBridge::IsAllowedToDestroyInstance(const UObject* Instance) const
+{
+	return true;
 }
 
 void UReplicationBridge::ReadAndExecuteDestructionInfoFromRemote(FReplicationBridgeSerializationContext& Context)
@@ -243,12 +251,11 @@ void UReplicationBridge::ReadAndExecuteDestructionInfoFromRemote(FReplicationBri
 
 	// Destroy the reference
 	// Resolve the reference in order to be able to destroy it
-	if (ObjectReferenceCache->ResolveObjectReference(ReferenceToDestroy, Context.SerializationContext.GetInternalContext()->ResolveContext) != nullptr)
+	if (const UObject* Instance = ObjectReferenceCache->ResolveObjectReference(ReferenceToDestroy, Context.SerializationContext.GetInternalContext()->ResolveContext))
 	{
-		// Need to forward this as we do not really know how to destroy the actual instance
-		constexpr bool bTearOff = false;
-		constexpr bool bShouldDestroyInstance = true;
-		CallDetachInstanceFromRemote(ReferenceToDestroy.GetRefHandle(), bTearOff, bShouldDestroyInstance);
+		constexpr EReplicationBridgeDestroyInstanceReason DestroyReason = EReplicationBridgeDestroyInstanceReason::Destroy;
+		const EReplicationBridgeDestroyInstanceFlags DestroyFlags = IsAllowedToDestroyInstance(Instance) ? EReplicationBridgeDestroyInstanceFlags::AllowDestroyInstanceFromRemote : EReplicationBridgeDestroyInstanceFlags::None;
+		CallDetachInstanceFromRemote(ReferenceToDestroy.GetRefHandle(), DestroyReason, DestroyFlags);
 	}
 }
 
@@ -956,3 +963,26 @@ void UReplicationBridge::ClearNetPushIds(UE::Net::Private::FInternalNetRefIndex
 	}
 #endif
 }
+
+const TCHAR* LexToString(EReplicationBridgeDestroyInstanceReason Reason)
+{
+	switch (Reason)
+	{
+	case EReplicationBridgeDestroyInstanceReason::DoNotDestroy:
+	{
+		return TEXT("DoNotDestroy");
+	}
+	case EReplicationBridgeDestroyInstanceReason::TearOff:
+	{
+		return TEXT("TearOff");
+	}
+	case EReplicationBridgeDestroyInstanceReason::Destroy:
+	{
+		return TEXT("Destroy");
+	}
+	default:
+	{
+		return TEXT("[Invalid]");
+	}
+	};
+}
diff --git a/Engine/Source/Runtime/Experimental/Iris/Core/Private/Iris/ReplicationSystem/ReplicationReader.cpp b/Engine/Source/Runtime/Experimental/Iris/Core/Private/Iris/ReplicationSystem/ReplicationReader.cpp
--- a/Engine/Source/Runtime/Experimental/Iris/Core/Private/Iris/ReplicationSystem/ReplicationReader.cpp
+++ b/Engine/Source/Runtime/Experimental/Iris/Core/Private/Iris/ReplicationSystem/ReplicationReader.cpp
@@ -300,7 +300,18 @@ void FReplicationReader::EndReplication(uint32 InternalIndex, bool bTearOff, boo
 
 		CleanupReferenceTracking(ObjectInfo);
 		Attachments.DropAllAttachments(ENetObjectAttachmentType::Normal, InternalIndex);
-		ReplicationBridge->DestroyNetObjectFromRemote(Data.RefHandle, bTearOff, bDestroyInstance);
+
+		EReplicationBridgeDestroyInstanceReason DestroyReason = EReplicationBridgeDestroyInstanceReason::DoNotDestroy;
+		if (bDestroyInstance)
+		{
+			DestroyReason = EReplicationBridgeDestroyInstanceReason::Destroy;
+		}
+		else if (bTearOff)
+		{
+			DestroyReason = EReplicationBridgeDestroyInstanceReason::TearOff;
+		}
+		EReplicationBridgeDestroyInstanceFlags DestroyFlags = (Data.bAllowDestroyInstanceFromRemote ? EReplicationBridgeDestroyInstanceFlags::AllowDestroyInstanceFromRemote : EReplicationBridgeDestroyInstanceFlags::None);
+		ReplicationBridge->DestroyNetObjectFromRemote(Data.RefHandle, DestroyReason, DestroyFlags);
 
 		CleanupObjectData(*ObjectInfo);
 
@@ -450,8 +461,8 @@ void FReplicationReader::ReadObject(FNetSerializationContext& Context)
 		// Get Bridge
 		FReplicationBridgeSerializationContext BridgeContext(Context, Parameters.ConnectionId);
 
-		FNetRefHandle NetRefHandle = ReplicationBridge->CallCreateNetRefHandleFromRemote(SubObjectOwnerHandle, IncompleteHandle, BridgeContext);
-
+		const FReplicationBridgeCreateNetRefHandleResult CreateResult = ReplicationBridge->CallCreateNetRefHandleFromRemote(SubObjectOwnerHandle, IncompleteHandle, BridgeContext);
+		FNetRefHandle NetRefHandle = CreateResult.NetRefHandle;
 		if (!NetRefHandle.IsValid())
 		{
 #if UE_NET_REPLICATION_SUPPORT_SKIP_INITIAL_STATE
diff --git a/Engine/Source/Runtime/Experimental/Iris/Core/Public/Iris/ReplicationSystem/ObjectReplicationBridge.h b/Engine/Source/Runtime/Experimental/Iris/Core/Public/Iris/ReplicationSystem/ObjectReplicationBridge.h
--- a/Engine/Source/Runtime/Experimental/Iris/Core/Public/Iris/ReplicationSystem/ObjectReplicationBridge.h
+++ b/Engine/Source/Runtime/Experimental/Iris/Core/Public/Iris/ReplicationSystem/ObjectReplicationBridge.h
@@ -25,6 +25,12 @@ namespace UE::Net
 	}
 }
 
+struct FObjectReplicationBridgeInstantiateResult
+{
+	UObject* Object = nullptr;
+	EReplicationBridgeCreateNetRefHandleResultFlags Flags = EReplicationBridgeCreateNetRefHandleResultFlags::None;
+};
+
 /*
 * Partial implementation of ReplicationBridge that can be used as a foundation for 
 * implementing support for replicating objects derived from UObject
@@ -149,9 +155,9 @@ protected:
 	IRISCORE_API virtual void UpdateInstancesWorldLocation() override;
 	IRISCORE_API virtual void PruneStaleObjects() override;	
 	IRISCORE_API virtual bool WriteNetRefHandleCreationInfo(FReplicationBridgeSerializationContext& Context, FNetRefHandle Handle) override;
-	IRISCORE_API virtual FNetRefHandle CreateNetRefHandleFromRemote(FNetRefHandle SubObjectOwnerNetHandle, FNetRefHandle WantedNetHandle, FReplicationBridgeSerializationContext& Context) override;
+	IRISCORE_API virtual FReplicationBridgeCreateNetRefHandleResult CreateNetRefHandleFromRemote(FNetRefHandle SubObjectOwnerNetHandle, FNetRefHandle WantedNetHandle, FReplicationBridgeSerializationContext& Context) override;
 	IRISCORE_API virtual void PostApplyInitialState(FNetRefHandle Handle) override;
-	IRISCORE_API virtual void DetachInstanceFromRemote(FNetRefHandle Handle, bool bTearOff, bool bShouldDestroyInstance) override;
+	IRISCORE_API virtual void DetachInstanceFromRemote(FNetRefHandle Handle, EReplicationBridgeDestroyInstanceReason DestroyReason, EReplicationBridgeDestroyInstanceFlags DestroyFlags) override;
 	IRISCORE_API virtual void DetachInstance(FNetRefHandle Handle) override;
 
 protected:
@@ -169,7 +175,7 @@ protected:
 	virtual FCreationHeader* ReadCreationHeader(UE::Net::FNetSerializationContext& Context) { return nullptr; };
 
 	/** Called when we instantiate/find object instance requested by remote. */
-	virtual UObject* BeginInstantiateFromRemote(FNetRefHandle SubObjectOwnerHandle, const UE::Net::FNetObjectResolveContext& ResolveContext, const FCreationHeader* Header) { return nullptr; };
+	virtual FObjectReplicationBridgeInstantiateResult BeginInstantiateFromRemote(FNetRefHandle SubObjectOwnerHandle, const UE::Net::FNetObjectResolveContext& ResolveContext, const FCreationHeader* Header) { return FObjectReplicationBridgeInstantiateResult(); };
 
 	/** Invoked before we start applying state data to instance on remote end. */
 	virtual bool OnInstantiatedFromRemote(UObject* Instance, const FCreationHeader* InHeader, uint32 ConnectionId) const { return true; }
@@ -178,7 +184,7 @@ protected:
 	virtual void EndInstantiateFromRemote(FNetRefHandle Handle) {};
 
 	/** Destroy or tear-off the game instance on request from remote. */
-	virtual void DestroyInstanceFromRemote(UObject* Instance, bool bTearOff) {};
+	virtual void DestroyInstanceFromRemote(UObject* Instance, EReplicationBridgeDestroyInstanceReason DestroyReason, EReplicationBridgeDestroyInstanceFlags DestroyFlags) {}
 
 protected:
 	/** Lookup the UObject associated with the provided Handle. This function will not try to resolve the reference. */
@@ -220,9 +226,10 @@ private:
 
 	void PreUpdateAndPollImpl(FNetRefHandle RefHandle);
 
-	void RegisterRemoteInstance(FNetRefHandle RefHandle, UObject* InstancePtr, const UE::Net::FReplicationProtocol* Protocol, UE::Net::FReplicationInstanceProtocol* InstanceProtocol, const FCreationHeader* Header, uint32 ConnectionId);
+	/** Remove mapping between handle and object instance. */
+	void UnregisterInstance(FNetRefHandle RefHandle);
 
-	void UnregisterRemoteInstance(FNetRefHandle RefHandle, bool bTearOff, bool bShouldDestroyInstance);
+	void RegisterRemoteInstance(FNetRefHandle RefHandle, UObject* InstancePtr, const UE::Net::FReplicationProtocol* Protocol, UE::Net::FReplicationInstanceProtocol* InstanceProtocol, const FCreationHeader* Header, uint32 ConnectionId);
 
 	void SetNetPushIdOnInstance(UE::Net::FReplicationInstanceProtocol* InstanceProtocol, FNetHandle NetHandle);
 
diff --git a/Engine/Source/Runtime/Experimental/Iris/Core/Public/Iris/ReplicationSystem/ReplicationBridge.h b/Engine/Source/Runtime/Experimental/Iris/Core/Public/Iris/ReplicationSystem/ReplicationBridge.h
--- a/Engine/Source/Runtime/Experimental/Iris/Core/Public/Iris/ReplicationSystem/ReplicationBridge.h
+++ b/Engine/Source/Runtime/Experimental/Iris/Core/Public/Iris/ReplicationSystem/ReplicationBridge.h
@@ -77,6 +77,36 @@ enum class EEndReplicationFlags : uint32
 };
 ENUM_CLASS_FLAGS(EEndReplicationFlags);
 
+enum class EReplicationBridgeCreateNetRefHandleResultFlags : unsigned
+{
+	None = 0U,
+	/** Whether the instance may be destroyed due to the remote peer requesting the object to be destroyed. If not then the object itself must not be destroyed. */
+	AllowDestroyInstanceFromRemote = 1U << 0U,
+};
+ENUM_CLASS_FLAGS(EReplicationBridgeCreateNetRefHandleResultFlags);
+
+struct FReplicationBridgeCreateNetRefHandleResult
+{
+	UE::Net::FNetRefHandle NetRefHandle;
+	EReplicationBridgeCreateNetRefHandleResultFlags Flags = EReplicationBridgeCreateNetRefHandleResultFlags::None;
+};
+
+enum class EReplicationBridgeDestroyInstanceReason : unsigned
+{
+	DoNotDestroy,
+	TearOff,
+	Destroy,
+};
+IRISCORE_API const TCHAR* LexToString(EReplicationBridgeDestroyInstanceReason Reason);
+
+enum class EReplicationBridgeDestroyInstanceFlags : unsigned
+{
+	None = 0U,
+	/** Whether the instance may be destroyed when instructed from the remote peer. This flag applies when the destroy reason is TearOff and torn off actors are to be destroyed as well as regular Destroy. */
+	AllowDestroyInstanceFromRemote = 1U << 0U,
+};
+ENUM_CLASS_FLAGS(EReplicationBridgeDestroyInstanceFlags);
+
 UCLASS(Transient, MinimalAPI)
 class UReplicationBridge : public UObject
 {
@@ -149,7 +179,7 @@ protected:
 	IRISCORE_API virtual bool WriteNetRefHandleCreationInfo(FReplicationBridgeSerializationContext& Context, FNetRefHandle Handle);
 
 	/** Read data required to instantiate NetObject from bitstream. */
-	IRISCORE_API virtual FNetRefHandle CreateNetRefHandleFromRemote(FNetRefHandle SubObjectOwnerNetHandle, FNetRefHandle WantedNetHandle, FReplicationBridgeSerializationContext& Context);
+	IRISCORE_API virtual FReplicationBridgeCreateNetRefHandleResult CreateNetRefHandleFromRemote(FNetRefHandle SubObjectOwnerNetHandle, FNetRefHandle WantedNetHandle, FReplicationBridgeSerializationContext& Context);
 
 	/** Invoke after we have applied the initial state for an object.*/
 	IRISCORE_API virtual void PostApplyInitialState(FNetRefHandle Handle);
@@ -157,10 +187,10 @@ protected:
 	/**
 	 * Called when the instance is detached from the protocol on request by the remote. 
 	 * @param Handle The handle of the object to destroy or tear off.
-	 * @param bTearOff Whether the object should be torn off, i.e. not destroyed. 
-	 * @param bShouldDestroyInstance Whether the object should be destroyed.
+	 * @param DestroyReason Reason for destroying the instance. 
+	 * @param DestroyFlags Special flags such as whether the instance may be destroyed when reason is TearOff, which may revert to destroying, or Destroy.
 	 */
-	IRISCORE_API virtual void DetachInstanceFromRemote(FNetRefHandle Handle, bool bTearOff, bool bShouldDestroyInstance);
+	IRISCORE_API virtual void DetachInstanceFromRemote(FNetRefHandle Handle, EReplicationBridgeDestroyInstanceReason DestroyReason, EReplicationBridgeDestroyInstanceFlags DestroyFlags);
 
 	/** Called when we detach instance protocol from the local instance */
 	IRISCORE_API virtual void DetachInstance(FNetRefHandle Handle);
@@ -212,18 +242,21 @@ protected:
 	/** Destroys the group associated with the level. */
 	IRISCORE_API void DestroyLevelGroup(const UObject* Level);
 
+	/** Called when destruction info is received to determine whether the instance may be destroyed. */
+	IRISCORE_API virtual bool IsAllowedToDestroyInstance(const UObject* Instance) const;
+
 private:
 
 	// Internal operations invoked by ReplicationSystem/ReplicationWriter
 	void ReadAndExecuteDestructionInfoFromRemote(FReplicationBridgeSerializationContext& Context);
-	void DetachSubObjectInstancesFromRemote(FNetRefHandle Handle, bool bTearOff, bool bShouldDestroyInstance);
-	void DestroyNetObjectFromRemote(FNetRefHandle Handle, bool bTearOff, bool bDestroyInstance);
+	void DetachSubObjectInstancesFromRemote(FNetRefHandle Handle, EReplicationBridgeDestroyInstanceReason DestroyReason, EReplicationBridgeDestroyInstanceFlags DestroyFlags);
+	void DestroyNetObjectFromRemote(FNetRefHandle Handle, EReplicationBridgeDestroyInstanceReason DestroyReason, EReplicationBridgeDestroyInstanceFlags DestroyFlags);
 
 	// Adds the Handle to the list of handles pending tear-off, if bIsImmediate is true the object will be destroyed after the next update, otherwise
 	// it will be kept around until EndReplication is called.
 	void TearOff(FNetRefHandle Handle, EEndReplicationFlags DestroyFlags, bool bIsImmediate);
 
-	FNetRefHandle CallCreateNetRefHandleFromRemote(FNetRefHandle SubObjectOwnerHandle, FNetRefHandle WantedNetHandle, FReplicationBridgeSerializationContext& Context);
+	FReplicationBridgeCreateNetRefHandleResult CallCreateNetRefHandleFromRemote(FNetRefHandle SubObjectOwnerHandle, FNetRefHandle WantedNetHandle, FReplicationBridgeSerializationContext& Context);
 	void CallPreSendUpdate(float DeltaSeconds);	
 	void CallPreSendUpdateSingleHandle(FNetRefHandle Handle);
 	void CallUpdateInstancesWorldLocation();
@@ -232,7 +265,7 @@ private:
 	void CallPruneStaleObjects();
 	void CallGetInitialDependencies(FNetRefHandle Handle, FNetDependencyInfoArray& OutDependencies) const;
 	void CallDetachInstance(FNetRefHandle Handle);
-	void CallDetachInstanceFromRemote(FNetRefHandle Handle, bool bTearOff, bool bShouldDestroyInstance);
+	void CallDetachInstanceFromRemote(FNetRefHandle Handle, EReplicationBridgeDestroyInstanceReason DestroyReason, EReplicationBridgeDestroyInstanceFlags DestroyFlags);
 
 private:
 
