Clustering optimization

diff --git a/Engine/Source/Programs/HeadlessChaos/Private/GeometryCollection/GeometryCollectionTestClustering.cpp b/Engine/Source/Programs/HeadlessChaos/Private/GeometryCollection/GeometryCollectionTestClustering.cpp
--- a/Engine/Source/Programs/HeadlessChaos/Private/GeometryCollection/GeometryCollectionTestClustering.cpp
+++ b/Engine/Source/Programs/HeadlessChaos/Private/GeometryCollection/GeometryCollectionTestClustering.cpp
@@ -307,7 +307,7 @@ namespace GeometryCollectionTest
 
 			if (Frame == 2)
 			{
-				ParticleHandles[0]->SetExternalStrain(50.0f);
+				Clustering.SetExternalStrain(ParticleHandles[0], 50.0f);
 				Clustering.BreakingModel();
 			}
 
diff --git a/Engine/Source/Programs/HeadlessChaos/Private/GeometryCollection/GeometryCollectionTestSimulationField.cpp b/Engine/Source/Programs/HeadlessChaos/Private/GeometryCollection/GeometryCollectionTestSimulationField.cpp
--- a/Engine/Source/Programs/HeadlessChaos/Private/GeometryCollection/GeometryCollectionTestSimulationField.cpp
+++ b/Engine/Source/Programs/HeadlessChaos/Private/GeometryCollection/GeometryCollectionTestSimulationField.cpp
@@ -860,7 +860,7 @@ namespace GeometryCollectionTest
 		UnitTest.Solver->RegisterSimOneShotCallback([&]()
 		{
 			EXPECT_EQ(ClusterMap.Num(), 3);
-			EXPECT_EQ(ClusteredParticleHandles[1]->Strain(), 101);
+			EXPECT_EQ(ClusteredParticleHandles[1]->GetInternalStrains(), 101);
 		});
 		FName TargetName = GetFieldPhysicsName(EFieldPhysicsType::Field_InternalClusterStrain);
 		FalloffField->Magnitude = 0.0;
diff --git a/Engine/Source/Programs/HeadlessChaos/Private/HeadlessChaosTestClustering.cpp b/Engine/Source/Programs/HeadlessChaos/Private/HeadlessChaosTestClustering.cpp
--- a/Engine/Source/Programs/HeadlessChaos/Private/HeadlessChaosTestClustering.cpp
+++ b/Engine/Source/Programs/HeadlessChaos/Private/HeadlessChaosTestClustering.cpp
@@ -258,7 +258,8 @@ namespace ChaosTest {
 		EXPECT_TRUE(RootClusterHandle->Disabled() == false);	//strain > 0 so no fracture yet
 
 		// todo: is this the correct replacement for strain?
-		static_cast<Chaos::FPBDRigidClusteredParticleHandle*>(ClusterHandlesDup[2])->SetStrain((FReal)0);	//fracture the third cluster, this should leave us with three pieces (0, 1), (2), (3,4,5,6,7)
+		//fracture the third cluster, this should leave us with three pieces (0, 1), (2), (3,4,5,6,7)
+		Evolution.GetRigidClustering().SetInternalStrain(static_cast<Chaos::FPBDRigidClusteredParticleHandle*>(ClusterHandlesDup[2]), 0.0);
 
 		Evolution.AdvanceOneTimeStep((FReal)1 / (FReal)60);
 		//EXPECT_TRUE(Evolution.GetParticles().Disabled(RootClusterHandle) == false);	//one of the connected pieces should re-use this
diff --git a/Engine/Source/Runtime/Experimental/Chaos/Private/Chaos/ChaosDebugDraw.cpp b/Engine/Source/Runtime/Experimental/Chaos/Private/Chaos/ChaosDebugDraw.cpp
--- a/Engine/Source/Runtime/Experimental/Chaos/Private/Chaos/ChaosDebugDraw.cpp
+++ b/Engine/Source/Runtime/Experimental/Chaos/Private/Chaos/ChaosDebugDraw.cpp
@@ -1458,7 +1458,7 @@ namespace Chaos
 						{
 							if (const FRigidClustering::FClusterHandle& Child = Rigid->CastToClustered())
 							{
-								MaxStrain = FMath::Max(MaxStrain, Child->Strain());
+								MaxStrain = FMath::Max(MaxStrain, Child->GetInternalStrains());
 							}
 						}
 						NumConnections += ActiveCluster.Value.Num();
diff --git a/Engine/Source/Runtime/Experimental/Chaos/Private/Chaos/PBDRigidClustering.cpp b/Engine/Source/Runtime/Experimental/Chaos/Private/Chaos/PBDRigidClustering.cpp
--- a/Engine/Source/Runtime/Experimental/Chaos/Private/Chaos/PBDRigidClustering.cpp
+++ b/Engine/Source/Runtime/Experimental/Chaos/Private/Chaos/PBDRigidClustering.cpp
@@ -124,7 +124,8 @@ namespace Chaos
 		NewParticle->SetInternalCluster(false);
 		NewParticle->SetClusterId(ClusterId(nullptr, Children.Num()));
 		NewParticle->SetClusterGroupIndex(ClusterGroupIndex);
-		NewParticle->SetStrains(0.0);
+		NewParticle->SetInternalStrains(0.0);
+		UpdateTopLevelParticle(NewParticle);
 		NewParticle->SetIsAnchored(Parameters.bIsAnchored);
 
 		// Update clustering data structures.
@@ -156,7 +157,8 @@ namespace Chaos
 				// TODO: Use INDEX_NONE instead of 0?
 				ClusteredChild->SetClusterGroupIndex(0);
 				ClusteredChild->ClusterIds().Id = NewParticle;
-				NewParticle->Strains() += ClusteredChild->Strains();
+				NewParticle->SetInternalStrains(NewParticle->GetInternalStrains() + ClusteredChild->GetInternalStrains());
+				UpdateTopLevelParticle(NewParticle);
 
 				NewParticle->SetCollisionImpulses(FMath::Max(NewParticle->CollisionImpulses(), ClusteredChild->CollisionImpulses()));
 
@@ -167,7 +169,8 @@ namespace Chaos
 		}
 		if (ChildrenSet.Num())
 		{
-			NewParticle->Strains() /= static_cast<FReal>(ChildrenSet.Num());
+			NewParticle->SetInternalStrains(NewParticle->GetInternalStrains() / static_cast<FReal>(ChildrenSet.Num()));
+			UpdateTopLevelParticle(NewParticle);
 		}
 
 		ensureMsgf(!ProxyGeometry || ForceMassOrientation, TEXT("If ProxyGeometry is passed, we must override the mass orientation as they are tied"));
@@ -568,7 +571,7 @@ namespace Chaos
 
 				MEvolution.SetPhysicsMaterial(NewCluster, MEvolution.GetPhysicsMaterial(ClusteredParent));
 
-				NewCluster->SetStrain(ClusteredParent->Strain());
+				NewCluster->SetInternalStrains(ClusteredParent->GetInternalStrains());
 				NewCluster->SetV(ClusteredParent->V());
 				NewCluster->SetW(ClusteredParent->W());
 				NewCluster->SetPreV(ClusteredParent->PreV());
@@ -576,6 +579,8 @@ namespace Chaos
 				NewCluster->SetP(NewCluster->X());
 				NewCluster->SetQ(NewCluster->R());
 
+				UpdateTopLevelParticle(NewCluster);
+
 				// Need to get the material from the previous particle and apply it to the new one
 				const FShapesArray& ChildShapes = ClusteredParent->ShapesArray();
 				const FShapesArray& NewShapes = NewCluster->ShapesArray();
@@ -602,6 +607,31 @@ namespace Chaos
 		}
 		return NewClusters;
 	}
+
+	void FRigidClustering::SetInternalStrain(FPBDRigidClusteredParticleHandle* Particle, FReal Strain)
+	{
+		Particle->SetInternalStrains(Strain);
+		UpdateTopLevelParticle(Particle);
+	}
+
+	void FRigidClustering::SetExternalStrain(FPBDRigidClusteredParticleHandle* Particle, FReal Strain)
+	{
+		Particle->SetExternalStrains(Strain);
+		UpdateTopLevelParticle(Particle);
+	}
+
+	void FRigidClustering::UpdateTopLevelParticle(FPBDRigidClusteredParticleHandle* Particle)
+	{
+		FPBDRigidClusteredParticleHandle* Parent = Particle->Parent();
+		if (Parent != nullptr)
+		{
+			TopLevelClusterParentsStrained.Add(Parent);
+		}
+		else
+		{
+			TopLevelClusterParentsStrained.Add(Particle);
+		}
+	}
 	
 	DECLARE_CYCLE_STAT(TEXT("TPBDRigidClustering<>::ReleaseClusterParticles(STRAIN)"), STAT_ReleaseClusterParticles_STRAIN, STATGROUP_Chaos);
 	TSet<FPBDRigidParticleHandle*> FRigidClustering::ReleaseClusterParticles(
@@ -673,9 +703,9 @@ namespace Chaos
 				continue;
 			}
 
-			// @todo(chaos) eventually shoudl get rid of collision impulse array and only use external strain
+			// @todo(chaos) eventually should get rid of collision impulse array and only use external strain
 			const FReal MaxAppliedStrain = FMath::Max(Child->CollisionImpulses(), Child->GetExternalStrain());
-			if ((MaxAppliedStrain >= Child->Strain()) || bForceRelease)
+			if ((MaxAppliedStrain >= Child->GetInternalStrains()) || bForceRelease)
 			{
 				//UE_LOG(LogTemp, Warning, TEXT("Releasing child %d from parent %p due to strain %.5f Exceeding internal strain %.5f (Source: %s)"), ChildIdx, ClusteredParticle, ChildStrain, Child->Strain(), bForceRelease ? TEXT("Forced by caller") : ExternalStrainMap ? TEXT("External") : TEXT("Collision"));
 				
@@ -685,6 +715,8 @@ namespace Chaos
 				RemoveChildFromParent(Child, ClusteredParticle);
 				ActivatedChildren.Add(Child);
 
+				UpdateTopLevelParticle(Child);
+
 				// Remove from the children array without freeing memory yet. 
 				// We're looping over Children and it'd be silly to free the array
 				// 1 entry at a time.
@@ -696,7 +728,7 @@ namespace Chaos
 			{
 				AppliedStrains.Add(Child, MaxAppliedStrain);
 			}
-			Child->ClearExternalStrain();
+			Child->SetExternalStrains(0.0);
 		}
 
 		// if necessary propagate strain through the graph
@@ -714,7 +746,7 @@ namespace Chaos
 				if (BreakDamagePropagationFactor > 0 && ActivatedChildren.Contains(AppliedStrain.Key))
 				{
 					// break damage propagation case: we only look at the broken pieces and propagate the strain remainder 
-					const FReal RemainingStrain = (AppliedStrainValue - ClusteredChild->Strain());
+					const FReal RemainingStrain = (AppliedStrainValue - ClusteredChild->GetInternalStrains());
 					if (RemainingStrain > 0)
 					{
 						const FReal AdjustedRemainingStrain = (FReal)BreakDamagePropagationFactor * RemainingStrain;
@@ -724,7 +756,7 @@ namespace Chaos
 				}
 				else if (ShockDamagePropagationFactor > 0)
 				{
-					// shock damage propagation case : for all the non broken pieces, proapagate the actual applied strain 
+					// shock damage propagation case : for all the non broken pieces, propagate the actual applied strain 
 					PropagatedStrainPerConnection = (FReal)ShockDamagePropagationFactor * AppliedStrainValue;
 				}
 
@@ -736,8 +768,8 @@ namespace Chaos
 						{
 							if (FPBDRigidClusteredParticleHandle* ClusteredSibling = Edge.Sibling->CastToClustered())
 							{
-								// todo(chaos) this may currently be non optimal as we are in the apply loop and this may be cleareed right after
-								ClusteredSibling->SetExternalStrain(FMath::Max(ClusteredSibling->GetExternalStrain(), PropagatedStrainPerConnection));
+								// todo(chaos) this may currently be non optimal as we are in the apply loop and this may be cleared right after
+								SetExternalStrain(ClusteredSibling, FMath::Max(ClusteredSibling->GetExternalStrain(), PropagatedStrainPerConnection));
 							}
 						}
 					}
@@ -821,7 +853,7 @@ namespace Chaos
 					{
 						if (ensure(!ClusterHandle || ClusteredChildHandle->ClusterIds().Id == ClusterHandle))
 						{
-							ClusteredChildHandle->SetExternalStrain(TNumericLimits<FReal>::Max());
+							SetExternalStrain(ClusteredChildHandle, TNumericLimits<FReal>::Max());
 							ClusterHandle = ClusteredChildHandle->ClusterIds().Id;
 						}
 						else
@@ -877,7 +909,7 @@ namespace Chaos
 			bool bPotentialBreak = false;
 			{
 				SCOPE_CYCLE_COUNTER(STAT_UpdateDirtyImpulses);
-				for (const auto& ActiveCluster : GetTopLevelClusterParents())
+				for (const auto& ActiveCluster : TopLevelClusterParentsStrained)
 				{
 					if (!ActiveCluster->Disabled())
 					{
@@ -888,7 +920,7 @@ namespace Chaos
 							{
 								if (FClusterHandle ClusteredChild = Child->CastToClustered())
 								{
-									if (ClusteredChild->Strain() <= 0.f)
+									if (ClusteredChild->GetInternalStrains() <= 0.f)
 									{
 										ClusteredChild->CollisionImpulse() = FLT_MAX;
 										MCollisionImpulseArrayDirty = true;
@@ -902,6 +934,7 @@ namespace Chaos
 						}
 					}
 				}
+				TopLevelClusterParentsStrained.Reset();
 			}
 
 			if (MCollisionImpulseArrayDirty || bPotentialBreak)
@@ -1350,7 +1383,8 @@ namespace Chaos
 			{
 				if (Chaos::FPBDRigidClusteredParticleHandle* ClusteredChildHandle = ChildHandle->CastToClustered())
 				{
-					ClusteredChildHandle->SetExternalStrain(MaxStrain);
+					ClusteredChildHandle->SetExternalStrains(MaxStrain);
+					SetExternalStrain(ClusteredChildHandle, MaxStrain);
 				}
 			}
 			if (ChildrenHandles->Num() > 0)
@@ -1380,7 +1414,7 @@ namespace Chaos
 					{
 						if (Chaos::FPBDRigidClusteredParticleHandle* ClusteredChildHandle = ChildHandle->CastToClustered())
 						{
-							ClusteredChildHandle->SetExternalStrain(MaxStrain);
+							SetExternalStrain(ClusteredChildHandle, MaxStrain);
 						}
 					}
 					if (Children->Num() > 0)
@@ -1400,7 +1434,7 @@ namespace Chaos
 		check(ClusteredChild1 && ClusteredChild2);
 		if (ClusteredChild1 == ClusteredChild2)
 			return;
-		const FReal AvgStrain = (ClusteredChild1->Strains() + ClusteredChild2->Strains()) * (FReal)0.5;
+		const FReal AvgStrain = (ClusteredChild1->GetInternalStrains() + ClusteredChild2->GetInternalStrains()) * (FReal)0.5;
 		TArray<TConnectivityEdge<FReal>>& Edges1 = ClusteredChild1->ConnectivityEdges();
 		TArray<TConnectivityEdge<FReal>>& Edges2 = ClusteredChild2->ConnectivityEdges();
 		if (//Edges1.Num() < Parameters.MaxNumConnections && 
diff --git a/Engine/Source/Runtime/Experimental/Chaos/Private/PhysicsProxy/GeometryCollectionPhysicsProxy.cpp b/Engine/Source/Runtime/Experimental/Chaos/Private/PhysicsProxy/GeometryCollectionPhysicsProxy.cpp
--- a/Engine/Source/Runtime/Experimental/Chaos/Private/PhysicsProxy/GeometryCollectionPhysicsProxy.cpp
+++ b/Engine/Source/Runtime/Experimental/Chaos/Private/PhysicsProxy/GeometryCollectionPhysicsProxy.cpp
@@ -32,6 +32,7 @@
 #include "Modules/ModuleManager.h"
 #include "Chaos/PullPhysicsDataImp.h"
 #include "Chaos/PBDRigidsEvolution.h"
+#include "Chaos/PBDRigidClustering.h"
 #include "Chaos/PhysicsObjectInternal.h"
 #include "Chaos/PhysicsObjectInterface.h"
 #include "GeometryCollection/Facades/CollectionAnchoringFacade.h"
@@ -973,7 +974,9 @@ void FGeometryCollectionPhysicsProxy::InitializeBodiesPT(Chaos::FPBDRigidsSolver
 					{
 						DamageThreshold = ComputeDamageThreshold(DynamicCollection, TransformGroupIndex);
 					}
-					Handle->SetStrain(DamageThreshold);
+					Chaos::FPhysicsSolver* RBDSolver = GetSolver<Chaos::FPhysicsSolver>();
+					RBDSolver->GetEvolution()->GetRigidClustering().SetInternalStrain(Handle, DamageThreshold);
+					
 				}
 
 				// #BGTODO - non-updating parameters - remove lin/ang drag arrays and always query material if this stays a material parameter
@@ -1197,7 +1200,7 @@ void FGeometryCollectionPhysicsProxy::InitializeBodiesPT(Chaos::FPBDRigidsSolver
 									if (SiblingTransformIndex > TransformGroupIndex)
 									{
 										Chaos::TConnectivityEdge<Chaos::FReal> Edge;
-										Edge.Strain = (ClusteredParticle->Strains() + OtherClusteredParticle->Strains()) * 0.5f;
+										Edge.Strain = (ClusteredParticle->GetInternalStrains() + OtherClusteredParticle->GetInternalStrains()) * 0.5f;
 									
 										Edge.Sibling = OtherClusteredParticle;
 										ClusteredParticle->ConnectivityEdges().Add(Edge);
@@ -1408,7 +1411,9 @@ FGeometryCollectionPhysicsProxy::BuildClusters_Internal(
 	SolverClusterHandles[CollectionClusterIndex] = Parent;
 
 	const float DamageThreshold = ComputeDamageThreshold(DynamicCollection, CollectionClusterIndex);
-	Parent->SetStrains(DamageThreshold);
+
+	Chaos::FRigidClustering& RigidClustering = static_cast<Chaos::FPBDRigidsSolver*>(Solver)->GetEvolution()->GetRigidClustering();
+	RigidClustering.SetInternalStrain(Parent, DamageThreshold);
 	
 	// #BGTODO This will not automatically update - material properties should only ever exist in the material, not in other arrays
 	const Chaos::FChaosPhysicsMaterial* CurMaterial = static_cast<Chaos::FPBDRigidsSolver*>(Solver)->GetSimMaterials().Get(Parameters.PhysicalMaterialHandle.InnerHandle);
@@ -1435,7 +1440,7 @@ FGeometryCollectionPhysicsProxy::BuildClusters_Internal(
 		{
 			if (Chaos::FPBDRigidClusteredParticleHandle* ClusteredChild = Child->CastToClustered())
 			{
-				ClusteredChild->SetStrains(DamageThreshold);
+				RigidClustering.SetInternalStrain(ClusteredChild, DamageThreshold);
 			}
 		}
 
@@ -1960,12 +1965,13 @@ void FGeometryCollectionPhysicsProxy::ApplyExternalStrain_External(FGeometryColl
 		{
 			if (Chaos::FPBDRigidClusteredParticleHandle* ClusteredHandle = FindClusteredParticleHandleByItemIndex_Internal(ItemIndex))
 			{
+				Chaos::FRigidClustering& RigidClustering = RBDSolver->GetEvolution()->GetRigidClustering();
 				ApplyToChildrenAtPointWithRadiusAndPropagation_Internal(
 					RBDSolver->GetEvolution()->GetRigidClustering(), *ClusteredHandle,
 					WorldLocation, Radius, PropagationDepth, PropagationFactor,
-					[StrainValue](Chaos::FPBDRigidClusteredParticleHandle* ClusteredChildHandle, Chaos::FReal PropagationMultiplier)
+					[StrainValue, &RigidClustering](Chaos::FPBDRigidClusteredParticleHandle* ClusteredChildHandle, Chaos::FReal PropagationMultiplier)
 					{
-						ClusteredChildHandle->SetExternalStrain(FMath::Max(ClusteredChildHandle->GetExternalStrain(), StrainValue * PropagationMultiplier));
+						RigidClustering.SetExternalStrain(ClusteredChildHandle, FMath::Max(ClusteredChildHandle->GetExternalStrain(), StrainValue * PropagationMultiplier));
 					});
 			}
 		});
@@ -1985,10 +1991,10 @@ void FGeometryCollectionPhysicsProxy::ApplyInternalStrain_External(FGeometryColl
 				ApplyToChildrenAtPointWithRadiusAndPropagation_Internal(
 					RBDSolver->GetEvolution()->GetRigidClustering(), *ClusteredHandle,
 					WorldLocation, Radius, PropagationDepth, PropagationFactor,
-					[StrainValue](Chaos::FPBDRigidClusteredParticleHandle* ClusteredChildHandle, Chaos::FReal PropagationMultiplier)
+					[StrainValue, RBDSolver](Chaos::FPBDRigidClusteredParticleHandle* ClusteredChildHandle, Chaos::FReal PropagationMultiplier)
 					{
-						const Chaos::FReal NewInternalStrain = ClusteredChildHandle->Strain() - ((Chaos::FReal)StrainValue * PropagationMultiplier);
-						ClusteredChildHandle->SetStrain(FMath::Max(0,NewInternalStrain));
+						const Chaos::FReal NewInternalStrain = ClusteredChildHandle->GetInternalStrains() - ((Chaos::FReal)StrainValue * PropagationMultiplier);
+						RBDSolver->GetEvolution()->GetRigidClustering().SetInternalStrain(ClusteredChildHandle, FMath::Max(0, NewInternalStrain));
 					});
 			}
 		});
@@ -2008,7 +2014,7 @@ static void ApplyToBreakingChildren_Internal(Chaos::FRigidClustering& Clustering
 			{
 				if (Chaos::FPBDRigidClusteredParticleHandle* ClusteredChildHandle = ChildHandle->CastToClustered())
 				{
-					if (ClusteredChildHandle->GetExternalStrain() > ClusteredChildHandle->Strain())
+					if (ClusteredChildHandle->GetExternalStrain() > ClusteredChildHandle->GetInternalStrains())
 					{
 						Action(ClusteredChildHandle);
 					}
@@ -2734,7 +2740,7 @@ void FGeometryCollectionPhysicsProxy::BufferPhysicsResults_Internal(Chaos::FPBDR
 #if WITH_EDITORONLY_DATA
 			FGeometryCollectionResults::FDamageInfo& DamageInfo = TargetResults.DamageInfo[TransformGroupIndex];
 			DamageInfo.Damage = static_cast<float>(Handle->CollisionImpulse());
-			DamageInfo.DamageThreshold = static_cast<float>(Handle->Strain());
+			DamageInfo.DamageThreshold = static_cast<float>(Handle->GetInternalStrains());
 #endif			
 		}    // end for
 	}        // STAT_CalcParticleToWorld scope
diff --git a/Engine/Source/Runtime/Experimental/Chaos/Public/Chaos/PBDRigidClustering.h b/Engine/Source/Runtime/Experimental/Chaos/Public/Chaos/PBDRigidClustering.h
--- a/Engine/Source/Runtime/Experimental/Chaos/Public/Chaos/PBDRigidClustering.h
+++ b/Engine/Source/Runtime/Experimental/Chaos/Public/Chaos/PBDRigidClustering.h
@@ -345,6 +345,8 @@ public:
 	FRigidEvolution& GetEvolution() { return MEvolution; }
 	const FRigidEvolution& GetEvolution() const { return MEvolution; }
 
+	void SetInternalStrain(FPBDRigidClusteredParticleHandle* Particle, FReal Strain);
+	void SetExternalStrain(FPBDRigidClusteredParticleHandle* Particle, FReal Strain);
 
  protected:
 
@@ -384,11 +386,14 @@ public:
 	using FParticleIsland = TArray<FPBDRigidParticleHandle*>;
 	TArray<FParticleIsland> FindIslandsInChildren(const FPBDRigidClusteredParticleHandle* ClusteredParticle);
 	TArray<FPBDRigidParticleHandle*> CreateClustersFromNewIslands(TArray<FParticleIsland>& Islands, FPBDRigidClusteredParticleHandle* ClusteredParent);
+
+	void UpdateTopLevelParticle(FPBDRigidClusteredParticleHandle* Particle);
 private:
 
 	FRigidEvolution& MEvolution;
 	FPBDRigidClusteredParticles& MParticles;
 	TSet<Chaos::FPBDRigidClusteredParticleHandle*> TopLevelClusterParents;
+	TSet<Chaos::FPBDRigidClusteredParticleHandle*> TopLevelClusterParentsStrained;
 
 	// Cluster data
 	FClusterMap MChildren;
diff --git a/Engine/Source/Runtime/Experimental/Chaos/Public/Chaos/ParticleHandle.h b/Engine/Source/Runtime/Experimental/Chaos/Public/Chaos/ParticleHandle.h
--- a/Engine/Source/Runtime/Experimental/Chaos/Public/Chaos/ParticleHandle.h
+++ b/Engine/Source/Runtime/Experimental/Chaos/Public/Chaos/ParticleHandle.h
@@ -19,6 +19,7 @@
 #endif
 
 class IPhysicsProxyBase;
+class UGeometryCollectionComponent;
 
 namespace Chaos
 {
@@ -1276,15 +1277,23 @@ public:
 	void SetCollisionImpulses(const T Value) { PBDRigidClusteredParticles->CollisionImpulses(ParticleIdx) = Value; }
 
 	T GetExternalStrain() const { return PBDRigidClusteredParticles->ExternalStrains(ParticleIdx); }
+	UE_DEPRECATED(5.2, "This method should not be used anymore. FRigidClustering::SetExternalStrain should be used instead.")
 	void SetExternalStrain(const T Value) { PBDRigidClusteredParticles->ExternalStrains(ParticleIdx) = Value; }
+	UE_DEPRECATED(5.2, "This method should not be used anymore. FRigidClustering::SetExternalStrain with 0 strain should be used instead.")
 	void ClearExternalStrain() { PBDRigidClusteredParticles->ExternalStrains(ParticleIdx) = static_cast<T>(0); }
 	
+	const T& GetInternalStrains() const { return PBDRigidClusteredParticles->Strains(ParticleIdx); }
 	const T& Strain() const { return PBDRigidClusteredParticles->Strains(ParticleIdx); }
+	UE_DEPRECATED(5.2, "This method should not be used anymore. FRigidClustering::SetInternalStrain should be used instead.")
 	T& Strain() { return PBDRigidClusteredParticles->Strains(ParticleIdx); }
+	UE_DEPRECATED(5.2, "This method should not be used anymore. FRigidClustering::SetInternalStrain should be used instead.")
 	void SetStrain(const T Value) { PBDRigidClusteredParticles->Strains(ParticleIdx) = Value; }
 	const T& Strains() const { return PBDRigidClusteredParticles->Strains(ParticleIdx); }
+	UE_DEPRECATED(5.2, "This method should not be used anymore. FRigidClustering::SetInternalStrain should be used instead.")
 	T& Strains() { return PBDRigidClusteredParticles->Strains(ParticleIdx); }
+	UE_DEPRECATED(5.2, "This method should not be used anymore. FRigidClustering::SetInternalStrain should be used instead.")
 	void SetStrains(const T Value) { PBDRigidClusteredParticles->Strains(ParticleIdx) = Value; }
+	void SetMaximumInternalStrain() { PBDRigidClusteredParticles->Strains(ParticleIdx) = MaxStrain; }
 
 	const TArray<TConnectivityEdge<T>>& ConnectivityEdges() const { return PBDRigidClusteredParticles->ConnectivityEdges(ParticleIdx); }
 	TArray<TConnectivityEdge<T>>& ConnectivityEdges() { return PBDRigidClusteredParticles->ConnectivityEdges(ParticleIdx); }
@@ -1301,6 +1310,12 @@ public:
 	static constexpr EParticleType StaticType() { return EParticleType::Rigid; }
 
 	int32 TransientParticleIndex() const { return ParticleIdx; }
+
+private:
+	void SetInternalStrains(const T Value) { PBDRigidClusteredParticles->Strains(ParticleIdx) = Value; }
+	void SetExternalStrains(const T Value) { PBDRigidClusteredParticles->ExternalStrains(ParticleIdx) = Value; }
+	friend class FRigidClustering;
+	static constexpr Chaos::FReal MaxStrain = TNumericLimits<Chaos::FReal>::Max() - TNumericLimits<Chaos::FReal>::Min();
 };
 
 template <typename T, int d, bool bPersistent = true>
diff --git a/Engine/Source/Runtime/Experimental/Chaos/Public/PhysicsProxy/FieldSystemProxyHelper.h b/Engine/Source/Runtime/Experimental/Chaos/Public/PhysicsProxy/FieldSystemProxyHelper.h
--- a/Engine/Source/Runtime/Experimental/Chaos/Public/PhysicsProxy/FieldSystemProxyHelper.h
+++ b/Engine/Source/Runtime/Experimental/Chaos/Public/PhysicsProxy/FieldSystemProxyHelper.h
@@ -494,6 +494,7 @@ namespace Chaos
 		{
 			SCOPE_CYCLE_COUNTER(STAT_ParamUpdateField_ExternalClusterStrain);
 			{
+				FRigidClustering& RigidClustering = RigidSolver->GetEvolution()->GetRigidClustering();
 				static_cast<const FFieldNode<float>*>(FieldCommand.RootNode.Get())->Evaluate(FieldContext, ResultsView);
 				for (const FFieldContextIndex& Index : FieldContext.GetEvaluatedSamples())
 				{
@@ -503,7 +504,7 @@ namespace Chaos
 						if (Chaos::FPBDRigidClusteredParticleHandle* ClusteredParticle = ParticleHandles[Index.Sample]->CastToClustered())
 						{
 							const FReal CurrentExternalStrains = ClusteredParticle->GetExternalStrain();
-							ClusteredParticle->SetExternalStrain(FMath::Max(CurrentExternalStrains, ExternalStrainValue));
+							RigidClustering.SetExternalStrain(ClusteredParticle, FMath::Max(CurrentExternalStrains, ExternalStrainValue));
 						}
 					}
 				}
@@ -559,13 +560,14 @@ namespace Chaos
 		{
 			SCOPE_CYCLE_COUNTER(STAT_ParamUpdateField_InternalClusterStrain);
 			{
+				FRigidClustering& RigidClustering = RigidSolver->GetEvolution()->GetRigidClustering();
 				static_cast<const FFieldNode<float>*>(FieldCommand.RootNode.Get())->Evaluate(FieldContext, ResultsView);
 				for (const FFieldContextIndex& Index : FieldContext.GetEvaluatedSamples())
 				{
 					Chaos::FPBDRigidClusteredParticleHandle* RigidHandle = ParticleHandles[Index.Sample]->CastToClustered();
 					if (RigidHandle && RigidHandle->ObjectState() == Chaos::EObjectStateType::Dynamic)
 					{
-						RigidHandle->Strain() += ResultsView[Index.Result];
+						RigidClustering.SetInternalStrain(RigidHandle, RigidHandle->GetInternalStrains() + ResultsView[Index.Result]);
 					}
 				}
 			}
diff --git a/Engine/Source/Runtime/Experimental/GeometryCollectionEngine/Private/GeometryCollection/GeometryCollectionComponent.cpp b/Engine/Source/Runtime/Experimental/GeometryCollectionEngine/Private/GeometryCollection/GeometryCollectionComponent.cpp
--- a/Engine/Source/Runtime/Experimental/GeometryCollectionEngine/Private/GeometryCollection/GeometryCollectionComponent.cpp
+++ b/Engine/Source/Runtime/Experimental/GeometryCollectionEngine/Private/GeometryCollection/GeometryCollectionComponent.cpp
@@ -1851,23 +1851,6 @@ void UGeometryCollectionComponent::SetInitialClusterBreaks(const TArray<int32>&
 	}
 }
 
-void SetHierarchyStrain(Chaos::TPBDRigidClusteredParticleHandle<Chaos::FReal, 3>* P, TMap<Chaos::TPBDRigidClusteredParticleHandle<Chaos::FReal, 3>*, TArray<Chaos::TPBDRigidParticleHandle<Chaos::FReal, 3>*>>& Map, float Strain)
-{
-	TArray<Chaos::TPBDRigidParticleHandle<Chaos::FReal, 3>*>* Children = Map.Find(P);
-
-	if(Children)	
-	{
-		for(Chaos::TPBDRigidParticleHandle<Chaos::FReal, 3> * ChildP : (*Children))
-		{
-			SetHierarchyStrain(ChildP->CastToClustered(), Map, Strain);
-		}
-	}
-
-	if(P)
-	{
-		P->SetStrain(Strain);
-	}
-}
 
 void UGeometryCollectionComponent::InitConstantData(FGeometryCollectionConstantData* ConstantData) const
 {
@@ -2723,11 +2706,11 @@ void UGeometryCollectionComponent::RegisterAndInitializePhysicsProxy()
 	FPhysScene_Chaos* Scene = GetInnerChaosScene();
 	Scene->AddObject(this, PhysicsProxy);
 
-	// If we're replicating we need some extra setup - check netmode as we don't need this for standalone runtimes where we aren't going to network the component
-	// IMPORTANT this need to happen after the object is registered so this will garantee that the particles are properly created by the time the callback below gets called
+	// If we're replicating we need some extra setup - check netmode as we don't need this for standalone runtime where we aren't going to network the component
+	// IMPORTANT this need to happen after the object is registered so this will guarantee that the particles are properly created by the time the callback below gets called
 	if (GetIsReplicated())
 	{
-		// Client side : geometry collection children of parents below the rep level need to be infintely strong so that client cannot break it 
+		// Client side : geometry collection children of parents below the rep level need to be infinitely strong so that client cannot break it 
 		if (Chaos::FPhysicsSolver* CurrSolver = GetSolver(*this))
 		{
 			CurrSolver->EnqueueCommandImmediate([Proxy = PhysicsProxy, AbandonAfterLevel = ReplicationAbandonAfterLevel, EnableAbandonAfterLevel = bEnableAbandonAfterLevel]()
@@ -2737,7 +2720,6 @@ void UGeometryCollectionComponent::RegisterAndInitializePhysicsProxy()
 						// As we're not in control we make it so our simulated proxy cannot break clusters
 						// We have to set the strain to a high value but be below the max for the data type
 						// so releasing on authority demand works
-						constexpr Chaos::FReal MaxStrain = TNumericLimits<Chaos::FReal>::Max() - TNumericLimits<Chaos::FReal>::Min();
 						for (Chaos::FPBDRigidClusteredParticleHandle* ParticleHandle : Proxy->GetParticles())
 						{
 							if (ParticleHandle)
@@ -2745,7 +2727,7 @@ void UGeometryCollectionComponent::RegisterAndInitializePhysicsProxy()
 								const int32 Level = EnableAbandonAfterLevel ? ComputeParticleLevel(ParticleHandle) : -1;
 								if (Level <= AbandonAfterLevel + 1)	//we only replicate up until level X, but it means we should replicate the breaking event of level X+1 (but not X+1's positions)
 								{
-									ParticleHandle->SetStrain(MaxStrain);
+									ParticleHandle->SetMaximumInternalStrain();
 								}
 							}
 						}
