Geometry collection : fix crash when blueprint set notification settings on a geo metryc collection that cause a unresolved callback to be evaluated after the proxy has been destroyed
- move setting of notification and other parameters to properly be set on the PT when GT to PT happens opn the proxy

diff --git a/Engine/Source/Runtime/Experimental/Chaos/Private/PhysicsProxy/GeometryCollectionPhysicsProxy.cpp b/Engine/Source/Runtime/Experimental/Chaos/Private/PhysicsProxy/GeometryCollectionPhysicsProxy.cpp
--- a/Engine/Source/Runtime/Experimental/Chaos/Private/PhysicsProxy/GeometryCollectionPhysicsProxy.cpp
+++ b/Engine/Source/Runtime/Experimental/Chaos/Private/PhysicsProxy/GeometryCollectionPhysicsProxy.cpp
@@ -2096,6 +2096,43 @@ void FGeometryCollectionPhysicsProxy::ApplyAngularVelocity_External(FGeometryCol
 	}
 }
 
+void FGeometryCollectionPhysicsProxy::SetProxyDirty_External()
+{
+	check(IsInGameThread());
+	if (Chaos::FPhysicsSolver* RBDSolver = GetSolver<Chaos::FPhysicsSolver>())
+	{
+		RBDSolver->AddDirtyProxy(this);
+	}
+}
+
+void FGeometryCollectionPhysicsProxy::SetEnableDamageFromCollision_External(bool bEnable)
+{
+	check(IsInGameThread());
+	GameThreadPerFrameData.SetEnableStrainOnCollision(bEnable);
+	SetProxyDirty_External();
+}
+
+void FGeometryCollectionPhysicsProxy::SetNotifyBreakings_External(bool bNotify)
+{
+	check(IsInGameThread());
+	GameThreadPerFrameData.SetNotifyBreakings(bNotify);
+	SetProxyDirty_External();
+}
+
+void FGeometryCollectionPhysicsProxy::SetNotifyRemovals_External(bool bNotify)
+{
+	check(IsInGameThread());
+	GameThreadPerFrameData.SetNotifyRemovals(bNotify);
+	SetProxyDirty_External();
+}
+
+void FGeometryCollectionPhysicsProxy::SetNotifyCrumblings_External(bool bNotify, bool bIncludeChildren)
+{
+	check(IsInGameThread());
+	GameThreadPerFrameData.SetNotifyCrumblings(bNotify, bIncludeChildren);
+	SetProxyDirty_External();
+}
+
 int32 FGeometryCollectionPhysicsProxy::CalculateHierarchyLevel(const FGeometryDynamicCollection& DynamicCollection, int32 TransformIndex)
 {
 	int32 Level = 0;
@@ -2261,11 +2298,7 @@ void FGeometryCollectionPhysicsProxy::SetWorldTransform_External(const FTransfor
 {
 	check(IsInGameThread());
 	GameThreadPerFrameData.SetWorldTransform(WorldTransform);
-
-	if (Chaos::FPhysicsSolver* RBDSolver = GetSolver<Chaos::FPhysicsSolver>())
-	{
-		RBDSolver->AddDirtyProxy(this);
-	}
+	SetProxyDirty_External();
 }
 
 void FGeometryCollectionPhysicsProxy::PushStateOnGameThread(Chaos::FPBDRigidsSolver* InSolver)
@@ -2287,6 +2320,21 @@ void FGeometryCollectionPhysicsProxy::PushStateOnGameThread(Chaos::FPBDRigidsSol
 		Parameters.SimulationFilterData = GameThreadPerFrameData.GetSimFilter();
 		GameThreadPerFrameData.ResetIsCollisionFilterDataDirty();
 	}
+
+	if (GameThreadPerFrameData.GetIsNotificationDataDirty())
+	{
+		Parameters.bGenerateBreakingData = GameThreadPerFrameData.GetNotifyBreakings();
+		Parameters.bGenerateRemovalsData = GameThreadPerFrameData.GetNotifyRemovals();
+		Parameters.bGenerateCrumblingData = GameThreadPerFrameData.GetNotifyCrumblings();
+		Parameters.bGenerateCrumblingChildrenData = GameThreadPerFrameData.GetCrumblingEventIncludesChildren();
+		GameThreadPerFrameData.ResetIsNotificationDataDirty();
+	}
+
+	if (GameThreadPerFrameData.GetIsDamageSettingsDataDirty())
+	{
+		Parameters.bEnableStrainOnCollision = GameThreadPerFrameData.GetEnableStrainOnCollision();
+		GameThreadPerFrameData.ResetIsDamageSettingsDataDirty();
+	}
 }
 
 void FGeometryCollectionPhysicsProxy::PushToPhysicsState()
@@ -3020,10 +3068,7 @@ void FGeometryCollectionPhysicsProxy::UpdateFilterData_External(const FCollision
 	GameThreadPerFrameData.SetSimFilter(NewSimFilter);
 	GameThreadPerFrameData.SetQueryFilter(NewQueryFilter);
 
-	if (Chaos::FPhysicsSolver* RBDSolver = GetSolver<Chaos::FPhysicsSolver>())
-	{
-		RBDSolver->AddDirtyProxy(this);
-	}
+	SetProxyDirty_External();
 }
 
 //==============================================================================
@@ -3569,6 +3614,7 @@ void FGeometryCollectionPhysicsProxy::InitializeSharedCollisionStructures(
 		MassSpaceParticles.X(Idx) = Vertex[Idx];	//mass space computation done later down
 	}
 	
+	
 	// compute the bounding box in mesh space
 	FVec3 MaxChildBounds(1);
 	ParallelFor(NumGeometries, [&](int32 GeometryIndex)
diff --git a/Engine/Source/Runtime/Experimental/Chaos/Public/GeometryCollectionProxyData.h b/Engine/Source/Runtime/Experimental/Chaos/Public/GeometryCollectionProxyData.h
--- a/Engine/Source/Runtime/Experimental/Chaos/Public/GeometryCollectionProxyData.h
+++ b/Engine/Source/Runtime/Experimental/Chaos/Public/GeometryCollectionProxyData.h
@@ -140,8 +140,16 @@ class FGeometryCollectioPerFrameData
 {
 public:
 	FGeometryCollectioPerFrameData()
-		: IsWorldTransformDirty(false)
-		, bIsCollisionFilterDataDirty(false) {}
+		: bIsWorldTransformDirty(false)
+		, bIsCollisionFilterDataDirty(false)
+		, bIsNotificationDataDirty(false) 
+		, bIsDamageSettingsDataDirty(false)
+		, bNotifyBreakings(false)
+		, bNotifyRemovals(false)
+		, bNotifyCrumblings(false)
+		, bCrumblingEventIncludesChildren(false)
+		, bEnableStrainOnCollision(false)
+	{}
 
 	const FTransform& GetWorldTransform() const { return WorldTransform; }
 
@@ -150,12 +158,12 @@ public:
 		if (!WorldTransform.Equals(InWorldTransform))
 		{
 			WorldTransform = InWorldTransform;
-			IsWorldTransformDirty = true;
+			bIsWorldTransformDirty = true;
 		}
 	}
 
-	bool GetIsWorldTransformDirty() const { return IsWorldTransformDirty; }
-	void ResetIsWorldTransformDirty() { IsWorldTransformDirty = false; }
+	bool GetIsWorldTransformDirty() const { return bIsWorldTransformDirty; }
+	void ResetIsWorldTransformDirty() { bIsWorldTransformDirty = false; }
 
 	const FCollisionFilterData& GetSimFilter() const { return SimFilter; }
 	void SetSimFilter(const FCollisionFilterData& NewSimFilter)
@@ -174,13 +182,66 @@ public:
 	bool GetIsCollisionFilterDataDirty() const { return bIsCollisionFilterDataDirty; }
 	void ResetIsCollisionFilterDataDirty() { bIsCollisionFilterDataDirty = false; }
 
+	void SetNotifyBreakings(bool bNotify)
+	{
+		bNotifyBreakings = bNotify;
+		bIsNotificationDataDirty = true;
+	}
+	bool GetNotifyBreakings() const { return bNotifyBreakings; }
+
+	void SetNotifyRemovals(bool bNotify)
+	{
+		bNotifyRemovals = bNotify;
+		bIsNotificationDataDirty = true;
+	}
+
+	bool GetNotifyRemovals() const { return bNotifyRemovals; }
+
+	void SetNotifyCrumblings(bool bNotify, bool bIncludeChildren)
+	{
+		bNotifyCrumblings = bNotify;
+		bCrumblingEventIncludesChildren = bIncludeChildren;
+		bIsNotificationDataDirty = true;
+	}
+
+	bool GetNotifyCrumblings() const { return bNotifyCrumblings; }
+	bool GetCrumblingEventIncludesChildren() const { return bCrumblingEventIncludesChildren; }
+
+	bool GetIsNotificationDataDirty() const { return bIsNotificationDataDirty; }
+	void ResetIsNotificationDataDirty() { bIsNotificationDataDirty = false; }
+
+	void SetEnableStrainOnCollision(bool bEnable)
+	{
+		bEnableStrainOnCollision = bEnable;
+		bIsDamageSettingsDataDirty = true;
+	}
+
+	bool GetEnableStrainOnCollision() const { return bEnableStrainOnCollision; }
+
+	bool GetIsDamageSettingsDataDirty() const { return bIsDamageSettingsDataDirty; }
+	void ResetIsDamageSettingsDataDirty() { bIsDamageSettingsDataDirty = false; }
+
 private:
-	FTransform WorldTransform;
-	bool IsWorldTransformDirty;
+	uint16 bIsWorldTransformDirty : 1;
+	uint16 bIsCollisionFilterDataDirty : 1;
+	uint16 bIsNotificationDataDirty: 1;
+	uint16 bIsDamageSettingsDataDirty : 1;
+
+	/** updated when bNotificationDataDirty is set */
+	uint16 bNotifyBreakings : 1;
+	uint16 bNotifyRemovals : 1;
+	uint16 bNotifyCrumblings : 1;
+	uint16 bCrumblingEventIncludesChildren : 1;
+
+	/** updated when bDamageSettingsDataDirty is set */
+	uint16 bEnableStrainOnCollision : 1;
+
+	/** updated when bIsWorldTransformDirty is set */
+	FTransform WorldTransform; 
 
+	/** updated when bIsCollisionFilterDataDirty is set */
 	FCollisionFilterData SimFilter;
 	FCollisionFilterData QueryFilter;
-	bool bIsCollisionFilterDataDirty;
 };
 
 /**
diff --git a/Engine/Source/Runtime/Experimental/Chaos/Public/PhysicsProxy/GeometryCollectionPhysicsProxy.h b/Engine/Source/Runtime/Experimental/Chaos/Public/PhysicsProxy/GeometryCollectionPhysicsProxy.h
--- a/Engine/Source/Runtime/Experimental/Chaos/Public/PhysicsProxy/GeometryCollectionPhysicsProxy.h
+++ b/Engine/Source/Runtime/Experimental/Chaos/Public/PhysicsProxy/GeometryCollectionPhysicsProxy.h
@@ -401,6 +401,13 @@ public:
 	void ApplyLinearVelocity_External(FGeometryCollectionItemIndex ItemIndex, const FVector& LinearVelocity);
 	void ApplyAngularVelocity_External(FGeometryCollectionItemIndex ItemIndex, const FVector& AngularVelocity);
 
+	void SetProxyDirty_External();
+
+	void SetEnableDamageFromCollision_External(bool bEnable);
+	void SetNotifyBreakings_External(bool bNotify);
+	void SetNotifyRemovals_External(bool bNotify);
+	void SetNotifyCrumblings_External(bool bNotify, bool bIncludeChildren);
+
 	FProxyInterpolationData& GetInterpolationData() { return InterpolationData; }
 	const FProxyInterpolationData& GetInterpolationData() const { return InterpolationData; }
 
diff --git a/Engine/Source/Runtime/Experimental/GeometryCollectionEngine/Private/GeometryCollection/GeometryCollectionComponent.cpp b/Engine/Source/Runtime/Experimental/GeometryCollectionEngine/Private/GeometryCollection/GeometryCollectionComponent.cpp
--- a/Engine/Source/Runtime/Experimental/GeometryCollectionEngine/Private/GeometryCollection/GeometryCollectionComponent.cpp
+++ b/Engine/Source/Runtime/Experimental/GeometryCollectionEngine/Private/GeometryCollection/GeometryCollectionComponent.cpp
@@ -737,15 +737,9 @@ void UGeometryCollectionComponent::SetNotifyBreaks(bool bNewNotifyBreaks)
 {
 	if (bNotifyBreaks != bNewNotifyBreaks)
 	{
-		if (Chaos::FPhysicsSolver* Solver = GetSolver(*this); Solver && PhysicsProxy)
+		if (PhysicsProxy)
 		{
-			Solver->EnqueueCommandImmediate(
-				[this, bNewNotifyBreaks]()
-				{
-					FSimulationParameters& SimParams = PhysicsProxy->GetSimParameters();
-					SimParams.bGenerateBreakingData = bNewNotifyBreaks;
-				}
-			);
+			PhysicsProxy->SetNotifyBreakings_External(bNewNotifyBreaks);
 		}
 		bNotifyBreaks = bNewNotifyBreaks;
 		UpdateBreakEventRegistration();
@@ -756,15 +750,9 @@ void UGeometryCollectionComponent::SetNotifyRemovals(bool bNewNotifyRemovals)
 {
 	if (bNotifyRemovals != bNewNotifyRemovals)
 	{
-		if (Chaos::FPhysicsSolver* Solver = GetSolver(*this); Solver && PhysicsProxy)
+		if (PhysicsProxy)
 		{
-			Solver->EnqueueCommandImmediate(
-				[this, bNewNotifyRemovals]()
-				{
-					FSimulationParameters& SimParams = PhysicsProxy->GetSimParameters();
-					SimParams.bGenerateBreakingData = bNewNotifyRemovals;
-				}
-			);
+			PhysicsProxy->SetNotifyRemovals_External(bNewNotifyRemovals);
 		}
 		bNotifyRemovals = bNewNotifyRemovals;
 		UpdateRemovalEventRegistration();
@@ -773,18 +761,12 @@ void UGeometryCollectionComponent::SetNotifyRemovals(bool bNewNotifyRemovals)
 
 void UGeometryCollectionComponent::SetNotifyCrumblings(bool bNewNotifyCrumblings, bool bNewCrumblingEventIncludesChildren)
 {
-	if (bNotifyCrumblings != bNewNotifyCrumblings)
+	if (bNotifyCrumblings != bNewNotifyCrumblings || 
+		bCrumblingEventIncludesChildren != bNewCrumblingEventIncludesChildren)
 	{
-		if (Chaos::FPhysicsSolver* Solver = GetSolver(*this); Solver && PhysicsProxy)
+		if (PhysicsProxy)
 		{
-			Solver->EnqueueCommandImmediate(
-				[this, bNewNotifyCrumblings, bNewCrumblingEventIncludesChildren]()
-				{
-					FSimulationParameters& SimParams = PhysicsProxy->GetSimParameters();
-					SimParams.bGenerateCrumblingData = bNewNotifyCrumblings;
-					SimParams.bGenerateCrumblingChildrenData = bNewCrumblingEventIncludesChildren;
-				}
-			);
+			PhysicsProxy->SetNotifyCrumblings_External(bNewNotifyCrumblings, bNewCrumblingEventIncludesChildren);
 		}
 		bNotifyCrumblings = bNewNotifyCrumblings;
 		bCrumblingEventIncludesChildren = bNewCrumblingEventIncludesChildren;
@@ -4760,16 +4742,9 @@ TArray<Chaos::FPhysicsObject*> UGeometryCollectionComponent::GetAllPhysicsObject
 
 void UGeometryCollectionComponent::SetEnableDamageFromCollision(bool bValue)
 {
-	if (Chaos::FPhysicsSolver* Solver = GetSolver(*this); Solver && PhysicsProxy)
+	bEnableDamageFromCollision = bValue;
+	if (PhysicsProxy)
 	{
-		Solver->EnqueueCommandImmediate(
-			[this, bValue]()
-			{
-				FSimulationParameters& SimParams = PhysicsProxy->GetSimParameters();
-				SimParams.bEnableStrainOnCollision = bValue;
-			}
-		);
+		PhysicsProxy->SetEnableDamageFromCollision_External(bValue);
 	}
-
-	bEnableDamageFromCollision = bValue;
 }
\ No newline at end of file
