rework how UVs are accessed in geometry collection: make each UV layer a separate attribute, so it is impossible to have an inconsistent number of UVs per vertex, and we do not allocate so many small arrays. Provide a variety of helper functions to make the transition as easy as possible.

diff --git a/Engine/Plugins/Experimental/ChaosEditor/Source/FractureEditor/Private/FractureToolUV.cpp b/Engine/Plugins/Experimental/ChaosEditor/Source/FractureEditor/Private/FractureToolUV.cpp
--- a/Engine/Plugins/Experimental/ChaosEditor/Source/FractureEditor/Private/FractureToolUV.cpp
+++ b/Engine/Plugins/Experimental/ChaosEditor/Source/FractureEditor/Private/FractureToolUV.cpp
@@ -57,7 +57,7 @@ int32 UFractureAutoUVSettings::GetSelectedChannelIndex(bool bForceToZeroOnFailur
 void UFractureAutoUVSettings::ChangeNumUVChannels(int32 Delta)
 {
 	int32 Target = UVChannelNamesList.Num() + Delta;
-	if (Target > 0 && Target < GeometryCollectionUV::MAX_NUM_UV_CHANNELS)
+	if (Target > 0 && Target <= GeometryCollectionUV::MAX_NUM_UV_CHANNELS)
 	{
 		UFractureToolAutoUV* AutoUVTool = Cast<UFractureToolAutoUV>(OwnerTool.Get());
 		AutoUVTool->UpdateUVChannels(Target);
@@ -168,11 +168,12 @@ void UFractureToolAutoUV::UpdateUVChannels(int32 TargetNumUVChannels)
 				{
 					// Copy UV layer 0 into new UV layer
 					FGeometryCollection& Collection = *GeometryCollectionComponent->GetRestCollection()->GetGeometryCollection();
-					for (int32 Idx = 0; Idx < Collection.UVs.Num(); Idx++)
+					GeometryCollection::UV::FUVLayers Layers = GeometryCollection::UV::FindActiveUVLayers(Collection);
+					for (int32 Idx = 0; Idx < Layers[0].Num(); Idx++)
 					{
 						for (int32 Ch = NumChannels; Ch < TargetNumUVChannels; Ch++)
 						{
-							Collection.UVs[Idx][Ch] = Collection.UVs[Idx][0];
+							Layers[Ch][Idx] = Layers[0][Idx];
 						}
 					}
 				}
diff --git a/Engine/Plugins/Experimental/PlanarCutPlugin/Source/PlanarCut/Private/FractureAutoUV.cpp b/Engine/Plugins/Experimental/PlanarCutPlugin/Source/PlanarCut/Private/FractureAutoUV.cpp
--- a/Engine/Plugins/Experimental/PlanarCutPlugin/Source/PlanarCut/Private/FractureAutoUV.cpp
+++ b/Engine/Plugins/Experimental/PlanarCutPlugin/Source/PlanarCut/Private/FractureAutoUV.cpp
@@ -168,10 +168,12 @@ struct FGeomMesh : public UE::Geometry::FUVPacker::IUVMeshView
 	TArray<FVector3d> GlobalVertices; // vertices transformed to global space
 	TArray<FVector3f> GlobalNormals; // normals transformed to global space
 	int32 UVLayer = 0;
+	TManagedArray<FVector2f>* CollectionUVs;
 
 	// Construct a mesh from a geometry collection and a mask of active triangles
 	FGeomMesh(int32 UVLayer, FGeometryCollection* Collection, TArrayView<bool> ActiveTriangles, int32 NumTriangles) 
-		: Collection(Collection), ActiveTriangles(ActiveTriangles), NumTriangles(NumTriangles), UVLayer(UVLayer)
+		: Collection(Collection), ActiveTriangles(ActiveTriangles), NumTriangles(NumTriangles), UVLayer(UVLayer),
+		CollectionUVs(Collection->FindUVLayer(UVLayer))
 	{
 		ValidateUVLayer();
 		InitVertices();
@@ -180,7 +182,7 @@ struct FGeomMesh : public UE::Geometry::FUVPacker::IUVMeshView
 	// Construct a mesh from an existing FGeomMesh and a mask of active triangles
 	FGeomMesh(const FGeomMesh& OtherMesh, TArrayView<bool> ActiveTriangles, int32 NumTriangles)
 		: Collection(OtherMesh.Collection), ActiveTriangles(ActiveTriangles), NumTriangles(NumTriangles),
-		  UVLayer(OtherMesh.UVLayer)
+		  UVLayer(OtherMesh.UVLayer), CollectionUVs(OtherMesh.Collection->FindUVLayer(OtherMesh.UVLayer))
 	{
 		ValidateUVLayer();
 		GlobalVertices = OtherMesh.GlobalVertices;
@@ -191,7 +193,9 @@ struct FGeomMesh : public UE::Geometry::FUVPacker::IUVMeshView
 		if (!ensure(UVLayer >= 0 && UVLayer < Collection->NumUVLayers()))
 		{
 			UVLayer = FMath::Clamp(UVLayer, 0, Collection->NumUVLayers());
+			CollectionUVs = Collection->FindUVLayer(UVLayer);
 		}
+		check(CollectionUVs);
 	}
 
 	void InitVertices()
@@ -227,14 +231,14 @@ struct FGeomMesh : public UE::Geometry::FUVPacker::IUVMeshView
 		return GlobalVertices[VID];
 	}
 
-	virtual FVector2f GetUV(int32 EID) const
+	virtual FVector2f GetUV(int32 VID) const
 	{
-		return FVector2f(Collection->UVs[EID][UVLayer]);
+		return (*CollectionUVs)[VID];
 	}
 
-	virtual void SetUV(int32 EID, FVector2f UVIn)
+	virtual void SetUV(int32 VID, FVector2f UVIn)
 	{
-		FVector2f& UV = Collection->UVs[EID][UVLayer];
+		FVector2f& UV = (*CollectionUVs)[VID];
 		UV.X = UVIn.X;
 		UV.Y = UVIn.Y;
 	}
@@ -300,9 +304,10 @@ struct FGeomFlatUVMesh
 	const TArrayView<bool> ActiveTriangles;
 	int32 NumTriangles;
 	int32 UVLayer;
+	TManagedArray<FVector2f>* CollectionUVs;
 
 	FGeomFlatUVMesh(int32 UVLayer, FGeometryCollection* Collection, TArrayView<bool> ActiveTriangles, int32 NumTriangles)
-		: Collection(Collection), ActiveTriangles(ActiveTriangles), NumTriangles(NumTriangles), UVLayer(UVLayer)
+		: Collection(Collection), ActiveTriangles(ActiveTriangles), NumTriangles(NumTriangles), UVLayer(UVLayer), CollectionUVs(Collection->FindUVLayer(UVLayer))
 	{
 		ValidateUVLayer();
 	}
@@ -312,7 +317,9 @@ struct FGeomFlatUVMesh
 		if (!ensure(UVLayer >= 0 && UVLayer < Collection->NumUVLayers()))
 		{
 			UVLayer = FMath::Clamp(UVLayer, 0, Collection->NumUVLayers());
+			CollectionUVs = Collection->FindUVLayer(UVLayer);
 		}
+		check(CollectionUVs);
 	}
 
 	inline FIndex3i GetTriangle(int32 TID) const
@@ -322,7 +329,7 @@ struct FGeomFlatUVMesh
 
 	inline FVector3d GetVertex(int32 VID) const
 	{
-		const FVector2f& UV = Collection->UVs[VID][UVLayer];
+		const FVector2f& UV = (*CollectionUVs)[VID];
 		return FVector3d(UV.X, UV.Y, 0);
 	}
 
@@ -340,7 +347,7 @@ struct FGeomFlatUVMesh
 	}
 	inline int32 MaxVertexID() const
 	{
-		return Collection->UVs.Num();
+		return (*CollectionUVs).Num();
 	}
 	inline int32 TriangleCount() const
 	{
@@ -348,7 +355,7 @@ struct FGeomFlatUVMesh
 	}
 	inline int32 VertexCount() const
 	{
-		return Collection->UVs.Num();
+		return (*CollectionUVs).Num();
 	}
 	constexpr inline uint64 GetChangeStamp() const
 	{
diff --git a/Engine/Plugins/Experimental/PlanarCutPlugin/Source/PlanarCut/Private/GeometryMeshConversion.cpp b/Engine/Plugins/Experimental/PlanarCutPlugin/Source/PlanarCut/Private/GeometryMeshConversion.cpp
--- a/Engine/Plugins/Experimental/PlanarCutPlugin/Source/PlanarCut/Private/GeometryMeshConversion.cpp
+++ b/Engine/Plugins/Experimental/PlanarCutPlugin/Source/PlanarCut/Private/GeometryMeshConversion.cpp
@@ -2477,7 +2477,8 @@ void FCellMeshes::CreateMeshesForSinglePlane(const FPlanarCells& Cells, const FA
 
 void FDynamicMeshCollection::Init(const FGeometryCollection* Collection, const TManagedArray<FTransform>& Transforms, const TArrayView<const int32>& TransformIndices, FTransform TransformCollection, bool bSaveIsolatedVertices)
 {
-	int32 NumUVLayers = Collection->NumUVLayers();
+	GeometryCollection::UV::FConstUVLayers UVLayers = GeometryCollection::UV::FindActiveUVLayers(*Collection);
+	int32 NumUVLayers = UVLayers.Num();
 	Meshes.Reset();
 	Bounds = FAxisAlignedBox3d::Empty();
 
@@ -2501,9 +2502,6 @@ void FDynamicMeshCollection::Init(const FGeometryCollection* Collection, const T
 		Mesh.EnableAttributes();
 		Mesh.Attributes()->EnableMaterialID();
 
-		//int32 NumUVLayers = Collection->NumUVLayers();
-		//AugmentedDynamicMesh::InitializeOverlayToPerVertexUVs(Mesh, NumUVLayers);
-
 		int32 VertexStart = Collection->VertexStart[GeometryIdx];
 		int32 VertexCount = Collection->VertexCount[GeometryIdx];
 		int32 FaceCount = Collection->FaceCount[GeometryIdx];
@@ -2524,7 +2522,7 @@ void FDynamicMeshCollection::Init(const FGeometryCollection* Collection, const T
 			
 			for (int32 UVLayer = 0; UVLayer < NumUVLayers; ++UVLayer)
 			{
-				AugmentedDynamicMesh::SetUV(Mesh, VID, Collection->UVs[Idx][UVLayer], UVLayer);
+				AugmentedDynamicMesh::SetUV(Mesh, VID, UVLayers[UVLayer][Idx], UVLayer);
 			}
 		}
 		FIntVector VertexOffset(VertexStart, VertexStart, VertexStart);
@@ -2547,7 +2545,7 @@ void FDynamicMeshCollection::Init(const FGeometryCollection* Collection, const T
 
 					for (int32 UVLayer = 0; UVLayer < NumUVLayers; ++UVLayer)
 					{
-						AugmentedDynamicMesh::SetUV(Mesh, NewVID, Collection->UVs[SrcIdx][UVLayer], UVLayer);
+						AugmentedDynamicMesh::SetUV(Mesh, NewVID, UVLayers[UVLayer][SrcIdx], UVLayer);
 					}
 
 					NewTri[SubIdx] = NewVID;
@@ -3221,6 +3219,7 @@ bool FDynamicMeshCollection::UpdateCollection(const FTransform& FromCollection,
 
 	int32 UVLayerCount = AugmentedDynamicMesh::NumEnabledUVChannels(Mesh);
 	Output.SetNumUVLayers(UVLayerCount);
+	GeometryCollection::UV::FUVLayers OutputUVLayers = GeometryCollection::UV::FindActiveUVLayers(Output);
 
 	int32 NewVertexCount = Mesh.VertexCount();
 	int32 NewTriangleCount = Mesh.TriangleCount();
@@ -3245,7 +3244,7 @@ bool FDynamicMeshCollection::UpdateCollection(const FTransform& FromCollection,
 		{
 			FVector2f UV;
 			AugmentedDynamicMesh::GetUV(Mesh, VID, UV, UVLayer);
-			Output.UVs[CopyToIdx][UVLayer] = UV;
+			OutputUVLayers[UVLayer][CopyToIdx] = UV;
 		}
 		
 		FVector3f TangentU, TangentV;
@@ -3335,6 +3334,9 @@ int32 FDynamicMeshCollection::AppendToCollection(const FTransform& FromCollectio
 	int32 FacesStart = Output.AddElements(NumTriangles, FGeometryCollection::FacesGroup);
 	int32 VerticesStart = Output.AddElements(NumVertices, FGeometryCollection::VerticesGroup);
 
+	Output.SetNumUVLayers(UVLayerCount);
+	GeometryCollection::UV::FUVLayers OutputUVLayers = GeometryCollection::UV::FindActiveUVLayers(Output);
+
 	for (int32 VID = 0; VID < Mesh.MaxVertexID(); VID++)
 	{
 		checkSlow(Mesh.IsVertex(VID)); // mesh is compact
@@ -3342,12 +3344,11 @@ int32 FDynamicMeshCollection::AppendToCollection(const FTransform& FromCollectio
 		Output.Vertex[CopyToIdx] = (FVector3f)FromCollection.InverseTransformPosition(FVector(Mesh.GetVertex(VID)));
 		Output.Normal[CopyToIdx] = (FVector3f)FromCollection.InverseTransformVectorNoScale(FVector(Mesh.GetVertexNormal(VID)));
 		
-		Output.UVs[CopyToIdx].SetNum(UVLayerCount);
 		for (int32 UVLayer = 0; UVLayer < UVLayerCount; ++UVLayer)
 		{
 			FVector2f UV;
 			AugmentedDynamicMesh::GetUV(Mesh, VID, UV, UVLayer);
-			Output.UVs[CopyToIdx][UVLayer] = UV;
+			OutputUVLayers[UVLayer][CopyToIdx] = UV;
 		}
 
 		FVector3f TangentU, TangentV;
diff --git a/Engine/Plugins/Experimental/PlanarCutPlugin/Source/PlanarCut/Private/PlanarCut.cpp b/Engine/Plugins/Experimental/PlanarCutPlugin/Source/PlanarCut/Private/PlanarCut.cpp
--- a/Engine/Plugins/Experimental/PlanarCutPlugin/Source/PlanarCut/Private/PlanarCut.cpp
+++ b/Engine/Plugins/Experimental/PlanarCutPlugin/Source/PlanarCut/Private/PlanarCut.cpp
@@ -89,7 +89,7 @@ int32 FInternalSurfaceMaterials::GetDefaultMaterialIDForGeometry(const FGeometry
 void FInternalSurfaceMaterials::SetUVScaleFromCollection(const GeometryCollection::Facades::FCollectionMeshFacade& CollectionMesh, int32 GeometryIdx)
 {
 	const auto& VertexArray = CollectionMesh.VertexAttribute.Get();
-	const auto& UVsArray = CollectionMesh.UVsAttribute.Get();
+	const auto& UV0Array = CollectionMesh.GetUVLayer(0);
 	const auto& IndicesArray = CollectionMesh.IndicesAttribute.Get();
 	const auto& FaceStartArray = CollectionMesh.FaceStartAttribute.Get();
 	const auto& FaceCountArray = CollectionMesh.FaceCountAttribute.Get();
@@ -107,11 +107,11 @@ void FInternalSurfaceMaterials::SetUVScaleFromCollection(const GeometryCollectio
 	{
 		const FIntVector& Tri = IndicesArray[FaceIdx];
 		WorldDistance += FVector3f::Distance(VertexArray[Tri.X], VertexArray[Tri.Y]);
-		UVDistance += FVector2D::Distance(FVector2D(UVsArray[Tri.X][0]), FVector2D(UVsArray[Tri.Y][0]));
+		UVDistance += FVector2D::Distance(FVector2D(UV0Array[Tri.X]), FVector2D(UV0Array[Tri.Y]));
 		WorldDistance += FVector3f::Distance(VertexArray[Tri.Z], VertexArray[Tri.Y]);
-		UVDistance += FVector2D::Distance(FVector2D(UVsArray[Tri.Z][0]), FVector2D(UVsArray[Tri.Y][0]));
+		UVDistance += FVector2D::Distance(FVector2D(UV0Array[Tri.Z]), FVector2D(UV0Array[Tri.Y]));
 		WorldDistance += FVector3f::Distance(VertexArray[Tri.X], VertexArray[Tri.Z]);
-		UVDistance += FVector2D::Distance(FVector2D(UVsArray[Tri.X][0]), FVector2D(UVsArray[Tri.Z][0]));
+		UVDistance += FVector2D::Distance(FVector2D(UV0Array[Tri.X]), FVector2D(UV0Array[Tri.Z]));
 	}
 
 	if (WorldDistance > 0)
diff --git a/Engine/Plugins/MeshPainting/Source/MeshPaintingToolset/MeshPaintingToolset.Build.cs b/Engine/Plugins/MeshPainting/Source/MeshPaintingToolset/MeshPaintingToolset.Build.cs
--- a/Engine/Plugins/MeshPainting/Source/MeshPaintingToolset/MeshPaintingToolset.Build.cs
+++ b/Engine/Plugins/MeshPainting/Source/MeshPaintingToolset/MeshPaintingToolset.Build.cs
@@ -26,7 +26,8 @@ namespace UnrealBuildTool.Rules
 					"TypedElementRuntime",
 					"UnrealEd",
 					"InterchangeEngine",
-					"InterchangePipelines"
+					"InterchangePipelines",
+					"Chaos"
 				}
                 );
 
diff --git a/Engine/Plugins/MeshPainting/Source/MeshPaintingToolset/Private/MeshPaintGeometryCollectionAdapter.cpp b/Engine/Plugins/MeshPainting/Source/MeshPaintingToolset/Private/MeshPaintGeometryCollectionAdapter.cpp
--- a/Engine/Plugins/MeshPainting/Source/MeshPaintingToolset/Private/MeshPaintGeometryCollectionAdapter.cpp
+++ b/Engine/Plugins/MeshPainting/Source/MeshPaintingToolset/Private/MeshPaintGeometryCollectionAdapter.cpp
@@ -250,10 +250,10 @@ void FMeshPaintGeometryCollectionComponentAdapter::ApplyOrRemoveTextureOverride(
 
 void FMeshPaintGeometryCollectionComponentAdapter::GetTextureCoordinate(int32 VertexIndex, int32 ChannelIndex, FVector2D& OutTextureCoordinate) const
 {
-	const TArray<FVector2f>& UV = GetGeometryCollectionObject()->GetGeometryCollection()->UVs[VertexIndex];
-	if (ChannelIndex < UV.Num())
+	const TManagedArray<FVector2f>* UVLayer = GetGeometryCollectionObject()->GetGeometryCollection()->FindUVLayer(ChannelIndex);
+	if (UVLayer)
 	{
-		OutTextureCoordinate = (FVector2D)UV[ChannelIndex];
+		OutTextureCoordinate = (FVector2D)(*UVLayer)[VertexIndex];
 	}
 }
 
diff --git a/Engine/Source/Runtime/Experimental/Chaos/Private/GeometryCollection/Facades/CollectionMeshFacade.cpp b/Engine/Source/Runtime/Experimental/Chaos/Private/GeometryCollection/Facades/CollectionMeshFacade.cpp
--- a/Engine/Source/Runtime/Experimental/Chaos/Private/GeometryCollection/Facades/CollectionMeshFacade.cpp
+++ b/Engine/Source/Runtime/Experimental/Chaos/Private/GeometryCollection/Facades/CollectionMeshFacade.cpp
@@ -11,12 +11,12 @@
 namespace GeometryCollection::Facades
 {
 	FCollectionMeshFacade::FCollectionMeshFacade(FManagedArrayCollection& InCollection)
-		: TransformToGeometryIndexAttribute(InCollection, "TransformToGeometryIndex", FGeometryCollection::TransformGroup)
+		: FCollectionUVFacade(InCollection)
+		, TransformToGeometryIndexAttribute(InCollection, "TransformToGeometryIndex", FGeometryCollection::TransformGroup)
 		, VertexAttribute(InCollection, "Vertex", FGeometryCollection::VerticesGroup)
 		, TangentUAttribute(InCollection, "TangentU", FGeometryCollection::VerticesGroup)
 		, TangentVAttribute(InCollection, "TangentV", FGeometryCollection::VerticesGroup)
 		, NormalAttribute(InCollection, "Normal", FGeometryCollection::VerticesGroup)
-		, UVsAttribute(InCollection, "UVs", FGeometryCollection::VerticesGroup)
 		, ColorAttribute(InCollection, "Color", FGeometryCollection::VerticesGroup)
 		, BoneMapAttribute(InCollection, "BoneMap", FGeometryCollection::VerticesGroup)
 		, VertexStartAttribute(InCollection, "VertexStart", FGeometryCollection::GeometryGroup)
@@ -31,12 +31,12 @@ namespace GeometryCollection::Facades
 	}
 
 	FCollectionMeshFacade::FCollectionMeshFacade(const FManagedArrayCollection& InCollection)
-		: TransformToGeometryIndexAttribute(InCollection, "TransformToGeometryIndex", FGeometryCollection::TransformGroup)
+		: FCollectionUVFacade(InCollection)
+		, TransformToGeometryIndexAttribute(InCollection, "TransformToGeometryIndex", FGeometryCollection::TransformGroup)
 		, VertexAttribute(InCollection, "Vertex", FGeometryCollection::VerticesGroup)
 		, TangentUAttribute(InCollection, "TangentU", FGeometryCollection::VerticesGroup)
 		, TangentVAttribute(InCollection, "TangentV", FGeometryCollection::VerticesGroup)
 		, NormalAttribute(InCollection, "Normal", FGeometryCollection::VerticesGroup)
-		, UVsAttribute(InCollection, "UVs", FGeometryCollection::VerticesGroup)
 		, ColorAttribute(InCollection, "Color", FGeometryCollection::VerticesGroup)
 		, BoneMapAttribute(InCollection, "BoneMap", FGeometryCollection::VerticesGroup)
 		, VertexStartAttribute(InCollection, "VertexStart", FGeometryCollection::GeometryGroup)
@@ -52,12 +52,12 @@ namespace GeometryCollection::Facades
 
 	bool FCollectionMeshFacade::IsValid() const
 	{
-		return TransformToGeometryIndexAttribute.IsValid()
+		return FCollectionUVFacade::IsValid()
+			&& TransformToGeometryIndexAttribute.IsValid()
 			&& VertexAttribute.IsValid()
 			&& TangentUAttribute.IsValid()
 			&& TangentVAttribute.IsValid()
 			&& NormalAttribute.IsValid()
-			&& UVsAttribute.IsValid()
 			&& ColorAttribute.IsValid()
 			&& BoneMapAttribute.IsValid()
 			&& VertexStartAttribute.IsValid()
@@ -73,12 +73,12 @@ namespace GeometryCollection::Facades
 
 	void FCollectionMeshFacade::DefineSchema()
 	{
+		FCollectionUVFacade::DefineSchema();
 		TransformToGeometryIndexAttribute.Add();
 		VertexAttribute.Add();
 		TangentUAttribute.Add();
 		TangentVAttribute.Add();
 		NormalAttribute.Add();
-		UVsAttribute.Add();
 		ColorAttribute.Add();
 		BoneMapAttribute.Add();
 		VertexStartAttribute.Add();
diff --git a/Engine/Source/Runtime/Experimental/Chaos/Private/GeometryCollection/Facades/CollectionUVFacade.cpp b/Engine/Source/Runtime/Experimental/Chaos/Private/GeometryCollection/Facades/CollectionUVFacade.cpp
new file mode 100644
--- /dev/null
+++ b/Engine/Source/Runtime/Experimental/Chaos/Private/GeometryCollection/Facades/CollectionUVFacade.cpp
@@ -0,0 +1,70 @@
+// Copyright Epic Games, Inc. All Rights Reserved.
+
+
+#include "GeometryCollection/Facades/CollectionUVFacade.h"
+#include "GeometryCollection/TransformCollection.h"
+#include "GeometryCollection/GeometryCollection.h"
+
+namespace GeometryCollection::UV
+{
+	const FName UVLayerNames[GeometryCollectionUV::MAX_NUM_UV_CHANNELS] = {
+		"UVLayer0",
+		"UVLayer1",
+		"UVLayer2",
+		"UVLayer3",
+		"UVLayer4",
+		"UVLayer5",
+		"UVLayer6",
+		"UVLayer7"
+	};
+
+	const FName VerticesGroupName = FGeometryCollection::VerticesGroup;
+
+	bool HasValidUVs(const FManagedArrayCollection& Collection)
+	{
+		return Collection.HasAttribute(UVLayerNames[0], VerticesGroupName);
+	}
+
+	void DefineUVSchema(FManagedArrayCollection& Collection)
+	{
+		// Only the first layer's attribute must always be present
+		Collection.AddAttribute<FVector2f>(UVLayerNames[0], VerticesGroupName);
+	}
+
+	bool SetNumUVLayers(FManagedArrayCollection& Collection, int32 NumLayers)
+	{
+		if (NumLayers < 1 || NumLayers > GeometryCollectionUV::MAX_NUM_UV_CHANNELS)
+		{
+			return false;
+		}
+
+		// enable layers 0 to Num
+		for (int32 LayerIdx = 0; LayerIdx < NumLayers; ++LayerIdx)
+		{
+			Collection.AddAttribute<FVector2f>(UVLayerNames[LayerIdx], VerticesGroupName); // Note: AddAttribute checks HasAttribute; no need to do so here
+		}
+		// disable layers Num to Max
+		for (int32 LayerIdx = NumLayers; LayerIdx < GeometryCollectionUV::MAX_NUM_UV_CHANNELS; ++LayerIdx)
+		{
+			Collection.RemoveAttribute(UVLayerNames[LayerIdx], VerticesGroupName);
+		}
+		return true;
+	}
+
+	int32 GetNumUVLayers(const FManagedArrayCollection& Collection)
+	{
+		// Any collection with valid UV layers must always have layer 0
+		checkSlow(Collection.HasAttribute(UVLayerNames[0], VerticesGroupName));
+
+		int32 LayerCount = 1;
+		for (; LayerCount < GeometryCollectionUV::MAX_NUM_UV_CHANNELS; ++LayerCount)
+		{
+			if (!Collection.HasAttribute(UVLayerNames[LayerCount], VerticesGroupName))
+			{
+				break;
+			}
+		}
+
+		return LayerCount;
+	}
+}
diff --git a/Engine/Source/Runtime/Experimental/Chaos/Private/GeometryCollection/GeometryCollection.cpp b/Engine/Source/Runtime/Experimental/Chaos/Private/GeometryCollection/GeometryCollection.cpp
--- a/Engine/Source/Runtime/Experimental/Chaos/Private/GeometryCollection/GeometryCollection.cpp
+++ b/Engine/Source/Runtime/Experimental/Chaos/Private/GeometryCollection/GeometryCollection.cpp
@@ -60,7 +60,7 @@ void FGeometryCollection::DefineGeometrySchema(FManagedArrayCollection& InCollec
 	// Vertices Group
 	InCollection.AddAttribute<FVector3f>("Vertex", FGeometryCollection::VerticesGroup);
 	InCollection.AddAttribute<FVector3f>("Normal", FGeometryCollection::VerticesGroup);
-	InCollection.AddAttribute<TArray<FVector2f>>("UVs", FGeometryCollection::VerticesGroup);
+	GeometryCollection::UV::DefineUVSchema(InCollection);
 	InCollection.AddAttribute<FLinearColor>("Color", FGeometryCollection::VerticesGroup);
 	InCollection.AddAttribute<FVector3f>("TangentU", FGeometryCollection::VerticesGroup);
 	InCollection.AddAttribute<FVector3f>("TangentV", FGeometryCollection::VerticesGroup);
@@ -102,11 +102,11 @@ void FGeometryCollection::Construct()
 	// Vertices Group
 	AddExternalAttribute<FVector3f>("Vertex", FGeometryCollection::VerticesGroup, Vertex);
 	AddExternalAttribute<FVector3f>("Normal", FGeometryCollection::VerticesGroup, Normal);
-	AddExternalAttribute<TArray<FVector2f>>("UVs", FGeometryCollection::VerticesGroup, UVs);
 	AddExternalAttribute<FLinearColor>("Color", FGeometryCollection::VerticesGroup, Color);
 	AddExternalAttribute<FVector3f>("TangentU", FGeometryCollection::VerticesGroup, TangentU);
 	AddExternalAttribute<FVector3f>("TangentV", FGeometryCollection::VerticesGroup, TangentV);
 	AddExternalAttribute<int32>("BoneMap", FGeometryCollection::VerticesGroup, BoneMap, TransformDependency);
+	GeometryCollection::UV::DefineUVSchema(*this);
 
 	// Faces Group
 	AddExternalAttribute<FIntVector>("Indices", FGeometryCollection::FacesGroup, Indices, VerticesDependency);
@@ -167,7 +167,6 @@ int32 FGeometryCollection::AppendGeometry(const FGeometryCollection & Element, i
 	int NumNewVertices = Element.NumElements(FGeometryCollection::VerticesGroup);
 	const TManagedArray<FVector3f>& ElementVertices = Element.Vertex;
 	const TManagedArray<FVector3f>& ElementNormals = Element.Normal;
-	const TManagedArray<TArray<FVector2f>>& ElementUVs = Element.UVs;
 	const TManagedArray<FLinearColor>& ElementColors = Element.Color;
 	const TManagedArray<FVector3f>& ElementTangentUs = Element.TangentU;
 	const TManagedArray<FVector3f>& ElementTangentVs = Element.TangentV;
@@ -211,23 +210,28 @@ int32 FGeometryCollection::AppendGeometry(const FGeometryCollection & Element, i
 	int VerticesIndex = AddElements(NumNewVertices, FGeometryCollection::VerticesGroup);
 	TManagedArray<FVector3f>& Vertices = Vertex;
 	TManagedArray<FVector3f>& Normals = Normal;
-	TManagedArray<TArray<FVector2f>>& AllUVs = UVs;
 	TManagedArray<FLinearColor>& Colors = Color;
 	TManagedArray<FVector3f>& TangentUs = TangentU;
 	TManagedArray<FVector3f>& TangentVs = TangentV;
 	TManagedArray<int32>& BoneMaps = BoneMap;
 	TManagedArray<FIntVector>& FaceIndices = Indices;
 
+	// Make sure we have enough UV layers to copy all the Element's UVs
+	int32 ExistingUVLayers = NumUVLayers();
+	int32 ElementUVLayers = Element.NumUVLayers();
+	if (ElementUVLayers > ExistingUVLayers)
+	{
+		SetNumUVLayers(ElementUVLayers);
+	}
+
 	for (int vdx = 0; vdx < NumNewVertices; vdx++)
 	{
 		Vertices[VerticesIndex + vdx] = ElementVertices[vdx];
 		Normals[VerticesIndex + vdx] = ElementNormals[vdx];
 		
-		const int32 NumUVLayers = ElementUVs[vdx].Num();
-		AllUVs[VerticesIndex + vdx].SetNum(NumUVLayers);
-		for (int UVChannelIndex = 0; UVChannelIndex < NumUVLayers; UVChannelIndex++)
+		for (int UVLayerIndex = 0; UVLayerIndex < ElementUVLayers; UVLayerIndex++)
 		{
-			AllUVs[VerticesIndex + vdx][UVChannelIndex] = ElementUVs[vdx][UVChannelIndex];
+			ModifyUV(VerticesIndex + vdx, UVLayerIndex) = Element.GetUV(vdx, UVLayerIndex);
 		}
 
 		Colors[VerticesIndex + vdx] = ElementColors[vdx];
@@ -672,6 +676,7 @@ void FGeometryCollection::Empty()
 	}
 	// re-initialize interfaces
 	InitializeInterfaces();
+	SetNumUVLayers(1);
 }
 
 void FGeometryCollection::InitializeInterfaces()
@@ -1151,7 +1156,7 @@ void FGeometryCollection::Serialize(Chaos::FChaosArchive& Ar)
 			{
 				// Note: As UVs is an external attribute that is always added by Construct, this should never be encountered
 				UE_LOG(FGeometryCollectionLogging, Log, TEXT("GeometryCollection updated to multiple UV sets."));
-				AddExternalAttribute<TArray<FVector2f>>("UVs", FGeometryCollection::VerticesGroup, UVs);
+				AddAttribute<TArray<FVector2f>>("UVs", FGeometryCollection::VerticesGroup);
 			}
 
 			TManagedArray<TArray<FVector2f>>& MultipleUVs = ModifyAttribute<TArray<FVector2f>>("UVs", FGeometryCollection::VerticesGroup);
@@ -1262,6 +1267,46 @@ void FGeometryCollection::Serialize(Chaos::FChaosArchive& Ar)
 			HierarchyFacade.GenerateLevelAttribute();
 		}
 
+		if (Version < 10 && 
+			(!HasAttribute(GeometryCollection::UV::UVLayerNames[0], VerticesGroup) || HasAttribute("UVs", VerticesGroup)) )
+		{
+			TManagedArray<TArray<FVector2f>>* OrigUVs = FindAttributeTyped<TArray<FVector2f>>("UVs", VerticesGroup);
+			if (OrigUVs)
+			{
+				// Note: We take the max of the num layers because in practice there have been some vertices with inconsistent layer counts
+				// and it seems better to transfer all the data (with missing data left as zeros) than to potentially lose data
+				int32 NumLayers = 1;
+				for (int32 Idx = 0; Idx < OrigUVs->Num(); ++Idx)
+				{
+					NumLayers = FMath::Max((*OrigUVs)[Idx].Num(), NumLayers);
+				}
+				NumLayers = FMath::Min((int32)GeometryCollectionUV::MAX_NUM_UV_CHANNELS, NumLayers); // make sure we never exceed max layers
+				SetNumUVLayers(NumLayers);
+
+				for (int32 Idx = 0; Idx < OrigUVs->Num(); ++Idx)
+				{
+					TArray<FVector2f>& VertexLayers = (*OrigUVs)[Idx];
+					int32 NumVertexLayers = FMath::Min(VertexLayers.Num(), NumLayers);
+					for (int32 LayerIdx = 0; LayerIdx < NumVertexLayers; ++LayerIdx)
+					{
+						ModifyUV(Idx, LayerIdx) = VertexLayers[LayerIdx];
+					}
+					for (int32 LayerIdx = NumVertexLayers; LayerIdx < NumLayers; ++LayerIdx)
+					{
+						ModifyUV(Idx, LayerIdx) = FVector2f(0, 0); // explicitly zero UVs of any missing layers
+					}
+				}
+				
+				RemoveAttribute("UVs", VerticesGroup);
+			}
+			else
+			{
+				SetNumUVLayers(1);
+			}
+
+			Version = 10;
+		}
+
 		// Finally, make sure expected interfaces are initialized
 		InitializeInterfaces();
 	}
@@ -1405,34 +1450,12 @@ bool FGeometryCollection::HasContiguousRenderFaces() const
 
 int32 FGeometryCollection::NumUVLayers() const
 {
-	if (NumElements(VerticesGroup) == 0)
-	{
-		return 0;
-	}
-
-	// A valid GeometryCollection has the same number of UVs on every vertex.
-	// For safety, we consider only the smallest number of layers.
-	int32 MinUVLayers = UVs[0].Num();
-	for (int32 VertIdx = 0; VertIdx < NumElements(FGeometryCollection::VerticesGroup); ++VertIdx)
-	{
-		MinUVLayers = FMath::Min(MinUVLayers, UVs[VertIdx].Num());
-	}
-
-	return MinUVLayers;
+	return GeometryCollection::UV::GetNumUVLayers(*this);
 }
 
 bool FGeometryCollection::SetNumUVLayers(int32 NumLayers)
 {
-	if (NumLayers < 1 || NumLayers >= GeometryCollectionUV::MAX_NUM_UV_CHANNELS)
-	{
-		return false;
-	}
-
-	for (int32 VertIdx = 0; VertIdx < NumElements(FGeometryCollection::VerticesGroup); ++VertIdx)
-	{
-		UVs[VertIdx].SetNum(NumLayers);
-	}
-	return true;
+	return GeometryCollection::UV::SetNumUVLayers(*this, NumLayers);
 }
 
 bool FGeometryCollection::IsVisible(int32 Element) const
@@ -1481,7 +1504,6 @@ void FGeometryCollection::Init(FGeometryCollection* Collection, const TArray<flo
 		TManagedArray<FVector3f>& Normals = Collection->Normal;
 		TManagedArray<FVector3f>& TangentU = Collection->TangentU;
 		TManagedArray<FVector3f>& TangentV = Collection->TangentV;
-		TManagedArray<TArray<FVector2f>>& UVs = Collection->UVs;
 		TManagedArray<FLinearColor>& Colors = Collection->Color;
 		TManagedArray<FIntVector>& Indices = Collection->Indices;
 		TManagedArray<bool>& Visible = Collection->Visible;
@@ -1496,10 +1518,11 @@ void FGeometryCollection::Init(FGeometryCollection* Collection, const TArray<flo
 			Vertices[Idx] = FVector3f(RawVertexArray[3 * Idx], RawVertexArray[3 * Idx + 1], RawVertexArray[3 * Idx + 2]);
 			TempVertices += FVector3d(Vertices[Idx]);
 
-			UVs[Idx].SetNumZeroed(GeometryCollectionUV::MAX_NUM_UV_CHANNELS);
 			Colors[Idx] = FLinearColor::White;
 		}
 
+		Collection->SetNumUVLayers(1);
+
 		// set the particle information
 		TempVertices /= (float)NumNewVertices;
 		Transform[0] = FTransform(TempVertices);
@@ -1794,7 +1817,6 @@ FGeometryCollection* FGeometryCollection::NewGeometryCollection(const TArray<flo
 	TManagedArray<FVector3f>&  Normals = RestCollection->Normal;
 	TManagedArray<FVector3f>&  TangentU = RestCollection->TangentU;
 	TManagedArray<FVector3f>&  TangentV = RestCollection->TangentV;
-	TManagedArray<TArray<FVector2f>>& UVs = RestCollection->UVs;
 	TManagedArray<FLinearColor>&  Colors = RestCollection->Color;
 	TManagedArray<int32>& BoneMap = RestCollection->BoneMap;
 	TManagedArray<FIntVector>&  Indices = RestCollection->Indices;
@@ -1814,10 +1836,11 @@ FGeometryCollection* FGeometryCollection::NewGeometryCollection(const TArray<flo
 		Vertices[Idx] = FVector3f(RawVertexArray[3 * Idx], RawVertexArray[3 * Idx + 1], RawVertexArray[3 * Idx + 2]);
 		BoneMap[Idx] = RawBoneMapArray[Idx];
 
-		UVs[Idx].SetNumZeroed(GeometryCollectionUV::MAX_NUM_UV_CHANNELS);
 		Colors[Idx] = FLinearColor::White;
 	}
 
+	RestCollection->SetNumUVLayers(1);
+
 	// Transforms
 	int NumNewTransforms = RawTransformArray.Num(); // 1 particle for this geometry structure
 	int TransformIndex = RestCollection->AddElements(NumNewTransforms, FGeometryCollection::TransformGroup);
diff --git a/Engine/Source/Runtime/Experimental/Chaos/Private/GeometryCollection/GeometryCollectionUtility.cpp b/Engine/Source/Runtime/Experimental/Chaos/Private/GeometryCollection/GeometryCollectionUtility.cpp
--- a/Engine/Source/Runtime/Experimental/Chaos/Private/GeometryCollection/GeometryCollectionUtility.cpp
+++ b/Engine/Source/Runtime/Experimental/Chaos/Private/GeometryCollection/GeometryCollectionUtility.cpp
@@ -32,7 +32,7 @@ namespace GeometryCollection
 		TManagedArray<FVector3f>&  Normals = RestCollection->Normal;
 		TManagedArray<FVector3f>&  TangentU = RestCollection->TangentU;
 		TManagedArray<FVector3f>&  TangentV = RestCollection->TangentV;
-		TManagedArray<TArray<FVector2f>>& UVs = RestCollection->UVs;
+		TManagedArray<FVector2f>& UV0 = *RestCollection->FindUVLayer(0);
 		TManagedArray<FLinearColor>&  Colors = RestCollection->Color;
 		TManagedArray<FIntVector>&  Indices = RestCollection->Indices;
 		TManagedArray<bool>&  Visible = RestCollection->Visible;
@@ -67,14 +67,14 @@ namespace GeometryCollection
 		Normals[6] = FVector3f(-1.f, 1.f, 1.f).GetSafeNormal();
 		Normals[7] = FVector3f(1.f, 1.f, 1.f).GetSafeNormal();
 
-		UVs[0].Add(FVector2f(0, 0));
-		UVs[1].Add(FVector2f(1, 0));
-		UVs[2].Add(FVector2f(0, 1));
-		UVs[3].Add(FVector2f(1, 1));
-		UVs[4].Add(FVector2f(0, 0));
-		UVs[5].Add(FVector2f(1, 0));
-		UVs[6].Add(FVector2f(0, 1));
-		UVs[7].Add(FVector2f(1, 1));
+		UV0[0] = FVector2f(0, 0);
+		UV0[1] = FVector2f(1, 0);
+		UV0[2] = FVector2f(0, 1);
+		UV0[3] = FVector2f(1, 1);
+		UV0[4] = FVector2f(0, 0);
+		UV0[5] = FVector2f(1, 0);
+		UV0[6] = FVector2f(0, 1);
+		UV0[7] = FVector2f(1, 1);
 
 		Colors[0] = FLinearColor::White;
 		Colors[1] = FLinearColor::White;
diff --git a/Engine/Source/Runtime/Experimental/Chaos/Private/GeometryCollection/ManagedArrayCollection.cpp b/Engine/Source/Runtime/Experimental/Chaos/Private/GeometryCollection/ManagedArrayCollection.cpp
--- a/Engine/Source/Runtime/Experimental/Chaos/Private/GeometryCollection/ManagedArrayCollection.cpp
+++ b/Engine/Source/Runtime/Experimental/Chaos/Private/GeometryCollection/ManagedArrayCollection.cpp
@@ -445,6 +445,8 @@ void FManagedArrayCollection::CopyMatchingAttributesFrom(
 	const FManagedArrayCollection& InCollection,
 	const TMap<FName, TSet<FName>>* SkipList)
 {
+	MatchOptionalDefaultAttributes(InCollection);
+
 	for (const auto& Pair : InCollection.GroupInfo)
 	{
 		SyncGroupSizeFrom(InCollection, Pair.Key);
diff --git a/Engine/Source/Runtime/Experimental/Chaos/Public/GeometryCollection/Facades/CollectionMeshFacade.h b/Engine/Source/Runtime/Experimental/Chaos/Public/GeometryCollection/Facades/CollectionMeshFacade.h
--- a/Engine/Source/Runtime/Experimental/Chaos/Public/GeometryCollection/Facades/CollectionMeshFacade.h
+++ b/Engine/Source/Runtime/Experimental/Chaos/Public/GeometryCollection/Facades/CollectionMeshFacade.h
@@ -5,12 +5,12 @@
 #include "GeometryCollection/ManagedArrayAccessor.h"
 #include "GeometryCollection/ManagedArrayCollection.h"
 #include "GeometryCollection/GeometryCollection.h"
-
+#include "GeometryCollection/Facades/CollectionUVFacade.h"
 
 namespace GeometryCollection::Facades
 {
 
-	class CHAOS_API FCollectionMeshFacade
+	class CHAOS_API FCollectionMeshFacade : public FCollectionUVFacade
 	{
 	public:
 		FCollectionMeshFacade(FManagedArrayCollection& InCollection);
@@ -62,7 +62,6 @@ namespace GeometryCollection::Facades
 		TManagedArrayAccessor<FVector3f> TangentUAttribute;
 		TManagedArrayAccessor<FVector3f> TangentVAttribute;
 		TManagedArrayAccessor<FVector3f> NormalAttribute;
-		TManagedArrayAccessor<TArray<FVector2f>> UVsAttribute;
 		TManagedArrayAccessor<FLinearColor> ColorAttribute;
 		TManagedArrayAccessor<int32> BoneMapAttribute;
 		TManagedArrayAccessor<int32> VertexStartAttribute;
diff --git a/Engine/Source/Runtime/Experimental/Chaos/Public/GeometryCollection/Facades/CollectionUVFacade.h b/Engine/Source/Runtime/Experimental/Chaos/Public/GeometryCollection/Facades/CollectionUVFacade.h
new file mode 100644
--- /dev/null
+++ b/Engine/Source/Runtime/Experimental/Chaos/Public/GeometryCollection/Facades/CollectionUVFacade.h
@@ -0,0 +1,210 @@
+// Copyright Epic Games, Inc. All Rights Reserved.
+
+#pragma once
+
+#include "GeometryCollection/ManagedArrayAccessor.h"
+#include "GeometryCollection/ManagedArrayCollection.h"
+
+// Note: Prefer namespace GeometryCollection::UV for new code
+namespace GeometryCollectionUV
+{
+	enum
+	{
+		MAX_NUM_UV_CHANNELS = 8,
+	};
+}
+
+namespace GeometryCollection::UV
+{
+	// Names of each UV layer attribute. In a valid collection with UVs, the first layer must be present, and the other layers are optional
+	extern CHAOS_API const FName UVLayerNames[GeometryCollectionUV::MAX_NUM_UV_CHANNELS];
+	// The UV layers are always contained in the vertices group in a geometry collection
+	extern CHAOS_API const FName VerticesGroupName;
+
+	bool CHAOS_API HasValidUVs(const FManagedArrayCollection& Collection);
+	bool CHAOS_API SetNumUVLayers(FManagedArrayCollection& Collection, int32 NumUVs);
+	int32 CHAOS_API GetNumUVLayers(const FManagedArrayCollection& Collection);
+
+	void CHAOS_API DefineUVSchema(FManagedArrayCollection& Collection);
+
+	inline const TManagedArray<FVector2f>* FindUVLayer(const FManagedArrayCollection& Collection, int32 UVLayer)
+	{
+		return Collection.FindAttributeTyped<FVector2f>(UVLayerNames[UVLayer], VerticesGroupName);
+	}
+	inline TManagedArray<FVector2f>* FindUVLayer(FManagedArrayCollection& Collection, int32 UVLayer)
+	{
+		return Collection.FindAttributeTyped<FVector2f>(UVLayerNames[UVLayer], VerticesGroupName);
+	}
+	inline const TManagedArray<FVector2f>& GetUVLayer(const FManagedArrayCollection& Collection, int32 UVLayer)
+	{
+		checkSlow(UVLayer >= 0 && UVLayer < GeometryCollectionUV::MAX_NUM_UV_CHANNELS);
+		return Collection.GetAttribute<FVector2f>(UVLayerNames[UVLayer], VerticesGroupName);
+	}
+	inline TManagedArray<FVector2f>& ModifyUVLayer(FManagedArrayCollection& Collection, int32 UVLayer)
+	{
+		checkSlow(UVLayer >= 0 && UVLayer < GeometryCollectionUV::MAX_NUM_UV_CHANNELS);
+		return Collection.ModifyAttribute<FVector2f>(UVLayerNames[UVLayer], VerticesGroupName);
+	}
+
+	template<typename ManagedArrayType, int32 MaxArraySize>
+	struct TArrayOfAttributesAccessor
+	{
+		TArray<ManagedArrayType*, TFixedAllocator<MaxArraySize>> Attributes;
+
+		ManagedArrayType& operator[](int32 Idx) const
+		{
+			return *Attributes[Idx];
+		}
+
+		int32 Num() const
+		{
+			return Attributes.Num();
+		}
+	};
+
+	using FUVLayers = TArrayOfAttributesAccessor<TManagedArray<FVector2f>, 8>;
+	using FConstUVLayers = TArrayOfAttributesAccessor<const TManagedArray<FVector2f>, 8>;
+
+	/**
+	 * Find the currently-valid UV layers for a given collection
+	 * @return A temporary accessor for all the active UV layers. Must be recreated if the collection attributes are changed / the number of UV layers is changed.
+	 */
+	inline FUVLayers FindActiveUVLayers(FManagedArrayCollection& Collection)
+	{
+		FUVLayers Layers;
+		for (int32 LayerIdx = 0; LayerIdx < GeometryCollectionUV::MAX_NUM_UV_CHANNELS; ++LayerIdx)
+		{
+			TManagedArray<FVector2f>* Layer = Collection.FindAttributeTyped<FVector2f>(UVLayerNames[LayerIdx], VerticesGroupName);
+			if (!Layer)
+			{
+				break;
+			}
+			Layers.Attributes.Add(Layer);
+		}
+		return Layers;
+	}
+	/**
+	 * Find the currently-valid UV layers for a given collection
+	 * @return A temporary accessor for all the active UV layers. Must be recreated if the collection attributes are changed / the number of UV layers is changed.
+	 */
+	inline FConstUVLayers FindActiveUVLayers(const FManagedArrayCollection& Collection)
+	{
+		FConstUVLayers Layers;
+		for (int32 LayerIdx = 0; LayerIdx < GeometryCollectionUV::MAX_NUM_UV_CHANNELS; ++LayerIdx)
+		{
+			const TManagedArray<FVector2f>* Layer = Collection.FindAttributeTyped<FVector2f>(UVLayerNames[LayerIdx], VerticesGroupName);
+			if (!Layer)
+			{
+				break;
+			}
+			Layers.Attributes.Add(Layer);
+		}
+		return Layers;
+	}
+
+	// Set all UVs in the array for the given vertex, stopping early if there are not enough matching UV layers in the collection
+	inline void SetUVs(FManagedArrayCollection& Collection, int32 VertexIdx, TArrayView<const FVector2f> UVs)
+	{
+		for (int32 LayerIdx = 0; LayerIdx < UVs.Num(); ++LayerIdx)
+		{
+			TManagedArray<FVector2f>* Layer = FindUVLayer(Collection, LayerIdx);
+			if (!ensure(Layer != nullptr))
+			{
+				break;
+			}
+			(*Layer)[VertexIdx] = UVs[LayerIdx];
+		}
+	}
+
+	inline void MatchUVLayerCount(FManagedArrayCollection& ToCollection, const FManagedArrayCollection& FromCollection)
+	{
+		SetNumUVLayers(ToCollection, GetNumUVLayers(FromCollection));
+	}
+}
+
+namespace GeometryCollection::Facades
+{
+	class FCollectionUVFacade
+	{
+	public:
+		FCollectionUVFacade(FManagedArrayCollection& InCollection) : Collection(&InCollection), ConstCollection(&InCollection)
+		{}
+		FCollectionUVFacade(const FManagedArrayCollection& InCollection) : Collection(nullptr), ConstCollection(&InCollection)
+		{}
+
+		/**
+		 * returns true if all the necessary attributes are present
+		 * if not then the API can be used to create
+		 */
+		inline bool IsValid() const
+		{
+			return GeometryCollection::UV::HasValidUVs(*ConstCollection);
+		}
+
+		/**
+		 * Add the necessary attributes if they are missing
+		 */
+		void DefineSchema()
+		{
+			GeometryCollection::UV::DefineUVSchema(*Collection);
+		}
+
+		inline bool SetNumUVLayers(int32 UVLayers)
+		{
+			return GeometryCollection::UV::SetNumUVLayers(*Collection, UVLayers);
+		}
+
+		inline int32 GetNumUVLayers() const
+		{
+			return GeometryCollection::UV::GetNumUVLayers(*ConstCollection);
+		}
+
+		inline const TManagedArray<FVector2f>& GetUVLayer(int32 UVLayer) const
+		{
+			return GeometryCollection::UV::GetUVLayer(*ConstCollection, UVLayer);
+		}
+		inline TManagedArray<FVector2f>& ModifyUVLayer(int32 UVLayer)
+		{
+			return GeometryCollection::UV::ModifyUVLayer(*Collection, UVLayer);
+		}
+		inline const TManagedArray<FVector2f>* FindUVLayer(int32 UVLayer) const
+		{
+			return GeometryCollection::UV::FindUVLayer(*ConstCollection, UVLayer);
+		}
+		inline TManagedArray<FVector2f>* FindUVLayer(int32 UVLayer)
+		{
+			return GeometryCollection::UV::FindUVLayer(*Collection, UVLayer);
+		}
+		inline FVector2f& ModifyUV(int32 VertexIdx, int32 LayerIdx)
+		{
+			return ModifyUVLayer(LayerIdx)[VertexIdx];
+		}
+		inline const FVector2f& GetUV(int32 VertexIdx, int32 LayerIdx) const
+		{
+			return GetUVLayer(LayerIdx)[VertexIdx];
+		}
+		inline void SetUV(int32 VertexIdx, int32 LayerIdx, const FVector2f& UV)
+		{
+			ModifyUVLayer(LayerIdx)[VertexIdx] = UV;
+		}
+		inline GeometryCollection::UV::FUVLayers FindActiveUVLayers()
+		{
+			return GeometryCollection::UV::FindActiveUVLayers(*Collection);
+		}
+		inline GeometryCollection::UV::FConstUVLayers FindActiveUVLayers() const
+		{
+			return GeometryCollection::UV::FindActiveUVLayers(*ConstCollection);
+		}
+
+		// Set all UVs in the array for the given vertex, stopping early if there are not enough matching UV layers in the collection
+		inline void SetUVs(int32 VertexIdx, TArrayView<const FVector2f> UVs)
+		{
+			GeometryCollection::UV::SetUVs(*Collection, VertexIdx, UVs);
+		}
+
+	private:
+		FManagedArrayCollection* Collection;
+		const FManagedArrayCollection* ConstCollection;
+	};
+
+}
\ No newline at end of file
diff --git a/Engine/Source/Runtime/Experimental/Chaos/Public/GeometryCollection/GeometryCollection.h b/Engine/Source/Runtime/Experimental/Chaos/Public/GeometryCollection/GeometryCollection.h
--- a/Engine/Source/Runtime/Experimental/Chaos/Public/GeometryCollection/GeometryCollection.h
+++ b/Engine/Source/Runtime/Experimental/Chaos/Public/GeometryCollection/GeometryCollection.h
@@ -9,6 +9,7 @@
 #include "GeometryCollection/GeometryCollectionConvexPropertiesInterface.h"
 #include "GeometryCollection/GeometryCollectionProximityPropertiesInterface.h"
 #include "GeometryCollection/ManagedArrayAccessor.h"
+#include "GeometryCollection/Facades/CollectionUVFacade.h"
 
 namespace Chaos
 {
@@ -17,13 +18,6 @@ namespace Chaos
 class FGeometryCollectionConvexPropertiesInterface;
 class FGeometryCollectionProximityPropertiesInterface;
 
-namespace GeometryCollectionUV
-{
-	enum
-	{
-		MAX_NUM_UV_CHANNELS = 8,
-	};
-}
 
 /**
 * FGeometryCollection (FTransformCollection)
@@ -284,8 +278,28 @@ public:
 
 	// Vertices Group
 	TManagedArray<FVector3f>		 Vertex;
-	// Outer array is the array of vertices, inner array is the uv channels
-	TManagedArray<TArray<FVector2f>> UVs;
+
+	// Note: UVs have been reworked, and unfortunately there is not a safe path to provide the original UVs managed array as a deprecated accessor.
+	// They are now stored in dynamically allocated attributes per UV channel (/ layer)
+	// See Facades/CollectionUVFacade.h for a more complete interface to access UV layers,
+	// but accesses of the form Collection.UVs[Vertex][Layer] can be replaced with Collection.GetUV(Vertex, Layer) (or ModifyUV)
+	FVector2f& ModifyUV(int32 VertexIndex, int32 UVLayer)
+	{
+		return GeometryCollection::UV::ModifyUVLayer(*this, UVLayer)[VertexIndex];
+	}
+	const FVector2f& GetUV(int32 VertexIndex, int32 UVLayer) const
+	{
+		return GeometryCollection::UV::GetUVLayer(*this, UVLayer)[VertexIndex];
+	}
+	inline TManagedArray<FVector2f>* FindUVLayer(int32 UVLayer)
+	{
+		return GeometryCollection::UV::FindUVLayer(*this, UVLayer);
+	}
+	inline const TManagedArray<FVector2f>* FindUVLayer(int32 UVLayer) const
+	{
+		return GeometryCollection::UV::FindUVLayer(*this, UVLayer);
+	}
+
 	TManagedArray<FLinearColor>      Color;
 	TManagedArray<FVector3f>         TangentU;
 	TManagedArray<FVector3f>         TangentV;
@@ -313,6 +327,15 @@ public:
 	
 protected:
 
+	/**
+	 * Virtual helper function called by CopyMatchingAttributesFrom; adds attributes 'default, but optional' attributes that are present in InCollection
+	 * This is used by FGeometryCollection to make sure all UV layers are copied over by CopyMatchingAttributesFrom()
+	 */
+	virtual void MatchOptionalDefaultAttributes(const FManagedArrayCollection& InCollection) override
+	{
+		GeometryCollection::UV::MatchUVLayerCount(*this, InCollection);
+	}
+
 	void Construct();
 
 	/**
diff --git a/Engine/Source/Runtime/Experimental/Chaos/Public/GeometryCollection/GeometryCollectionUtility.h b/Engine/Source/Runtime/Experimental/Chaos/Public/GeometryCollection/GeometryCollectionUtility.h
--- a/Engine/Source/Runtime/Experimental/Chaos/Public/GeometryCollection/GeometryCollectionUtility.h
+++ b/Engine/Source/Runtime/Experimental/Chaos/Public/GeometryCollection/GeometryCollectionUtility.h
@@ -39,7 +39,7 @@ namespace GeometryCollection
 		TManagedArray<FVector3f>& Normals = RestCollection->Normal;
 		TManagedArray<FVector3f>& TangentU = RestCollection->TangentU;
 		TManagedArray<FVector3f>& TangentV = RestCollection->TangentV;
-		TManagedArray<TArray<FVector2f>>& UVs = RestCollection->UVs;
+		TManagedArray<FVector2f>& UV0 = *RestCollection->FindUVLayer(0);
 		TManagedArray<FLinearColor>& Colors = RestCollection->Color;
 		TManagedArray<FIntVector>& Indices = RestCollection->Indices;
 		TManagedArray<bool>& Visible = RestCollection->Visible;
@@ -60,7 +60,7 @@ namespace GeometryCollection
 			Normals[Idx] = NormalsIn.Num() > Idx ? FVector3f(NormalsIn[Idx][0], NormalsIn[Idx][1], NormalsIn[Idx][2]) : FVector3f(0);
 
 			FVector2D UV = UVsIn.Num() > Idx ? FVector2D(UVsIn[Idx][0], UVsIn[Idx][1]) : FVector2D(0);
-			UVs[Idx].Add(FVector2f(UV));
+			UV0[Idx] = FVector2f(UV);
 			
 			Colors[Idx] = FLinearColor::White;
 		}
diff --git a/Engine/Source/Runtime/Experimental/Chaos/Public/GeometryCollection/ManagedArrayCollection.h b/Engine/Source/Runtime/Experimental/Chaos/Public/GeometryCollection/ManagedArrayCollection.h
--- a/Engine/Source/Runtime/Experimental/Chaos/Public/GeometryCollection/ManagedArrayCollection.h
+++ b/Engine/Source/Runtime/Experimental/Chaos/Public/GeometryCollection/ManagedArrayCollection.h
@@ -656,6 +656,13 @@ private:
 
 protected:
 
+	/**
+	 * Virtual helper function called by CopyMatchingAttributesFrom; adds attributes 'default, but optional' attributes that are present in InCollection
+	 * This is used by FGeometryCollection to make sure all UV layers are copied over by CopyMatchingAttributesFrom()
+	 */
+	virtual void MatchOptionalDefaultAttributes(const FManagedArrayCollection& InCollection)
+	{}
+
 	/**
 	* Size and order a group so that it matches the group found in the input collection.
 	* @param InCollection - The collection we are ordering our group against. 
diff --git a/Engine/Source/Runtime/Experimental/GeometryCollectionEngine/Private/GeometryCollection/GeometryCollectionComponent.cpp b/Engine/Source/Runtime/Experimental/GeometryCollectionEngine/Private/GeometryCollection/GeometryCollectionComponent.cpp
--- a/Engine/Source/Runtime/Experimental/GeometryCollectionEngine/Private/GeometryCollection/GeometryCollectionComponent.cpp
+++ b/Engine/Source/Runtime/Experimental/GeometryCollectionEngine/Private/GeometryCollection/GeometryCollectionComponent.cpp
@@ -1852,7 +1852,6 @@ void UGeometryCollectionComponent::InitConstantData(FGeometryCollectionConstantD
 		const TManagedArray<FVector3f>& TangentU = Collection->TangentU;
 		const TManagedArray<FVector3f>& TangentV = Collection->TangentV;
 		const TManagedArray<FVector3f>& Normal = Collection->Normal;
-		const TManagedArray<TArray<FVector2f>>& UVs = Collection->UVs;
 		const TManagedArray<FLinearColor>& Color = Collection->Color;
 		const TManagedArray<FLinearColor>& BoneColors = Collection->BoneColor;
 		
@@ -1868,13 +1867,8 @@ void UGeometryCollectionComponent::InitConstantData(FGeometryCollectionConstantD
 		ConstantData->Normals = TArray<FVector3f>(Normal.GetData(), Normal.Num());
 		ConstantData->Colors = TArray<FLinearColor>(Color.GetData(), Color.Num());
 
-		// find the maximum number of channels 
-		int32 NumUVChannels = 0;
-		for (const TArray<FVector2f>& VertexUVs : Collection->UVs.GetConstArray())
-		{
-			NumUVChannels = FMath::Max(NumUVChannels, VertexUVs.Num());
-		}
-
+		int32 NumUVChannels = Collection->NumUVLayers();
+		GeometryCollection::UV::FConstUVLayers CollectionUVs = GeometryCollection::UV::FindActiveUVLayers(*Collection);
 		ConstantData->UVChannels.SetNum(NumUVChannels);
 		for (int32 UVChannelIndex = 0; UVChannelIndex < NumUVChannels; UVChannelIndex++)
 		{
@@ -1891,14 +1885,9 @@ void UGeometryCollectionComponent::InitConstantData(FGeometryCollectionConstantD
 					const int32 BoneIndex = ConstantData->BoneMap[InPointIndex];
 					ConstantData->BoneColors[InPointIndex] = BoneColors[BoneIndex];
 
-					const TArray<FVector2f>& VertexUVs = Collection->UVs[InPointIndex];
-					for (int32 UVChannelIndex = 0; UVChannelIndex < VertexUVs.Num(); UVChannelIndex++)
-					{
-						ConstantData->UVChannels[UVChannelIndex][InPointIndex] = Collection->UVs[InPointIndex][UVChannelIndex];
-					}
-					for (int32 UVChannelIndex = VertexUVs.Num(); UVChannelIndex < NumUVChannels; UVChannelIndex++)
+					for (int32 UVChannelIndex = 0; UVChannelIndex < NumUVChannels; UVChannelIndex++)
 					{
-						ConstantData->UVChannels[UVChannelIndex][InPointIndex] = FVector2f::ZeroVector;
+						ConstantData->UVChannels[UVChannelIndex][InPointIndex] = CollectionUVs[UVChannelIndex][InPointIndex];
 					}
 				});
 		}
@@ -1914,7 +1903,7 @@ void UGeometryCollectionComponent::InitConstantData(FGeometryCollectionConstantD
 			{
 				for (int32 PointIndex = 0; PointIndex < NumPoints; PointIndex++)
 				{
-					ConstantData->UVChannels[UVChannelIndex][PointIndex] = Collection->UVs[PointIndex][UVChannelIndex];
+					ConstantData->UVChannels[UVChannelIndex][PointIndex] = CollectionUVs[UVChannelIndex][PointIndex];
 				}
 			}
 		}
diff --git a/Engine/Source/Runtime/Experimental/GeometryCollectionEngine/Private/GeometryCollection/GeometryCollectionEngineConversion.cpp b/Engine/Source/Runtime/Experimental/GeometryCollectionEngine/Private/GeometryCollection/GeometryCollectionEngineConversion.cpp
--- a/Engine/Source/Runtime/Experimental/GeometryCollectionEngine/Private/GeometryCollection/GeometryCollectionEngineConversion.cpp
+++ b/Engine/Source/Runtime/Experimental/GeometryCollectionEngine/Private/GeometryCollection/GeometryCollectionEngineConversion.cpp
@@ -126,12 +126,16 @@ void FGeometryCollectionEngineConversion::AppendMeshDescription(
 	TManagedArray<FVector3f>& TargetTangentU = GeometryCollection->TangentU;
 	TManagedArray<FVector3f>& TargetTangentV = GeometryCollection->TangentV;
 	TManagedArray<FVector3f>& TargetNormal = GeometryCollection->Normal;
-	TManagedArray<TArray<FVector2f>>& TargetUVs = GeometryCollection->UVs;
 	TManagedArray<FLinearColor>& TargetColor = GeometryCollection->Color;
 	TManagedArray<int32>& TargetBoneMap = GeometryCollection->BoneMap;
 	TManagedArray<FLinearColor>& TargetBoneColor = GeometryCollection->BoneColor;
 	TManagedArray<FString>& TargetBoneName = GeometryCollection->BoneName;
 
+	if (GeometryCollection->NumUVLayers() < NumUVLayers)
+	{
+		GeometryCollection->SetNumUVLayers(NumUVLayers);
+	}
+
 	const int32 VertexStart = GeometryCollection->NumElements(FGeometryCollection::VerticesGroup);
 	int32 VertexCount = 0;
 		
@@ -177,7 +181,7 @@ void FGeometryCollectionEngineConversion::AppendMeshDescription(
 			TargetTangentU[CurrentVertex] = SourceTangent[ExemplarInstanceID];
 			TargetTangentV[CurrentVertex] = (FVector3f)SourceBinormalSign[ExemplarInstanceID] * FVector3f::CrossProduct(TargetNormal[CurrentVertex], TargetTangentU[CurrentVertex]);
 
-			TargetUVs[CurrentVertex] = SplitVertex.Key.UVs;
+			GeometryCollection::UV::SetUVs(*GeometryCollection, CurrentVertex, SplitVertex.Key.UVs);
 
 			if (SourceColor.Num() > 0)
 			{
@@ -693,7 +697,6 @@ bool FGeometryCollectionEngineConversion::AppendGeometryCollection(const FGeomet
 	const TManagedArray<FVector3f>& SourceTangentU = SourceGeometryCollectionPtr->TangentU;
 	const TManagedArray<FVector3f>& SourceTangentV = SourceGeometryCollectionPtr->TangentV;
 	const TManagedArray<FVector3f>& SourceNormal = SourceGeometryCollectionPtr->Normal;
-	const TManagedArray<TArray<FVector2f>>& SourceUVs = SourceGeometryCollectionPtr->UVs;
 	const TManagedArray<FLinearColor>& SourceColor = SourceGeometryCollectionPtr->Color;
 	const TManagedArray<int32>& SourceBoneMap = SourceGeometryCollectionPtr->BoneMap;
 
@@ -702,10 +705,13 @@ bool FGeometryCollectionEngineConversion::AppendGeometryCollection(const FGeomet
 	TManagedArray<FVector3f>& TargetTangentU = TargetGeometryCollection->TangentU;
 	TManagedArray<FVector3f>& TargetTangentV = TargetGeometryCollection->TangentV;
 	TManagedArray<FVector3f>& TargetNormal = TargetGeometryCollection->Normal;
-	TManagedArray<TArray<FVector2f>>& TargetUVs = TargetGeometryCollection->UVs;
 	TManagedArray<FLinearColor>& TargetColor = TargetGeometryCollection->Color;
 	TManagedArray<int32>& TargetBoneMap = TargetGeometryCollection->BoneMap;
 
+	TargetGeometryCollection->SetNumUVLayers(FMath::Max(TargetGeometryCollection->NumUVLayers(), SourceGeometryCollectionPtr->NumUVLayers()));
+	GeometryCollection::UV::FUVLayers TargetUVLayers = GeometryCollection::UV::FindActiveUVLayers(*TargetGeometryCollection);
+	GeometryCollection::UV::FConstUVLayers SourceUVLayers = GeometryCollection::UV::FindActiveUVLayers(*SourceGeometryCollectionPtr);
+
 	// append vertices
 	for (int32 VertexIndex = 0; VertexIndex < VertexCount; VertexIndex++)
 	{
@@ -715,7 +721,11 @@ bool FGeometryCollectionEngineConversion::AppendGeometryCollection(const FGeomet
 		TargetTangentU[VertexOffset] = SourceTangentU[VertexIndex];
 		TargetTangentV[VertexOffset] = SourceTangentV[VertexIndex];
 		TargetNormal[VertexOffset] = SourceNormal[VertexIndex];
-		TargetUVs[VertexOffset] = SourceUVs[VertexIndex];
+
+		for (int32 UVLayer = 0; UVLayer < SourceUVLayers.Num(); ++UVLayer)
+		{
+			TargetUVLayers[UVLayer][VertexOffset] = SourceUVLayers[UVLayer][VertexIndex];
+		}
 		TargetColor[VertexOffset] = SourceColor[VertexIndex];
 
 		TargetBoneMap[VertexOffset] = SourceBoneMap[VertexIndex] + TransformStart;
@@ -1115,7 +1125,6 @@ bool FGeometryCollectionEngineConversion::AppendSkeletalMesh(const USkeletalMesh
 	//
 	TManagedArray<FVector3f>& Vertex = InCollection->ModifyAttribute<FVector3f>("Vertex", FGeometryCollection::VerticesGroup);
 	TManagedArray<FVector3f>& Normal = InCollection->ModifyAttribute<FVector3f>("Normal", FGeometryCollection::VerticesGroup);
-	TManagedArray<TArray<FVector2f>>& UVs = InCollection->ModifyAttribute<TArray<FVector2f>>("UVs", FGeometryCollection::VerticesGroup);
 	TManagedArray<FLinearColor>& Color = InCollection->ModifyAttribute<FLinearColor>("Color", FGeometryCollection::VerticesGroup);
 	TManagedArray<FVector3f>& TangentU = InCollection->ModifyAttribute<FVector3f>("TangentU", FGeometryCollection::VerticesGroup);
 	TManagedArray<FVector3f>& TangentV = InCollection->ModifyAttribute<FVector3f>("TangentV", FGeometryCollection::VerticesGroup);
@@ -1131,6 +1140,8 @@ bool FGeometryCollectionEngineConversion::AppendSkeletalMesh(const USkeletalMesh
 	int InitialNumVertices = InCollection->NumElements(FGeometryCollection::VerticesGroup);
 	int VertexBaseIndex = InCollection->AddElements(VertexCount, FGeometryCollection::VerticesGroup);
 	const int32 NumUVLayers = VertexBuffers.StaticMeshVertexBuffer.GetNumTexCoords();
+	GeometryCollection::UV::SetNumUVLayers(*InCollection, NumUVLayers);
+	GeometryCollection::UV::FUVLayers UVLayers = GeometryCollection::UV::FindActiveUVLayers(*InCollection);
 	for (int32 VertexIndex = 0; VertexIndex < VertexCount; VertexIndex++)
 	{
 		int VertexOffset = VertexBaseIndex + VertexIndex;
@@ -1147,10 +1158,9 @@ bool FGeometryCollectionEngineConversion::AppendSkeletalMesh(const USkeletalMesh
 		TangentV[VertexOffset] = VertexBuffers.StaticMeshVertexBuffer.VertexTangentY(VertexIndex);
 		Normal[VertexOffset] = VertexBuffers.StaticMeshVertexBuffer.VertexTangentZ(VertexIndex);
 
-		UVs[VertexOffset].SetNum(NumUVLayers);
 		for (int32 UVLayerIdx = 0; UVLayerIdx < NumUVLayers; ++UVLayerIdx)
 		{
-			UVs[VertexOffset][UVLayerIdx] = VertexBuffers.StaticMeshVertexBuffer.GetVertexUV(VertexIndex, UVLayerIdx);
+			UVLayers[UVLayerIdx][VertexOffset] = VertexBuffers.StaticMeshVertexBuffer.GetVertexUV(VertexIndex, UVLayerIdx);
 		}
 
 		if (VertexBuffers.ColorVertexBuffer.GetNumVertices() == VertexCount)
diff --git a/Engine/Source/Runtime/Experimental/GeometryCollectionEngine/Private/GeometryCollection/GeometryCollectionEngineUtility.cpp b/Engine/Source/Runtime/Experimental/GeometryCollectionEngine/Private/GeometryCollection/GeometryCollectionEngineUtility.cpp
--- a/Engine/Source/Runtime/Experimental/GeometryCollectionEngine/Private/GeometryCollection/GeometryCollectionEngineUtility.cpp
+++ b/Engine/Source/Runtime/Experimental/GeometryCollectionEngine/Private/GeometryCollection/GeometryCollectionEngineUtility.cpp
@@ -328,7 +328,7 @@ void GeometryCollectionEngineUtility::ComputeNormals(FGeometryCollection* Geomet
 
 	TManagedArray<FIntVector>& Indices = GeometryCollection->Indices;
 	const int32 UVChannelIndex = 0;
-	TManagedArray<TArray<FVector2f>>& UVs = GeometryCollection->UVs;
+	const TManagedArray<FVector2f>& UV0 = *GeometryCollection->FindUVLayer(UVChannelIndex);
 	for (int i = 0; i < NumIndices; ++i)
 	{
 		FIntVector CurrTri = Indices[i];
@@ -337,9 +337,9 @@ void GeometryCollectionEngineUtility::ComputeNormals(FGeometryCollection* Geomet
 		TmpIndices.Add(CurrTri.Y);
 		TmpIndices.Add(CurrTri.Z);
 
-		TmpUV.Add(UVs[CurrTri.X][UVChannelIndex]);
-		TmpUV.Add(UVs[CurrTri.Y][UVChannelIndex]);
-		TmpUV.Add(UVs[CurrTri.Z][UVChannelIndex]);
+		TmpUV.Add(UV0[CurrTri.X]);
+		TmpUV.Add(UV0[CurrTri.Y]);
+		TmpUV.Add(UV0[CurrTri.Z]);
 	}
 
 	// Make a copy of the vertex array
@@ -386,7 +386,7 @@ void GeometryCollectionEngineUtility::ComputeTangents(FGeometryCollection* Geome
 
 	TManagedArray<FIntVector>& Indices = GeometryCollection->Indices;
 	const int32 UVChannelIndex = 0;
-	TManagedArray<TArray<FVector2f>>& UVs = GeometryCollection->UVs;
+	TManagedArray<FVector2f>& UV0 = *GeometryCollection->FindUVLayer(UVChannelIndex);
 	for (int i = 0; i < NumIndices; ++i)
 	{
 		FIntVector CurrTri = Indices[i];
@@ -395,9 +395,9 @@ void GeometryCollectionEngineUtility::ComputeTangents(FGeometryCollection* Geome
 		TmpIndices.Add(CurrTri.Y);
 		TmpIndices.Add(CurrTri.Z);
 
-		TmpUV.Add(UVs[CurrTri.X][UVChannelIndex]);
-		TmpUV.Add(UVs[CurrTri.Y][UVChannelIndex]);
-		TmpUV.Add(UVs[CurrTri.Z][UVChannelIndex]);
+		TmpUV.Add(UV0[CurrTri.X]);
+		TmpUV.Add(UV0[CurrTri.Y]);
+		TmpUV.Add(UV0[CurrTri.Z]);
 	}
 
 	// Make a copy of the vertex array
diff --git a/Engine/Source/Runtime/Experimental/GeometryCollectionEngine/Private/GeometryCollection/GeometryCollectionObject.cpp b/Engine/Source/Runtime/Experimental/GeometryCollectionEngine/Private/GeometryCollection/GeometryCollectionObject.cpp
--- a/Engine/Source/Runtime/Experimental/GeometryCollectionEngine/Private/GeometryCollection/GeometryCollectionObject.cpp
+++ b/Engine/Source/Runtime/Experimental/GeometryCollectionEngine/Private/GeometryCollection/GeometryCollectionObject.cpp
@@ -1004,7 +1004,7 @@ TUniquePtr<FGeometryCollectionNaniteData> UGeometryCollection::CreateNaniteData(
 
 	// Vertices Group
 	const TManagedArray<FVector3f>& VertexArray = Collection->Vertex;
-	const TManagedArray<TArray<FVector2f>>& UVsArray = Collection->UVs;
+	GeometryCollection::UV::FUVLayers UVsLayers = GeometryCollection::UV::FindActiveUVLayers(*Collection);
 	const TManagedArray<FLinearColor>& ColorArray = Collection->Color;
 	const TManagedArray<FVector3f>& TangentUArray = Collection->TangentU;
 	const TManagedArray<FVector3f>& TangentVArray = Collection->TangentV;
@@ -1055,9 +1055,9 @@ TUniquePtr<FGeometryCollectionNaniteData> UGeometryCollection::CreateNaniteData(
 			Vertex.TangentX = FVector3f::ZeroVector;
 			Vertex.TangentY = FVector3f::ZeroVector;
 			Vertex.TangentZ = NormalArray[VertexStart + VertexIndex];
-			for (int32 UVIdx = 0; UVIdx < UVsArray[VertexStart + VertexIndex].Num(); ++UVIdx)
+			for (int32 UVIdx = 0; UVIdx < UVsLayers.Num(); ++UVIdx)
 			{
-				Vertex.UVs[UVIdx] = UVsArray[VertexStart + VertexIndex][UVIdx];
+				Vertex.UVs[UVIdx] = UVsLayers[UVIdx][VertexStart + VertexIndex];
 				if (Vertex.UVs[UVIdx].ContainsNaN())
 				{
 					Vertex.UVs[UVIdx] = FVector2f::ZeroVector;
diff --git a/Engine/Source/Runtime/Experimental/GeometryCollectionEngine/Public/GeometryCollection/GeometryCollectionComponent.h b/Engine/Source/Runtime/Experimental/GeometryCollectionEngine/Public/GeometryCollection/GeometryCollectionComponent.h
--- a/Engine/Source/Runtime/Experimental/GeometryCollectionEngine/Public/GeometryCollection/GeometryCollectionComponent.h
+++ b/Engine/Source/Runtime/Experimental/GeometryCollectionEngine/Public/GeometryCollection/GeometryCollectionComponent.h
@@ -660,7 +660,6 @@ public:
 #define COPY_ON_WRITE_ATTRIBUTES \
 	/* Vertices Group */ \
 	COPY_ON_WRITE_ATTRIBUTE(FVector3f, Vertex, FGeometryCollection::VerticesGroup)		/* GetVertexArray, GetVertexArrayCopyOnWrite, GetVertexArrayRest */ \
-	COPY_ON_WRITE_ATTRIBUTE(TArray<FVector2f>, UVs, FGeometryCollection::VerticesGroup) /* GetUVsArray */		\
 	COPY_ON_WRITE_ATTRIBUTE(FLinearColor, Color, FGeometryCollection::VerticesGroup)	/* GetColorArray */		\
 	COPY_ON_WRITE_ATTRIBUTE(FVector3f, TangentU, FGeometryCollection::VerticesGroup)	/* GetTangentUArray */	\
 	COPY_ON_WRITE_ATTRIBUTE(FVector3f, TangentV, FGeometryCollection::VerticesGroup)	/* //... */		\
