Dataflow : Disable serialization of cooked dataflow assets. (editor only asset atm)

diff --git a/Engine/Source/Runtime/Experimental/Dataflow/Engine/Public/Dataflow/DataflowObject.h b/Engine/Source/Runtime/Experimental/Dataflow/Engine/Public/Dataflow/DataflowObject.h
--- a/Engine/Source/Runtime/Experimental/Dataflow/Engine/Public/Dataflow/DataflowObject.h
+++ b/Engine/Source/Runtime/Experimental/Dataflow/Engine/Public/Dataflow/DataflowObject.h
@@ -73,6 +73,9 @@ public:
 	UE_DEPRECATED(5.1, "Use Blueprint library version of the function")
 	void EvaluateTerminalNodeByName(FName NodeName, UObject* Asset);
 
+	virtual bool IsEditorOnly() const { return true; }
+
+
 public:
 	UPROPERTY(EditAnywhere, Category = "Evaluation")
 	bool bActive = true;
