fix for new UV channels ending up in group None instead of group "Vertices"

diff --git a/Engine/Source/Runtime/Experimental/Chaos/Private/GeometryCollection/Facades/CollectionUVFacade.cpp b/Engine/Source/Runtime/Experimental/Chaos/Private/GeometryCollection/Facades/CollectionUVFacade.cpp
--- a/Engine/Source/Runtime/Experimental/Chaos/Private/GeometryCollection/Facades/CollectionUVFacade.cpp
+++ b/Engine/Source/Runtime/Experimental/Chaos/Private/GeometryCollection/Facades/CollectionUVFacade.cpp
@@ -18,7 +18,7 @@ namespace GeometryCollection::UV
 		"UVLayer7"
 	};
 
-	const FName VerticesGroupName = FGeometryCollection::VerticesGroup;
+	const FName VerticesGroupName = "Vertices";
 
 	bool HasValidUVs(const FManagedArrayCollection& Collection)
 	{
