remove clustered particles from the TopLevelClusterParentsStrained array in FRigidClustering upon destroy/disable

diff --git a/Engine/Source/Runtime/Experimental/Chaos/Private/Chaos/PBDRigidClustering.cpp b/Engine/Source/Runtime/Experimental/Chaos/Private/Chaos/PBDRigidClustering.cpp
--- a/Engine/Source/Runtime/Experimental/Chaos/Private/Chaos/PBDRigidClustering.cpp
+++ b/Engine/Source/Runtime/Experimental/Chaos/Private/Chaos/PBDRigidClustering.cpp
@@ -1292,6 +1292,7 @@ namespace Chaos
 		// #note: we don't recursively descend to the children
 		MEvolution.DisableParticle(ClusteredParticle);
 		TopLevelClusterParents.Remove(ClusteredParticle);
+		TopLevelClusterParentsStrained.Remove(ClusteredParticle);
 		GetChildrenMap().Remove(ClusteredParticle);
 		ClusteredParticle->ClusterIds() = ClusterId();
 		ClusteredParticle->ClusterGroupIndex() = 0;
@@ -1328,6 +1329,7 @@ namespace Chaos
 
 	// reset the structures
 		TopLevelClusterParents.Remove(ClusteredParticle);
+		TopLevelClusterParentsStrained.Remove(ClusteredParticle);
 
 		// disconnect from the parents
 		if (ClusteredParticle->ClusterIds().Id)
