Fix CIS non-unity build.

diff --git a/Engine/Plugins/2D/Paper2D/Source/Paper2D/Private/PaperRenderSceneProxy.cpp b/Engine/Plugins/2D/Paper2D/Source/Paper2D/Private/PaperRenderSceneProxy.cpp
--- a/Engine/Plugins/2D/Paper2D/Source/Paper2D/Private/PaperRenderSceneProxy.cpp
+++ b/Engine/Plugins/2D/Paper2D/Source/Paper2D/Private/PaperRenderSceneProxy.cpp
@@ -4,6 +4,7 @@
 #include "PrimitiveViewRelevance.h"
 #include "SceneManagement.h"
 #include "MaterialDomain.h"
+#include "MaterialShared.h"
 #include "Materials/Material.h"
 #include "Materials/MaterialRenderProxy.h"
 #include "PhysicsEngine/BodySetup.h"
diff --git a/Engine/Plugins/2D/Paper2D/Source/Paper2DEditor/Private/SpriteEditor/SpriteEditorViewportClient.cpp b/Engine/Plugins/2D/Paper2D/Source/Paper2DEditor/Private/SpriteEditor/SpriteEditorViewportClient.cpp
--- a/Engine/Plugins/2D/Paper2D/Source/Paper2DEditor/Private/SpriteEditor/SpriteEditorViewportClient.cpp
+++ b/Engine/Plugins/2D/Paper2D/Source/Paper2DEditor/Private/SpriteEditor/SpriteEditorViewportClient.cpp
@@ -3,6 +3,7 @@
 #include "SpriteEditor/SpriteEditorViewportClient.h"
 #include "InputKeyEventArgs.h"
 #include "Materials/MaterialInterface.h"
+#include "MaterialShared.h"
 #include "CanvasItem.h"
 #include "Misc/PackageName.h"
 #include "Utils.h"
diff --git a/Engine/Plugins/2D/Paper2D/Source/Paper2DEditor/Private/TileMapEditing/TileMapEditorViewportClient.cpp b/Engine/Plugins/2D/Paper2D/Source/Paper2DEditor/Private/TileMapEditing/TileMapEditorViewportClient.cpp
--- a/Engine/Plugins/2D/Paper2D/Source/Paper2DEditor/Private/TileMapEditing/TileMapEditorViewportClient.cpp
+++ b/Engine/Plugins/2D/Paper2D/Source/Paper2DEditor/Private/TileMapEditing/TileMapEditorViewportClient.cpp
@@ -2,6 +2,7 @@
 
 #include "TileMapEditing/TileMapEditorViewportClient.h"
 #include "Materials/MaterialInterface.h"
+#include "MaterialShared.h"
 #include "PaperTileMapComponent.h"
 #include "SceneView.h"
 #include "ScopedTransaction.h"
diff --git a/Engine/Source/Developer/MeshUtilities/Private/MeshUtilities.cpp b/Engine/Source/Developer/MeshUtilities/Private/MeshUtilities.cpp
--- a/Engine/Source/Developer/MeshUtilities/Private/MeshUtilities.cpp
+++ b/Engine/Source/Developer/MeshUtilities/Private/MeshUtilities.cpp
@@ -43,6 +43,7 @@
 #include "PhysicsEngine/AggregateGeom.h"
 #include "PhysicsEngine/BodySetup.h"
 #include "MaterialUtilities.h"
+#include "MaterialShared.h"
 #include "IHierarchicalLODUtilities.h"
 #include "HierarchicalLODUtilitiesModule.h"
 #include "MeshBoneReduction.h"
diff --git a/Engine/Source/Editor/UnrealEd/Private/EditorExporters.cpp b/Engine/Source/Editor/UnrealEd/Private/EditorExporters.cpp
--- a/Engine/Source/Editor/UnrealEd/Private/EditorExporters.cpp
+++ b/Engine/Source/Editor/UnrealEd/Private/EditorExporters.cpp
@@ -65,6 +65,7 @@
 #include "StaticMeshAttributes.h"
 #include "StaticMeshOperations.h"
 #include "MaterialUtilities.h"
+#include "MaterialShared.h"
 #include "InstancedFoliageActor.h"
 #include "LandscapeProxy.h"
 #include "Landscape.h"
diff --git a/Engine/Source/Runtime/Experimental/GeometryCollectionEngine/Private/GeometryCollection/GeometryCollectionSceneProxy.cpp b/Engine/Source/Runtime/Experimental/GeometryCollectionEngine/Private/GeometryCollection/GeometryCollectionSceneProxy.cpp
--- a/Engine/Source/Runtime/Experimental/GeometryCollectionEngine/Private/GeometryCollection/GeometryCollectionSceneProxy.cpp
+++ b/Engine/Source/Runtime/Experimental/GeometryCollectionEngine/Private/GeometryCollection/GeometryCollectionSceneProxy.cpp
@@ -7,6 +7,7 @@
 #include "GeometryCollection/GeometryCollection.h"
 #include "GeometryCollection/GeometryCollectionObject.h"
 #include "MaterialDomain.h"
+#include "MaterialShared.h"
 #include "MaterialShaderType.h"
 #include "Materials/Material.h"
 #include "Materials/MaterialRenderProxy.h"
