Chaos Cloth - Add a Reset function to the Property Collection adapter, and fix some issues with some missing attribute initialization and a couple of missing template instantiations.

#rnx

diff --git a/Engine/Source/Runtime/Experimental/Chaos/Private/Chaos/PropertyCollectionAdapter.cpp b/Engine/Source/Runtime/Experimental/Chaos/Private/Chaos/PropertyCollectionAdapter.cpp
--- a/Engine/Source/Runtime/Experimental/Chaos/Private/Chaos/PropertyCollectionAdapter.cpp
+++ b/Engine/Source/Runtime/Experimental/Chaos/Private/Chaos/PropertyCollectionAdapter.cpp
@@ -15,7 +15,8 @@ namespace Chaos::Softs
 	FPropertyCollectionConstAdapter::FPropertyCollectionConstAdapter(const TSharedPtr<const FManagedArrayCollection>& InManagedArrayCollection)
 		: ManagedArrayCollection(InManagedArrayCollection)
 	{
-		Initialize();
+		UpdateArrays();
+		RebuildKeyIndices();
 	}
 
 	FPropertyCollectionConstAdapter::FPropertyCollectionConstAdapter(const TSharedPtr<const FManagedArrayCollection>& InManagedArrayCollection, ENoInit)
@@ -32,11 +33,8 @@ namespace Chaos::Softs
 		FlagsArray = GetArray<uint8>(FlagsName);
 	}
 	
-	void FPropertyCollectionConstAdapter::Initialize()
+	void FPropertyCollectionConstAdapter::RebuildKeyIndices()
 	{
-		// Update arrays
-		UpdateArrays();
-
 		// Create a fast access search map (although it might only be faster for a large enough number of properties)
 		const int32 NumKeys = KeyArray.Num();
 		KeyIndices.Empty(NumKeys);
@@ -51,6 +49,8 @@ namespace Chaos::Softs
 	{
 		return ValueArray[KeyIndex];
 	}
+	template CHAOS_API FVector3f FPropertyCollectionConstAdapter::GetValue<FVector3f, FVector3f>(int32 KeyIndex, const TConstArrayView<FVector3f>& ValueArray) const;
+	template CHAOS_API const FVector3f& FPropertyCollectionConstAdapter::GetValue<const FVector3f&, FVector3f>(int32 KeyIndex, const TConstArrayView<FVector3f>& ValueArray) const;
 	template CHAOS_API FString FPropertyCollectionConstAdapter::GetValue<FString, FString>(int32 KeyIndex, const TConstArrayView<FString>& ValueArray) const;
 	template CHAOS_API const FString& FPropertyCollectionConstAdapter::GetValue<const FString&, FString>(int32 KeyIndex, const TConstArrayView<FString>& ValueArray) const;
 	template CHAOS_API uint8 FPropertyCollectionConstAdapter::GetValue<uint8, uint8>(int32 KeyIndex, const TConstArrayView<uint8>& ValueArray) const;
@@ -83,7 +83,8 @@ namespace Chaos::Softs
 	FPropertyCollectionAdapter::FPropertyCollectionAdapter(const TSharedPtr<FManagedArrayCollection>& InManagedArrayCollection)
 		: FPropertyCollectionConstAdapter(InManagedArrayCollection, NoInit)
 	{
-		Initialize();
+		UpdateArrays();
+		RebuildKeyIndices();
 	}
 	
 	FPropertyCollectionAdapter::FPropertyCollectionAdapter(const TSharedPtr<FManagedArrayCollection>& InManagedArrayCollection, ENoInit)
@@ -117,7 +118,8 @@ namespace Chaos::Softs
 		: FPropertyCollectionAdapter(InManagedArrayCollection, NoInit)
 	{
 		Construct();
-		Initialize();
+		UpdateArrays();
+		RebuildKeyIndices();
 	}
 
 	void FPropertyCollectionMutableAdapter::Construct()
@@ -135,6 +137,9 @@ namespace Chaos::Softs
 		const int32 Index = GetManagedArrayCollection()->AddElements(1, PropertyGroup);
 		const uint8 Flags = (bEnabled ? (uint8)EPropertyFlag::Enabled : 0) | (bAnimatable ? (uint8)EPropertyFlag::Animatable : 0);
 
+		// Update the arrayviews in case the new element triggered a reallocation 
+		UpdateArrays();
+
 		// Setup the new element's default value and enable the property by default
 		GetKeyArray()[Index] = Key;
 		GetLowValueArray()[Index] = GetHighValueArray()[Index] = FVector3f::ZeroVector;
@@ -143,9 +148,6 @@ namespace Chaos::Softs
 		// Update search map
 		KeyIndices.Emplace(KeyArray[Index], Index);
 
-		// Update the array view count
-		UpdateArrays();
-
 		return Index;
 	}
 
@@ -156,6 +158,9 @@ namespace Chaos::Softs
 			const int32 StartIndex = GetManagedArrayCollection()->AddElements(NumProperties, PropertyGroup);
 			const uint8 Flags = (bEnabled ? (uint8)EPropertyFlag::Enabled : 0) | (bAnimatable ? (uint8)EPropertyFlag::Animatable : 0);
 
+			// Update the arrayviews in case the new elements triggered a reallocation 
+			UpdateArrays();
+
 			for (int32 Index = StartIndex; Index < NumProperties + StartIndex; ++Index)
 			{
 				// Setup the new elements' default value and enable the property by default
@@ -167,25 +172,33 @@ namespace Chaos::Softs
 				KeyIndices.Emplace(KeyArray[Index], Index);
 			}
 
-			// Update the array view count
-			UpdateArrays();
-
 			return StartIndex;
 		}
 		return INDEX_NONE;
 	}
 
+	void FPropertyCollectionMutableAdapter::Reset()
+	{
+		GetManagedArrayCollection()->Reset();
+		Construct();
+		UpdateArrays();
+		RebuildKeyIndices();  // Reset search map
+	}
+
 	void FPropertyCollectionMutableAdapter::Append(const TSharedPtr<const FManagedArrayCollection>& InManagedArrayCollection)
 	{
 		TArray<FName> GroupsToSkip = InManagedArrayCollection->GroupNames();
 		GroupsToSkip.RemoveSingleSwap(PropertyGroup);
 
 		InManagedArrayCollection->CopyTo(GetManagedArrayCollection().Get(), GroupsToSkip);
+		UpdateArrays();
+		RebuildKeyIndices();
 	}
 
 	void FPropertyCollectionMutableAdapter::Copy(const TSharedPtr<const FManagedArrayCollection>& InManagedArrayCollection)
 	{
 		GetManagedArrayCollection()->Reset();
+		Construct();
 		Append(InManagedArrayCollection);
 	}
 
diff --git a/Engine/Source/Runtime/Experimental/Chaos/Public/Chaos/PropertyCollectionAdapter.h b/Engine/Source/Runtime/Experimental/Chaos/Public/Chaos/PropertyCollectionAdapter.h
--- a/Engine/Source/Runtime/Experimental/Chaos/Public/Chaos/PropertyCollectionAdapter.h
+++ b/Engine/Source/Runtime/Experimental/Chaos/Public/Chaos/PropertyCollectionAdapter.h
@@ -97,8 +97,8 @@ namespace Chaos::Softs
 		// Update the array views
 		void UpdateArrays();
 
-		// Initialize the array views and the search map
-		void Initialize();
+		// Rebuild the search map
+		void RebuildKeyIndices();
 
 		template<typename T, typename ElementType>
 		T GetValue(int32 KeyIndex, const TConstArrayView<ElementType>& ValueArray) const;
@@ -227,6 +227,9 @@ namespace Chaos::Softs
 		/** Add new properties, and return the index of the first added property. */
 		int32 AddProperties(const TArray<FString>& Keys, bool bEnabled = true, bool bAnimatable = false);
 
+		/** Remove all properties and values from this property collection. */
+		void Reset();
+
 		/**
 		 * Append all properties and values from an existing collection to this property collection.
 		 * This won't copy any other groups, only data from PropertyGroup.
