Chaos Cloth - Fixed a few missing defaults in get values initializations, causing the property index to not be correctly initialized.

#rnx

diff --git a/Engine/Source/Runtime/Experimental/Chaos/Public/Chaos/PBDAxialSpringConstraints.h b/Engine/Source/Runtime/Experimental/Chaos/Public/Chaos/PBDAxialSpringConstraints.h
--- a/Engine/Source/Runtime/Experimental/Chaos/Public/Chaos/PBDAxialSpringConstraints.h
+++ b/Engine/Source/Runtime/Experimental/Chaos/Public/Chaos/PBDAxialSpringConstraints.h
@@ -65,7 +65,7 @@ public:
 			ParticleCount,
 			InConstraints,
 			StiffnessMultipliers,
-			FSolverVec2(GetWeightedFloatAreaSpringStiffness(PropertyCollection)),
+			FSolverVec2(GetWeightedFloatAreaSpringStiffness(PropertyCollection, 1.f)),
 			bTrimKinematicConstraints)
 	{}
 
diff --git a/Engine/Source/Runtime/Experimental/Chaos/Public/Chaos/PBDLongRangeConstraints.h b/Engine/Source/Runtime/Experimental/Chaos/Public/Chaos/PBDLongRangeConstraints.h
--- a/Engine/Source/Runtime/Experimental/Chaos/Public/Chaos/PBDLongRangeConstraints.h
+++ b/Engine/Source/Runtime/Experimental/Chaos/Public/Chaos/PBDLongRangeConstraints.h
@@ -37,8 +37,8 @@ public:
 			InTethers,
 			StiffnessMultipliers,
 			ScaleMultipliers,
-			FSolverVec2(GetWeightedFloatTetherStiffness(PropertyCollection)),
-			FSolverVec2(GetWeightedFloatTetherScale(PropertyCollection)),  // Scale clamping done in constructor
+			FSolverVec2(GetWeightedFloatTetherStiffness(PropertyCollection, 1.f)),
+			FSolverVec2(GetWeightedFloatTetherScale(PropertyCollection, 1.f)),  // Scale clamping done in constructor
 			FPBDStiffness::DefaultPBDMaxStiffness,
 			MeshScale)
 	{}
diff --git a/Engine/Source/Runtime/Experimental/Chaos/Public/Chaos/XPBDAxialSpringConstraints.h b/Engine/Source/Runtime/Experimental/Chaos/Public/Chaos/XPBDAxialSpringConstraints.h
--- a/Engine/Source/Runtime/Experimental/Chaos/Public/Chaos/XPBDAxialSpringConstraints.h
+++ b/Engine/Source/Runtime/Experimental/Chaos/Public/Chaos/XPBDAxialSpringConstraints.h
@@ -187,7 +187,7 @@ public:
 			ParticleCount,
 			InConstraints,
 			StiffnessMultipliers,
-			FSolverVec2(GetWeightedFloatXPBDAreaSpringStiffness(PropertyCollection)),
+			FSolverVec2(GetWeightedFloatXPBDAreaSpringStiffness(PropertyCollection, MaxStiffness)),
 			bTrimKinematicConstraints)
 	{}
 
diff --git a/Engine/Source/Runtime/Experimental/Chaos/Public/Chaos/XPBDBendingConstraints.h b/Engine/Source/Runtime/Experimental/Chaos/Public/Chaos/XPBDBendingConstraints.h
--- a/Engine/Source/Runtime/Experimental/Chaos/Public/Chaos/XPBDBendingConstraints.h
+++ b/Engine/Source/Runtime/Experimental/Chaos/Public/Chaos/XPBDBendingConstraints.h
@@ -52,7 +52,7 @@ public:
 			bTrimKinematicConstraints,
 			MaxStiffness)
 		, DampingRatio(
-			FSolverVec2(GetWeightedFloatXPBDBendingElementDamping(PropertyCollection, 0.f)).ClampAxes(MinDamping, MaxDamping),
+			FSolverVec2(GetWeightedFloatXPBDBendingElementDamping(PropertyCollection, MinDamping)).ClampAxes(MinDamping, MaxDamping),
 			DampingMultipliers,
 			TConstArrayView<TVec2<int32>>(ConstraintSharedEdges),
 			ParticleOffset,
diff --git a/Engine/Source/Runtime/Experimental/Chaos/Public/Chaos/XPBDLongRangeConstraints.h b/Engine/Source/Runtime/Experimental/Chaos/Public/Chaos/XPBDLongRangeConstraints.h
--- a/Engine/Source/Runtime/Experimental/Chaos/Public/Chaos/XPBDLongRangeConstraints.h
+++ b/Engine/Source/Runtime/Experimental/Chaos/Public/Chaos/XPBDLongRangeConstraints.h
@@ -40,8 +40,8 @@ public:
 			InTethers,
 			StiffnessMultipliers,
 			ScaleMultipliers,
-			FSolverVec2(GetWeightedFloatXPBDTetherStiffness(PropertyCollection)),
-			FSolverVec2(GetWeightedFloatXPBDTetherScale(PropertyCollection)),  // Scale clamping done in constructor
+			FSolverVec2(GetWeightedFloatXPBDTetherStiffness(PropertyCollection, MaxStiffness)),
+			FSolverVec2(GetWeightedFloatXPBDTetherScale(PropertyCollection, 1.f)),  // Scale clamping done in constructor
 			MaxStiffness)
 	{
 		NumTethers = 0;
@@ -82,7 +82,7 @@ public:
 
 	virtual ~FXPBDLongRangeConstraints() override {}
 
-	void SetProperties(const FPropertyCollectionConstAdapter& PropertyCollection)
+	void SetProperties(const FPropertyCollectionConstAdapter& PropertyCollection, FSolverReal MeshScale)
 	{
 		if (IsXPBDTetherStiffnessMutable(PropertyCollection))
 		{
@@ -90,15 +90,15 @@ public:
 		}
 		if (IsXPBDTetherScaleMutable(PropertyCollection))
 		{
-			TetherScale.SetWeightedValue(FSolverVec2(GetWeightedFloatXPBDTetherScale(PropertyCollection)).ClampAxes(Base::MinTetherScale, Base::MaxTetherScale));
+			TetherScale.SetWeightedValue(FSolverVec2(GetWeightedFloatXPBDTetherScale(PropertyCollection)).ClampAxes(Base::MinTetherScale, Base::MaxTetherScale) * MeshScale);
 		}
 	}
 
 	// Set the stiffness and scale values used by the constraint
-	void SetProperties(const FSolverVec2& InStiffness, const FSolverVec2& InScale)
+	void SetProperties(const FSolverVec2& InStiffness, const FSolverVec2& InTetherScale, FSolverReal MeshScale = (FSolverReal)1)
 	{
 		Stiffness.SetWeightedValue(InStiffness, MaxStiffness);
-		TetherScale.SetWeightedValue(InScale.ClampAxes(Base::MinTetherScale, Base::MaxTetherScale));
+		TetherScale.SetWeightedValue(InTetherScale.ClampAxes(Base::MinTetherScale, Base::MaxTetherScale) * MeshScale);
 	}
 
 	// Set stiffness offset and range, as well as the simulation stiffness exponent
diff --git a/Engine/Source/Runtime/Experimental/Chaos/Public/Chaos/XPBDSpringConstraints.h b/Engine/Source/Runtime/Experimental/Chaos/Public/Chaos/XPBDSpringConstraints.h
--- a/Engine/Source/Runtime/Experimental/Chaos/Public/Chaos/XPBDSpringConstraints.h
+++ b/Engine/Source/Runtime/Experimental/Chaos/Public/Chaos/XPBDSpringConstraints.h
@@ -179,8 +179,8 @@ public:
 			InConstraints,
 			StiffnessMultipliers,
 			DampingMultipliers,
-			FSolverVec2(GetWeightedFloatXPBDEdgeSpringStiffness(PropertyCollection)),
-			FSolverVec2(GetWeightedFloatXPBDEdgeSpringDamping(PropertyCollection)),
+			FSolverVec2(GetWeightedFloatXPBDEdgeSpringStiffness(PropertyCollection, MaxStiffness)),
+			FSolverVec2(GetWeightedFloatXPBDEdgeSpringDamping(PropertyCollection, MinDampingRatio)),
 			bTrimKinematicConstraints)
 	{}
 
@@ -230,8 +230,8 @@ public:
 			InConstraints,
 			StiffnessMultipliers,
 			DampingMultipliers,
-			FSolverVec2(GetWeightedFloatXPBDBendingSpringStiffness(PropertyCollection)),
-			FSolverVec2(GetWeightedFloatXPBDBendingSpringDamping(PropertyCollection)),
+			FSolverVec2(GetWeightedFloatXPBDBendingSpringStiffness(PropertyCollection, MaxStiffness)),
+			FSolverVec2(GetWeightedFloatXPBDBendingSpringDamping(PropertyCollection, MinDampingRatio)),
 			bTrimKinematicConstraints)
 	{}
 
