additional cis fixes for geometry collection UV change:
 - fix legacy NumUVLayers() call during much older serialization update function to use the old behavior
 - zero-initialize blank UV layer 0 for legacy GeometryCollection Init() functions (used in tests)
 - add Chaos as a dependency for fracture unit tests

diff --git a/Engine/Source/Runtime/Experimental/Chaos/Private/GeometryCollection/GeometryCollection.cpp b/Engine/Source/Runtime/Experimental/Chaos/Private/GeometryCollection/GeometryCollection.cpp
--- a/Engine/Source/Runtime/Experimental/Chaos/Private/GeometryCollection/GeometryCollection.cpp
+++ b/Engine/Source/Runtime/Experimental/Chaos/Private/GeometryCollection/GeometryCollection.cpp
@@ -1160,7 +1160,12 @@ void FGeometryCollection::Serialize(Chaos::FChaosArchive& Ar)
 			}
 
 			TManagedArray<TArray<FVector2f>>& MultipleUVs = ModifyAttribute<TArray<FVector2f>>("UVs", FGeometryCollection::VerticesGroup);
-			if (NumUVLayers() < 1)
+			int32 MinUVLayers = 8;
+			for (int32 VertIdx = 0; VertIdx < MultipleUVs.Num(); ++VertIdx)
+			{
+				MinUVLayers = FMath::Min(MultipleUVs[VertIdx].Num(), MinUVLayers);
+			}
+			if (MinUVLayers < 1)
 			{
 				for (int32 VertIdx = 0; VertIdx < MultipleUVs.Num(); ++VertIdx)
 				{
@@ -1510,18 +1515,22 @@ void FGeometryCollection::Init(FGeometryCollection* Collection, const TArray<flo
 		TManagedArray<int32>& MaterialID = Collection->MaterialID;
 		TManagedArray<int32>& MaterialIndex = Collection->MaterialIndex;
 		TManagedArray<FTransform>& Transform = Collection->Transform;
+		
+		Collection->SetNumUVLayers(1);
 
 		// set the vertex information
+		TManagedArray<FVector2f>* UV0 = Collection->FindUVLayer(0);
 		FVector3d TempVertices(0.f, 0.f, 0.f);
 		for (int32 Idx = 0; Idx < NumNewVertices; ++Idx)
 		{
 			Vertices[Idx] = FVector3f(RawVertexArray[3 * Idx], RawVertexArray[3 * Idx + 1], RawVertexArray[3 * Idx + 2]);
 			TempVertices += FVector3d(Vertices[Idx]);
+			(*UV0)[Idx] = FVector2f::ZeroVector;
 
 			Colors[Idx] = FLinearColor::White;
 		}
 
-		Collection->SetNumUVLayers(1);
+		
 
 		// set the particle information
 		TempVertices /= (float)NumNewVertices;
@@ -1830,17 +1839,19 @@ FGeometryCollection* FGeometryCollection::NewGeometryCollection(const TArray<flo
 	TManagedArray<int32>& StatusFlags = RestCollection->StatusFlags;
 	TManagedArray<int32>& InitialDynamicState = RestCollection->InitialDynamicState;
 
+	RestCollection->SetNumUVLayers(1);
+
 	// set the vertex information
+	TManagedArray<FVector2f>* UV0 = RestCollection->FindUVLayer(0);
 	for (int32 Idx = 0; Idx < NumNewVertices; ++Idx)
 	{
 		Vertices[Idx] = FVector3f(RawVertexArray[3 * Idx], RawVertexArray[3 * Idx + 1], RawVertexArray[3 * Idx + 2]);
 		BoneMap[Idx] = RawBoneMapArray[Idx];
+		(*UV0)[Idx] = FVector2f::ZeroVector;
 
 		Colors[Idx] = FLinearColor::White;
 	}
 
-	RestCollection->SetNumUVLayers(1);
-
 	// Transforms
 	int NumNewTransforms = RawTransformArray.Num(); // 1 particle for this geometry structure
 	int TransformIndex = RestCollection->AddElements(NumNewTransforms, FGeometryCollection::TransformGroup);
