Bug fix: In Create Cluster Particle remove particles from TopLevelClusterParentsStrained

diff --git a/Engine/Source/Runtime/Experimental/Chaos/Private/Chaos/PBDRigidClustering.cpp b/Engine/Source/Runtime/Experimental/Chaos/Private/Chaos/PBDRigidClustering.cpp
--- a/Engine/Source/Runtime/Experimental/Chaos/Private/Chaos/PBDRigidClustering.cpp
+++ b/Engine/Source/Runtime/Experimental/Chaos/Private/Chaos/PBDRigidClustering.cpp
@@ -152,6 +152,7 @@ namespace Chaos
 			if (FPBDRigidClusteredParticleHandle* ClusteredChild = Child->CastToClustered())
 			{
 				TopLevelClusterParents.Remove(ClusteredChild);
+				TopLevelClusterParentsStrained.Remove(ClusteredChild);
 
 				// Cluster group id 0 means "don't union with other things"
 				// TODO: Use INDEX_NONE instead of 0?
@@ -268,6 +269,7 @@ namespace Chaos
 				if (DeactivateClusterChildren)
 				{
 					TopLevelClusterParents.Remove(ClusteredChild);
+					TopLevelClusterParentsStrained.Remove(ClusteredChild);
 				}
 
 				ClusteredChild->SetCollisionImpulses(FMath::Max(NewParticle->CollisionImpulses(), ClusteredChild->CollisionImpulses()));
