Chaos - Simd collison solver first pass. WIP: partially functional and not optimized.
- Chaos simd helper class first pass.
- AB Testing collision solver

diff --git a/Engine/Source/Runtime/Experimental/Chaos/Private/Chaos/Collision/PBDCollisionContainerSolverSimd.cpp b/Engine/Source/Runtime/Experimental/Chaos/Private/Chaos/Collision/PBDCollisionContainerSolverSimd.cpp
--- a/Engine/Source/Runtime/Experimental/Chaos/Private/Chaos/Collision/PBDCollisionContainerSolverSimd.cpp
+++ b/Engine/Source/Runtime/Experimental/Chaos/Private/Chaos/Collision/PBDCollisionContainerSolverSimd.cpp
@@ -37,8 +37,18 @@ namespace Chaos
 		//////////////////////////////////////////////////////////////////////////////////////////////////
 		//////////////////////////////////////////////////////////////////////////////////////////////////
 
-		// Transform the Constraint's local-space data into world space for use by  the collision solver and also calculate tangents, errors, etc
-		FORCEINLINE_DEBUGGABLE void UpdateCollisionSolverContactPointFromConstraint(Private::FPBDCollisionSolverSimd& Solver, const int32 SolverPointIndex, const FPBDCollisionConstraint* Constraint, const int32 ConstraintPointIndex, const FRealSingle Dt, const FSolverBody& Body0, const FSolverBody& Body1)
+		// Transform and copy a single manifold point for use in the solver
+		template<int TNumLanes>
+		void UpdateSolverContactPointFromConstraint(
+			TArrayView<TPBDCollisionSolverManifoldPointsSimd<TNumLanes>>& ManifoldPointsBuffer,
+			TArrayView<TSolverBodyPtrPairSimd<4>>& SolverBodiesBuffer,
+			Private::FPBDCollisionSolverSimd& Solver,
+			const int32 SolverPointIndex, 
+			const FPBDCollisionConstraint* Constraint, 
+			const int32 ConstraintPointIndex, 
+			const FRealSingle Dt, 
+			const FSolverBody& Body0, 
+			const FSolverBody& Body1)
 		{
 			const FManifoldPoint& ManifoldPoint = Constraint->GetManifoldPoint(ConstraintPointIndex);
 
@@ -113,19 +123,31 @@ namespace Chaos
 				}
 			}
 
-			FWorldContactPoint& WorldContactPoint = Solver.GetWorldContactPoint(SolverPointIndex);
-			WorldContactPoint.RelativeContactPoints[0] = WorldRelativeContact0;
-			WorldContactPoint.RelativeContactPoints[1] = WorldRelativeContact1;
-			WorldContactPoint.ContactNormal = WorldContactNormal;
-			WorldContactPoint.ContactTangentU = WorldContactTangentU;
-			WorldContactPoint.ContactTangentV = WorldContactTangentV;
-			WorldContactPoint.ContactDeltaNormal = WorldContactDeltaNormal;
-			WorldContactPoint.ContactDeltaTangentU = WorldContactDeltaTangentU;
-			WorldContactPoint.ContactDeltaTangentV = WorldContactDeltaTangentV;
-			WorldContactPoint.ContactTargetVelocityNormal = WorldContactTargetVelocityNormal;
+			Solver.SetManifoldPoint(
+				ManifoldPointsBuffer,
+				SolverBodiesBuffer,
+				SolverPointIndex,
+				WorldRelativeContact0,
+				WorldRelativeContact1,
+				WorldContactNormal,
+				WorldContactTangentU,
+				WorldContactTangentV,
+				WorldContactDeltaNormal,
+				WorldContactDeltaTangentU,
+				WorldContactDeltaTangentV,
+				WorldContactTargetVelocityNormal);
 		}
 
-		FORCEINLINE_DEBUGGABLE void UpdateCollisionSolverManifoldFromConstraint(Private::FPBDCollisionSolverSimd& Solver, const FPBDCollisionConstraint* Constraint, const FSolverReal Dt, const int32 ConstraintPointBeginIndex, const int32 ConstraintPointEndIndex)
+		// Transform and copy all of a constraint's manifold point data for use by the solver
+		template<int TNumLanes>
+		void UpdateSolverManifoldFromConstraint(
+			TArrayView<TPBDCollisionSolverManifoldPointsSimd<TNumLanes>>& ManifoldPointsBuffer,
+			TArrayView<TSolverBodyPtrPairSimd<4>>& SolverBodiesBuffer,
+			Private::FPBDCollisionSolverSimd& Solver,
+			const FPBDCollisionConstraint* Constraint, 
+			const FSolverReal Dt, 
+			const int32 ConstraintPointBeginIndex, 
+			const int32 ConstraintPointEndIndex)
 		{
 			const FSolverBody& Body0 = Solver.SolverBody0().SolverBody();
 			const FSolverBody& Body1 = Solver.SolverBody1().SolverBody();
@@ -140,15 +162,29 @@ namespace Chaos
 
 					// Transform the constraint contact data into world space for use by the solver
 					// We build this data directly into the solver's world-space contact data which looks a bit odd with "Init" called after but there you go
-					UpdateCollisionSolverContactPointFromConstraint(Solver, SolverManifoldPointIndex, Constraint, ConstraintManifoldPointIndex, Dt, Body0, Body1);
-
-					// Prepare the solver (set constraint-space mass etc)
-					Solver.FinalizeManifoldPoint(SolverManifoldPointIndex, Dt);
+					UpdateSolverContactPointFromConstraint(
+						ManifoldPointsBuffer,
+						SolverBodiesBuffer,
+						Solver, 
+						SolverManifoldPointIndex, 
+						Constraint, 
+						ConstraintManifoldPointIndex, 
+						Dt, 
+						Body0, 
+						Body1);
 				}
 			}
 		}
 
-		FORCEINLINE_DEBUGGABLE void UpdateCollisionSolverFromConstraint(Private::FPBDCollisionSolverSimd& Solver, const FPBDCollisionConstraint* Constraint, const FSolverReal Dt, const FPBDCollisionSolverSettings& SolverSettings, bool& bOutPerIterationCollision)
+		// Transform and copy all constraint data for use by the solver
+		template<int TNumLanes>
+		void UpdateSolverFromConstraint(
+			TArrayView<TPBDCollisionSolverManifoldPointsSimd<TNumLanes>>& ManifoldPointsBuffer,
+			TArrayView<TSolverBodyPtrPairSimd<4>>& SolverBodiesBuffer,
+			FPBDCollisionSolverSimd& Solver,
+			const FPBDCollisionConstraint* Constraint,
+			const FSolverReal Dt, 
+			const FPBDCollisionSolverSettings& SolverSettings)
 		{
 			// Friction values. Static and Dynamic friction are applied in the position solve for most shapes.
 			// We can also run in a mode without static friction at all. This is faster but stacking is not possible.
@@ -189,12 +225,22 @@ namespace Chaos
 			Solver.SolverBody1().SetInvIScale(Constraint->GetInvInertiaScale1());
 			Solver.SolverBody1().SetShockPropagationScale(FReal(1));
 
-			bOutPerIterationCollision = (!Constraint->GetUseManifold() || Constraint->GetUseIncrementalCollisionDetection());
-
-			UpdateCollisionSolverManifoldFromConstraint(Solver, Constraint, Dt, 0, Constraint->NumManifoldPoints());
+			UpdateSolverManifoldFromConstraint(
+				ManifoldPointsBuffer,
+				SolverBodiesBuffer,
+				Solver, 
+				Constraint, 
+				Dt, 
+				0, 
+				Constraint->NumManifoldPoints());
 		}
 
-		FORCEINLINE_DEBUGGABLE void UpdateCollisionConstraintFromSolver(FPBDCollisionConstraint* Constraint, const Private::FPBDCollisionSolverSimd& Solver, const FSolverReal Dt)
+		template<int TNumLanes>
+		void UpdateConstraintFromSolver(
+			TArrayView<TPBDCollisionSolverManifoldPointsSimd<TNumLanes>>& ManifoldPointsBuffer,
+			FPBDCollisionConstraint* Constraint,
+			const Private::FPBDCollisionSolverSimd& Solver, 
+			const FSolverReal Dt)
 		{
 			Constraint->ResetSolverResults();
 
@@ -207,21 +253,11 @@ namespace Chaos
 				FSolverVec3 NetImpulse = FSolverVec3(0);
 				FSolverReal StaticFrictionRatio = FSolverReal(0);
 
-				const Private::FPBDCollisionSolverManifoldPointSimd& SolverManifoldPoint = Solver.GetManifoldPoint(SolverManifoldPointIndex);
 				if (!Constraint->GetManifoldPoint(ManifoldPointIndex).Flags.bDisabled)
 				{
-					NetPushOut =
-						SolverManifoldPoint.NetPushOutNormal * SolverManifoldPoint.WorldContact.ContactNormal +
-						SolverManifoldPoint.NetPushOutTangentU * SolverManifoldPoint.WorldContact.ContactTangentU +
-						SolverManifoldPoint.NetPushOutTangentV * SolverManifoldPoint.WorldContact.ContactTangentV;
-
-					NetImpulse =
-						SolverManifoldPoint.NetImpulseNormal * SolverManifoldPoint.WorldContact.ContactNormal +
-						SolverManifoldPoint.NetImpulseTangentU * SolverManifoldPoint.WorldContact.ContactTangentU +
-						SolverManifoldPoint.NetImpulseTangentV * SolverManifoldPoint.WorldContact.ContactTangentV;
-
-					StaticFrictionRatio = SolverManifoldPoint.StaticFrictionRatio;
-
+					NetPushOut = Solver.GetNetPushOut(ManifoldPointsBuffer, SolverManifoldPointIndex);
+					NetImpulse = Solver.GetNetImpulse(ManifoldPointsBuffer, SolverManifoldPointIndex);
+					StaticFrictionRatio = Solver.GetStaticFrictionRatio(ManifoldPointsBuffer, SolverManifoldPointIndex);
 					++SolverManifoldPointIndex;
 				}
 
@@ -245,7 +281,9 @@ namespace Chaos
 		FPBDCollisionContainerSolverSimd::FPBDCollisionContainerSolverSimd(const FPBDCollisionConstraints& InConstraintContainer, const int32 InPriority)
 			: FConstraintContainerSolver(InPriority)
 			, ConstraintContainer(InConstraintContainer)
-			, bPerIterationCollisionDetection(false)
+			, NumConstraints(0)
+			, DummySolverBody(FSolverBody::MakeInitialized())
+			, DummyConstraintSolverBody(DummySolverBody)
 		{
 	#if !UE_BUILD_SHIPPING && !UE_BUILD_TEST && INTEL_ISPC
 			Private::FPBDCollisionSolverHelperSimd::CheckISPC();
@@ -258,247 +296,243 @@ namespace Chaos
 
 		void FPBDCollisionContainerSolverSimd::Reset(const int32 MaxCollisions)
 		{
-			// A over-allocation policy to avoid reallocation every frame in the common case where a pile of objects is dropped
-			// and the number of contacts increases every tick.
-			if (MaxCollisions > CollisionSolvers.Max())
-			{
-				CollisionSolvers.Reset((5 * MaxCollisions) / 4);	// +25%
-			}
+			// @todo(chaos): allocation policy to reduce number of resizes as contacts increase
+			Constraints.Reset(MaxCollisions);
+			Solvers.Reset(MaxCollisions);
 
-			CollisionSolvers.Reset(MaxCollisions);
-			CollisionConstraints.Reset(CollisionSolvers.Max());
-			bCollisionConstraintPerIterationCollisionDetection.Reset(CollisionSolvers.Max());
-			CollisionSolverManifoldPoints.Reset(CollisionSolvers.Max() * Private::FPBDCollisionSolverSimd::MaxPointsPerConstraint);
+			// This is over-allocating by up to a factor of NumLanes!
+			Data.SimdManifoldPoints.Reset(MaxCollisions);
+			Data.SimdConstraintIndices.Reset(MaxCollisions);
 
-			// Just set the array size for these right away - all data will be initialized later and there is no constructor
-			// We use CollisionConstraints.Num() to track how many constraints we currently have
-			CollisionSolvers.SetNumUninitialized(MaxCollisions);
-			bCollisionConstraintPerIterationCollisionDetection.SetNumUninitialized(MaxCollisions);
+			for (int32 LaneIndex = 0; LaneIndex < GetNumLanes(); ++LaneIndex)
+			{
+				Data.SimdNumConstraints[LaneIndex] = 0;
+				Data.SimdNumManifoldPoints[LaneIndex] = 0;
+			}
+			NumConstraints = 0;
 		}
 
 		void FPBDCollisionContainerSolverSimd::AddConstraints()
 		{
-			Reset(ConstraintContainer.NumConstraints());
+			// Not supported with RBAN (requires islands)
+			check(false);
+		}
 
-			for (FPBDCollisionConstraintHandle* ConstraintHandle : ConstraintContainer.GetConstraintHandles())
+		void FPBDCollisionContainerSolverSimd::AddConstraints(const TArrayView<Private::FPBDIslandConstraint>& IslandConstraints)
+		{
+			// Decide what lane this island goes into: Find the lane with the least constraints in it
+			// @todo(chaos): should use manifold point count, not constraint count
+			int32 IslandLaneIndex = 0;
+			int32 IslandLaneNumConstraints = Data.SimdNumConstraints[0];
+			for (int32 LaneIndex = 1; LaneIndex < GetNumLanes(); ++LaneIndex)
 			{
-				FPBDCollisionConstraint& Constraint = ConstraintHandle->GetContact();
+				if (Data.SimdNumConstraints[LaneIndex] < IslandLaneNumConstraints)
+				{
+					IslandLaneIndex = LaneIndex;
+					IslandLaneNumConstraints = Data.SimdNumConstraints[LaneIndex];
+				}
+			}
 
-				AddConstraint(Constraint);
+			// Make sure we have enough constraint rows for these constraints (space is pre-allocated)
+			if (Data.SimdConstraintIndices.Num() < IslandLaneNumConstraints + IslandConstraints.Num())
+			{
+				Data.SimdConstraintIndices.SetNum(IslandLaneNumConstraints + IslandConstraints.Num());
 			}
-		}
 
-		void FPBDCollisionContainerSolverSimd::AddConstraints(const TArrayView<Private::FPBDIslandConstraint>& IslandConstraints)
-		{
+			// Add all the constraints in the island to the selected lane
 			for (Private::FPBDIslandConstraint& IslandConstraint : IslandConstraints)
 			{
-				// We will only ever be given constraints from our container (asserts in non-shipping)
+				// NOTE: We will only ever be given constraints from our container (asserts in non-shipping)
 				FPBDCollisionConstraint& Constraint = IslandConstraint.GetConstraint()->AsUnsafe<FPBDCollisionConstraintHandle>()->GetContact();
 
-				AddConstraint(Constraint);
+				// Add the constraint to our list
+				const int32 ConstraintIndex = Constraints.Add(&Constraint);
+
+				// Add the constraint to its island's lane
+				const int32 RowIndex = Data.SimdNumConstraints[IslandLaneIndex];
+				Data.SimdConstraintIndices[RowIndex].ConstraintIndex[IslandLaneIndex] = ConstraintIndex;
+
+				++Data.SimdNumConstraints[IslandLaneIndex];
+				++NumConstraints;
 			}
 		}
 
-		void FPBDCollisionContainerSolverSimd::AddConstraint(FPBDCollisionConstraint& Constraint)
+		void FPBDCollisionContainerSolverSimd::CreateSolvers()
 		{
-			// NOTE: No need to add to CollisionSolvers or bCollisionConstraintPerIterationCollisionDetection
-			const int32 Index = CollisionConstraints.Add(&Constraint);
-
-			// Allocate the manifold points for this constraints solver
-			// NOTE: we don't know how many points we will create if collision detection is deferred or incrememntal, so just allocate space for the max allowed
-			// @todo(chaos): see if we can do better - we don't want to pay for this logic when it is only used rarely
-			const bool bDeferredCollisionDetection = ConstraintContainer.GetDetectorSettings().bDeferNarrowPhase;
-			const int32 ManifoldPointMax = (bDeferredCollisionDetection || Constraint.GetUseIncrementalManifold()) ? Private::FPBDCollisionSolverSimd::MaxPointsPerConstraint : Constraint.NumManifoldPoints();
-			if (ManifoldPointMax > 0)
+			// Allocate the solvers
+			Solvers.SetNum(Constraints.Num());
+
+			// Reset manifold pointer per lane counter
+			for (int32 LaneIndex = 0; LaneIndex < GetNumLanes(); ++LaneIndex)
 			{
-				const int32 ManifoldPointStart = CollisionSolverManifoldPoints.Num();
-				check(CollisionSolverManifoldPoints.Max() >= ManifoldPointStart + ManifoldPointMax);	// No resize following reset
-				CollisionSolverManifoldPoints.SetNum(ManifoldPointStart + ManifoldPointMax);
+				Data.SimdNumManifoldPoints[LaneIndex] = 0;
+			}
 
-				CollisionSolvers[Index].SetManifoldPointsBuffer(&CollisionSolverManifoldPoints[ManifoldPointStart], ManifoldPointMax);
+			// Determine how many manifold point rows we need and tell each solver where its points are
+			int32 NumManifoldPointRows = 0;
+			for (int32 ConstraintRowIndex = 0; ConstraintRowIndex < Data.SimdConstraintIndices.Num(); ++ConstraintRowIndex)
+			{
+				TConstraintIndexSimd<4>& ConstraintIndices = Data.SimdConstraintIndices[ConstraintRowIndex];
+
+				for (int32 LaneIndex = 0; LaneIndex < GetNumLanes(); ++LaneIndex)
+				{
+					const int32 ConstraintIndex = ConstraintIndices.ConstraintIndex[LaneIndex];
+					if (ConstraintIndex != INDEX_NONE)
+					{
+						const int32 NumManifoldPoints = Constraints[ConstraintIndex]->NumManifoldPoints();
+
+						// Tell the solver where its manifold points are in the manifold points buffer
+						Solvers[ConstraintIndex].SetManifoldPointsBuffer(
+							ConstraintIndex,
+							LaneIndex,
+							Data.SimdNumManifoldPoints[LaneIndex],
+							NumManifoldPoints);
+
+						Data.SimdNumManifoldPoints[LaneIndex] += NumManifoldPoints;
+						NumManifoldPointRows = FMath::Max(NumManifoldPointRows, Data.SimdNumManifoldPoints[LaneIndex]);
+					}
+				}
 			}
-			else
+
+			// Allocate the manifold point solver rows
+			Data.SimdManifoldPoints.SetNum(NumManifoldPointRows);
+
+			// Initalize the set of body pointers with a dummy body.
+			// This allows us to avoid some branches in the body data gather for rows where not all lanes are used
+			Data.SimdSolverBodies.SetNum(NumManifoldPointRows);
+			for (TSolverBodyPtrPairSimd<4>& SolverBodies : Data.SimdSolverBodies)
 			{
-				// Ideally we wouldn't have any active constraint with no manifold points, but this can
-				// currently happen for CCD objects and incremental collisions
-				CollisionSolvers[Index].Reset();
+				SolverBodies.Body0.SetValues(&DummySolverBody);
+				SolverBodies.Body1.SetValues(&DummySolverBody);
 			}
 		}
 
-
 		void FPBDCollisionContainerSolverSimd::AddBodies(FSolverBodyContainer& SolverBodyContainer)
 		{
-			for (int32 SolverIndex = 0, SolverEndIndex = NumSolvers(); SolverIndex < SolverEndIndex; ++SolverIndex)
+			// All constraints and bodies are now known, so we can initialize the array of solvers.
+			CreateSolvers();
+
+			for (int32 ConstraintIndex = 0; ConstraintIndex < Constraints.Num(); ++ConstraintIndex)
 			{
-				Private::FPBDCollisionSolverSimd& CollisionSolver = CollisionSolvers[SolverIndex];
-				FPBDCollisionConstraint* Constraint = CollisionConstraints[SolverIndex];
-				check(Constraint != nullptr);
+				FPBDCollisionConstraint* Constraint = Constraints[ConstraintIndex];
+				FPBDCollisionSolverSimd& Solver = Solvers[ConstraintIndex];
 
-				// Find the solver bodies for the particles we constrain. This will add them to the container
-				// if they aren't there already, and ensure that they are populated with the latest data.
 				FSolverBody* Body0 = SolverBodyContainer.FindOrAdd(Constraint->GetParticle0());
 				FSolverBody* Body1 = SolverBodyContainer.FindOrAdd(Constraint->GetParticle1());
 				check(Body0 != nullptr);
 				check(Body1 != nullptr);
 
-				CollisionSolver.SetSolverBodies(*Body0, *Body1);
+				Solver.SetSolverBodies(*Body0, *Body1);
 			}
 		}
 
 		void FPBDCollisionContainerSolverSimd::GatherInput(const FReal Dt)
 		{
-			GatherInput(Dt, 0, NumSolvers());
+			GatherInput(Dt, 0, GetNumConstraints());
 		}
 
-		void FPBDCollisionContainerSolverSimd::CachePrefetchSolver(const int32 ConstraintIndex) const
-		{
-			if (ConstraintIndex < NumSolvers())
-			{
-				const int32 NumManifoldPoints = CollisionSolvers[ConstraintIndex].NumManifoldPoints();
-				if (NumManifoldPoints > 0)
-				{
-					FPlatformMisc::PrefetchBlock(&CollisionSolvers[ConstraintIndex].GetManifoldPoint(0), NumManifoldPoints * sizeof(Private::FPBDCollisionSolverManifoldPointSimd));
-				}
-				FPlatformMisc::PrefetchBlock(&CollisionSolvers[ConstraintIndex], sizeof(FPBDCollisionContainerSolverSimd));
-			}
-		}
-
-		void FPBDCollisionContainerSolverSimd::GatherInput(const FReal InDt, const int32 BeginIndex, const int32 EndIndex)
+		void FPBDCollisionContainerSolverSimd::GatherInput(const FReal InDt, const int32 ConstraintBeginIndex, const int32 ConstraintEndIndex)
 		{
 			// NOTE: may be called in parallel. Should not change the container or any elements outside of [BeginIndex, EndIndex)
 
-			check(BeginIndex >= 0);
-			check(EndIndex <= NumSolvers());
-
-			const FSolverReal Dt = FSolverReal(InDt);
-
-			const int32 PrefetchCount = 4;
-			for (int32 PrefetchIndex = 0; PrefetchIndex < PrefetchCount; ++PrefetchIndex)
+			// Map the constraint range to a row/constraint range
+			// @toso(chaos): change the API to ask the container solver for the number of items it needs to gather
+			if (ConstraintBeginIndex != 0)
 			{
-				CachePrefetchSolver(PrefetchIndex);
+				return;
 			}
 
-			bool bWantPerIterationCollisionDetection = false;
-			for (int32 ConstraintIndex = BeginIndex; ConstraintIndex < EndIndex; ++ConstraintIndex)
-			{
-				Private::FPBDCollisionSolverSimd& CollisionSolver = CollisionSolvers[ConstraintIndex];
-				FPBDCollisionConstraint* Constraint = CollisionConstraints[ConstraintIndex];
-				bool& bPerIterationCollision = bCollisionConstraintPerIterationCollisionDetection[ConstraintIndex];
+			int32 RowBeginIndex = 0;
+			int32 RowEndIndex = Constraints.Num();
 
-				UpdateCollisionSolverFromConstraint(CollisionSolver, Constraint, Dt, ConstraintContainer.GetSolverSettings(), bPerIterationCollision);
+			check(RowBeginIndex >= 0);
+			check(RowEndIndex <= Constraints.Num());
 
-				// We need to run collision every iteration if we are not using manifolds, or are using incremental manifolds
-				bWantPerIterationCollisionDetection = bWantPerIterationCollisionDetection || bPerIterationCollision;
+			const FSolverReal Dt = FSolverReal(InDt);
 
-				CachePrefetchSolver(ConstraintIndex + PrefetchCount);
-			}
+			TArrayView<TPBDCollisionSolverManifoldPointsSimd<4>> ManifoldPointsBuffer = MakeArrayView(Data.SimdManifoldPoints);
+			TArrayView<TSolverBodyPtrPairSimd<4>> SolverBodiesBuffer = MakeArrayView(Data.SimdSolverBodies);
 
-			if (bWantPerIterationCollisionDetection)
+			// @todo(chaos): would it be better to iterate over manifold point rows?
+			for (int32 ConstraintIndex = 0; ConstraintIndex < Constraints.Num(); ++ConstraintIndex)
 			{
-				// We should lock here? We only ever set to true or do nothing so I think it doesn't matter if this happens on multiple threads...
-				bPerIterationCollisionDetection = true;
+				FPBDCollisionConstraint* Constraint = Constraints[ConstraintIndex];
+				FPBDCollisionSolverSimd& Solver = Solvers[ConstraintIndex];
+
+				UpdateSolverFromConstraint(
+					ManifoldPointsBuffer,
+					SolverBodiesBuffer,
+					Solver,
+					Constraint, 
+					Dt, 
+					ConstraintContainer.GetSolverSettings());
 			}
 		}
 
 		void FPBDCollisionContainerSolverSimd::ScatterOutput(const FReal Dt)
 		{
-			ScatterOutput(Dt, 0, NumSolvers());
+			ScatterOutput(Dt, 0, GetNumConstraints());
 		}
 
-		void FPBDCollisionContainerSolverSimd::ScatterOutput(const FReal InDt, const int32 BeginIndex, const int32 EndIndex)
+		void FPBDCollisionContainerSolverSimd::ScatterOutput(const FReal InDt, const int32 ConstraintBeginIndex, const int32 ConstraintEndIndex)
 		{
 			// NOTE: may be called in parallel. Should not change the container or any elements outside of [BeginIndex, EndIndex)
 
-			check(BeginIndex >= 0);
-			check(EndIndex <= NumSolvers());
-
-			const FSolverReal Dt = FSolverReal(InDt);
-
-			for (int32 ConstraintIndex = BeginIndex; ConstraintIndex < EndIndex; ++ConstraintIndex)
+			// Map the constraint range to a row/constraint range
+			// @toso(chaos): change the API to ask the container solver for the number of items it needs to gather
+			if (ConstraintBeginIndex != 0)
 			{
-				Private::FPBDCollisionSolverSimd& CollisionSolver = CollisionSolvers[ConstraintIndex];
-				FPBDCollisionConstraint* Constraint = CollisionConstraints[ConstraintIndex];
-
-				UpdateCollisionConstraintFromSolver(Constraint, CollisionSolver, Dt);
-
-				// Reset the collision solver here as the pointers will be invalid on the next tick
-				// @todo(chaos): Pointers are always reinitalized before use next tick so this isn't strictly necessary
-				CollisionSolver.Reset();
+				return;
 			}
-		}
 
-		void FPBDCollisionContainerSolverSimd::ApplyPositionConstraints(const FReal Dt, const int32 It, const int32 NumIts)
-		{
-			SolvePositionImpl(Dt, It, NumIts, 0, NumSolvers(), ConstraintContainer.GetSolverSettings());
-		}
+			int32 RowBeginIndex = 0;
+			int32 RowEndIndex = Constraints.Num();
 
-		void FPBDCollisionContainerSolverSimd::ApplyVelocityConstraints(const FReal Dt, const int32 It, const int32 NumIts)
-		{
-			SolveVelocityImpl(Dt, It, NumIts, 0, NumSolvers(), ConstraintContainer.GetSolverSettings());
-		}
+			check(RowBeginIndex >= 0);
+			check(RowEndIndex <= Constraints.Num());
 
-		void FPBDCollisionContainerSolverSimd::ApplyProjectionConstraints(const FReal Dt, const int32 It, const int32 NumIts)
-		{
-			// Not supported for collisions
-		}
+			const FSolverReal Dt = FSolverReal(InDt);
 
-		void FPBDCollisionContainerSolverSimd::UpdatePositionShockPropagation(const FReal Dt, const int32 It, const int32 NumIts, const int32 BeginIndex, const int32 EndIndex, const FPBDCollisionSolverSettings& SolverSettings)
-		{
-			// If this is the first shock propagation iteration, enable it on each solver
-			const bool bEnableShockPropagation = (It == NumIts - SolverSettings.NumPositionShockPropagationIterations);
-			if (bEnableShockPropagation)
-			{
-				for (int32 SolverIndex = BeginIndex; SolverIndex < EndIndex; ++SolverIndex)
-				{
-					CollisionSolvers[SolverIndex].EnablePositionShockPropagation();
-				}
-			}
-		}
+			TArrayView<TPBDCollisionSolverManifoldPointsSimd<4>> ManifoldPointsBuffer = MakeArrayView(Data.SimdManifoldPoints);
 
-		void FPBDCollisionContainerSolverSimd::UpdateVelocityShockPropagation(const FReal Dt, const int32 It, const int32 NumIts, const int32 BeginIndex, const int32 EndIndex, const FPBDCollisionSolverSettings& SolverSettings)
-		{
-			// Set/reset the shock propagation based on current iteration. The position solve may
-			// have left the bodies with a mass scale and we want to change or reset it.
-			const bool bEnableShockPropagation = (It == NumIts - SolverSettings.NumVelocityShockPropagationIterations);
-			if (bEnableShockPropagation)
+			// @todo(chaos): would it be better to iterate over manifold point rows?
+			for (int32 ConstraintIndex = 0; ConstraintIndex < Constraints.Num(); ++ConstraintIndex)
 			{
-				for (int32 SolverIndex = BeginIndex; SolverIndex < EndIndex; ++SolverIndex)
-				{
-					CollisionSolvers[SolverIndex].EnableVelocityShockPropagation();
-				}
-			}
-			else if (It == 0)
-			{
-				for (int32 SolverIndex = BeginIndex; SolverIndex < EndIndex; ++SolverIndex)
-				{
-					CollisionSolvers[SolverIndex].DisableShockPropagation();
-				}
+				FPBDCollisionConstraint* Constraint = Constraints[ConstraintIndex];
+				FPBDCollisionSolverSimd& Solver = Solvers[ConstraintIndex];
+
+				UpdateConstraintFromSolver(
+					ManifoldPointsBuffer,
+					Constraint, 
+					Solver, 
+					Dt);
+
+				// Reset the collision solver here as the pointers will be invalid on the next tick
+				// @todo(chaos): Pointers are always reinitalized before use next tick so this isn't strictly necessary
+				Solver.Reset();
 			}
 		}
 
-		bool FPBDCollisionContainerSolverSimd::SolvePositionImpl(const FReal InDt, const int32 It, const int32 NumIts, const int32 BeginIndex, const int32 EndIndex, const FPBDCollisionSolverSettings& SolverSettings)
+		void FPBDCollisionContainerSolverSimd::ApplyPositionConstraints(const FReal InDt, const int32 It, const int32 NumIts)
 		{
 			SCOPE_CYCLE_COUNTER(STAT_Collisions_Apply);
 			if (!CVars::bChaos_PBDCollisionSolver_Position_SolveEnabled)
 			{
-				return false;
+				return;
 			}
 
-			if (EndIndex <= BeginIndex)
-			{
-				return false;
-			}
+			const FPBDCollisionSolverSettings& SolverSettings = ConstraintContainer.GetSolverSettings();
 
-			UpdatePositionShockPropagation(InDt, It, NumIts, BeginIndex, EndIndex, SolverSettings);
+			UpdatePositionShockPropagation(InDt, It, NumIts, SolverSettings);
 
 			// We run collision detection here under two conditions (normally it is run after Integration and before the constraint solver phase):
 			// 1) When deferring collision detection until the solver phase for better joint-collision behaviour (RBAN). In this case, we only do this on the first iteration.
 			// 2) When using no manifolds or incremental manifolds, where we may add/replace manifold points every iteration.
-			const bool bRunDeferredCollisionDetection = (It == 0) && ConstraintContainer.GetDetectorSettings().bDeferNarrowPhase;
-			if (bRunDeferredCollisionDetection || bPerIterationCollisionDetection)
-			{
-				UpdateCollisions(InDt, BeginIndex, EndIndex);
-			}
+			//const bool bRunDeferredCollisionDetection = (It == 0) && ConstraintContainer.GetDetectorSettings().bDeferNarrowPhase;
+			//if (bRunDeferredCollisionDetection || bPerIterationCollisionDetection)
+			//{
+			//	UpdateCollisions(InDt, BeginIndex, EndIndex);
+			//}
 
 			// Only apply friction for the last few (tunable) iterations
 			// Adjust max pushout to attempt to make it iteration count independent
@@ -509,86 +543,139 @@ namespace Chaos
 			// Apply the position correction
 			if (bApplyStaticFriction)
 			{
-				Private::FPBDCollisionSolverHelperSimd::SolvePositionWithFriction(MakeArrayView(&CollisionSolvers[BeginIndex], EndIndex - BeginIndex), Dt, MaxPushOut);
+				Private::FPBDCollisionSolverHelperSimd::SolvePositionWithFriction(
+					MakeArrayView(Data.SimdManifoldPoints),
+					MakeArrayView(Data.SimdSolverBodies),
+					Dt,
+					MaxPushOut);
 			}
 			else
 			{
-				Private::FPBDCollisionSolverHelperSimd::SolvePositionNoFriction(MakeArrayView(&CollisionSolvers[BeginIndex], EndIndex - BeginIndex), Dt, MaxPushOut);
+				Private::FPBDCollisionSolverHelperSimd::SolvePositionNoFriction(
+					MakeArrayView(Data.SimdManifoldPoints),
+					MakeArrayView(Data.SimdSolverBodies),
+					Dt,
+					MaxPushOut);
 			}
-
-			return true;
 		}
 
-		bool FPBDCollisionContainerSolverSimd::SolveVelocityImpl(const FReal InDt, const int32 It, const int32 NumIts, const int32 BeginIndex, const int32 EndIndex, const FPBDCollisionSolverSettings& SolverSettings)
+		void FPBDCollisionContainerSolverSimd::ApplyVelocityConstraints(const FReal InDt, const int32 It, const int32 NumIts)
 		{
 			SCOPE_CYCLE_COUNTER(STAT_Collisions_ApplyPushOut);
 			if (!CVars::bChaos_PBDCollisionSolver_Velocity_SolveEnabled)
 			{
-				return false;
+				return;
 			}
 
-			if (EndIndex <= BeginIndex)
-			{
-				return false;
-			}
+			const FPBDCollisionSolverSettings& SolverSettings = ConstraintContainer.GetSolverSettings();
 
-			UpdateVelocityShockPropagation(InDt, It, NumIts, BeginIndex, EndIndex, SolverSettings);
+			UpdateVelocityShockPropagation(InDt, It, NumIts, SolverSettings);
 
 			const FSolverReal Dt = FSolverReal(InDt);
 			const bool bApplyDynamicFriction = (It >= NumIts - SolverSettings.NumVelocityFrictionIterations);
 
-			Private::FPBDCollisionSolverHelperSimd::SolveVelocity(MakeArrayView(&CollisionSolvers[BeginIndex], EndIndex - BeginIndex), Dt, bApplyDynamicFriction);
-
-			return true;
+			if (bApplyDynamicFriction)
+			{
+				Private::FPBDCollisionSolverHelperSimd::SolveVelocityWithFriction(
+					MakeArrayView(Data.SimdManifoldPoints),
+					MakeArrayView(Data.SimdSolverBodies),
+					Dt);
+			}
+			else
+			{
+				Private::FPBDCollisionSolverHelperSimd::SolveVelocityNoFriction(
+					MakeArrayView(Data.SimdManifoldPoints),
+					MakeArrayView(Data.SimdSolverBodies),
+					Dt);
+			}
 		}
 
-		void FPBDCollisionContainerSolverSimd::UpdateCollisions(const FReal InDt, const int32 BeginIndex, const int32 EndIndex)
+		void FPBDCollisionContainerSolverSimd::ApplyProjectionConstraints(const FReal Dt, const int32 It, const int32 NumIts)
 		{
-			const FSolverReal Dt = FSolverReal(InDt);
-			const bool bDeferredCollisionDetection = ConstraintContainer.GetDetectorSettings().bDeferNarrowPhase;
-
-			bool bNeedsAnotherIteration = false;
-			for (int32 SolverIndex = BeginIndex; SolverIndex < EndIndex; ++SolverIndex)
-			{
-				if (bDeferredCollisionDetection || bCollisionConstraintPerIterationCollisionDetection[SolverIndex])
-				{
-					Private::FPBDCollisionSolverSimd& CollisionSolver = CollisionSolvers[SolverIndex];
-					FPBDCollisionConstraint* Constraint = CollisionConstraints[SolverIndex];
-
-					// Run collision detection at the current transforms including any correction from previous iterations
-					const FSolverBody& Body0 = CollisionSolver.SolverBody0().SolverBody();
-					const FSolverBody& Body1 = CollisionSolver.SolverBody1().SolverBody();
-					const FRigidTransform3 CorrectedActorWorldTransform0 = FRigidTransform3(Body0.CorrectedActorP(), Body0.CorrectedActorQ());
-					const FRigidTransform3 CorrectedActorWorldTransform1 = FRigidTransform3(Body1.CorrectedActorP(), Body1.CorrectedActorQ());
-					const FRigidTransform3 CorrectedShapeWorldTransform0 = Constraint->GetShapeRelativeTransform0() * CorrectedActorWorldTransform0;
-					const FRigidTransform3 CorrectedShapeWorldTransform1 = Constraint->GetShapeRelativeTransform1() * CorrectedActorWorldTransform1;
-
-					// @todo(chaos): this is ugly - pass these to the required functions instead and remove from the constraint class
-					// This is now only needed for LevelSet collision (see UpdateLevelsetLevelsetConstraint)
-					Constraint->SetSolverBodies(&Body0, &Body1);
-
-					// Reset the manifold if we are not using manifolds (we just use the first manifold point)
-					if (!Constraint->GetUseManifold())
-					{
-						Constraint->ResetActiveManifoldContacts();
-						CollisionSolver.ResetManifold();
-					}
-
-					// We need to know how many points were added to the manifold
-					const int32 BeginPointIndex = Constraint->NumManifoldPoints();
+			// Not supported for collisions
+		}
 
-					// NOTE: We deliberately have not updated the ShapwWorldTranforms on the constraint. If we did that, we would calculate 
-					// errors incorrectly in UpdateManifoldPoints, because the solver assumes nothing has been moved as we iterate (we accumulate 
-					// corrections that will be applied later.)
-					Constraint->ResetPhi(Constraint->GetCullDistance());
-					Collisions::UpdateConstraint(*Constraint, CorrectedShapeWorldTransform0, CorrectedShapeWorldTransform1, Dt);
+		void FPBDCollisionContainerSolverSimd::UpdatePositionShockPropagation(const FReal Dt, const int32 It, const int32 NumIts, const FPBDCollisionSolverSettings& SolverSettings)
+		{
+			// If this is the first shock propagation iteration, enable it on each solver
+			//const bool bEnableShockPropagation = (It == NumIts - SolverSettings.NumPositionShockPropagationIterations);
+			//if (bEnableShockPropagation)
+			//{
+			//	for (int32 SolverIndex = BeginIndex; SolverIndex < EndIndex; ++SolverIndex)
+			//	{
+			//		CollisionSolvers[SolverIndex].EnablePositionShockPropagation();
+			//	}
+			//}
+		}
 
-					// Update the manifold based on the new or updated contacts
-					UpdateCollisionSolverManifoldFromConstraint(CollisionSolver, Constraint, Dt, BeginPointIndex, Constraint->NumManifoldPoints());
+		void FPBDCollisionContainerSolverSimd::UpdateVelocityShockPropagation(const FReal Dt, const int32 It, const int32 NumIts, const FPBDCollisionSolverSettings& SolverSettings)
+		{
+			// Set/reset the shock propagation based on current iteration. The position solve may
+			// have left the bodies with a mass scale and we want to change or reset it.
+			//const bool bEnableShockPropagation = (It == NumIts - SolverSettings.NumVelocityShockPropagationIterations);
+			//if (bEnableShockPropagation)
+			//{
+			//	for (int32 SolverIndex = BeginIndex; SolverIndex < EndIndex; ++SolverIndex)
+			//	{
+			//		CollisionSolvers[SolverIndex].EnableVelocityShockPropagation();
+			//	}
+			//}
+			//else if (It == 0)
+			//{
+			//	for (int32 SolverIndex = BeginIndex; SolverIndex < EndIndex; ++SolverIndex)
+			//	{
+			//		CollisionSolvers[SolverIndex].DisableShockPropagation();
+			//	}
+			//}
+		}
 
-					Constraint->SetSolverBodies(nullptr, nullptr);
-				}
-			}
+		void FPBDCollisionContainerSolverSimd::UpdateCollisions(const FReal InDt)
+		{
+			//const FSolverReal Dt = FSolverReal(InDt);
+			//const bool bDeferredCollisionDetection = ConstraintContainer.GetDetectorSettings().bDeferNarrowPhase;
+
+			//bool bNeedsAnotherIteration = false;
+			//for (int32 SolverIndex = BeginIndex; SolverIndex < EndIndex; ++SolverIndex)
+			//{
+			//	if (bDeferredCollisionDetection || bCollisionConstraintPerIterationCollisionDetection[SolverIndex])
+			//	{
+			//		Private::FPBDCollisionSolverSimd& CollisionSolver = CollisionSolvers[SolverIndex];
+			//		FPBDCollisionConstraint* Constraint = CollisionConstraints[SolverIndex];
+
+			//		// Run collision detection at the current transforms including any correction from previous iterations
+			//		const FSolverBody& Body0 = CollisionSolver.SolverBody0().SolverBody();
+			//		const FSolverBody& Body1 = CollisionSolver.SolverBody1().SolverBody();
+			//		const FRigidTransform3 CorrectedActorWorldTransform0 = FRigidTransform3(Body0.CorrectedActorP(), Body0.CorrectedActorQ());
+			//		const FRigidTransform3 CorrectedActorWorldTransform1 = FRigidTransform3(Body1.CorrectedActorP(), Body1.CorrectedActorQ());
+			//		const FRigidTransform3 CorrectedShapeWorldTransform0 = Constraint->GetShapeRelativeTransform0() * CorrectedActorWorldTransform0;
+			//		const FRigidTransform3 CorrectedShapeWorldTransform1 = Constraint->GetShapeRelativeTransform1() * CorrectedActorWorldTransform1;
+
+			//		// @todo(chaos): this is ugly - pass these to the required functions instead and remove from the constraint class
+			//		// This is now only needed for LevelSet collision (see UpdateLevelsetLevelsetConstraint)
+			//		Constraint->SetSolverBodies(&Body0, &Body1);
+
+			//		// Reset the manifold if we are not using manifolds (we just use the first manifold point)
+			//		if (!Constraint->GetUseManifold())
+			//		{
+			//			Constraint->ResetActiveManifoldContacts();
+			//			CollisionSolver.ResetManifold();
+			//		}
+
+			//		// We need to know how many points were added to the manifold
+			//		const int32 BeginPointIndex = Constraint->NumManifoldPoints();
+
+			//		// NOTE: We deliberately have not updated the ShapwWorldTranforms on the constraint. If we did that, we would calculate 
+			//		// errors incorrectly in UpdateManifoldPoints, because the solver assumes nothing has been moved as we iterate (we accumulate 
+			//		// corrections that will be applied later.)
+			//		Constraint->ResetPhi(Constraint->GetCullDistance());
+			//		Collisions::UpdateConstraint(*Constraint, CorrectedShapeWorldTransform0, CorrectedShapeWorldTransform1, Dt);
+
+			//		// Update the manifold based on the new or updated contacts
+			//		UpdateCollisionSolverManifoldFromConstraint(CollisionSolver, Constraint, Dt, BeginPointIndex, Constraint->NumManifoldPoints());
+
+			//		Constraint->SetSolverBodies(nullptr, nullptr);
+			//	}
+			//}
 		}
 
 	}	// namespace Private
diff --git a/Engine/Source/Runtime/Experimental/Chaos/Private/Chaos/Collision/PBDCollisionSolverSimd.cpp b/Engine/Source/Runtime/Experimental/Chaos/Private/Chaos/Collision/PBDCollisionSolverSimd.cpp
--- a/Engine/Source/Runtime/Experimental/Chaos/Private/Chaos/Collision/PBDCollisionSolverSimd.cpp
+++ b/Engine/Source/Runtime/Experimental/Chaos/Private/Chaos/Collision/PBDCollisionSolverSimd.cpp
@@ -20,6 +20,8 @@ namespace Chaos
 	{
 		extern float Chaos_PBDCollisionSolver_Position_MinInvMassScale;
 		extern float Chaos_PBDCollisionSolver_Velocity_MinInvMassScale;
+		extern bool bChaos_PBDCollisionSolver_Velocity_FrictionEnabled;
+		extern float Chaos_PBDCollisionSolver_Position_StaticFrictionStiffness;
 	}
 	using namespace CVars;
 
@@ -77,153 +79,175 @@ namespace Chaos
 				}
 
 				// If the masses changed, we need to rebuild the contact mass for each manifold point
-				if (bInvMassUpdated)
-				{
-					for (int32 PointIndex = 0; PointIndex < NumManifoldPoints(); ++PointIndex)
-					{
-						State.ManifoldPoints[PointIndex].UpdateMassNormal(Body0, Body1);
-					}
-				}
+				//if (bInvMassUpdated)
+				//{
+				//	for (int32 PointIndex = 0; PointIndex < NumManifoldPoints(); ++PointIndex)
+				//	{
+				//		State.ManifoldPoints[PointIndex].UpdateMassNormal(Body0, Body1);
+				//	}
+				//}
 			}
 		}
 
-		void FPBDCollisionSolverSimd::SolveVelocityAverage(const FSolverReal Dt)
+		//////////////////////////////////////////////////////////////////////////////////////////////////
+		//////////////////////////////////////////////////////////////////////////////////////////////////
+		//////////////////////////////////////////////////////////////////////////////////////////////////
+		//////////////////////////////////////////////////////////////////////////////////////////////////
+
+		// NOTE: Currently only implemeted for 4-lane SIMD
+
+		// Same as FPlatformMisc::PrefetchBlock() but auto-unrolls the loop and avoids the divide.
+		// Somehow this makes things slower
+		template<typename T>
+		FORCEINLINE static void PrefetchObject(const T* Object)
 		{
-			FConstraintSolverBody& Body0 = SolverBody0();
-			FConstraintSolverBody& Body1 = SolverBody1();
+			constexpr int32 CacheLineSize = PLATFORM_CACHE_LINE_SIZE;
+			constexpr int32 NumBytes = sizeof(T);
+			constexpr int32 NumLines = (NumBytes + CacheLineSize - 1) / CacheLineSize;
 
-			// Generate a new contact point at the average of all the active contacts
-			int32 NumActiveManifoldPoints = 0;
-			FSolverVec3 RelativeContactPosition0 = FSolverVec3(0);
-			FSolverVec3 RelativeContactPosition1 = FSolverVec3(0);
-			FSolverVec3 WorldContactNormal = FSolverVec3(0);
-			FSolverReal NetPushOutNormal = FSolverReal(0);
-			FSolverReal WorldContactVelocityTargetNormal = FSolverReal(0);
-			for (int32 PointIndex = 0; PointIndex < NumManifoldPoints(); ++PointIndex)
+			int32 Offset = 0;
+			for (int32 Line = 0; Line < NumLines; ++Line)
 			{
-				FPBDCollisionSolverManifoldPointSimd& SolverManifoldPoint = State.ManifoldPoints[PointIndex];
-				if (SolverManifoldPoint.ShouldSolveVelocity())
-				{
-					RelativeContactPosition0 += SolverManifoldPoint.WorldContact.RelativeContactPoints[0];
-					RelativeContactPosition1 += SolverManifoldPoint.WorldContact.RelativeContactPoints[1];
-					WorldContactVelocityTargetNormal += SolverManifoldPoint.WorldContact.ContactTargetVelocityNormal;
-					WorldContactNormal = SolverManifoldPoint.WorldContact.ContactNormal;	// Take last value - should all be similar
-					NetPushOutNormal += SolverManifoldPoint.NetPushOutNormal;
-					++NumActiveManifoldPoints;
-				}
+				FPlatformMisc::Prefetch(Object, Offset);
+				Offset += CacheLineSize;
+			}
+		}
+
+		FORCEINLINE void PrefetchManifoldPoint(const TArrayView<TPBDCollisionSolverManifoldPointsSimd<4>>& ManifoldPoints, const int32 Index)
+		{
+			if (Index < ManifoldPoints.Num())
+			{
+				//PrefetchObject(&ManifoldPoints[Index]);
+				FPlatformMisc::PrefetchBlock(&ManifoldPoints[Index], sizeof(TPBDCollisionSolverManifoldPointsSimd<4>));
 			}
+		}
 
-			// Solve for velocity at the new conatct point
-			// We only do this if we have multiple active contacts
-			// NOTE: this average contact isn't really correct, especially when all contacts do not equally
-			// contribute to the pushout (which is normal even for a simple box on a plane). E.g., the WorldContactVelocityTargetNormal
-			// and NetPushOutNormal will not be right. The goal here though is just to get close to the correct result in fewer iterations
-			// than we would have if we just solved the corners of the box in sequence.
-			if (NumActiveManifoldPoints > 1)
+		FORCEINLINE void PrefetchPositionSolverBodies(const TArrayView<TSolverBodyPtrPairSimd<4>>& SolverBodies, const int32 Index)
+		{
+			if (Index < SolverBodies.Num())
 			{
-				const FSolverReal InvCount = FSolverReal(1) / FSolverReal(NumActiveManifoldPoints);
-
-				FPBDCollisionSolverManifoldPointSimd AverageManifoldPoint;
-				AverageManifoldPoint.WorldContact.RelativeContactPoints[0] = RelativeContactPosition0 * InvCount;
-				AverageManifoldPoint.WorldContact.RelativeContactPoints[1] = RelativeContactPosition1 * InvCount;
-				AverageManifoldPoint.WorldContact.ContactNormal = WorldContactNormal;
-				AverageManifoldPoint.WorldContact.ContactTargetVelocityNormal = WorldContactVelocityTargetNormal * InvCount;
-			
-				// Total pushout (not average) which is correct if the average point is also the centroid but otherwise probably an overestimate.
-				// This is used to limit the possibly-attractive impulse that corrects implicit velocity errors from the PBD solve, but it
-				// which might cause stickiness if overestimated. Although it should be recitifed when we solve the normal contact point velocities
-				// so I don't think it matters...we could do better though if it is a problem
-				AverageManifoldPoint.NetPushOutNormal = NetPushOutNormal;
-
-				// Calculate the contact mass (and derived properties) for this point
-				FSolverVec3 WorldContactNormalAngular0 = FSolverVec3(0);
-				FSolverVec3 WorldContactNormalAngular1 = FSolverVec3(0);
-				FSolverReal ContactMassInvNormal = FSolverReal(0);
-				if (Body0.IsDynamic())
-				{
-					const FSolverVec3 R0xN = FSolverVec3::CrossProduct(RelativeContactPosition0, WorldContactNormal);
-					WorldContactNormalAngular0 = Body0.InvI() * R0xN;
-					ContactMassInvNormal += FSolverVec3::DotProduct(R0xN, WorldContactNormalAngular0) + Body0.InvM();
-				}
-				if (Body1.IsDynamic())
+				for (int32 LaneIndex = 0; LaneIndex < 4; ++LaneIndex)
 				{
-					const FSolverVec3 R1xN = FSolverVec3::CrossProduct(RelativeContactPosition1, WorldContactNormal);
-					WorldContactNormalAngular1 = Body1.InvI() * R1xN;
-					ContactMassInvNormal += FSolverVec3::DotProduct(R1xN, WorldContactNormalAngular1) + Body1.InvM();
+					SolverBodies[Index].Body0.GetValue(LaneIndex)->PrefetchPositionSolverData();
+					SolverBodies[Index].Body1.GetValue(LaneIndex)->PrefetchPositionSolverData();
 				}
-				AverageManifoldPoint.ContactMassNormal = (ContactMassInvNormal > FSolverReal(UE_SMALL_NUMBER)) ? FSolverReal(1) / ContactMassInvNormal : FSolverReal(0);
-				AverageManifoldPoint.WorldContactNormalAngular0 = WorldContactNormalAngular0;
-				AverageManifoldPoint.WorldContactNormalAngular1 = WorldContactNormalAngular1;
-
-				// @todo(chaos): we don't use these - maybe do the calculation without an actual manifold point object...
-				AverageManifoldPoint.WorldContact.ContactTangentU = FSolverVec3(0);
-				AverageManifoldPoint.WorldContact.ContactTangentV = FSolverVec3(0);
-				AverageManifoldPoint.WorldContact.ContactDeltaNormal = FSolverReal(0);
-				AverageManifoldPoint.WorldContact.ContactDeltaTangentU = FSolverReal(0);
-				AverageManifoldPoint.WorldContact.ContactDeltaTangentV = FSolverReal(0);
-				AverageManifoldPoint.WorldContactTangentUAngular0 = FSolverVec3(0);
-				AverageManifoldPoint.WorldContactTangentVAngular0 = FSolverVec3(0);
-				AverageManifoldPoint.WorldContactTangentUAngular1 = FSolverVec3(0);
-				AverageManifoldPoint.WorldContactTangentVAngular1 = FSolverVec3(0);
-				AverageManifoldPoint.ContactMassTangentU = FSolverReal(0);
-				AverageManifoldPoint.ContactMassTangentV = FSolverReal(0);
-				AverageManifoldPoint.NetPushOutTangentU = FSolverReal(0);
-				AverageManifoldPoint.NetPushOutTangentV = FSolverReal(0);
-				AverageManifoldPoint.NetImpulseNormal = FSolverReal(0);
-				AverageManifoldPoint.NetImpulseTangentU = FSolverReal(0);
-				AverageManifoldPoint.NetImpulseTangentV = FSolverReal(0);
-				AverageManifoldPoint.StaticFrictionRatio = FSolverReal(0);
-
-				FSolverReal ContactVelocityDeltaNormal;
-				AverageManifoldPoint.CalculateContactVelocityErrorNormal(Body0, Body1, ContactVelocityDeltaNormal);
-
-				const FSolverReal MinImpulseNormal = FMath::Min(FSolverReal(0), -AverageManifoldPoint.NetPushOutNormal / Dt);
-
-				AverageManifoldPoint.ApplyVelocityCorrectionNormal(
-					State.Stiffness,
-					ContactVelocityDeltaNormal,
-					MinImpulseNormal,
-					Body0,
-					Body1);
-
-				// Now distribute the net impulse among the active points so we don't over-correct pushout from initial overlaps
-				for (int32 PointIndex = 0; PointIndex < NumManifoldPoints(); ++PointIndex)
+			}
+		}
+
+
+		FORCEINLINE void PrefetchVelocitySolverBodies(const TArrayView<TSolverBodyPtrPairSimd<4>>& SolverBodies, const int32 Index)
+		{
+			if (Index < SolverBodies.Num())
+			{
+				for (int32 LaneIndex = 0; LaneIndex < 4; ++LaneIndex)
 				{
-					FPBDCollisionSolverManifoldPointSimd& SolverManifoldPoint = State.ManifoldPoints[PointIndex];
-					if (SolverManifoldPoint.ShouldSolveVelocity())
-					{
-						SolverManifoldPoint.NetImpulseNormal += AverageManifoldPoint.NetImpulseNormal * InvCount;
-					}
+					SolverBodies[Index].Body0.GetValue(LaneIndex)->PrefetchVelocitySolverData();
+					SolverBodies[Index].Body1.GetValue(LaneIndex)->PrefetchVelocitySolverData();
 				}
 			}
 		}
 
-		//////////////////////////////////////////////////////////////////////////////////////////////////
-		//////////////////////////////////////////////////////////////////////////////////////////////////
-		//////////////////////////////////////////////////////////////////////////////////////////////////
-		//////////////////////////////////////////////////////////////////////////////////////////////////
+		static const int32 PrefetchCount = 4;
 
-		void FPBDCollisionSolverHelperSimd::SolvePositionNoFriction(const TArrayView<FPBDCollisionSolverSimd>& CollisionSolvers, const FSolverReal Dt, const FSolverReal MaxPushOut)
+		template<>
+		void FPBDCollisionSolverHelperSimd::SolvePositionNoFriction<4>(
+			const TArrayView<TPBDCollisionSolverManifoldPointsSimd<4>>& ManifoldPoints,
+			const TArrayView<TSolverBodyPtrPairSimd<4>>& SolverBodies,
+			const FSolverReal InDt,
+			const FSolverReal InMaxPushOut)
 		{
-			for (FPBDCollisionSolverSimd& CollisionSolver : CollisionSolvers)
+			const FSimd4Realf MaxPushOut = FSimd4Realf::Make(InMaxPushOut);
+
+			for (int32 Index = 0; Index < PrefetchCount; ++Index)
 			{
-				CollisionSolver.SolvePositionNoFriction(Dt, MaxPushOut);
+				PrefetchManifoldPoint(ManifoldPoints, Index);
+				PrefetchPositionSolverBodies(SolverBodies, Index);
+			}
+
+			for (int32 Index = 0; Index < ManifoldPoints.Num(); ++Index)
+			{
+				PrefetchManifoldPoint(ManifoldPoints, Index + PrefetchCount);
+				PrefetchPositionSolverBodies(SolverBodies, Index + PrefetchCount);
+
+				ManifoldPoints[Index].SolvePositionNoFriction(SolverBodies[Index].Body0, SolverBodies[Index].Body1, MaxPushOut);
 			}
 		}
 
-		void FPBDCollisionSolverHelperSimd::SolvePositionWithFriction(const TArrayView<FPBDCollisionSolverSimd>& CollisionSolvers, const FSolverReal Dt, const FSolverReal MaxPushOut)
+		template<>
+		void FPBDCollisionSolverHelperSimd::SolvePositionWithFriction<4>(
+			const TArrayView<TPBDCollisionSolverManifoldPointsSimd<4>>& ManifoldPoints,
+			const TArrayView<TSolverBodyPtrPairSimd<4>>& SolverBodies,
+			const FSolverReal InDt,
+			const FSolverReal InMaxPushOut)
 		{
-			for (FPBDCollisionSolverSimd& CollisionSolver : CollisionSolvers)
+			const FSimd4Realf MaxPushOut = FSimd4Realf::Make(InMaxPushOut);
+			const FSimd4Realf FrictionStiffnessScale = FSimd4Realf::Make(CVars::Chaos_PBDCollisionSolver_Position_StaticFrictionStiffness);
+
+			for (int32 Index = 0; Index < PrefetchCount; ++Index)
+			{
+				PrefetchManifoldPoint(ManifoldPoints, Index);
+				PrefetchPositionSolverBodies(SolverBodies, Index);
+			}
+
+			for (int32 Index = 0; Index < ManifoldPoints.Num(); ++Index)
 			{
-				CollisionSolver.SolvePositionWithFriction(Dt, MaxPushOut);
+				PrefetchManifoldPoint(ManifoldPoints, Index + PrefetchCount);
+				PrefetchPositionSolverBodies(SolverBodies, Index + PrefetchCount);
+
+				ManifoldPoints[Index].SolvePositionWithFriction(SolverBodies[Index].Body0, SolverBodies[Index].Body1, MaxPushOut, FrictionStiffnessScale);
 			}
 		}
 
-		void FPBDCollisionSolverHelperSimd::SolveVelocity(const TArrayView<FPBDCollisionSolverSimd>& CollisionSolvers, const FSolverReal Dt, const bool bApplyDynamicFriction)
+		template<>
+		void FPBDCollisionSolverHelperSimd::SolveVelocityNoFriction<4>(
+			const TArrayView<TPBDCollisionSolverManifoldPointsSimd<4>>& ManifoldPoints,
+			const TArrayView<TSolverBodyPtrPairSimd<4>>& SolverBodies,
+			const FSolverReal InDt)
 		{
-			for (FPBDCollisionSolverSimd& CollisionSolver : CollisionSolvers)
+			const FSimd4Realf Dt = FSimd4Realf::Make(InDt);
+
+			for (int32 Index = 0; Index < PrefetchCount; ++Index)
 			{
-				CollisionSolver.SolveVelocity(Dt, bApplyDynamicFriction);
+				PrefetchManifoldPoint(ManifoldPoints, Index);
+				PrefetchVelocitySolverBodies(SolverBodies, Index);
+			}
+
+			for (int32 Index = 0; Index < ManifoldPoints.Num(); ++Index)
+			{
+				PrefetchManifoldPoint(ManifoldPoints, Index + PrefetchCount);
+				PrefetchVelocitySolverBodies(SolverBodies, Index + PrefetchCount);
+
+				ManifoldPoints[Index].SolveVelocityNoFriction(SolverBodies[Index].Body0, SolverBodies[Index].Body1, Dt);
+			}
+		}
+
+		template<>
+		void FPBDCollisionSolverHelperSimd::SolveVelocityWithFriction<4>(
+			const TArrayView<TPBDCollisionSolverManifoldPointsSimd<4>>& ManifoldPoints,
+			const TArrayView<TSolverBodyPtrPairSimd<4>>& SolverBodies,
+			const FSolverReal InDt)
+		{
+			if (!CVars::bChaos_PBDCollisionSolver_Velocity_FrictionEnabled)
+			{
+				SolveVelocityNoFriction(ManifoldPoints, SolverBodies, InDt);
+				return;
+			}
+
+			for (int32 Index = 0; Index < PrefetchCount; ++Index)
+			{
+				PrefetchManifoldPoint(ManifoldPoints, Index);
+				PrefetchVelocitySolverBodies(SolverBodies, Index);
+			}
+
+			const FSimd4Realf Dt = FSimd4Realf::Make(InDt);
+			const FSimd4Realf FrictionStiffnessScale = FSimd4Realf::Make(CVars::Chaos_PBDCollisionSolver_Position_StaticFrictionStiffness);
+
+			for (int32 Index = 0; Index < ManifoldPoints.Num(); ++Index)
+			{
+				PrefetchManifoldPoint(ManifoldPoints, Index + PrefetchCount);
+				PrefetchVelocitySolverBodies(SolverBodies, Index + PrefetchCount);
+
+				ManifoldPoints[Index].SolveVelocityWithFriction(SolverBodies[Index].Body0, SolverBodies[Index].Body1, Dt, FrictionStiffnessScale);
 			}
 		}
 
diff --git a/Engine/Source/Runtime/Experimental/Chaos/Private/Chaos/Evolution/ConstraintGroupSolver.cpp b/Engine/Source/Runtime/Experimental/Chaos/Private/Chaos/Evolution/ConstraintGroupSolver.cpp
--- a/Engine/Source/Runtime/Experimental/Chaos/Private/Chaos/Evolution/ConstraintGroupSolver.cpp
+++ b/Engine/Source/Runtime/Experimental/Chaos/Private/Chaos/Evolution/ConstraintGroupSolver.cpp
@@ -193,6 +193,17 @@ namespace Chaos
 				});
 		}
 
+		void FPBDConstraintGroupSolver::PreApplyPositionConstraints(const FReal Dt)
+		{
+			for (int32 ContainerIndex = 0; ContainerIndex < ConstraintContainerSolvers.Num(); ++ContainerIndex)
+			{
+				if (ConstraintContainerSolvers[ContainerIndex] != nullptr)
+				{
+					ConstraintContainerSolvers[ContainerIndex]->PreApplyPositionConstraints(Dt);
+				}
+			}
+		}
+
 		void FPBDConstraintGroupSolver::ApplyPositionConstraints(const FReal Dt)
 		{
 			const int32 NumIts = Iterations.GetNumPositionIterations();
@@ -208,9 +219,20 @@ namespace Chaos
 					}
 				}
 			}
+		}
 
+		void FPBDConstraintGroupSolver::PreApplyVelocityConstraints(const FReal Dt)
+		{
 			// Calculate the velocity from the net change in position after applying position constraints
 			SolverBodyContainer.SetImplicitVelocities(Dt);
+
+			for (int32 ContainerIndex = 0; ContainerIndex < ConstraintContainerSolvers.Num(); ++ContainerIndex)
+			{
+				if (ConstraintContainerSolvers[ContainerIndex] != nullptr)
+				{
+					ConstraintContainerSolvers[ContainerIndex]->PreApplyVelocityConstraints(Dt);
+				}
+			}
 		}
 
 		void FPBDConstraintGroupSolver::ApplyVelocityConstraints(const FReal Dt)
@@ -230,15 +252,26 @@ namespace Chaos
 			}
 		}
 
-		void FPBDConstraintGroupSolver::ApplyProjectionConstraints(const FReal Dt)
+		void FPBDConstraintGroupSolver::PreApplyProjectionConstraints(const FReal Dt)
 		{
-			const int32 NumIts = Iterations.GetNumProjectionIterations();
-
 			// Update the body transforms from the deltas calculated in the constraint solve phases 1 and 2
 			// NOTE: deliberately not updating the world-space inertia as it is not used by joint projection
 			// and no other constraints currently implement projection
 			SolverBodyContainer.ApplyCorrections();
 
+			for (int32 ContainerIndex = 0; ContainerIndex < ConstraintContainerSolvers.Num(); ++ContainerIndex)
+			{
+				if (ConstraintContainerSolvers[ContainerIndex] != nullptr)
+				{
+					ConstraintContainerSolvers[ContainerIndex]->PreApplyProjectionConstraints(Dt);
+				}
+			}
+		}
+
+		void FPBDConstraintGroupSolver::ApplyProjectionConstraints(const FReal Dt)
+		{
+			const int32 NumIts = Iterations.GetNumProjectionIterations();
+
 			// NOTE: We loop over prioritized solvers here
 			for (int32 It = 0; It < NumIts; ++It)
 			{
diff --git a/Engine/Source/Runtime/Experimental/Chaos/Private/Chaos/Island/IslandGroupManager.cpp b/Engine/Source/Runtime/Experimental/Chaos/Private/Chaos/Island/IslandGroupManager.cpp
--- a/Engine/Source/Runtime/Experimental/Chaos/Private/Chaos/Island/IslandGroupManager.cpp
+++ b/Engine/Source/Runtime/Experimental/Chaos/Private/Chaos/Island/IslandGroupManager.cpp
@@ -249,18 +249,21 @@ namespace Chaos
 				SCOPE_CYCLE_COUNTER(STAT_Evolution_ApplyConstraintsPhase1);
 				CSV_SCOPED_ISLANDGROUP_TIMING_STAT(PerIslandSolve_ApplyTotalSerialized, GroupIndex);
 
+				IslandGroup->PreApplyPositionConstraints(Dt);
 				IslandGroup->ApplyPositionConstraints(Dt);
 			}
 			{
 				SCOPE_CYCLE_COUNTER(STAT_Evolution_ApplyConstraintsPhase2);
 				CSV_SCOPED_ISLANDGROUP_TIMING_STAT(PerIslandSolve_ApplyPushOutTotalSerialized, GroupIndex);
 
+				IslandGroup->PreApplyVelocityConstraints(Dt);
 				IslandGroup->ApplyVelocityConstraints(Dt);
 			}
 			{
 				SCOPE_CYCLE_COUNTER(STAT_Evolution_ApplyConstraintsPhase3);
 				CSV_SCOPED_ISLANDGROUP_TIMING_STAT(PerIslandSolve_ApplyProjectionTotalSerialized, GroupIndex);
 
+				IslandGroup->PreApplyProjectionConstraints(Dt);
 				IslandGroup->ApplyProjectionConstraints(Dt);
 			}
 		}
diff --git a/Engine/Source/Runtime/Experimental/Chaos/Private/Chaos/PBDCollisionConstraints.cpp b/Engine/Source/Runtime/Experimental/Chaos/Private/Chaos/PBDCollisionConstraints.cpp
--- a/Engine/Source/Runtime/Experimental/Chaos/Private/Chaos/PBDCollisionConstraints.cpp
+++ b/Engine/Source/Runtime/Experimental/Chaos/Private/Chaos/PBDCollisionConstraints.cpp
@@ -13,6 +13,7 @@
 #include "Chaos/Collision/PBDCollisionContainerSolverSimd.h"
 #include "Chaos/Defines.h"
 #include "Chaos/Evolution/SolverBodyContainer.h"
+#include "Chaos/Evolution/ABTestingConstraintContainerSolver.h"
 #include "Chaos/GeometryQueries.h"
 #include "Chaos/Island/IslandManager.h"
 #include "Chaos/SpatialAccelerationCollection.h"
@@ -92,6 +93,134 @@ namespace Chaos
 	DECLARE_CYCLE_STAT(TEXT("Collisions::EndDetect"), STAT_Collisions_EndDetect, STATGROUP_ChaosCollision);
 	DECLARE_CYCLE_STAT(TEXT("Collisions::DetectProbeCollisions"), STAT_Collisions_DetectProbeCollisions, STATGROUP_ChaosCollision);
 
+#if CHAOS_ABTEST_CONSTRAINTSOLVER_ENABLED
+	namespace CVars
+	{
+		bool bCollisionsEnableSolverABTest = false;
+		FAutoConsoleVariableRef CVarCollisionsEnableSolverABTest(TEXT("p.Chaos.Collision.ABTestSolver"), bCollisionsEnableSolverABTest, TEXT(""));
+	}
+
+	// An AB Testing collision solver for use while developing the Simd version
+	using FABTestingCollisionContainerSolver = Private::TABTestingConstraintContainerSolver<FPBDCollisionContainerSolver, Private::FPBDCollisionContainerSolverSimd>;
+
+	// Simd AB testing callback
+	void ABTestCollisionContainerSolver(
+		FABTestingCollisionContainerSolver::ESolverPhase Phase,
+		const FPBDCollisionContainerSolver& SolverA,
+		const Private::FPBDCollisionContainerSolverSimd& SolverB,
+		const FSolverBodyContainer& SolverBodyContainerA,
+		const FSolverBodyContainer& SolverBodyContainerB)
+	{
+		for (int32 ConstraintIndex = 0; ConstraintIndex < SolverA.GetNumConstraints(); ++ConstraintIndex)
+		{
+			const Private::FPBDCollisionSolver& ConstraintSolverA = SolverA.GetConstraintSolver(ConstraintIndex);
+			const Private::FPBDCollisionSolverSimd& ConstraintSolverB = SolverB.GetConstraintSolver(ConstraintIndex);
+
+			if (ConstraintSolverA.NumManifoldPoints() != ConstraintSolverB.NumManifoldPoints())
+			{
+				UE_LOG(LogChaos, Warning, TEXT("ABTestSolverBodies: ManifoldPoint count mismatch"));
+				return;
+			}
+
+			for (int32 ManifoldPointIndex = 0; ManifoldPointIndex < ConstraintSolverA.NumManifoldPoints(); ++ManifoldPointIndex)
+			{
+				const Private::FPBDCollisionSolverManifoldPoint& ManifoldPointSolverA = ConstraintSolverA.GetManifoldPoint(ManifoldPointIndex);
+				const Private::TPBDCollisionSolverManifoldPointsSimd<4>& ManifoldPointSolverB = SolverB.GetManifoldPointSolver(ConstraintSolverB.GetManifoldPointBeginIndex() + ManifoldPointIndex);
+				const int32 LaneIndexB = ConstraintSolverB.GetLaneIndex();
+
+				if (ManifoldPointSolverA.WorldContact.ContactNormal != ManifoldPointSolverB.SimdContactNormal.GetValue(LaneIndexB))
+				{
+					UE_LOG(LogChaos, Warning, TEXT("ABTestSolverBodies: ContactNormal mismatch"));
+				}
+				if (ManifoldPointSolverA.WorldContact.RelativeContactPoints[0] != ManifoldPointSolverB.SimdRelativeContactPoint0.GetValue(LaneIndexB))
+				{
+					UE_LOG(LogChaos, Warning, TEXT("ABTestSolverBodies: RelativeContactPoints mismatch"));
+				}
+				if (ManifoldPointSolverA.WorldContact.RelativeContactPoints[1] != ManifoldPointSolverB.SimdRelativeContactPoint1.GetValue(LaneIndexB))
+				{
+					UE_LOG(LogChaos, Warning, TEXT("ABTestSolverBodies: RelativeContactPoints mismatch"));
+				}
+				if (ManifoldPointSolverA.ContactMassNormal != ManifoldPointSolverB.SimdContactMassNormal.GetValue(LaneIndexB))
+				{
+					UE_LOG(LogChaos, Warning, TEXT("ABTestSolverBodies: ContactMassNormal mismatch"));
+				}
+
+				if (Phase == FABTestingCollisionContainerSolver::ESolverPhase::PostApplyPositionConstraints)
+				{
+					if (ManifoldPointSolverA.NetPushOutNormal != ManifoldPointSolverB.SimdNetPushOutNormal.GetValue(LaneIndexB))
+					{
+						UE_LOG(LogChaos, Warning, TEXT("ABTestSolverBodies: NetPushOutNormal mismatch"));
+					}
+					if (ManifoldPointSolverA.NetPushOutTangentU != ManifoldPointSolverB.SimdNetPushOutTangentU.GetValue(LaneIndexB))
+					{
+						UE_LOG(LogChaos, Warning, TEXT("ABTestSolverBodies: NetPushOutTangentU mismatch"));
+					}
+					if (ManifoldPointSolverA.NetPushOutTangentV != ManifoldPointSolverB.SimdNetPushOutTangentV.GetValue(LaneIndexB))
+					{
+						UE_LOG(LogChaos, Warning, TEXT("ABTestSolverBodies: NetPushOutTangentV mismatch"));
+					}
+				}
+
+				if (Phase == FABTestingCollisionContainerSolver::ESolverPhase::PostApplyVelocityConstraints)
+				{
+					if (ManifoldPointSolverA.NetImpulseNormal != ManifoldPointSolverB.SimdNetImpulseNormal.GetValue(LaneIndexB))
+					{
+						UE_LOG(LogChaos, Warning, TEXT("ABTestSolverBodies: NetImpulseNormal mismatch"));
+					}
+					if (ManifoldPointSolverA.NetImpulseTangentU != ManifoldPointSolverB.SimdNetImpulseTangentU.GetValue(LaneIndexB))
+					{
+						UE_LOG(LogChaos, Warning, TEXT("ABTestSolverBodies: NetImpulseTangentU mismatch"));
+					}
+					if (ManifoldPointSolverA.NetImpulseTangentV != ManifoldPointSolverB.SimdNetImpulseTangentV.GetValue(LaneIndexB))
+					{
+						UE_LOG(LogChaos, Warning, TEXT("ABTestSolverBodies: NetImpulseTangentV mismatch"));
+					}
+				}
+			}
+		}
+
+		if (SolverBodyContainerA.Num() != SolverBodyContainerB.Num())
+		{
+			UE_LOG(LogChaos, Warning, TEXT("ABTestSolverBodies: Body count mismatch"));
+			return;
+		}
+
+		for (int BodyIndex = 0; BodyIndex < SolverBodyContainerA.Num(); ++BodyIndex)
+		{
+			const FSolverBody& BodyA = SolverBodyContainerA.GetSolverBody(BodyIndex);
+			const FSolverBody& BodyB = SolverBodyContainerB.GetSolverBody(BodyIndex);
+
+			if (BodyA.P() != BodyB.P())
+			{
+				UE_LOG(LogChaos, Warning, TEXT("Position mismatch: (%f %f %f) != (%f %f %f)"), 
+					BodyA.P().X, BodyA.P().Y, BodyA.P().Z,
+					BodyB.P().X, BodyB.P().Y, BodyB.P().Z);
+			}
+
+			if (BodyA.DP() != BodyB.DP())
+			{
+				UE_LOG(LogChaos, Warning, TEXT("PositionDelta mismatch: (%f %f %f) != (%f %f %f)"),
+					BodyA.DP().X, BodyA.DP().Y, BodyA.DP().Z,
+					BodyB.DP().X, BodyB.DP().Y, BodyB.DP().Z);
+			}
+
+			if (BodyA.DQ() != BodyB.DQ())
+			{
+				UE_LOG(LogChaos, Warning, TEXT("RotationDelta mismatch: (%f %f %f) != (%f %f %f)"),
+					BodyA.DQ().X, BodyA.DQ().Y, BodyA.DQ().Z,
+					BodyB.DQ().X, BodyB.DQ().Y, BodyB.DQ().Z);
+			}
+
+			if (BodyA.V() != BodyB.V())
+			{
+				UE_LOG(LogChaos, Warning, TEXT("Velocity mismatch: (%f %f %f) != (%f %f %f)"),
+					BodyA.V().X, BodyA.V().Y, BodyA.V().Z,
+					BodyB.V().X, BodyB.V().Y, BodyB.V().Z);
+			}
+		}
+	}
+#endif
+
 	//
 	// Collision Constraint Container
 	//
@@ -145,8 +274,21 @@ namespace Chaos
 		{
 		case Private::ECollisionSolverType::GaussSeidel:
 			return MakeUnique<FPBDCollisionContainerSolver>(*this, Priority);
+
 		case Private::ECollisionSolverType::GaussSeidelSimd:
+#if CHAOS_ABTEST_CONSTRAINTSOLVER_ENABLED
+			if (CVars::bCollisionsEnableSolverABTest)
+			{
+				// Create an AB testing collison solver for simd testing
+				return MakeUnique<FABTestingCollisionContainerSolver>(
+					MakeUnique<FPBDCollisionContainerSolver>(*this, Priority),
+					MakeUnique<Private::FPBDCollisionContainerSolverSimd>(*this, Priority),
+					Priority,
+					&ABTestCollisionContainerSolver);
+			}
+#endif
 			return MakeUnique<Private::FPBDCollisionContainerSolverSimd>(*this, Priority);
+
 		case Private::ECollisionSolverType::PartialJacobi:
 			return MakeUnique<Private::FPBDCollisionContainerSolverJacobi>(*this, Priority);
 		}
diff --git a/Engine/Source/Runtime/Experimental/Chaos/Public/Chaos/Collision/PBDCollisionContainerSolver.h b/Engine/Source/Runtime/Experimental/Chaos/Public/Chaos/Collision/PBDCollisionContainerSolver.h
--- a/Engine/Source/Runtime/Experimental/Chaos/Public/Chaos/Collision/PBDCollisionContainerSolver.h
+++ b/Engine/Source/Runtime/Experimental/Chaos/Public/Chaos/Collision/PBDCollisionContainerSolver.h
@@ -51,6 +51,10 @@ namespace Chaos
 		virtual void ApplyVelocityConstraints(const FReal Dt, const int32 It, const int32 NumIts) override final;
 		virtual void ApplyProjectionConstraints(const FReal Dt, const int32 It, const int32 NumIts) override final;
 
+		// For testing
+		const Private::FPBDCollisionSolver& GetConstraintSolver(const int32 ConstraintIndex) const { return CollisionSolvers[ConstraintIndex]; }
+		const Private::FPBDCollisionSolverManifoldPoint& GetManifoldPointSolver(const int32 RowIndex) const { return CollisionSolverManifoldPoints[RowIndex]; }
+
 	private:
 		void CachePrefetchSolver(const int32 ConstraintIndex) const;
 		void AddConstraint(FPBDCollisionConstraint& Constraint);
diff --git a/Engine/Source/Runtime/Experimental/Chaos/Public/Chaos/Collision/PBDCollisionContainerSolverSimd.h b/Engine/Source/Runtime/Experimental/Chaos/Public/Chaos/Collision/PBDCollisionContainerSolverSimd.h
--- a/Engine/Source/Runtime/Experimental/Chaos/Public/Chaos/Collision/PBDCollisionContainerSolverSimd.h
+++ b/Engine/Source/Runtime/Experimental/Chaos/Public/Chaos/Collision/PBDCollisionContainerSolverSimd.h
@@ -33,11 +33,9 @@ namespace Chaos
 			FPBDCollisionContainerSolverSimd(const FPBDCollisionConstraints& InConstraintContainer, const int32 InPriority);
 			~FPBDCollisionContainerSolverSimd();
 
-			int32 NumSolvers() const { return CollisionConstraints.Num(); }
-
 			virtual void Reset(const int32 InMaxCollisions) override final;
 
-			virtual int32 GetNumConstraints() const override final { return NumSolvers(); }
+			virtual int32 GetNumConstraints() const override final { return NumConstraints; }
 
 			//
 			// IslandGroup API
@@ -53,21 +51,60 @@ namespace Chaos
 			virtual void ApplyVelocityConstraints(const FReal Dt, const int32 It, const int32 NumIts) override final;
 			virtual void ApplyProjectionConstraints(const FReal Dt, const int32 It, const int32 NumIts) override final;
 
+			template<int TNumLanes>
+			struct TConstraintIndexSimd
+			{
+				TConstraintIndexSimd()
+				{
+					for (int32 LaneIndex = 0; LaneIndex < TNumLanes; ++LaneIndex)
+					{
+						ConstraintIndex[LaneIndex] = INDEX_NONE;
+					}
+				}
+
+				int32 ConstraintIndex[TNumLanes];
+			};
+
+			// Rows of collision constraints and solvers. 
+			// We solve each row using SIMD. 
+			// NumLanes is the width of a row and will be the same as the float SIMD register width (or less).
+			template<int TNumLanes>
+			struct FDataSimd
+			{
+				FDataSimd()
+					: SimdNumConstraints{ 0, }
+					, SimdNumManifoldPoints{ 0, }
+				{
+				}
+
+				int32 SimdNumConstraints[TNumLanes];
+				int32 SimdNumManifoldPoints[TNumLanes];
+
+				TArray<TSolverBodyPtrPairSimd<TNumLanes>> SimdSolverBodies;
+				TArray<TConstraintIndexSimd<TNumLanes>> SimdConstraintIndices;
+				TArray<TPBDCollisionSolverManifoldPointsSimd<TNumLanes>> SimdManifoldPoints;
+			};
+
+			// For testing
+			const FPBDCollisionSolverSimd& GetConstraintSolver(const int32 ConstraintIndex) const { return Solvers[ConstraintIndex]; }
+			const TPBDCollisionSolverManifoldPointsSimd<4>& GetManifoldPointSolver(const int32 RowIndex) const { return Data.SimdManifoldPoints[RowIndex]; }
+
 		private:
-			void CachePrefetchSolver(const int32 ConstraintIndex) const;
-			void AddConstraint(FPBDCollisionConstraint& Constraint);
-			void UpdatePositionShockPropagation(const FReal Dt, const int32 It, const int32 NumIts, const int32 BeginIndex, const int32 EndIndex, const FPBDCollisionSolverSettings& SolverSettings);
-			void UpdateVelocityShockPropagation(const FReal Dt, const int32 It, const int32 NumIts, const int32 BeginIndex, const int32 EndIndex, const FPBDCollisionSolverSettings& SolverSettings);
-			bool SolvePositionImpl(const FReal Dt, const int32 It, const int32 NumIts, const int32 BeginIndex, const int32 EndIndex, const FPBDCollisionSolverSettings& SolverSettings);
-			bool SolveVelocityImpl(const FReal Dt, const int32 It, const int32 NumIts, const int32 BeginIndex, const int32 EndIndex, const FPBDCollisionSolverSettings& SolverSettings);
-			void UpdateCollisions(const FReal InDt, const int32 BeginIndex, const int32 EndIndex);
+			int32 GetNumLanes() const { return 4; }
+			void CreateSolvers();
+			void UpdatePositionShockPropagation(const FReal Dt, const int32 It, const int32 NumIts, const FPBDCollisionSolverSettings& SolverSettings);
+			void UpdateVelocityShockPropagation(const FReal Dt, const int32 It, const int32 NumIts, const FPBDCollisionSolverSettings& SolverSettings);
+			void UpdateCollisions(const FReal InDt);
 
 			const FPBDCollisionConstraints& ConstraintContainer;
-			TArray<Private::FPBDCollisionSolverSimd> CollisionSolvers;
-			TArray<Private::FPBDCollisionSolverManifoldPointSimd> CollisionSolverManifoldPoints;
-			TArray<FPBDCollisionConstraint*> CollisionConstraints;
-			TArray<bool> bCollisionConstraintPerIterationCollisionDetection;
-			bool bPerIterationCollisionDetection;
+
+			TArray<FPBDCollisionConstraint*> Constraints;
+			TArray<FPBDCollisionSolverSimd> Solvers;
+			FDataSimd<4> Data;
+			int32 NumConstraints;
+
+			FSolverBody DummySolverBody;
+			FConstraintSolverBody DummyConstraintSolverBody;
 		};
 
 	}	// namespace Private
diff --git a/Engine/Source/Runtime/Experimental/Chaos/Public/Chaos/Collision/PBDCollisionSolverSimd.h b/Engine/Source/Runtime/Experimental/Chaos/Public/Chaos/Collision/PBDCollisionSolverSimd.h
--- a/Engine/Source/Runtime/Experimental/Chaos/Public/Chaos/Collision/PBDCollisionSolverSimd.h
+++ b/Engine/Source/Runtime/Experimental/Chaos/Public/Chaos/Collision/PBDCollisionSolverSimd.h
@@ -5,6 +5,7 @@
 #include "Chaos/Collision/ContactPoint.h"
 #include "Chaos/Defines.h"
 #include "Chaos/Evolution/SolverBody.h"
+#include "Chaos/Simd4.h"
 
 // Set to 0 to use a linearized error calculation, and set to 1 to use a non-linear error calculation in collision detection. 
 // In principle nonlinear is more accurate when large rotation corrections occur, but this is not too important for collisions because 
@@ -22,310 +23,558 @@ namespace Chaos
 		class FPBDCollisionSolverSimd;
 	}
 
-	namespace CVars
-	{
-		extern bool bChaos_PBDCollisionSolver_Velocity_AveragePointEnabled;
-		extern bool bChaos_PBDCollisionSolver_Velocity_FrictionEnabled;
-		extern float Chaos_PBDCollisionSolver_Position_StaticFrictionStiffness;
-		extern float Chaos_PBDCollisionSolver_JacobiStiffness;
-		extern float Chaos_PBDCollisionSolver_JacobiPositionTolerance;
-		extern float Chaos_PBDCollisionSolver_JacobiRotationTolerance;
-	}
-
 	namespace Private
 	{
+		template<int TNumLanes>
+		using TSolverBodyPtrSimd = TSimdValue<FSolverBody*, TNumLanes>;
+
+		template<int TNumLanes>
+		struct TSolverBodyPtrPairSimd
+		{
+			TSolverBodyPtrSimd<TNumLanes> Body0;
+			TSolverBodyPtrSimd<TNumLanes> Body1;
+		};
+
 		/**
-		 * @brief A single contact point in a FPBDCollisionSolverSimd
-		 * @note Internal Chaos class subject to API changes, renaming or removal
+		 * @brief A SIMD row of contact points from a set of FPBDCollisionSolverSimd
 		*/
-		class FPBDCollisionSolverManifoldPointSimd
+		template<int TNumLanes>
+		class TPBDCollisionSolverManifoldPointsSimd
 		{
 		public:
+			using FSimdVec3f = TSimdVec3f<TNumLanes>;
+			using FSimdRealf = TSimdRealf<TNumLanes>;
+			using FSimdInt32 = TSimdInt32<TNumLanes>;
+			using FSimdSelector = TSimdSelector<TNumLanes>;
+			using FSimdSolverBodyPtr = TSolverBodyPtrSimd<TNumLanes>;
+
+			TPBDCollisionSolverManifoldPointsSimd()
+			{
+				SimdConstraintIndex.SetValues(INDEX_NONE);
+				SimdManifoldPointIndex.SetValues(INDEX_NONE);
+			}
+
+			void SetSharedData(
+				const int32 LaneIndex,
+				const int32 InConstraintIndex,
+				const int32 InManifoldPointIndex)
+			{
+				SimdConstraintIndex.SetValue(LaneIndex, InConstraintIndex);
+				SimdManifoldPointIndex.SetValue(LaneIndex, InManifoldPointIndex);
+			}
+
 			/**
-			 * @brief Initialize the geometric data for the contact
+			 * @brief Initialize contact geometry (contact positions, normal, etc)
 			*/
 			void InitContact(
-				const FSolverReal Dt,
-				const FConstraintSolverBody& Body0,
-				const FConstraintSolverBody& Body1);
+				const int32 LaneIndex,
+				const FSolverVec3& InRelativeContactPosition0,
+				const FSolverVec3& InRelativeContactPosition1,
+				const FSolverVec3& InWorldContactNormal,
+				const FSolverVec3& InWorldContactTangentU,
+				const FSolverVec3& InWorldContactTangentV,
+				const FSolverReal InWorldContactDeltaNormal,
+				const FSolverReal InWorldContactDeltaTangentU,
+				const FSolverReal InWorldContactDeltaTangentV,
+				const FSolverReal InWorldContactVelocityTargetNormal,
+				const FSolverReal InStiffness,
+				const FSolverReal InStaticFriction,
+				const FSolverReal InDynamicFriction,
+				const FSolverReal InVelocityFriction)
+			{
+				SimdRelativeContactPoint0.SetValue(LaneIndex, InRelativeContactPosition0);
+				SimdRelativeContactPoint1.SetValue(LaneIndex, InRelativeContactPosition1);
+				SimdContactNormal.SetValue(LaneIndex, InWorldContactNormal);
+				SimdContactTangentU.SetValue(LaneIndex, InWorldContactTangentU);
+				SimdContactTangentV.SetValue(LaneIndex, InWorldContactTangentV);
+				SimdContactDeltaNormal.SetValue(LaneIndex, InWorldContactDeltaNormal);
+				SimdContactDeltaTangentU.SetValue(LaneIndex, InWorldContactDeltaTangentU);
+				SimdContactDeltaTangentV.SetValue(LaneIndex, InWorldContactDeltaTangentV);
+				SimdContactTargetVelocityNormal.SetValue(LaneIndex, InWorldContactVelocityTargetNormal);
+
+				SimdStiffness.SetValue(LaneIndex, InStiffness);
+				SimdStaticFriction.SetValue(LaneIndex, InStaticFriction);
+				SimdDynamicFriction.SetValue(LaneIndex, InDynamicFriction);
+				SimdVelocityFriction.SetValue(LaneIndex, InVelocityFriction);
+			}
 
 			/**
-			 * @brief Initialize the material related properties of the contact
+			 * @brief Initialize contact state (net impulses, masses, etc). Must come after InitContact
+			 * because the mass properties depend on contact offsets etc.
 			*/
-			//void InitMaterial(
-			//	const FConstraintSolverBody& Body0,
-			//	const FConstraintSolverBody& Body1,
-			//	const FSolverReal InRestitution,
-			//	const FSolverReal InRestitutionVelocityThreshold);
+			void FinalizeContact(
+				const int32 LaneIndex,
+				const FConstraintSolverBody& Body0,
+				const FConstraintSolverBody& Body1)
+			{
+				SimdNetPushOutNormal.SetValue(LaneIndex, 0);
+				SimdNetPushOutTangentU.SetValue(LaneIndex, 0);
+				SimdNetPushOutTangentV.SetValue(LaneIndex, 0);
+				SimdNetImpulseNormal.SetValue(LaneIndex, 0);
+				SimdNetImpulseTangentU.SetValue(LaneIndex, 0);
+				SimdNetImpulseTangentV.SetValue(LaneIndex, 0);
+				SimdStaticFrictionRatio.SetValue(LaneIndex, 0);
+
+				UpdateMass(LaneIndex, Body0, Body1);
+			}
 
 			/**
 			 * @brief Update the cached mass properties based on the current body transforms
 			*/
-			void UpdateMass(const FConstraintSolverBody& Body0, const FConstraintSolverBody& Body1);
+			void UpdateMass(
+				const int32 LaneIndex,
+				const FConstraintSolverBody& Body0,
+				const FConstraintSolverBody& Body1)
+			{
+				FSolverReal ContactMassInvNormal = FSolverReal(0);
+				FSolverReal ContactMassInvTangentU = FSolverReal(0);
+				FSolverReal ContactMassInvTangentV = FSolverReal(0);
+
+				SimdInvM0.SetValue(LaneIndex, 0);
+				SimdInvM1.SetValue(LaneIndex, 0);
+				SimdContactNormalAngular0.SetValue(LaneIndex, FVec3f(0));
+				SimdContactTangentUAngular0.SetValue(LaneIndex, FVec3f(0));
+				SimdContactTangentVAngular0.SetValue(LaneIndex, FVec3f(0));
+				SimdContactNormalAngular1.SetValue(LaneIndex, FVec3f(0));
+				SimdContactTangentUAngular1.SetValue(LaneIndex, FVec3f(0));
+				SimdContactTangentVAngular1.SetValue(LaneIndex, FVec3f(0));
+
+				const FVec3f ContactNormal = SimdContactNormal.GetValue(LaneIndex);
+				const FVec3f ContactTangentU = SimdContactTangentU.GetValue(LaneIndex);
+				const FVec3f ContactTangentV = SimdContactTangentV.GetValue(LaneIndex);
 
-			/**
-			 * @brief Update the contact mass for the normal correction
-			 * This is used by shock propagation.
-			*/
-			void UpdateMassNormal(const FConstraintSolverBody& Body0, const FConstraintSolverBody& Body1);
+				if (Body0.IsDynamic())
+				{
+					const FVec3f RelativeContactPoint0 = SimdRelativeContactPoint0.GetValue(LaneIndex);
 
-			/**
-			 * @brief Calculate the position error at the current transforms
-			 * @param MaxPushOut a limit on the position error for this iteration to prevent initial-penetration explosion (a common PBD problem)
-			*/
-			void CalculateContactPositionErrorNormal(const FConstraintSolverBody& Body0, const FConstraintSolverBody& Body1, const FSolverReal MaxPushOut, FSolverReal& OutContactDeltaNormal) const;
-			void CalculateContactPositionErrorTangential(const FConstraintSolverBody& Body0, const FConstraintSolverBody& Body1, FSolverReal& OutContactDeltaTanget0, FSolverReal& OutContactDeltaTangent1) const;
+					const FSolverVec3 R0xN = FSolverVec3::CrossProduct(RelativeContactPoint0, ContactNormal);
+					const FSolverVec3 R0xU = FSolverVec3::CrossProduct(RelativeContactPoint0, ContactTangentU);
+					const FSolverVec3 R0xV = FSolverVec3::CrossProduct(RelativeContactPoint0, ContactTangentV);
 
-			/**
-			 * @brief Calculate the velocity error at the current transforms
-			*/
-			void CalculateContactVelocityErrorNormal(const FConstraintSolverBody& Body0, const FConstraintSolverBody& Body1, FSolverReal& OutContactVelocityDeltaNormal) const;
-			void CalculateContactVelocityError(const FConstraintSolverBody& Body0, const FConstraintSolverBody& Body1, const FSolverReal DynamicFriction, const FSolverReal Dt, FSolverReal& OutContactVelocityDeltaNormal, FSolverReal& OutContactVelocityDeltaTangent0, FSolverReal& OutContactVelocityDeltaTangent1) const;
+					const FSolverMatrix33 InvI0 = Body0.InvI();
+					const FSolverVec3 IR0xN = InvI0 * R0xN;
+					const FSolverVec3 IR0xU = InvI0 * R0xU;
+					const FSolverVec3 IR0xV = InvI0 * R0xV;
 
+					SimdInvM0.SetValue(LaneIndex, Body0.InvM());
 
-			FORCEINLINE_DEBUGGABLE void CalculatePositionCorrectionNormal(
-				const FSolverReal Stiffness,
-				const FSolverReal ContactDeltaNormal,
-				FSolverReal& OutPushOutNormal)
-			{
-				const FSolverReal PushOutNormal = -Stiffness * ContactDeltaNormal * ContactMassNormal;
+					SimdContactNormalAngular0.SetValue(LaneIndex, IR0xN);
+					SimdContactTangentUAngular0.SetValue(LaneIndex, IR0xU);
+					SimdContactTangentVAngular0.SetValue(LaneIndex, IR0xV);
 
-				// The total pushout so far this sub-step
-				// Unilateral constraint: Net-negative impulses not allowed (negative incremental impulses are allowed as long as the net is positive)
-				if ((NetPushOutNormal + PushOutNormal) > FSolverReal(0))
-				{
-					OutPushOutNormal = PushOutNormal;
+					ContactMassInvNormal += FSolverVec3::DotProduct(R0xN, IR0xN) + Body0.InvM();
+					ContactMassInvTangentU += FSolverVec3::DotProduct(R0xU, IR0xU) + Body0.InvM();
+					ContactMassInvTangentV += FSolverVec3::DotProduct(R0xV, IR0xV) + Body0.InvM();
 				}
-				else
+				if (Body1.IsDynamic())
 				{
-					OutPushOutNormal = -NetPushOutNormal;
+					const FVec3f RelativeContactPoint1 = SimdRelativeContactPoint1.GetValue(LaneIndex);
+
+					const FSolverVec3 R1xN = FSolverVec3::CrossProduct(RelativeContactPoint1, ContactNormal);
+					const FSolverVec3 R1xU = FSolverVec3::CrossProduct(RelativeContactPoint1, ContactTangentU);
+					const FSolverVec3 R1xV = FSolverVec3::CrossProduct(RelativeContactPoint1, ContactTangentV);
+
+					const FSolverMatrix33 InvI1 = Body1.InvI();
+					const FSolverVec3 IR1xN = InvI1 * R1xN;
+					const FSolverVec3 IR1xU = InvI1 * R1xU;
+					const FSolverVec3 IR1xV = InvI1 * R1xV;
+
+					SimdInvM1.SetValue(LaneIndex, Body1.InvM());
+
+					SimdContactNormalAngular1.SetValue(LaneIndex, IR1xN);
+					SimdContactTangentUAngular1.SetValue(LaneIndex, IR1xU);
+					SimdContactTangentVAngular1.SetValue(LaneIndex, IR1xV);
+
+					ContactMassInvNormal += FSolverVec3::DotProduct(R1xN, IR1xN) + Body1.InvM();
+					ContactMassInvTangentU += FSolverVec3::DotProduct(R1xU, IR1xU) + Body1.InvM();
+					ContactMassInvTangentV += FSolverVec3::DotProduct(R1xV, IR1xV) + Body1.InvM();
 				}
+
+				SimdContactMassNormal.SetValue(LaneIndex, (ContactMassInvNormal > FSolverReal(UE_SMALL_NUMBER)) ? FSolverReal(1) / ContactMassInvNormal : FSolverReal(0));
+				SimdContactMassTangentU.SetValue(LaneIndex, (ContactMassInvTangentU > FSolverReal(UE_SMALL_NUMBER)) ? FSolverReal(1) / ContactMassInvTangentU : FSolverReal(0));
+				SimdContactMassTangentV.SetValue(LaneIndex, (ContactMassInvTangentV > FSolverReal(UE_SMALL_NUMBER)) ? FSolverReal(1) / ContactMassInvTangentV : FSolverReal(0));
 			}
 
-			FORCEINLINE_DEBUGGABLE void CalculatePositionCorrectionTangents(
-				const FSolverReal Stiffness,
-				const FSolverReal ContactDeltaTangentU,
-				const FSolverReal ContactDeltaTangentV,
-				FSolverReal& OutPushOutTangentU,
-				FSolverReal& OutPushOutTangentV)
+			/**
+			 * @brief Update the contact mass for the normal correction
+			 * This is used by shock propagation.
+			*/
+			void UpdateMassNormal(
+				const int32 LaneIndex,
+				FConstraintSolverBody& Body0,
+				FConstraintSolverBody& Body1)
 			{
-				// Bilateral constraint - negative values allowed (unlike the normal correction)
-				OutPushOutTangentU = -Stiffness * ContactMassTangentU * ContactDeltaTangentU;
-				OutPushOutTangentV = -Stiffness * ContactMassTangentV * ContactDeltaTangentV;
 			}
 
-
-			FORCEINLINE_DEBUGGABLE void ApplyFrictionCone(
-				const FSolverReal StaticFriction,
-				const FSolverReal DynamicFriction,
-				const FSolverReal MaxFrictionPushOut,
-				FSolverReal& InOutPushOutTangentU,
-				FSolverReal& InOutPushOutTangentV)
+			FORCEINLINE void GatherBodyCorrections(
+				const FSimdSelector& IsValid,
+				const FSimdSolverBodyPtr& Body0,
+				const FSimdSolverBodyPtr& Body1,
+				FSimdVec3f& DP0,
+				FSimdVec3f& DQ0,
+				FSimdVec3f& DP1,
+				FSimdVec3f& DQ1)
 			{
-				// Assume we stay in the friction cone...
-				StaticFrictionRatio = FSolverReal(1);
-
-				if (MaxFrictionPushOut < FSolverReal(UE_KINDA_SMALL_NUMBER))
+				for (int32 LaneIndex = 0; LaneIndex < TNumLanes; ++LaneIndex)
 				{
-					// Note: we have already added the current iteration's PushOut to the NetPushOut but it has not been applied to the body
-					// so we must subtract it again to calculate the actual pushout we want to undo (i.e., the net pushout that has been applied 
-					// to the body so far from previous iterations)
-					InOutPushOutTangentU = -(NetPushOutTangentU - InOutPushOutTangentU);
-					InOutPushOutTangentV = -(NetPushOutTangentV - InOutPushOutTangentV);
-					NetPushOutTangentU = FSolverReal(0);
-					NetPushOutTangentV = FSolverReal(0);
-					StaticFrictionRatio = FSolverReal(0);
+					if (IsValid.GetValue(LaneIndex))
+					{
+						const FSolverBody* LaneBody0 = Body0.GetValue(LaneIndex);
+						const FSolverBody* LaneBody1 = Body1.GetValue(LaneIndex);
+						DP0.SetValue(LaneIndex, LaneBody0->DP());
+						DQ0.SetValue(LaneIndex, LaneBody0->DQ());
+						DP1.SetValue(LaneIndex, LaneBody1->DP());
+						DQ1.SetValue(LaneIndex, LaneBody1->DQ());
+					}
 				}
-				else
+			}
+
+			FORCEINLINE void ScatterBodyCorrections(
+				const FSimdSelector& IsValid,
+				const FSimdVec3f& DP0,
+				const FSimdVec3f& DQ0,
+				const FSimdVec3f& DP1,
+				const FSimdVec3f& DQ1,
+				const FSimdSolverBodyPtr& Body0,
+				const FSimdSolverBodyPtr& Body1)
+			{
+				for (int32 LaneIndex = 0; LaneIndex < TNumLanes; ++LaneIndex)
 				{
-					// If we exceed the static friction cone, clip to the dynamic friction cone
-					const FSolverReal MaxStaticPushOutTangentSq = FMath::Square(StaticFriction * MaxFrictionPushOut);
-					const FSolverReal NetPushOutTangentSq = FMath::Square(NetPushOutTangentU) + FMath::Square(NetPushOutTangentV);
-					if (NetPushOutTangentSq > MaxStaticPushOutTangentSq)
+					if (IsValid.GetValue(LaneIndex))
 					{
-						const FSolverReal MaxDynamicPushOutTangent = DynamicFriction * MaxFrictionPushOut;
-						const FSolverReal FrictionMultiplier = MaxDynamicPushOutTangent * FMath::InvSqrt(NetPushOutTangentSq);
-						const FSolverReal NewNetPushOutTangentU = FrictionMultiplier * NetPushOutTangentU;
-						const FSolverReal NewNetPushOutTangentV = FrictionMultiplier * NetPushOutTangentV;
-						InOutPushOutTangentU = NewNetPushOutTangentU - (NetPushOutTangentU - InOutPushOutTangentU);
-						InOutPushOutTangentV = NewNetPushOutTangentV - (NetPushOutTangentV - InOutPushOutTangentV);
-						NetPushOutTangentU = NewNetPushOutTangentU;
-						NetPushOutTangentV = NewNetPushOutTangentV;
-						StaticFrictionRatio = FrictionMultiplier;
+						FSolverBody* LaneBody0 = Body0.GetValue(LaneIndex);
+						FSolverBody* LaneBody1 = Body1.GetValue(LaneIndex);
+						LaneBody0->SetDP(DP0.GetValue(LaneIndex));
+						LaneBody0->SetDQ(DQ0.GetValue(LaneIndex));
+						LaneBody1->SetDP(DP1.GetValue(LaneIndex));
+						LaneBody1->SetDQ(DQ1.GetValue(LaneIndex));
 					}
 				}
 			}
 
-			FORCEINLINE_DEBUGGABLE void ApplyPositionCorrectionTangential(
-				const FSolverReal Stiffness,
-				const FSolverReal StaticFriction,
-				const FSolverReal DynamicFriction,
-				const FSolverReal MaxFrictionPushOut,
-				const FSolverReal ContactDeltaTangentU,
-				const FSolverReal ContactDeltaTangentV,
-				FConstraintSolverBody& Body0,
-				FConstraintSolverBody& Body1)
+			void SolvePositionNoFriction(
+				const FSimdSolverBodyPtr& Body0,
+				const FSimdSolverBodyPtr& Body1,
+				const FSimd4Realf& MaxPushOut)
 			{
-				FSolverReal PushOutTangentU = FSolverReal(0);
-				FSolverReal PushOutTangentV = FSolverReal(0);
+				// Which lanes a have a valid constraint setup?
+				const FSimdSelector IsValid = SimdGreaterEqual(SimdConstraintIndex, FSimdInt32::Zero());
+
+				// Get the current corrections for each body. 
+				// NOTE: This is a gather operation
+				FSimdVec3f DP0, DQ0, DP1, DQ1;
+				GatherBodyCorrections(IsValid, Body0, Body1, DP0, DQ0, DP1, DQ1);
+
+				// Calculate the contact error
+				const FSimdVec3f DQ0xR0 = SimdCrossProduct(DQ0, SimdRelativeContactPoint0);
+				const FSimdVec3f DQ1xR1 = SimdCrossProduct(DQ1, SimdRelativeContactPoint1);
+				const FSimdVec3f ContactDelta0 = SimdAdd(DP0, DQ0xR0);
+				const FSimdVec3f ContactDelta1 = SimdAdd(DP1, DQ1xR1);
+				const FSimdVec3f ContactDelta = SimdSubtract(ContactDelta0, ContactDelta1);
+				FSimdRealf ContactErrorNormal = SimdAdd(SimdContactDeltaNormal, SimdDotProduct(ContactDelta, SimdContactNormal));
+
+				// Apply MaxPushOut clamping if required
+				//if ((MaxPushOut > 0) && (ContactErrorNormal < -MaxPushOut)) { ContactErrorNormal = -MaxPushOut; }
+				const FSimdRealf NegMaxPushOut = SimdNegate(MaxPushOut);
+				const FSimdSelector ShouldClampError = SimdAnd(SimdLess(NegMaxPushOut, FSimdRealf::Zero()), SimdLess(ContactErrorNormal, NegMaxPushOut));
+				ContactErrorNormal = SimdSelect(ShouldClampError, NegMaxPushOut, ContactErrorNormal);
+
+				// Determine which lanes to process: only those with an overlap or that applied a pushout on a prior iteration
+				const FSimdSelector IsErrorNegative = SimdLess(ContactErrorNormal, FSimdRealf::Zero());
+				const FSimdSelector IsNetPushOutPositive = SimdGreater(SimdNetPushOutNormal, FSimdRealf::Zero());
+				const FSimdSelector ShouldProcess = SimdAnd(IsValid, SimdOr(IsErrorNegative, IsNetPushOutPositive));
+
+				// If all lanes are to be skipped, early-out
+				if (!SimdAnyTrue(ShouldProcess))
+				{
+					return;
+				}
 
-				CalculatePositionCorrectionTangents(
-					Stiffness,
-					ContactDeltaTangentU,
-					ContactDeltaTangentV,
-					PushOutTangentU,					// Out
-					PushOutTangentV);					// Out
+				// Zero out the error for points we should not process so we don't apply a correction for them
+				ContactErrorNormal = SimdSelect(ShouldProcess, ContactErrorNormal, FSimdRealf::Zero());
 
-				NetPushOutTangentU += PushOutTangentU;
-				NetPushOutTangentV += PushOutTangentV;
+				FSimdRealf PushOutNormal = SimdNegate(SimdMultiply(ContactErrorNormal, SimdMultiply(SimdStiffness, SimdContactMassNormal)));
 
-				ApplyFrictionCone(
-					StaticFriction,
-					DynamicFriction,
-					MaxFrictionPushOut,
-					PushOutTangentU,					// InOut
-					PushOutTangentV);					// InOut
+				// Unilateral constraint: Net-negative pushout is not allowed, but
+				// PushOutNormal may be negative on any iteration as long as the net is positive
+				// If the net goes negative, apply a pushout to make it zero
+				const FSimdSelector IsPositive = SimdGreater(SimdAdd(SimdNetPushOutNormal, PushOutNormal), FSimdRealf::Zero());
+				PushOutNormal = SimdSelect(IsPositive, PushOutNormal, SimdNegate(SimdNetPushOutNormal));
 
-				// Update the particle state based on the pushout
-				const FVec3 PushOut = PushOutTangentU * WorldContact.ContactTangentU + PushOutTangentV * WorldContact.ContactTangentV;
-				if (Body0.IsDynamic())
-				{
-					const FSolverVec3 DX0 = Body0.InvM() * PushOut;
-					const FSolverVec3 DR0 = WorldContactTangentUAngular0 * PushOutTangentU + WorldContactTangentVAngular0 * PushOutTangentV;
-					Body0.ApplyPositionDelta(DX0);
-					Body0.ApplyRotationDelta(DR0);
-				}
-				if (Body1.IsDynamic())
-				{
-					const FSolverVec3 DX1 = -Body1.InvM() * PushOut;
-					const FSolverVec3 DR1 = WorldContactTangentUAngular1 * -PushOutTangentU + WorldContactTangentVAngular1 * -PushOutTangentV;
-					Body1.ApplyPositionDelta(DX1);
-					Body1.ApplyRotationDelta(DR1);
-				}
-			}
+				// New net pushout
+				SimdNetPushOutNormal = SimdAdd(PushOutNormal, SimdNetPushOutNormal);
 
-			FORCEINLINE_DEBUGGABLE void ApplyPositionCorrectionNormal(
-				const FSolverReal Stiffness,
-				const FSolverReal ContactDeltaNormal,
-				FConstraintSolverBody& Body0,
-				FConstraintSolverBody& Body1)
-			{
-				FSolverReal PushOutNormal = FSolverReal(0);
+				// Convert the positional impulse into position and rotation corrections for each body
+				const FSimdVec3f DDP0 = SimdMultiply(SimdContactNormal, SimdMultiply(SimdInvM0, PushOutNormal));
+				const FSimdVec3f DDQ0 = SimdMultiply(SimdContactNormalAngular0, PushOutNormal);
+				const FSimdVec3f DDP1 = SimdMultiply(SimdContactNormal, SimdMultiply(SimdInvM1, PushOutNormal));
+				const FSimdVec3f DDQ1 = SimdMultiply(SimdContactNormalAngular1, PushOutNormal);
 
-				CalculatePositionCorrectionNormal(
-					Stiffness,
-					ContactDeltaNormal,
-					PushOutNormal);						// Out
+				DP0 = SimdAdd(DP0, DDP0);
+				DQ0 = SimdAdd(DQ0, DDQ0);
+				DP1 = SimdSubtract(DP1, DDP1);
+				DQ1 = SimdSubtract(DQ1, DDQ1);
 
-				NetPushOutNormal += PushOutNormal;
+				// Update the corrections on the bodies. 
+				// NOTE: This is a scatter operation
+				ScatterBodyCorrections(IsValid, DP0, DQ0, DP1, DQ1, Body0, Body1);
+			}
 
-				// Update the particle state based on the pushout
-				if (Body0.IsDynamic())
-				{
-					const FSolverVec3 DX0 = (Body0.InvM() * PushOutNormal) * WorldContact.ContactNormal;
-					const FSolverVec3 DR0 = WorldContactNormalAngular0 * PushOutNormal;
-					Body0.ApplyPositionDelta(DX0);
-					Body0.ApplyRotationDelta(DR0);
-				}
-				if (Body1.IsDynamic())
+			void SolvePositionWithFriction(
+				const FSimdSolverBodyPtr& Body0,
+				const FSimdSolverBodyPtr& Body1,
+				const FSimd4Realf& MaxPushOut,
+				const FSimd4Realf& FrictionStiffnessScale)
+			{
+				// Which lanes a have a valid constraint setup?
+				const FSimdSelector IsValid = SimdGreaterEqual(SimdConstraintIndex, FSimdInt32::Zero());
+
+				// Get the current corrections for each body. 
+				// NOTE: This is a gather operation
+				FSimdVec3f DP0, DQ0, DP1, DQ1;
+				GatherBodyCorrections(IsValid, Body0, Body1, DP0, DQ0, DP1, DQ1);
+
+				// Calculate the contact error
+				const FSimdVec3f DQ0xR0 = SimdCrossProduct(DQ0, SimdRelativeContactPoint0);
+				const FSimdVec3f DQ1xR1 = SimdCrossProduct(DQ1, SimdRelativeContactPoint1);
+				const FSimdVec3f ContactDelta0 = SimdAdd(DP0, DQ0xR0);
+				const FSimdVec3f ContactDelta1 = SimdAdd(DP1, DQ1xR1);
+				const FSimdVec3f ContactDelta = SimdSubtract(ContactDelta0, ContactDelta1);
+				FSimdRealf ContactErrorNormal = SimdAdd(SimdContactDeltaNormal, SimdDotProduct(ContactDelta, SimdContactNormal));
+
+				// Apply MaxPushOut clamping if required
+				//if ((MaxPushOut > 0) && (ContactErrorNormal < -MaxPushOut)) { ContactErrorNormal = -MaxPushOut; }
+				const FSimdRealf NegMaxPushOut = SimdNegate(MaxPushOut);
+				const FSimdSelector ShouldClampError = SimdAnd(SimdLess(NegMaxPushOut, FSimdRealf::Zero()), SimdLess(ContactErrorNormal, NegMaxPushOut));
+				ContactErrorNormal = SimdSelect(ShouldClampError, NegMaxPushOut, ContactErrorNormal);
+
+				// Determine which lanes to process: only those with an overlap or that applied a pushout on a prior iteration
+				const FSimdSelector IsErrorNegative = SimdLess(ContactErrorNormal, FSimdRealf::Zero());
+				const FSimdSelector IsNetPushOutPositive = SimdGreater(SimdNetPushOutNormal, FSimdRealf::Zero());
+				const FSimdSelector ShouldProcess = SimdAnd(IsValid, SimdOr(IsErrorNegative, IsNetPushOutPositive));
+
+				// If all lanes are to be skipped, early-out
+				//if (SimdAnyTrue(ShouldProcess))
 				{
-					const FSolverVec3 DX1 = (Body1.InvM() * -PushOutNormal) * WorldContact.ContactNormal;
-					const FSolverVec3 DR1 = WorldContactNormalAngular1 * -PushOutNormal;
-					Body1.ApplyPositionDelta(DX1);
-					Body1.ApplyRotationDelta(DR1);
+					// Zero out the error for points we should not process so we don't apply a correction for them
+					ContactErrorNormal = SimdSelect(ShouldProcess, ContactErrorNormal, FSimdRealf::Zero());
+
+					FSimdRealf PushOutNormal = SimdNegate(SimdMultiply(SimdStiffness, SimdMultiply(ContactErrorNormal, SimdContactMassNormal)));
+
+					// Unilateral constraint: Net-negative pushout is not allowed, but
+					// PushOutNormal may be negative on any iteration as long as the net is positive
+					// If the net goes negative, apply a pushout to make it zero
+					const FSimdSelector IsPositive = SimdGreater(SimdAdd(SimdNetPushOutNormal, PushOutNormal), FSimdRealf::Zero());
+					PushOutNormal = SimdSelect(IsPositive, PushOutNormal, SimdNegate(SimdNetPushOutNormal));
+
+					// New net normal pushout
+					SimdNetPushOutNormal = SimdAdd(PushOutNormal, SimdNetPushOutNormal);
+					SimdStaticFrictionRatio = FSimdRealf::One();
+
+					// Convert the positional impulse into position and rotation corrections for each body
+					FSimdVec3f DDP0 = SimdMultiply(SimdContactNormal, SimdMultiply(SimdInvM0, PushOutNormal));
+					FSimdVec3f DDQ0 = SimdMultiply(SimdContactNormalAngular0, PushOutNormal);
+					FSimdVec3f DDP1 = SimdMultiply(SimdContactNormal, SimdMultiply(SimdInvM1, PushOutNormal));
+					FSimdVec3f DDQ1 = SimdMultiply(SimdContactNormalAngular1, PushOutNormal);
+
+					// Should we apply tangential corrections for friction? 
+					// bUpdateFriction = ((NetPushOutNormal > 0) || (NetPushOutTangentU != 0) || (NetPushOutTangentV != 0))
+					const FSimdSelector HasNormalPushout = SimdGreater(SimdNetPushOutNormal, FSimdRealf::Zero());
+					const FSimdSelector HasTangentUPushout = SimdNotEqual(SimdNetPushOutTangentU, FSimdRealf::Zero());
+					const FSimdSelector HasTangentVPushout = SimdNotEqual(SimdNetPushOutTangentV, FSimdRealf::Zero());
+					//const FSimdSelector UpdateFriction = SimdAnd(IsValid, SimdOr(HasNormalPushout, SimdOr(HasTangentUPushout, HasTangentVPushout)));
+					//if (SimdAnyTrue(UpdateFriction))
+					{
+						// Calculate tangential errors
+						const FSimdRealf ContactErrorTangentU = SimdAdd(SimdContactDeltaTangentU, SimdDotProduct(ContactDelta, SimdContactTangentU));
+						const FSimdRealf ContactErrorTangentV = SimdAdd(SimdContactDeltaTangentV, SimdDotProduct(ContactDelta, SimdContactTangentV));
+
+						// Calculate tangential correction
+						const FSimdRealf FrictionStiffness = SimdMultiply(FrictionStiffnessScale, SimdStiffness);
+						FSimdRealf PushOutTangentU = SimdNegate(SimdMultiply(FrictionStiffness, SimdMultiply(ContactErrorTangentU, SimdContactMassTangentU)));
+						FSimdRealf PushOutTangentV = SimdNegate(SimdMultiply(FrictionStiffness, SimdMultiply(ContactErrorTangentV, SimdContactMassTangentV)));
+
+						// New net tangential pushouts
+						FSimdRealf NetPushOutTangentU = SimdAdd(SimdNetPushOutTangentU, PushOutTangentU);
+						FSimdRealf NetPushOutTangentV = SimdAdd(SimdNetPushOutTangentV, PushOutTangentV);
+						FSimdRealf StaticFrictionRatio = SimdStaticFrictionRatio;
+
+						// If a lane has a zero normal pushout, we need to undo all previously applied friction
+						//if (!SimdAllTrue(HasNormalPushout))
+						{
+							PushOutTangentU = SimdSelect(HasNormalPushout, PushOutTangentU, SimdNegate(SimdNetPushOutTangentU));
+							PushOutTangentV = SimdSelect(HasNormalPushout, PushOutTangentV, SimdNegate(SimdNetPushOutTangentV));
+							NetPushOutTangentU = SimdSelect(HasNormalPushout, NetPushOutTangentU, FSimdRealf::Zero());
+							NetPushOutTangentV = SimdSelect(HasNormalPushout, NetPushOutTangentV, FSimdRealf::Zero());
+							StaticFrictionRatio = SimdSelect(HasNormalPushout, StaticFrictionRatio, FSimdRealf::Zero());
+						}
+
+						// Apply cone limits to tangential pushouts on lanes that exceed the limit
+						// NOTE: if HasNormalPushout is false in any lane, we have already zeroed the net pushout and don't need to do anything here
+						const FSimdSelector ApplyFrictionCone = SimdAnd(IsValid, SimdAnd(HasNormalPushout, SimdOr(HasTangentUPushout, HasTangentVPushout)));
+						//if (SimdAnyTrue(ApplyFrictionCone))
+						{
+							const FSimdRealf MaxPushOutTangent = SimdNetPushOutNormal;
+							const FSimdRealf MaxStaticPushOutTangentSq = SimdSquare(SimdMultiply(SimdStaticFriction, MaxPushOutTangent));
+							const FSimdRealf NetPushOutTangentSq = SimdAdd(SimdSquare(NetPushOutTangentU), SimdSquare(NetPushOutTangentV));
+							const FSimdSelector ExceededFrictionCone = SimdAnd(ApplyFrictionCone, SimdGreater(NetPushOutTangentSq, MaxStaticPushOutTangentSq));
+							//if (SimdAnyTrue(ExceededFrictionCone))
+							{
+								const FSimdRealf MaxDynamicPushOutTangent = SimdMultiply(SimdDynamicFriction, MaxPushOutTangent);
+								const FSimdRealf FrictionMultiplier = SimdMultiply(MaxDynamicPushOutTangent, SimdInvSqrt(NetPushOutTangentSq));
+								const FSimdRealf ClampedNetPushOutTangentU = SimdMultiply(FrictionMultiplier, NetPushOutTangentU);
+								const FSimdRealf ClampedNetPushOutTangentV = SimdMultiply(FrictionMultiplier, NetPushOutTangentV);
+								const FSimdRealf ClampedPushOutTangentU = SimdSubtract(ClampedNetPushOutTangentU, SimdNetPushOutTangentU);
+								const FSimdRealf ClampedPushOutTangentV = SimdSubtract(ClampedNetPushOutTangentV, SimdNetPushOutTangentV);
+
+								PushOutTangentU = SimdSelect(ExceededFrictionCone, ClampedPushOutTangentU, PushOutTangentU);
+								PushOutTangentV = SimdSelect(ExceededFrictionCone, ClampedPushOutTangentV, PushOutTangentV);
+								NetPushOutTangentU = SimdSelect(ExceededFrictionCone, ClampedNetPushOutTangentU, NetPushOutTangentU);
+								NetPushOutTangentV = SimdSelect(ExceededFrictionCone, ClampedNetPushOutTangentV, NetPushOutTangentV);
+								StaticFrictionRatio = SimdSelect(ExceededFrictionCone, FrictionMultiplier, StaticFrictionRatio);
+							}
+						}
+
+						SimdNetPushOutTangentU = NetPushOutTangentU;
+						SimdNetPushOutTangentV = NetPushOutTangentV;
+						SimdStaticFrictionRatio = StaticFrictionRatio;
+
+						// Add the tangential corrections to the applied correction
+						DDP0 = SimdAdd(DDP0, SimdMultiply(SimdContactTangentU, SimdMultiply(SimdInvM0, PushOutTangentU)));
+						DDP0 = SimdAdd(DDP0, SimdMultiply(SimdContactTangentV, SimdMultiply(SimdInvM0, PushOutTangentV)));
+						DDQ0 = SimdAdd(DDQ0, SimdMultiply(SimdContactTangentUAngular0, PushOutTangentU));
+						DDQ0 = SimdAdd(DDQ0, SimdMultiply(SimdContactTangentVAngular0, PushOutTangentV));
+						DDP1 = SimdAdd(DDP1, SimdMultiply(SimdContactTangentU, SimdMultiply(SimdInvM1, PushOutTangentU)));
+						DDP1 = SimdAdd(DDP1, SimdMultiply(SimdContactTangentV, SimdMultiply(SimdInvM1, PushOutTangentV)));
+						DDQ1 = SimdAdd(DDQ1, SimdMultiply(SimdContactTangentUAngular1, PushOutTangentU));
+						DDQ1 = SimdAdd(DDQ1, SimdMultiply(SimdContactTangentVAngular1, PushOutTangentV));
+
+						DP0 = SimdAdd(DP0, DDP0);
+						DQ0 = SimdAdd(DQ0, DDQ0);
+						DP1 = SimdSubtract(DP1, DDP1);
+						DQ1 = SimdSubtract(DQ1, DDQ1);
+					}
+
+					// Update the corrections on the bodies. 
+					// NOTE: This is a scatter operation
+					ScatterBodyCorrections(IsValid, DP0, DQ0, DP1, DQ1, Body0, Body1);
 				}
 			}
 
-
-			FORCEINLINE_DEBUGGABLE void ApplyVelocityCorrection(
-				const FSolverReal Stiffness,
-				const FSolverReal Dt,
-				const FSolverReal DynamicFriction,
-				const FSolverReal ContactVelocityDeltaNormal,
-				const FSolverReal ContactVelocityDeltaTangent0,
-				const FSolverReal ContactVelocityDeltaTangent1,
-				const FSolverReal MinImpulseNormal,
-				FConstraintSolverBody& Body0,
-				FConstraintSolverBody& Body1)
+			void SolveVelocityNoFriction(
+				const FSimdSolverBodyPtr& Body0,
+				const FSimdSolverBodyPtr& Body1,
+				const FSimdRealf& Dt)
 			{
-				FSolverReal ImpulseNormal = -Stiffness * ContactMassNormal * ContactVelocityDeltaNormal;
-				FSolverReal ImpulseTangentU = -Stiffness * ContactMassTangentU * ContactVelocityDeltaTangent0;
-				FSolverReal ImpulseTangentV = -Stiffness * ContactMassTangentV * ContactVelocityDeltaTangent1;
+				// Which lanes a have a valid constraint setup?
+				const FSimdSelector IsValid = SimdGreaterEqual(SimdConstraintIndex, FSimdInt32::Zero());
 
-				// Clamp the total impulse to be positive along the normal. We can apply a net negative impulse, 
-				// but only to correct the velocity that was added by pushout (in which case MinImpulseNormal will be negative).
-				if ((NetImpulseNormal + ImpulseNormal) < MinImpulseNormal)
+				// Only lanes that applied a position correction or that started with an overlap require a velocity correction
+				FSimdSelector ShouldSolveVelocity = SimdOr(SimdGreater(SimdNetPushOutNormal, FSimdRealf::Zero()), SimdLess(SimdContactDeltaNormal, FSimdRealf::Zero()));
+				ShouldSolveVelocity = SimdAnd(IsValid, ShouldSolveVelocity);
+				if (!SimdAnyTrue(ShouldSolveVelocity))
 				{
-					// We are trying to apply a net negative impulse larger than one to counteract the effective pushout impulse
-					// so clamp the net impulse to be equal to minus the pushout impulse along the normal.
-					ImpulseNormal = MinImpulseNormal - NetImpulseNormal;
+					return;
 				}
 
-				// Clamp the tangential impulses to the friction cone
-				if ((DynamicFriction > 0) && (Dt > 0))
+				// Gather the body data we need
+				FSimdVec3f V0, W0, V1, W1;
+				for (int32 LaneIndex = 0; LaneIndex < TNumLanes; ++LaneIndex)
 				{
-					const FSolverReal MaxImpulseTangent = FMath::Max(FSolverReal(0), DynamicFriction * (NetImpulseNormal + ImpulseNormal + NetPushOutNormal / Dt));
-					const FSolverReal MaxImpulseTangentSq = FMath::Square(MaxImpulseTangent);
-					const FSolverReal ImpulseTangentSq = FMath::Square(ImpulseTangentU) + FMath::Square(ImpulseTangentV);
-					if (ImpulseTangentSq > (MaxImpulseTangentSq + UE_SMALL_NUMBER))
+					if (IsValid.GetValue(LaneIndex))
 					{
-						const FSolverReal ImpulseTangentScale = MaxImpulseTangent * FMath::InvSqrt(ImpulseTangentSq);
-						ImpulseTangentU *= ImpulseTangentScale;
-						ImpulseTangentV *= ImpulseTangentScale;
+						const FSolverBody* LaneBody0 = Body0.GetValue(LaneIndex);
+						const FSolverBody* LaneBody1 = Body1.GetValue(LaneIndex);
+						V0.SetValue(LaneIndex, LaneBody0->V());
+						W0.SetValue(LaneIndex, LaneBody0->W());
+						V1.SetValue(LaneIndex, LaneBody1->V());
+						W1.SetValue(LaneIndex, LaneBody1->W());
 					}
 				}
 
-				NetImpulseNormal += ImpulseNormal;
-				NetImpulseTangentU += ImpulseTangentU;
-				NetImpulseTangentV += ImpulseTangentV;
-
-				// Apply the velocity deltas from the impulse
-				const FSolverVec3 Impulse = ImpulseNormal * WorldContact.ContactNormal + ImpulseTangentU * WorldContact.ContactTangentU + ImpulseTangentV * WorldContact.ContactTangentV;
-				if (Body0.IsDynamic())
+				// Calculate the velocity error we need to correct
+				const FSimdVec3f ContactVelocity0 = SimdAdd(V0, SimdCrossProduct(W0, SimdRelativeContactPoint0));
+				const FSimdVec3f ContactVelocity1 = SimdAdd(V1, SimdCrossProduct(W1, SimdRelativeContactPoint1));
+				const FSimdVec3f ContactVelocity = SimdSubtract(ContactVelocity0, ContactVelocity1);
+				const FSimdRealf ContactVelocityNormal = SimdDotProduct(ContactVelocity, SimdContactNormal);
+				FSimdRealf ContactVelocityErrorNormal = SimdSubtract(ContactVelocityNormal, SimdContactTargetVelocityNormal);
+
+				// Make sure we don't add a correction to the lanes we should ignore
+				ContactVelocityErrorNormal = SimdSelect(ShouldSolveVelocity, ContactVelocityErrorNormal, FSimdRealf::Zero());
+
+				// Calculate the velocity correction for the error
+				FSimdRealf ImpulseNormal = SimdNegate(SimdMultiply(SimdMultiply(SimdStiffness, SimdContactMassNormal), ContactVelocityErrorNormal));
+
+				// The minimum normal impulse we can apply. We are allowed to apply a negative impulse 
+				// up to an amount that would conteract the implciit velocity applied by the pushout
+				// MinImpulseNormal = FMath::Min(0, -NetPushOutNormal / Dt)
+				// if (NetImpulseNormal + ImpulseNormal < MinImpulseNormal) {...}
+				const FSimdRealf MinImpulseNormal = SimdMin(SimdNegate(SimdDivide(SimdNetPushOutNormal, Dt)), FSimdRealf::Zero());
+				const FSimdSelector ShouldClampImpulse = SimdLess(SimdAdd(SimdNetImpulseNormal, ImpulseNormal), MinImpulseNormal);
+				const FSimdRealf ClampedImpulseNormal = SimdSubtract(MinImpulseNormal, SimdNetImpulseNormal);
+				ImpulseNormal = SimdSelect(ShouldClampImpulse, ClampedImpulseNormal, ImpulseNormal);
+
+				SimdNetImpulseNormal = SimdAdd(SimdNetImpulseNormal, ImpulseNormal);
+
+				const FSimdVec3f DV0 = SimdMultiply(SimdMultiply(ImpulseNormal, SimdInvM0), SimdContactNormal);
+				const FSimdVec3f DV1 = SimdMultiply(SimdMultiply(ImpulseNormal, SimdInvM1), SimdContactNormal);
+				const FSimdVec3f DW0 = SimdMultiply(ImpulseNormal, SimdContactNormalAngular0);
+				const FSimdVec3f DW1 = SimdMultiply(ImpulseNormal, SimdContactNormalAngular1);
+				V0 = SimdAdd(V0, DV0);
+				W0 = SimdAdd(W0, DW0);
+				V1 = SimdSubtract(V1, DV1);
+				W1 = SimdSubtract(W1, DW1);
+				for (int32 LaneIndex = 0; LaneIndex < TNumLanes; ++LaneIndex)
 				{
-					const FSolverVec3 DV0 = Body0.InvM() * Impulse;
-					const FSolverVec3 DW0 = WorldContactNormalAngular0 * ImpulseNormal + WorldContactTangentUAngular0 * ImpulseTangentU + WorldContactTangentVAngular0 * ImpulseTangentV;
-					Body0.ApplyVelocityDelta(DV0, DW0);
-				}
-				if (Body1.IsDynamic())
-				{
-					const FSolverVec3 DV1 = -Body1.InvM() * Impulse;
-					const FSolverVec3 DW1 = WorldContactNormalAngular1 * -ImpulseNormal + WorldContactTangentUAngular1 * -ImpulseTangentU + WorldContactTangentVAngular1 * -ImpulseTangentV;
-					Body1.ApplyVelocityDelta(DV1, DW1);
+					if (IsValid.GetValue(LaneIndex))
+					{
+						FSolverBody* LaneBody0 = Body0.GetValue(LaneIndex);
+						FSolverBody* LaneBody1 = Body1.GetValue(LaneIndex);
+						LaneBody0->SetV(V0.GetValue(LaneIndex));
+						LaneBody0->SetW(W0.GetValue(LaneIndex));
+						LaneBody1->SetV(V1.GetValue(LaneIndex));
+						LaneBody1->SetW(W1.GetValue(LaneIndex));
+					}
 				}
 			}
 
-			FORCEINLINE_DEBUGGABLE void ApplyVelocityCorrectionNormal(
-				const FSolverReal Stiffness,
-				const FSolverReal ContactVelocityDeltaNormal,
-				const FSolverReal MinImpulseNormal,
-				FConstraintSolverBody& Body0,
-				FConstraintSolverBody& Body1)
+			void SolveVelocityWithFriction(
+				const FSimdSolverBodyPtr& Body0,
+				const FSimdSolverBodyPtr& Body1,
+				const FSimdRealf& Dt,
+				const FSimd4Realf& FrictionStiffnessScale)
 			{
-				FSolverReal ImpulseNormal = -(Stiffness * ContactMassNormal) * ContactVelocityDeltaNormal;
-
-				// See comments in ApplyVelocityCorrection
-				if (NetImpulseNormal + ImpulseNormal < MinImpulseNormal)
+				// If all the lanes have zero velocity friction, run the zero-friction path
+				// (only spheres and capsules use the velocity-based dynamic friction path)
+				const FSimdSelector IsValid = SimdGreaterEqual(SimdConstraintIndex, FSimdInt32::Zero());
+				const FSimdSelector HasNonZeroFriction = SimdGreater(SimdVelocityFriction, FSimdRealf::Zero());
+				if (!SimdAnyTrue(SimdAnd(IsValid, HasNonZeroFriction)))
 				{
-					ImpulseNormal = MinImpulseNormal - NetImpulseNormal;
+					SolveVelocityNoFriction(Body0, Body1, Dt);
+					return;
 				}
 
-				NetImpulseNormal += ImpulseNormal;
+				// @todo(chaos): velocity friction
+				SolveVelocityNoFriction(Body0, Body1, Dt);
+			}
 
-				// Calculate the velocity deltas from the impulse
-				FSolverVec3 Impulse = ImpulseNormal * WorldContact.ContactNormal;
-				if (Body0.IsDynamic())
-				{
-					const FSolverVec3 DV0 = Body0.InvM() * Impulse;
-					const FSolverVec3 DW0 = WorldContactNormalAngular0 * ImpulseNormal;
-					Body0.ApplyVelocityDelta(DV0, DW0);
-				}
-				if (Body1.IsDynamic())
-				{
-					const FSolverVec3 DV1 = -Body1.InvM() * Impulse;
-					const FSolverVec3 DW1 = WorldContactNormalAngular1 * -ImpulseNormal;
-					Body1.ApplyVelocityDelta(DV1, DW1);
-				}
+			FSolverVec3 GetNetPushOut(const int32 LaneIndex) const
+			{
+				return SimdNetPushOutNormal.GetValue(LaneIndex) * SimdContactNormal.GetValue(LaneIndex) +
+					SimdNetPushOutTangentU.GetValue(LaneIndex) * SimdContactTangentU.GetValue(LaneIndex) +
+					SimdNetPushOutTangentV.GetValue(LaneIndex) * SimdContactTangentV.GetValue(LaneIndex);
+			}
+
+			FSolverVec3 GetNetImpulse(const int32 LaneIndex) const
+			{
+				return SimdNetImpulseNormal.GetValue(LaneIndex) * SimdContactNormal.GetValue(LaneIndex) +
+					SimdNetImpulseTangentU.GetValue(LaneIndex) * SimdContactTangentU.GetValue(LaneIndex) +
+					SimdNetImpulseTangentV.GetValue(LaneIndex) * SimdContactTangentV.GetValue(LaneIndex);
+			}
+
+			FSolverReal GetStaticFrictionRatio(const int32 LaneIndex) const
+			{
+				return SimdStaticFrictionRatio.GetValue(LaneIndex);
 			}
 
 			// @todo(chaos): make private
@@ -335,36 +584,55 @@ namespace Chaos
 			/**
 			 * @brief Whether we need to solve velocity for this manifold point (only if we were penetrating or applied a pushout)
 			*/
-			bool ShouldSolveVelocity() const;
-
-			// World-space contact data
-			FWorldContactPoint WorldContact;
-
-			// I^-1.(R x A) for each body where A is each axis (Normal, TangentU, TangentV)
-			FSolverVec3 WorldContactNormalAngular0;
-			FSolverVec3 WorldContactTangentUAngular0;
-			FSolverVec3 WorldContactTangentVAngular0;
-			FSolverVec3 WorldContactNormalAngular1;
-			FSolverVec3 WorldContactTangentUAngular1;
-			FSolverVec3 WorldContactTangentVAngular1;
-
-			// Contact mass (for non-friction)
-			FSolverReal ContactMassNormal;
-			FSolverReal ContactMassTangentU;
-			FSolverReal ContactMassTangentV;
-
-			// Solver outputs
-			FSolverReal NetPushOutNormal;
-			FSolverReal NetPushOutTangentU;
-			FSolverReal NetPushOutTangentV;
-			FSolverReal NetImpulseNormal;
-			FSolverReal NetImpulseTangentU;
-			FSolverReal NetImpulseTangentV;
-
-			// A measure of how much we exceeded the static friction threshold.
-			// Equal to (NormalPushOut / TangentialPushOut) before clamping to the friction cone.
-			// Used to move the static friction anchors to the edge of the cone in Scatter.
-			FSolverReal StaticFrictionRatio;
+			bool ShouldSolveVelocity() const
+			{
+				return true;
+			}
+
+			FSimdInt32 SimdConstraintIndex;
+			FSimdInt32 SimdManifoldPointIndex;
+
+			FSimdRealf SimdStiffness;
+
+			// World-space contact point relative to each particle's center of mass
+			FSimdVec3f SimdRelativeContactPoint0;
+			FSimdVec3f SimdRelativeContactPoint1;
+
+			// Normal PushOut
+			FSimdVec3f SimdContactNormal;
+			FSimdRealf SimdContactDeltaNormal;
+			FSimdRealf SimdNetPushOutNormal;
+			FSimdRealf SimdContactMassNormal;
+			FSimdVec3f SimdContactNormalAngular0;
+			FSimdVec3f SimdContactNormalAngular1;
+
+			// M^-1 for each body
+			FSimdRealf SimdInvM0;
+			FSimdRealf SimdInvM1;
+
+			// Tangential PushOut
+			FSimdVec3f SimdContactTangentU;
+			FSimdVec3f SimdContactTangentV;
+			FSimdRealf SimdContactDeltaTangentU;
+			FSimdRealf SimdContactDeltaTangentV;
+			FSimdRealf SimdNetPushOutTangentU;
+			FSimdRealf SimdNetPushOutTangentV;
+			FSimdRealf SimdStaticFrictionRatio;
+			FSimdRealf SimdContactMassTangentU;
+			FSimdRealf SimdContactMassTangentV;
+			FSimdRealf SimdStaticFriction;
+			FSimdRealf SimdDynamicFriction;
+			FSimdRealf SimdVelocityFriction;
+			FSimdVec3f SimdContactTangentUAngular0;
+			FSimdVec3f SimdContactTangentVAngular0;
+			FSimdVec3f SimdContactTangentUAngular1;
+			FSimdVec3f SimdContactTangentVAngular1;
+
+			// Normal Impulse
+			FSimdRealf SimdContactTargetVelocityNormal;
+			FSimdRealf SimdNetImpulseNormal;
+			FSimdRealf SimdNetImpulseTangentU;
+			FSimdRealf SimdNetImpulseTangentV;
 		};
 
 		/**
@@ -396,12 +664,23 @@ namespace Chaos
 			// NOTE: Does not initialize any properties. See MakeInitialized
 			FPBDCollisionSolverSimd() {}
 
+			int32 GetManifoldPointBeginIndex() const
+			{
+				return State.ManifoldPointBeginIndex;
+			}
+
+			int32 GetLaneIndex() const
+			{
+				return State.LaneIndex;
+			}
+
 			/** Reset the state of the collision solver */
 			void Reset()
 			{
-				State.SolverBodies[0].Reset();
-				State.SolverBodies[1].Reset();
-				State.ManifoldPoints = nullptr;
+				State.SolverBody0.Reset();
+				State.SolverBody1.Reset();
+				State.LaneIndex = 0;
+				State.ManifoldPointBeginIndex = 0;
 				State.NumManifoldPoints = 0;
 				State.MaxManifoldPoints = 0;
 			}
@@ -429,8 +708,17 @@ namespace Chaos
 
 			void SetSolverBodies(FSolverBody& SolverBody0, FSolverBody& SolverBody1)
 			{
-				State.SolverBodies[0].SetSolverBody(SolverBody0);
-				State.SolverBodies[1].SetSolverBody(SolverBody1);
+				State.SolverBody0.SetSolverBody(SolverBody0);
+				State.SolverBody1.SetSolverBody(SolverBody1);
+			}
+
+			void SetManifoldPointsBuffer(const int32 ConstraintIndex, const int32 LaneIndex, const int32 BeginIndex, const int32 Num)
+			{
+				State.ConstraintIndex = ConstraintIndex;
+				State.LaneIndex = LaneIndex;
+				State.ManifoldPointBeginIndex = BeginIndex;
+				State.NumManifoldPoints = 0;
+				State.MaxManifoldPoints = Num;
 			}
 
 			int32 NumManifoldPoints() const
@@ -444,30 +732,51 @@ namespace Chaos
 				return State.NumManifoldPoints++;
 			}
 
-			void SetManifoldPointsBuffer(FPBDCollisionSolverManifoldPointSimd* InManifoldPoints, const int32 InMaxManifoldPoints)
-			{
-				State.ManifoldPoints = InManifoldPoints;
-				State.MaxManifoldPoints = InMaxManifoldPoints;
-				State.NumManifoldPoints = 0;
-			}
-
-			const FPBDCollisionSolverManifoldPointSimd& GetManifoldPoint(const int32 ManifoldPointIndex) const
-			{
-				check(ManifoldPointIndex < NumManifoldPoints());
-				return State.ManifoldPoints[ManifoldPointIndex];
-			}
-
-			FWorldContactPoint& GetWorldContactPoint(const int32 ManifoldPointIndex)
+			/**
+			 * Set up a manifold point (calls InitManifodlPoint and FinalizeManifoldPoint)
+			*/
+			template<int TNumLanes>
+			void SetManifoldPoint(
+				TArrayView<TPBDCollisionSolverManifoldPointsSimd<TNumLanes>>& ManifoldPointsBuffer,
+				TArrayView<TSolverBodyPtrPairSimd<4>>& SolverBodiesBuffer,
+				const int32 ManifoldPointIndex,
+				const FSolverVec3& InRelativeContactPosition0,
+				const FSolverVec3& InRelativeContactPosition1,
+				const FSolverVec3& InWorldContactNormal,
+				const FSolverVec3& InWorldContactTangentU,
+				const FSolverVec3& InWorldContactTangentV,
+				const FSolverReal InWorldContactDeltaNormal,
+				const FSolverReal InWorldContactDeltaTangentU,
+				const FSolverReal InWorldContactDeltaTangentV,
+				const FSolverReal InWorldContactVelocityTargetNormal)
 			{
-				return State.ManifoldPoints[ManifoldPointIndex].WorldContact;
+				InitManifoldPoint(
+					ManifoldPointsBuffer,
+					SolverBodiesBuffer,
+					ManifoldPointIndex,
+					InRelativeContactPosition0,
+					InRelativeContactPosition1,
+					InWorldContactNormal,
+					InWorldContactTangentU,
+					InWorldContactTangentV,
+					InWorldContactDeltaNormal,
+					InWorldContactDeltaTangentU,
+					InWorldContactDeltaTangentV,
+					InWorldContactVelocityTargetNormal);
+
+				FinalizeManifoldPoint(
+					ManifoldPointsBuffer,
+					ManifoldPointIndex);
 			}
 
 			/**
-			 * Set up a manifold point (also calls FinalizeManifoldPoint)
+			 * Initialize the contact data for a manifold point (must call FinalizeManifoldPoint before use)
 			*/
-			void SetManifoldPoint(
-				const int32 ManifoldPoiontIndex,
-				const FSolverReal Dt,
+			template<int TNumLanes>
+			void InitManifoldPoint(
+				TArrayView<TPBDCollisionSolverManifoldPointsSimd<TNumLanes>>& ManifoldPointsBuffer,
+				TArrayView<TSolverBodyPtrPairSimd<4>>& SolverBodiesBuffer,
+				const int32 ManifoldPointIndex,
 				const FSolverVec3& InRelativeContactPosition0,
 				const FSolverVec3& InRelativeContactPosition1,
 				const FSolverVec3& InWorldContactNormal,
@@ -476,27 +785,68 @@ namespace Chaos
 				const FSolverReal InWorldContactDeltaNormal,
 				const FSolverReal InWorldContactDeltaTangentU,
 				const FSolverReal InWorldContactDeltaTangentV,
-				const FSolverReal InWorldContactVelocityTargetNormal);
+				const FSolverReal InWorldContactVelocityTargetNormal)
+			{
+				const int32 BufferIndex = GetBufferIndex(ManifoldPointIndex);
+
+				SolverBodiesBuffer[BufferIndex].Body0.SetValue(
+					State.LaneIndex,
+					&SolverBody0().SolverBody());
+
+				SolverBodiesBuffer[BufferIndex].Body1.SetValue(
+					State.LaneIndex,
+					&SolverBody1().SolverBody());
+
+				ManifoldPointsBuffer[BufferIndex].SetSharedData(
+					State.LaneIndex,
+					State.ConstraintIndex,
+					ManifoldPointIndex);
+
+				ManifoldPointsBuffer[BufferIndex].InitContact(
+					State.LaneIndex,
+					InRelativeContactPosition0,
+					InRelativeContactPosition1,
+					InWorldContactNormal,
+					InWorldContactTangentU,
+					InWorldContactTangentV,
+					InWorldContactDeltaNormal,
+					InWorldContactDeltaTangentU,
+					InWorldContactDeltaTangentV,
+					InWorldContactVelocityTargetNormal,
+					State.Stiffness,
+					State.StaticFriction,
+					State.DynamicFriction,
+					State.VelocityFriction);
+			}
 
 			/** 
 			 * Finish manifold point setup.
 			 * NOTE: Can only be called after the WorldContact has been set up (e.g., see SetManifoldPoint)
 			*/
-			void FinalizeManifoldPoint(const int32 ManifoldPoiontIndex, const FSolverReal Dt);
+			template<int TNumLanes>
+			void FinalizeManifoldPoint(
+				TArrayView<TPBDCollisionSolverManifoldPointsSimd<TNumLanes>>& ManifoldPointsBuffer,
+				const int32 ManifoldPointIndex)
+			{
+				ManifoldPointsBuffer[GetBufferIndex(ManifoldPointIndex)].FinalizeContact(
+					State.LaneIndex,
+					SolverBody0(),
+					SolverBody1());
+			}
 
 			/**
 			 * @brief Get the first (decorated) solver body
 			 * The decorator add a possible mass scale
 			*/
-			FConstraintSolverBody& SolverBody0() { return State.SolverBodies[0]; }
-			const FConstraintSolverBody& SolverBody0() const { return State.SolverBodies[0]; }
+			FConstraintSolverBody& SolverBody0() { return State.SolverBody0; }
+			const FConstraintSolverBody& SolverBody0() const { return State.SolverBody0; }
 
 			/**
 			 * @brief Get the second (decorated) solver body
 			 * The decorator add a possible mass scale
 			*/
-			FConstraintSolverBody& SolverBody1() { return State.SolverBodies[1]; }
-			const FConstraintSolverBody& SolverBody1() const { return State.SolverBodies[1]; }
+			FConstraintSolverBody& SolverBody1() { return State.SolverBody1; }
+			const FConstraintSolverBody& SolverBody1() const { return State.SolverBody1; }
 
 			/**
 			 * @brief Set up the mass scaling for shock propagation, using the position-phase mass scale
@@ -513,34 +863,43 @@ namespace Chaos
 			*/
 			void DisableShockPropagation();
 
-			/**
-			 * @brief Calculate and apply the position correction for this iteration
-			 * @return true if we need to run more iterations, false if we did not apply any correction
-			*/
-			bool SolvePositionWithFriction(const FSolverReal Dt, const FSolverReal MaxPushOut);
-			bool SolvePositionNoFriction(const FSolverReal Dt, const FSolverReal MaxPushOut);
+			template<int TNumLanes>
+			FSolverVec3 GetNetPushOut(
+				TArrayView<TPBDCollisionSolverManifoldPointsSimd<TNumLanes>>& ManifoldPointsBuffer,
+				const int32 ManifoldPointIndex) const
+			{
+				return ManifoldPointsBuffer[GetBufferIndex(ManifoldPointIndex)].GetNetPushOut(State.LaneIndex);
+			}
 
-			/**
-			 * @brief Calculate and apply the velocity correction for this iteration
-			 * @return true if we need to run more iterations, false if we did not apply any correction
-			*/
-			bool SolveVelocity(const FSolverReal Dt, const bool bApplyDynamicFriction);
+			template<int TNumLanes>
+			FSolverVec3 GetNetImpulse(
+				TArrayView<TPBDCollisionSolverManifoldPointsSimd<TNumLanes>>& ManifoldPointsBuffer,
+				const int32 ManifoldPointIndex) const
+			{
+				return ManifoldPointsBuffer[GetBufferIndex(ManifoldPointIndex)].GetNetImpulse(State.LaneIndex);
+			}
+
+			template<int TNumLanes>
+			FSolverReal GetStaticFrictionRatio(
+				TArrayView<TPBDCollisionSolverManifoldPointsSimd<TNumLanes>>& ManifoldPointsBuffer,
+				const int32 ManifoldPointIndex) const
+			{
+				return ManifoldPointsBuffer[GetBufferIndex(ManifoldPointIndex)].GetStaticFrictionRatio(State.LaneIndex);
+			}
 
 		private:
+			int32 GetBufferIndex(const int32 ManifoldPointIndex) const
+			{
+				check(ManifoldPointIndex < State.NumManifoldPoints);
+				return State.ManifoldPointBeginIndex + ManifoldPointIndex;
+			}
+
 			/**
 			 * @brief Apply the inverse mass scale the body with the lower level
 			 * @param InvMassScale
 			*/
 			void SetShockPropagationInvMassScale(const FSolverReal InvMassScale);
 
-			/**
-			 * @brief Run a velocity solve on the average point from all the points that received a position impulse
-			 * This is used to enforce Restitution constraints without introducing rotation artefacts without
-			 * adding more velocity iterations.
-			 * This will only perform work if there is more than one active contact.
-			*/
-			void SolveVelocityAverage(const FSolverReal Dt);
-
 			struct FState
 			{
 				FState()
@@ -549,17 +908,25 @@ namespace Chaos
 
 				void Init()
 				{
-					SolverBodies[0].Init();
-					SolverBodies[1].Init();
+					ConstraintIndex = INDEX_NONE;
+					LaneIndex = INDEX_NONE;
+					ManifoldPointBeginIndex = INDEX_NONE;
+					NumManifoldPoints = 0;
+					MaxManifoldPoints = 0;
 					StaticFriction = 0;
 					DynamicFriction = 0;
 					VelocityFriction = 0;
 					Stiffness = 1;
-					ManifoldPoints = nullptr;
-					NumManifoldPoints = 0;
-					MaxManifoldPoints = 0;
+					SolverBody0.Init();
+					SolverBody1.Init();
 				}
 
+				int32 ConstraintIndex;
+				int32 LaneIndex;
+				int32 ManifoldPointBeginIndex;
+				int32 NumManifoldPoints;
+				int32 MaxManifoldPoints;
+
 				// Static Friction in the position-solve phase
 				FSolverReal StaticFriction;
 
@@ -573,10 +940,8 @@ namespace Chaos
 				FSolverReal Stiffness;
 
 				// Bodies and contacts
-				FConstraintSolverBody SolverBodies[MaxConstrainedBodies];
-				FPBDCollisionSolverManifoldPointSimd* ManifoldPoints;
-				int32 NumManifoldPoints;
-				int32 MaxManifoldPoints;
+				FConstraintSolverBody SolverBody0;
+				FConstraintSolverBody SolverBody1;
 			};
 
 			FState State;
@@ -587,374 +952,68 @@ namespace Chaos
 		//////////////////////////////////////////////////////////////////////////////////////////////////
 		//////////////////////////////////////////////////////////////////////////////////////////////////
 
-		FORCEINLINE_DEBUGGABLE void FPBDCollisionSolverManifoldPointSimd::InitContact(
-			const FSolverReal Dt,
-			const FConstraintSolverBody& Body0,
-			const FConstraintSolverBody& Body1)
-		{
-			NetPushOutNormal = FSolverReal(0);
-			NetPushOutTangentU = FSolverReal(0);
-			NetPushOutTangentV = FSolverReal(0);
-			NetImpulseNormal = FSolverReal(0);
-			NetImpulseTangentU = FSolverReal(0);
-			NetImpulseTangentV = FSolverReal(0);
-			StaticFrictionRatio = FSolverReal(0);
-
-			UpdateMass(Body0, Body1);
-		}
-
-		FORCEINLINE_DEBUGGABLE void FPBDCollisionSolverManifoldPointSimd::UpdateMass(const FConstraintSolverBody& Body0, const FConstraintSolverBody& Body1)
-		{
-			FSolverReal ContactMassInvNormal = FSolverReal(0);
-			FSolverReal ContactMassInvTangentU = FSolverReal(0);
-			FSolverReal ContactMassInvTangentV = FSolverReal(0);
-
-			// These are not used if not initialized below so no need to clear
-			//WorldContactNormalAngular0 = FSolverVec3(0);
-			//WorldContactTangentUAngular0 = FSolverVec3(0);
-			//WorldContactTangentVAngular0 = FSolverVec3(0);
-			//WorldContactNormalAngular1 = FSolverVec3(0);
-			//WorldContactTangentUAngular1 = FSolverVec3(0);
-			//WorldContactTangentVAngular1 = FSolverVec3(0);
-
-			if (Body0.IsDynamic())
-			{
-				const FSolverVec3 R0xN = FSolverVec3::CrossProduct(WorldContact.RelativeContactPoints[0], WorldContact.ContactNormal);
-				const FSolverVec3 R0xU = FSolverVec3::CrossProduct(WorldContact.RelativeContactPoints[0], WorldContact.ContactTangentU);
-				const FSolverVec3 R0xV = FSolverVec3::CrossProduct(WorldContact.RelativeContactPoints[0], WorldContact.ContactTangentV);
-
-				const FSolverMatrix33 InvI0 = Body0.InvI();
-
-				WorldContactNormalAngular0 = InvI0 * R0xN;
-				WorldContactTangentUAngular0 = InvI0 * R0xU;
-				WorldContactTangentVAngular0 = InvI0 * R0xV;
-
-				ContactMassInvNormal += FSolverVec3::DotProduct(R0xN, WorldContactNormalAngular0) + Body0.InvM();
-				ContactMassInvTangentU += FSolverVec3::DotProduct(R0xU, WorldContactTangentUAngular0) + Body0.InvM();
-				ContactMassInvTangentV += FSolverVec3::DotProduct(R0xV, WorldContactTangentVAngular0) + Body0.InvM();
-			}
-			if (Body1.IsDynamic())
-			{
-				const FSolverVec3 R1xN = FSolverVec3::CrossProduct(WorldContact.RelativeContactPoints[1], WorldContact.ContactNormal);
-				const FSolverVec3 R1xU = FSolverVec3::CrossProduct(WorldContact.RelativeContactPoints[1], WorldContact.ContactTangentU);
-				const FSolverVec3 R1xV = FSolverVec3::CrossProduct(WorldContact.RelativeContactPoints[1], WorldContact.ContactTangentV);
-
-				const FSolverMatrix33 InvI1 = Body1.InvI();
-
-				WorldContactNormalAngular1 = InvI1 * R1xN;
-				WorldContactTangentUAngular1 = InvI1 * R1xU;
-				WorldContactTangentVAngular1 = InvI1 * R1xV;
-
-				ContactMassInvNormal += FSolverVec3::DotProduct(R1xN, WorldContactNormalAngular1) + Body1.InvM();
-				ContactMassInvTangentU += FSolverVec3::DotProduct(R1xU, WorldContactTangentUAngular1) + Body1.InvM();
-				ContactMassInvTangentV += FSolverVec3::DotProduct(R1xV, WorldContactTangentVAngular1) + Body1.InvM();
-			}
 
-			ContactMassNormal = (ContactMassInvNormal > FSolverReal(UE_SMALL_NUMBER)) ? FSolverReal(1) / ContactMassInvNormal : FSolverReal(0);
-			ContactMassTangentU = (ContactMassInvTangentU > FSolverReal(UE_SMALL_NUMBER)) ? FSolverReal(1) / ContactMassInvTangentU : FSolverReal(0);
-			ContactMassTangentV = (ContactMassInvTangentV > FSolverReal(UE_SMALL_NUMBER)) ? FSolverReal(1) / ContactMassInvTangentV : FSolverReal(0);
-		}
-
-		FORCEINLINE_DEBUGGABLE void FPBDCollisionSolverManifoldPointSimd::UpdateMassNormal(const FConstraintSolverBody& Body0, const FConstraintSolverBody& Body1)
-		{
-			FSolverReal ContactMassInvNormal = FSolverReal(0);
-			if (Body0.IsDynamic())
-			{
-				const FSolverVec3 R0xN = FSolverVec3::CrossProduct(WorldContact.RelativeContactPoints[0], WorldContact.ContactNormal);
-				const FSolverMatrix33 InvI0 = Body0.InvI();
-				WorldContactNormalAngular0 = InvI0 * R0xN;
-				ContactMassInvNormal += FSolverVec3::DotProduct(R0xN, WorldContactNormalAngular0) + Body0.InvM();
-			}
-			if (Body1.IsDynamic())
-			{
-				const FSolverVec3 R1xN = FSolverVec3::CrossProduct(WorldContact.RelativeContactPoints[1], WorldContact.ContactNormal);
-				const FSolverMatrix33 InvI1 = Body1.InvI();
-				WorldContactNormalAngular1 = InvI1 * R1xN;
-				ContactMassInvNormal += FSolverVec3::DotProduct(R1xN, WorldContactNormalAngular1) + Body1.InvM();
-			}
-			ContactMassNormal = (ContactMassInvNormal > FSolverReal(UE_SMALL_NUMBER)) ? FSolverReal(1) / ContactMassInvNormal : FSolverReal(0);
-		}
-
-		FORCEINLINE_DEBUGGABLE void FPBDCollisionSolverManifoldPointSimd::CalculateContactPositionErrorNormal(const FConstraintSolverBody& Body0, const FConstraintSolverBody& Body1, const FSolverReal MaxPushOut, FSolverReal& OutContactDeltaNormal) const
+		/**
+		 * A helper for solving arrays of constraints.
+		 * @note Only works with 4 SIMD lanes for now.
+		 */
+		class FPBDCollisionSolverHelperSimd
 		{
-			// Linear version: calculate the contact delta assuming linear motion after applying a positional impulse at the contact point. There will be an error that depends on the size of the rotation.
-			const FSolverVec3 ContactDelta0 = Body0.DP() + FSolverVec3::CrossProduct(Body0.DQ(), WorldContact.RelativeContactPoints[0]);
-			const FSolverVec3 ContactDelta1 = Body1.DP() + FSolverVec3::CrossProduct(Body1.DQ(), WorldContact.RelativeContactPoints[1]);
-			const FSolverVec3 ContactDelta = ContactDelta0 - ContactDelta1;
-			OutContactDeltaNormal = WorldContact.ContactDeltaNormal + FSolverVec3::DotProduct(ContactDelta, WorldContact.ContactNormal);
-
-			// NOTE: OutContactDeltaNormal is negative for penetration
-			// NOTE: MaxPushOut == 0 disables the pushout limits
-			if ((MaxPushOut > 0) && (OutContactDeltaNormal < -MaxPushOut))
-			{
-				OutContactDeltaNormal = -MaxPushOut;
-			}
-		}
+		public:
+			template<int TNumLanes>
+			static CHAOS_API void SolvePositionNoFriction(
+				const TArrayView<TPBDCollisionSolverManifoldPointsSimd<TNumLanes>>& ManifoldPoints,
+				const TArrayView<TSolverBodyPtrPairSimd<TNumLanes>>& SolverBodies,
+				const FSolverReal Dt,
+				const FSolverReal MaxPushOut);
 
-		FORCEINLINE_DEBUGGABLE void FPBDCollisionSolverManifoldPointSimd::CalculateContactPositionErrorTangential(const FConstraintSolverBody& Body0, const FConstraintSolverBody& Body1, FSolverReal& OutContactDeltaTangentU, FSolverReal& OutContactDeltaTangentV) const
-		{
-			// Linear version: calculate the contact delta assuming linear motion after applying a positional impulse at the contact point. There will be an error that depends on the size of the rotation.
-			const FSolverVec3 ContactDelta0 = Body0.DP() + FSolverVec3::CrossProduct(Body0.DQ(), WorldContact.RelativeContactPoints[0]);
-			const FSolverVec3 ContactDelta1 = Body1.DP() + FSolverVec3::CrossProduct(Body1.DQ(), WorldContact.RelativeContactPoints[1]);
-			const FSolverVec3 ContactDelta = ContactDelta0 - ContactDelta1;
-			OutContactDeltaTangentU = WorldContact.ContactDeltaTangentU + FSolverVec3::DotProduct(ContactDelta, WorldContact.ContactTangentU);
-			OutContactDeltaTangentV = WorldContact.ContactDeltaTangentV + FSolverVec3::DotProduct(ContactDelta, WorldContact.ContactTangentV);
-		}
-
-		FORCEINLINE_DEBUGGABLE void FPBDCollisionSolverManifoldPointSimd::CalculateContactVelocityError(const FConstraintSolverBody& Body0, const FConstraintSolverBody& Body1, const FSolverReal DynamicFriction, const FSolverReal Dt, FSolverReal& OutContactVelocityDeltaNormal, FSolverReal& OutContactVelocityDeltaTangent0, FSolverReal& OutContactVelocityDeltaTangent1) const
-		{
-			const FSolverVec3 ContactVelocity0 = Body0.V() + FSolverVec3::CrossProduct(Body0.W(), WorldContact.RelativeContactPoints[0]);
-			const FSolverVec3 ContactVelocity1 = Body1.V() + FSolverVec3::CrossProduct(Body1.W(), WorldContact.RelativeContactPoints[1]);
-			const FSolverVec3 ContactVelocity = ContactVelocity0 - ContactVelocity1;
-			const FSolverReal ContactVelocityNormal = FSolverVec3::DotProduct(ContactVelocity, WorldContact.ContactNormal);
-			const FSolverReal ContactVelocityTangent0 = FSolverVec3::DotProduct(ContactVelocity, WorldContact.ContactTangentU);
-			const FSolverReal ContactVelocityTangent1 = FSolverVec3::DotProduct(ContactVelocity, WorldContact.ContactTangentV);
-
-			OutContactVelocityDeltaNormal = (ContactVelocityNormal - WorldContact.ContactTargetVelocityNormal);
-			OutContactVelocityDeltaTangent0 = ContactVelocityTangent0;
-			OutContactVelocityDeltaTangent1 = ContactVelocityTangent1;
-		}
-
-		FORCEINLINE_DEBUGGABLE void FPBDCollisionSolverManifoldPointSimd::CalculateContactVelocityErrorNormal(const FConstraintSolverBody& Body0, const FConstraintSolverBody& Body1, FSolverReal& OutContactVelocityDeltaNormal) const
-		{
-			const FSolverVec3 ContactVelocity0 = Body0.V() + FSolverVec3::CrossProduct(Body0.W(), WorldContact.RelativeContactPoints[0]);
-			const FSolverVec3 ContactVelocity1 = Body1.V() + FSolverVec3::CrossProduct(Body1.W(), WorldContact.RelativeContactPoints[1]);
-			const FSolverVec3 ContactVelocity = ContactVelocity0 - ContactVelocity1;
-			const FSolverReal ContactVelocityNormal = FSolverVec3::DotProduct(ContactVelocity, WorldContact.ContactNormal);
+			template<int TNumLanes>
+			static CHAOS_API void SolvePositionWithFriction(
+				const TArrayView<TPBDCollisionSolverManifoldPointsSimd<TNumLanes>>& ManifoldPoints,
+				const TArrayView<TSolverBodyPtrPairSimd<TNumLanes>>& SolverBodies,
+				const FSolverReal Dt,
+				const FSolverReal MaxPushOut);
 
-			// Add up the errors in the velocity (current velocity - desired velocity)
-			OutContactVelocityDeltaNormal = (ContactVelocityNormal - WorldContact.ContactTargetVelocityNormal);
-		}
+			template<int TNumLanes>
+			static CHAOS_API void SolveVelocityNoFriction(
+				const TArrayView<TPBDCollisionSolverManifoldPointsSimd<TNumLanes>>& ManifoldPoints,
+				const TArrayView<TSolverBodyPtrPairSimd<TNumLanes>>& SolverBodies,
+				const FSolverReal Dt);
 
-		FORCEINLINE_DEBUGGABLE bool FPBDCollisionSolverManifoldPointSimd::ShouldSolveVelocity() const
-		{
-			// We ensure positive separating velocity for close contacts even if they didn't receive a pushout
-			return (NetPushOutNormal > FSolverReal(0)) || (WorldContact.ContactDeltaNormal < FSolverReal(0));
-		}
+			template<int TNumLanes>
+			static CHAOS_API void SolveVelocityWithFriction(
+				const TArrayView<TPBDCollisionSolverManifoldPointsSimd<TNumLanes>>& ManifoldPoints,
+				const TArrayView<TSolverBodyPtrPairSimd<TNumLanes>>& SolverBodies,
+				const FSolverReal Dt);
 
-		//////////////////////////////////////////////////////////////////////////////////////////////////
-		//////////////////////////////////////////////////////////////////////////////////////////////////
-		//////////////////////////////////////////////////////////////////////////////////////////////////
-		//////////////////////////////////////////////////////////////////////////////////////////////////
+			static CHAOS_API void CheckISPC();
+		};
 
-		FORCEINLINE_DEBUGGABLE  void FPBDCollisionSolverSimd::FinalizeManifoldPoint(const int32 ManifoldPoiontIndex, const FSolverReal Dt)
-		{
-			State.ManifoldPoints[ManifoldPoiontIndex].InitContact(
-				FSolverReal(Dt),
-				State.SolverBodies[0],
-				State.SolverBodies[1]);
-		}
-
-		FORCEINLINE_DEBUGGABLE void FPBDCollisionSolverSimd::SetManifoldPoint(
-			const int32 ManifoldPoiontIndex,
+		template<> 
+		void FPBDCollisionSolverHelperSimd::SolvePositionNoFriction(
+			const TArrayView<TPBDCollisionSolverManifoldPointsSimd<4>>& ManifoldPoints,
+			const TArrayView<TSolverBodyPtrPairSimd<4>>& SolverBodies,
 			const FSolverReal Dt,
-			const FSolverVec3& InRelativeContactPosition0,
-			const FSolverVec3& InRelativeContactPosition1,
-			const FSolverVec3& InWorldContactNormal,
-			const FSolverVec3& InWorldContactTangentU,
-			const FSolverVec3& InWorldContactTangentV,
-			const FSolverReal InWorldContactDeltaNormal,
-			const FSolverReal InWorldContactDeltaTangentU,
-			const FSolverReal InWorldContactDeltaTangentV,
-			const FSolverReal InWorldContactVelocityTargetNormal)
-		{
-			FWorldContactPoint& WorldContactPoint = State.ManifoldPoints[ManifoldPoiontIndex].WorldContact;
-			WorldContactPoint.RelativeContactPoints[0] = InRelativeContactPosition0;
-			WorldContactPoint.RelativeContactPoints[1] = InRelativeContactPosition1;
-			WorldContactPoint.ContactNormal = InWorldContactNormal;
-			WorldContactPoint.ContactTangentU = InWorldContactTangentU;
-			WorldContactPoint.ContactTangentV = InWorldContactTangentV;
-			WorldContactPoint.ContactDeltaNormal = InWorldContactDeltaNormal;
-			WorldContactPoint.ContactDeltaTangentU = InWorldContactDeltaTangentU;
-			WorldContactPoint.ContactDeltaTangentV = InWorldContactDeltaTangentV;
-			WorldContactPoint.ContactTargetVelocityNormal = InWorldContactVelocityTargetNormal;
-
-			State.ManifoldPoints[ManifoldPoiontIndex].InitContact(
-				FSolverReal(Dt),
-				State.SolverBodies[0],
-				State.SolverBodies[1]);
-		}
-
-		FORCEINLINE_DEBUGGABLE bool FPBDCollisionSolverSimd::SolvePositionWithFriction(const FSolverReal Dt, const FSolverReal MaxPushOut)
-		{
-			// SolverBody decorator used to add mass scaling
-			FConstraintSolverBody& Body0 = SolverBody0();
-			FConstraintSolverBody& Body1 = SolverBody1();
-
-			// Accumulate net pushout for friction limits below
-			bool bApplyFriction[MaxPointsPerConstraint] = { false, };
-			int32 NumFrictionContacts = 0;
-			FSolverReal TotalPushOutNormal = FSolverReal(0);
-
-			// Apply the position correction along the normal and determine if we want to run friction on each point
-			for (int32 PointIndex = 0; PointIndex < NumManifoldPoints(); ++PointIndex)
-			{
-				FPBDCollisionSolverManifoldPointSimd& SolverManifoldPoint = State.ManifoldPoints[PointIndex];
+			const FSolverReal MaxPushOut);
 
-				FSolverReal ContactDeltaNormal;
-				SolverManifoldPoint.CalculateContactPositionErrorNormal(Body0.SolverBody(), Body1.SolverBody(), MaxPushOut, ContactDeltaNormal);
-
-				// Apply a normal correction if we still have penetration or if we are now separated but have previously applied a correction that we may want to undo
-				const bool bProcessManifoldPoint = (ContactDeltaNormal < FSolverReal(0)) || (SolverManifoldPoint.NetPushOutNormal > FSolverReal(UE_SMALL_NUMBER));
-				if (bProcessManifoldPoint)
-				{
-					SolverManifoldPoint.ApplyPositionCorrectionNormal(
-						State.Stiffness,
-						ContactDeltaNormal,
-						Body0,
-						Body1);
-
-					TotalPushOutNormal += SolverManifoldPoint.NetPushOutNormal;
-				}
-
-				// Friction gets updated for any point with a net normal correction or where we have previously had a normal correction and 
-				// already applied friction (in which case we may need to zero it)
-				if ((SolverManifoldPoint.NetPushOutNormal != 0) || (SolverManifoldPoint.NetPushOutTangentU != 0) || (SolverManifoldPoint.NetPushOutTangentV != 0))
-				{
-					bApplyFriction[PointIndex] = true;
-					++NumFrictionContacts;
-				}
-			}
-
-			// Apply the tangential position correction if required
-			if (NumFrictionContacts > 0)
-			{
-				// We clip the tangential correction at each contact to the friction cone, but we use to average impulse
-				// among all contacts as the clipping limit. This is not really correct but it is much more stable to 
-				// differences in contacts from tick to tick
-				// @todo(chaos): try a decaying maximum per contact point rather than an average (again - we had that once!)
-				const FSolverReal FrictionMaxPushOut = TotalPushOutNormal / FSolverReal(NumFrictionContacts);
-				const FSolverReal FrictionStiffness = State.Stiffness * CVars::Chaos_PBDCollisionSolver_Position_StaticFrictionStiffness;
-
-				for (int32 PointIndex = 0; PointIndex < NumManifoldPoints(); ++PointIndex)
-				{
-					if (bApplyFriction[PointIndex])
-					{
-						FPBDCollisionSolverManifoldPointSimd& SolverManifoldPoint = State.ManifoldPoints[PointIndex];
-
-						FSolverReal ContactDeltaTangentU, ContactDeltaTangentV;
-						SolverManifoldPoint.CalculateContactPositionErrorTangential(Body0.SolverBody(), Body1.SolverBody(), ContactDeltaTangentU, ContactDeltaTangentV);
-
-						SolverManifoldPoint.ApplyPositionCorrectionTangential(
-							FrictionStiffness,
-							State.StaticFriction,
-							State.DynamicFriction,
-							FrictionMaxPushOut,
-							ContactDeltaTangentU,
-							ContactDeltaTangentV,
-							Body0,
-							Body1);
-					}
-				}
-			}
-
-			return false;
-		}
-
-		FORCEINLINE_DEBUGGABLE bool FPBDCollisionSolverSimd::SolvePositionNoFriction(const FSolverReal Dt, const FSolverReal MaxPushOut)
-		{
-			// SolverBody decorator used to add mass scaling
-			FConstraintSolverBody& Body0 = SolverBody0();
-			FConstraintSolverBody& Body1 = SolverBody1();
-
-			// Apply the position correction so that all contacts have zero separation
-			for (int32 PointIndex = 0; PointIndex < NumManifoldPoints(); ++PointIndex)
-			{
-				FPBDCollisionSolverManifoldPointSimd& SolverManifoldPoint = State.ManifoldPoints[PointIndex];
-
-				FSolverReal ContactDeltaNormal;
-				SolverManifoldPoint.CalculateContactPositionErrorNormal(Body0.SolverBody(), Body1.SolverBody(), MaxPushOut, ContactDeltaNormal);
-
-				const bool bProcessManifoldPoint = (ContactDeltaNormal < FSolverReal(0)) || (SolverManifoldPoint.NetPushOutNormal > FSolverReal(UE_SMALL_NUMBER));
-				if (bProcessManifoldPoint)
-				{
-					SolverManifoldPoint.ApplyPositionCorrectionNormal(
-						State.Stiffness,
-						ContactDeltaNormal,
-						Body0,
-						Body1);
-				}
-			}
-
-			return false;
-		}
-
-		FORCEINLINE_DEBUGGABLE bool FPBDCollisionSolverSimd::SolveVelocity(const FSolverReal Dt, const bool bApplyDynamicFriction)
-		{
-			// Apply restitution at the average contact point
-			// This means we don't need to run as many iterations to get stable bouncing
-			// It also helps with zero restitution to counter any velocioty added by the PBD solve
-			const bool bSolveAverageContact = (NumManifoldPoints() > 1) && CVars::bChaos_PBDCollisionSolver_Velocity_AveragePointEnabled;
-			if (bSolveAverageContact)
-			{
-				SolveVelocityAverage(Dt);
-			}
-
-			FConstraintSolverBody& Body0 = SolverBody0();
-			FConstraintSolverBody& Body1 = SolverBody1();
-
-			// NOTE: this dynamic friction implementation is iteration-count sensitive
-			// @todo(chaos): fix iteration count dependence of dynamic friction
-			const FSolverReal DynamicFriction = (bApplyDynamicFriction && (Dt > 0) && CVars::bChaos_PBDCollisionSolver_Velocity_FrictionEnabled) ? State.VelocityFriction : FSolverReal(0);
-
-			for (int32 PointIndex = 0; PointIndex < NumManifoldPoints(); ++PointIndex)
-			{
-				FPBDCollisionSolverManifoldPointSimd& SolverManifoldPoint = State.ManifoldPoints[PointIndex];
-
-				if (SolverManifoldPoint.ShouldSolveVelocity())
-				{
-					const FSolverReal MinImpulseNormal = FMath::Min(FSolverReal(0), -SolverManifoldPoint.NetPushOutNormal / Dt);
-
-					if (DynamicFriction > 0)
-					{
-						FSolverReal ContactVelocityDeltaNormal, ContactVelocityDeltaTangentU, ContactVelocityDeltaTangentV;
-						SolverManifoldPoint.CalculateContactVelocityError(Body0, Body1, DynamicFriction, Dt, ContactVelocityDeltaNormal, ContactVelocityDeltaTangentU, ContactVelocityDeltaTangentV);
-
-						SolverManifoldPoint.ApplyVelocityCorrection(
-							State.Stiffness,
-							Dt,
-							DynamicFriction,
-							ContactVelocityDeltaNormal,
-							ContactVelocityDeltaTangentU,
-							ContactVelocityDeltaTangentV,
-							MinImpulseNormal,
-							Body0,
-							Body1);
-					}
-					else
-					{
-						FSolverReal ContactVelocityDeltaNormal;
-						SolverManifoldPoint.CalculateContactVelocityErrorNormal(Body0, Body1, ContactVelocityDeltaNormal);
-
-						SolverManifoldPoint.ApplyVelocityCorrectionNormal(
-							State.Stiffness,
-							ContactVelocityDeltaNormal,
-							MinImpulseNormal,
-							Body0,
-							Body1);
-					}
-				}
-			}
-
-			// Early-out support for the velocity solve is not currently very important because we
-			// only run one iteration in the velocity solve phase.
-			// @todo(chaos): support early-out in velocity solve if necessary
-			return true;
-		}
-
-
-		/**
-		 * A helper for solving arrays of constraints
-		 */
-		class FPBDCollisionSolverHelperSimd
-		{
-		public:
-			static void SolvePositionNoFriction(const TArrayView<FPBDCollisionSolverSimd>& CollisionSolvers, const FSolverReal Dt, const FSolverReal MaxPushOut);
-			static void SolvePositionWithFriction(const TArrayView<FPBDCollisionSolverSimd>& CollisionSolvers, const FSolverReal Dt, const FSolverReal MaxPushOut);
-			static void SolveVelocity(const TArrayView<FPBDCollisionSolverSimd>& CollisionSolvers, const FSolverReal Dt, const bool bApplyDynamicFriction);
-
-			static void CheckISPC();
-		};
+		template<>
+		void FPBDCollisionSolverHelperSimd::SolvePositionWithFriction(
+			const TArrayView<TPBDCollisionSolverManifoldPointsSimd<4>>& ManifoldPoints,
+			const TArrayView<TSolverBodyPtrPairSimd<4>>& SolverBodies,
+			const FSolverReal Dt,
+			const FSolverReal MaxPushOut);
+
+		template<>
+		void FPBDCollisionSolverHelperSimd::SolveVelocityNoFriction(
+			const TArrayView<TPBDCollisionSolverManifoldPointsSimd<4>>& ManifoldPoints,
+			const TArrayView<TSolverBodyPtrPairSimd<4>>& SolverBodies,
+			const FSolverReal Dt);
+
+		template<>
+		void FPBDCollisionSolverHelperSimd::SolveVelocityWithFriction(
+			const TArrayView<TPBDCollisionSolverManifoldPointsSimd<4>>& ManifoldPoints,
+			const TArrayView<TSolverBodyPtrPairSimd<4>>& SolverBodies,
+			const FSolverReal Dt);
 
 	}	// namespace Private
 }	// namespace Chaos
diff --git a/Engine/Source/Runtime/Experimental/Chaos/Public/Chaos/Evolution/ABTestingConstraintContainerSolver.h b/Engine/Source/Runtime/Experimental/Chaos/Public/Chaos/Evolution/ABTestingConstraintContainerSolver.h
new file mode 100644
--- /dev/null
+++ b/Engine/Source/Runtime/Experimental/Chaos/Public/Chaos/Evolution/ABTestingConstraintContainerSolver.h
@@ -0,0 +1,205 @@
+// Copyright Epic Games, Inc. All Rights Reserved.
+#pragma once
+
+#include "Chaos/ConstraintHandle.h"
+
+#include <functional>
+
+#ifndef CHAOS_ABTEST_CONSTRAINTSOLVER_ENABLED
+#define CHAOS_ABTEST_CONSTRAINTSOLVER_ENABLED (!UE_BUILD_TEST && !UE_BUILD_SHIPPING)
+#endif
+
+
+namespace Chaos
+{
+	namespace Private
+	{
+		class FPBDIslandConstraint;
+	}
+	class FConstraintHandleHolder;
+	class FSolverBodyContainer;
+
+	namespace Private
+	{
+#if CHAOS_ABTEST_CONSTRAINTSOLVER_ENABLED
+
+		/**
+		* Used for testing a new solver that is supposed to produce the same output as anotehr solver.
+		* E.g., for testing the Simd version of a solver against the non-simd version.
+		* 
+		* The first solver passed in is assumed to be the reference - its results will be used to generate the outputs.
+		* The second solver results will be compared and discarded.
+		* 
+		* @todo(chaos): this isn't good enough on its own. Really we need to AB Test the whole IslandGroupSolver
+		* but with a different set of constraint container solvers in each. The way it is now will work fine
+		* as long as we only have one constraint type in the island group which will do for now.
+		*/
+		template<typename T1, typename T2>
+		class CHAOS_API TABTestingConstraintContainerSolver : public FConstraintContainerSolver
+		{
+		public:
+			enum class ESolverPhase
+			{
+				PreApplyPositionConstraints,
+				PostApplyPositionConstraints,
+				PostApplyVelocityConstraints,
+				PostApplyProjectionConstraints,
+			};
+
+			using FContainerSolverTypeA = T1;
+			using FContainerSolverTypeB = T2;
+			using FABTestFunctor = std::function<void(
+				const ESolverPhase Phase, 
+				const FContainerSolverTypeA& SolverA, 
+				const FContainerSolverTypeB& SolverB, 
+				const FSolverBodyContainer& SolverBodyContainerA, 
+				const FSolverBodyContainer& SolverBodyContainerB)>;
+
+			TABTestingConstraintContainerSolver(
+				TUniquePtr<FContainerSolverTypeA>&& InSolverA, 
+				TUniquePtr<FContainerSolverTypeB>&& InSolverB, 
+				const int32 InPriority,
+				const FABTestFunctor& InABTestFunctor)
+				: FConstraintContainerSolver(InPriority)
+				, SolverA(MoveTemp(InSolverA))
+				, SolverB(MoveTemp(InSolverB))
+				, SolverBodyContainerA(nullptr)
+				, SolverBodyContainerB()
+				, ABTestFunctor(InABTestFunctor)
+			{
+			}
+
+			virtual void Reset(const int32 MaxConstraints) override final
+			{
+				SolverA->Reset(MaxConstraints);
+				SolverB->Reset(MaxConstraints);
+				SolverBodyContainerA = nullptr;
+			}
+
+			virtual int32 GetNumConstraints() const override final
+			{
+				return SolverA->GetNumConstraints();
+			}
+
+			virtual void AddConstraints() override final
+			{
+				SolverA->AddConstraints();
+				SolverB->AddConstraints();
+			}
+
+			virtual void AddConstraints(const TArrayView<Private::FPBDIslandConstraint>& Constraints) override final
+			{
+				SolverA->AddConstraints(Constraints);
+				SolverB->AddConstraints(Constraints);
+			}
+
+			virtual void AddBodies(FSolverBodyContainer& SolverBodyContainer) override final
+			{
+				SolverBodyContainerA = &SolverBodyContainer;
+
+				SolverA->AddBodies(*SolverBodyContainerA);
+			}
+
+			virtual void GatherInput(const FReal Dt) override final
+			{
+				SolverA->GatherInput(Dt);
+				// Solver B will ge gathered in PreApply
+			}
+
+			virtual void GatherInput(const FReal Dt, const int32 BeginIndex, const int32 EndIndex) override final
+			{
+				SolverA->GatherInput(Dt, BeginIndex, EndIndex);
+				// Solver B will ge gathered in PreApply
+			}
+
+			virtual void ScatterOutput(const FReal Dt) override final
+			{
+				SolverA->ScatterOutput(Dt);
+				SolverBodyContainerA = nullptr;
+				// Do not scatter output from B!
+			}
+
+			virtual void ScatterOutput(const FReal Dt, const int32 BeginIndex, const int32 EndIndex) override final
+			{
+				SolverA->ScatterOutput(Dt, BeginIndex, EndIndex);
+				SolverBodyContainerA = nullptr;
+				// Do not scatter output from B!
+			}
+
+			virtual void PreApplyPositionConstraints(const FReal Dt) override final
+			{
+				// Copy all the solver bodies for use by the second solver
+				SolverBodyContainerA->CopyTo(SolverBodyContainerB);
+
+				// Bind SolverB to the copied bodies
+				SolverB->AddBodies(SolverBodyContainerB);
+
+				// Gather the second solver's constraints
+				SolverB->GatherInput(Dt);
+
+				SolverA->PreApplyPositionConstraints(Dt);
+				SolverB->PreApplyPositionConstraints(Dt);
+
+				CallABTestFunctor(ESolverPhase::PreApplyPositionConstraints);
+			}
+
+			virtual void ApplyPositionConstraints(const FReal Dt, const int32 It, const int32 NumIts) override final
+			{
+				SolverA->ApplyPositionConstraints(Dt, It, NumIts);
+				SolverB->ApplyPositionConstraints(Dt, It, NumIts);
+
+				CallABTestFunctor(ESolverPhase::PostApplyPositionConstraints);
+			}
+
+			virtual void PreApplyVelocityConstraints(const FReal Dt) override final
+			{
+				SolverBodyContainerB.SetImplicitVelocities(Dt);
+
+				SolverA->PreApplyVelocityConstraints(Dt);
+				SolverB->PreApplyVelocityConstraints(Dt);
+			}
+
+			virtual void ApplyVelocityConstraints(const FReal Dt, const int32 It, const int32 NumIts) override final
+			{
+				SolverA->ApplyVelocityConstraints(Dt, It, NumIts);
+				SolverB->ApplyVelocityConstraints(Dt, It, NumIts);
+
+				CallABTestFunctor(ESolverPhase::PostApplyVelocityConstraints);
+			}
+
+			virtual void PreApplyProjectionConstraints(const FReal Dt) override final
+			{
+				SolverBodyContainerB.ApplyCorrections();
+
+				SolverA->PreApplyProjectionConstraints(Dt);
+				SolverB->PreApplyProjectionConstraints(Dt);
+			}
+
+			virtual void ApplyProjectionConstraints(const FReal Dt, const int32 It, const int32 NumIts) override final
+			{
+				SolverA->ApplyProjectionConstraints(Dt, It, NumIts);
+				SolverB->ApplyProjectionConstraints(Dt, It, NumIts);
+
+				CallABTestFunctor(ESolverPhase::PostApplyProjectionConstraints);
+			}
+
+		protected:
+			void CallABTestFunctor(const ESolverPhase Phase)
+			{
+				if (ABTestFunctor != nullptr)
+				{
+					ABTestFunctor(Phase, *SolverA, *SolverB, *SolverBodyContainerA, SolverBodyContainerB);
+				}
+			}
+
+			TUniquePtr<FContainerSolverTypeA> SolverA;
+			TUniquePtr<FContainerSolverTypeB> SolverB;
+			FSolverBodyContainer* SolverBodyContainerA;
+			FSolverBodyContainer SolverBodyContainerB;
+			FABTestFunctor ABTestFunctor;
+		};
+
+#endif // CHAOS_ABTEST_CONSTRAINTSOLVER_ENABLED
+
+	}	// namespace Private
+}	// namespace Chaos
\ No newline at end of file
diff --git a/Engine/Source/Runtime/Experimental/Chaos/Public/Chaos/Evolution/ConstraintGroupSolver.h b/Engine/Source/Runtime/Experimental/Chaos/Public/Chaos/Evolution/ConstraintGroupSolver.h
--- a/Engine/Source/Runtime/Experimental/Chaos/Public/Chaos/Evolution/ConstraintGroupSolver.h
+++ b/Engine/Source/Runtime/Experimental/Chaos/Public/Chaos/Evolution/ConstraintGroupSolver.h
@@ -102,6 +102,14 @@ namespace Chaos
 			*/
 			void GatherConstraints(const FReal Dt, const int32 BeginConstraintIndex, const int32 EndConstraintIndex);
 
+			/**
+			 * For additional processing after gathering all the data.
+			 * Calls PreApplyConstraints on each container solver.
+			 */
+			void PreApplyPositionConstraints(const FReal Dt);
+			void PreApplyVelocityConstraints(const FReal Dt);
+			void PreApplyProjectionConstraints(const FReal Dt);
+
 			/**
 			 * Apply positional constraints, and set the body velocities
 			*/
diff --git a/Engine/Source/Runtime/Experimental/Chaos/Public/Chaos/Evolution/SolverBody.h b/Engine/Source/Runtime/Experimental/Chaos/Public/Chaos/Evolution/SolverBody.h
--- a/Engine/Source/Runtime/Experimental/Chaos/Public/Chaos/Evolution/SolverBody.h
+++ b/Engine/Source/Runtime/Experimental/Chaos/Public/Chaos/Evolution/SolverBody.h
@@ -253,11 +253,13 @@ namespace Chaos
 		 * @brief Net world-space position displacement applied by the constraints
 		*/
 		inline const FSolverVec3& DP() const { return State.DP; }
+		inline void SetDP(const FSolverVec3& InDP) { State.DP = InDP; }
 
 		/**
 		 * @brief Net world-space rotation displacement applied by the constraints (axis-angle vector equivalent to angular velocity but for position)
 		*/
 		inline const FSolverVec3& DQ() const { return State.DQ; }
+		inline void SetDQ(const FSolverVec3& InDQ) { State.DQ = InDQ; }
 
 		/**
 		 * @brief Net world-space position correction applied by the constraints
@@ -422,6 +424,20 @@ namespace Chaos
 		*/
 		void UpdateRotationDependentState();
 
+
+
+		void PrefetchPositionSolverData() const
+		{
+			// The position solver only uses DP,DQ
+			FPlatformMisc::PrefetchBlock(&State.DP, 2 * sizeof(FSolverVec3));
+		}
+
+		void PrefetchVelocitySolverData() const
+		{
+			// The velocity solver only uses V,W
+			FPlatformMisc::PrefetchBlock(&State.V, 2 * sizeof(FSolverVec3));
+		}
+
 	private:
 
 		// The struct exists only so that we can use the variable names
@@ -444,10 +460,10 @@ namespace Chaos
 				InvM = 0;
 				DP = FSolverVec3(0);
 				DQ = FSolverVec3(0);
-				CP = FSolverVec3(0);
-				CQ = FSolverVec3(0);
 				V = FSolverVec3(0);
 				W = FSolverVec3(0);
+				CP = FSolverVec3(0);
+				CQ = FSolverVec3(0);
 				Level = 0;
 			}
 
@@ -488,6 +504,12 @@ namespace Chaos
 			// Net rotation delta applied by all constraints (constantly changing as we iterate over constraints)
 			FSolverVec3 DQ;
 
+			// World-space center of mass velocity
+			FSolverVec3 V;
+
+			// World-space center of mass angular velocity
+			FSolverVec3 W;
+
 			// Net position correction delta applied by all constraints (constantly changing as we iterate over constraints)
 			// Will translate the body without introducing linear velocity
 			FSolverVec3 CP;
@@ -496,12 +518,6 @@ namespace Chaos
 			// Will rotate the body without introducing angular velocity
 			FSolverVec3 CQ;
 
-			// World-space center of mass velocity
-			FSolverVec3 V;
-
-			// World-space center of mass angular velocity
-			FSolverVec3 W;
-
 			// Distance to a kinematic body (through the contact graph). Used by collision shock propagation
 			int32 Level;
 		};
diff --git a/Engine/Source/Runtime/Experimental/Chaos/Public/Chaos/Evolution/SolverBodyContainer.h b/Engine/Source/Runtime/Experimental/Chaos/Public/Chaos/Evolution/SolverBodyContainer.h
--- a/Engine/Source/Runtime/Experimental/Chaos/Public/Chaos/Evolution/SolverBodyContainer.h
+++ b/Engine/Source/Runtime/Experimental/Chaos/Public/Chaos/Evolution/SolverBodyContainer.h
@@ -116,6 +116,15 @@ namespace Chaos
 			bLocked = true;
 		}
 
+		// Used to AB test solvers - not intended for general use (very expensive!)
+		void CopyTo(FSolverBodyContainer& Dest)
+		{
+			Dest.SolverBodies = SolverBodies;
+			Dest.Particles = Particles;
+			Dest.ParticleToIndexMap = ParticleToIndexMap;
+			Dest.bLocked = bLocked;
+		}
+
 	private:
 		int32 AddParticle(FGenericParticleHandle InParticle);
 
diff --git a/Engine/Source/Runtime/Experimental/Chaos/Public/Chaos/Evolution/SolverConstraintContainer.h b/Engine/Source/Runtime/Experimental/Chaos/Public/Chaos/Evolution/SolverConstraintContainer.h
--- a/Engine/Source/Runtime/Experimental/Chaos/Public/Chaos/Evolution/SolverConstraintContainer.h
+++ b/Engine/Source/Runtime/Experimental/Chaos/Public/Chaos/Evolution/SolverConstraintContainer.h
@@ -78,6 +78,10 @@ namespace Chaos
 		virtual void ScatterOutput(const FReal Dt) = 0;
 		virtual void ScatterOutput(const FReal Dt, const int32 BeginIndex, const int32 EndIndex) = 0;
 
+		virtual void PreApplyPositionConstraints(const FReal Dt) {}
+		virtual void PreApplyVelocityConstraints(const FReal Dt) {}
+		virtual void PreApplyProjectionConstraints(const FReal Dt) {}
+
 		/**
 		 * Apply the position solve to all constraints in the container
 		*/
diff --git a/Engine/Source/Runtime/Experimental/ChaosCore/Public/Chaos/Simd4.h b/Engine/Source/Runtime/Experimental/ChaosCore/Public/Chaos/Simd4.h
new file mode 100644
--- /dev/null
+++ b/Engine/Source/Runtime/Experimental/ChaosCore/Public/Chaos/Simd4.h
@@ -0,0 +1,514 @@
+// Copyright Epic Games, Inc. All Rights Reserved.
+#pragma once
+
+#include "Chaos/SimdTypes.h"
+
+namespace Chaos
+{
+	namespace Private
+	{
+		// WIP: Not ready for public use. Keep in the Private namespace for now...
+		// NOTE: the below only work with TNumLanes == 4 for now
+
+		///////////////////////////////////////////////////////////////////////
+		//
+		// 4-wide float typedefs
+		//
+		///////////////////////////////////////////////////////////////////////
+
+		using FSimd4Realf = TSimdRealf<4>;
+		using FSimd4Vec3f = TSimdVec3f<4>;
+		using FSimd4Int32 = TSimdInt32<4>;
+		using FSimd4Selector = TSimdSelector<4>;
+
+		///////////////////////////////////////////////////////////////////////
+		//
+		// 4-wide template specializations
+		//
+		///////////////////////////////////////////////////////////////////////
+
+		template<>
+		FORCEINLINE TSimdSelector<4> TSimdSelector<4>::True()
+		{
+			TSimdSelector<4> Selector;
+			VectorIntStoreAligned(MakeVectorRegisterIntConstant(0xFFFFFFFF, 0xFFFFFFFF, 0xFFFFFFFF, 0xFFFFFFFF), Selector.V);
+			return Selector;
+		}
+
+		template<>
+		FORCEINLINE TSimdSelector<4> TSimdSelector<4>::False()
+		{
+			TSimdSelector<4> Selector;
+			VectorIntStoreAligned(MakeVectorRegisterIntConstant(0, 0, 0, 0), Selector.V);
+			return Selector;
+		}
+
+		template<>
+		FORCEINLINE TSimdInt32<4> TSimdInt32<4>::Make(const int32 I)
+		{
+			TSimdInt32<4> R;
+			VectorIntStoreAligned(MakeVectorRegisterInt(I, I, I, I), R.V);
+			return R;
+		}
+
+		template<>
+		FORCEINLINE TSimdInt32<4> TSimdInt32<4>::Zero()
+		{
+			TSimdInt32<4> I;
+			VectorIntStoreAligned(MakeVectorRegisterIntConstant(0, 0, 0, 0), I.V);
+			return I;
+		}
+
+		template<>
+		FORCEINLINE TSimdRealf<4> TSimdRealf<4>::Make(const float F)
+		{
+			TSimdRealf<4> Out;
+			VectorStoreAligned(VectorSetFloat1(F), Out.V);
+			return Out;
+		}
+
+		template<>
+		FORCEINLINE TSimdRealf<4> TSimdRealf<4>::Zero()
+		{
+			TSimdRealf<4> Out;
+			VectorStoreAligned(VectorZeroFloat(), Out.V);
+			return Out;
+		}
+
+		template<>
+		FORCEINLINE TSimdRealf<4> TSimdRealf<4>::One()
+		{
+			TSimdRealf<4> Out;
+			VectorStoreAligned(VectorOneFloat(), Out.V);
+			return Out;
+		}
+
+		
+#if 0
+		// Failed experiment!
+		template<>
+		FORCEINLINE TSimdVec3f<4> TSimdVec3f<4>::Make(const FVec3f* InV0, const FVec3f* InV1, const FVec3f* InV2, const FVec3f* InV3)
+		{
+			VectorRegister4f V0 = VectorLoadFloat3(InV0);
+			VectorRegister4f V1 = VectorLoadFloat3(InV1);
+			VectorRegister4f V2 = VectorLoadFloat3(InV2);
+			VectorRegister4f V3 = VectorLoadFloat3(InV3);
+
+			// Load V0 and V1 into the component vectors in first and second lanes
+			VectorRegister4f Mask0100 = MakeVectorRegister((uint32)-1, (uint32)0, (uint32)-1, (uint32)-1);
+			VectorRegister4f X = VectorSelect(Mask0100, VectorReplicate(V0, 0), VectorReplicate(V1, 0));
+			VectorRegister4f Y = VectorSelect(Mask0100, VectorReplicate(V0, 1), VectorReplicate(V1, 1));
+			VectorRegister4f Z = VectorSelect(Mask0100, VectorReplicate(V0, 2), VectorReplicate(V1, 2));
+
+			// Load V2 into the component vectors in the third lane
+			VectorRegister4f Mask0010 = MakeVectorRegister((uint32)-1, (uint32)-1, (uint32)0, (uint32)-1);
+			X = VectorSelect(Mask0010, X, VectorReplicate(V2, 0));
+			Y = VectorSelect(Mask0010, Y, VectorReplicate(V2, 1));
+			Z = VectorSelect(Mask0010, Z, VectorReplicate(V2, 2));
+
+			// Load V3 into the component vectors in the fourth lane
+			VectorRegister4f Mask0001 = MakeVectorRegister((uint32)-1, (uint32)-1, (uint32)-1, (uint32)0);
+			X = VectorSelect(Mask0001, X, VectorReplicate(V3, 0));
+			Y = VectorSelect(Mask0001, Y, VectorReplicate(V3, 1));
+			Z = VectorSelect(Mask0001, Z, VectorReplicate(V3, 2));
+
+			TSimdVec3f<4> Out;
+			VectorStoreAligned(X, Out.VX);
+			VectorStoreAligned(Y, Out.VY);
+			VectorStoreAligned(Z, Out.VZ);
+			return Out;
+		}
+#endif
+
+		///////////////////////////////////////////////////////////////////////
+		//
+		// 4-wide Logical operations
+		//
+		///////////////////////////////////////////////////////////////////////
+
+		FORCEINLINE bool SimdAnyTrue(const FSimd4Selector& InL)
+		{
+			VectorRegister4f L = VectorLoadAligned(InL.V);
+			return VectorMaskBits(L) != 0;
+		}
+
+		FORCEINLINE bool SimdAllTrue(const FSimd4Selector& InL)
+		{
+			VectorRegister4f L = VectorLoadAligned(InL.V);
+			return ((VectorMaskBits(L) & 0xF) == 0xF);
+		}
+
+		FORCEINLINE FSimd4Selector SimdOr(const FSimd4Selector& InL, const FSimd4Selector& InR)
+		{
+			VectorRegister4f L = VectorLoadAligned(InL.V);
+			VectorRegister4f R = VectorLoadAligned(InR.V);
+
+			VectorRegister4f Mask = VectorBitwiseOr(L, R);
+
+			FSimd4Selector Out;
+			VectorStoreAligned(Mask, Out.V);
+			return Out;
+		}
+
+		FORCEINLINE FSimd4Selector SimdAnd(const FSimd4Selector& InL, const FSimd4Selector& InR)
+		{
+			VectorRegister4f L = VectorLoadAligned(InL.V);
+			VectorRegister4f R = VectorLoadAligned(InR.V);
+
+			VectorRegister4f Mask = VectorBitwiseAnd(L, R);
+
+			FSimd4Selector Out;
+			VectorStoreAligned(Mask, Out.V);
+			return Out;
+		}
+
+		FORCEINLINE FSimd4Selector SimdEqual(const FSimd4Realf& InL, const FSimd4Realf& InR)
+		{
+			VectorRegister4f L = VectorLoadAligned(InL.V);
+			VectorRegister4f R = VectorLoadAligned(InR.V);
+
+			VectorRegister4f Mask = VectorCompareEQ(L, R);
+
+			FSimd4Selector Out;
+			VectorStoreAligned(Mask, Out.V);
+			return Out;
+		}
+
+		FORCEINLINE FSimd4Selector SimdNotEqual(const FSimd4Realf& InL, const FSimd4Realf& InR)
+		{
+			VectorRegister4f L = VectorLoadAligned(InL.V);
+			VectorRegister4f R = VectorLoadAligned(InR.V);
+
+			VectorRegister4f Mask = VectorCompareNE(L, R);
+
+			FSimd4Selector Out;
+			VectorStoreAligned(Mask, Out.V);
+			return Out;
+		}
+
+		FORCEINLINE FSimd4Selector SimdGreaterEqual(const FSimd4Int32& InL, const FSimd4Int32& InR)
+		{
+			VectorRegister4i L = VectorIntLoadAligned(InL.V);
+			VectorRegister4i R = VectorIntLoadAligned(InR.V);
+
+			union { VectorRegister4Float V; VectorRegister4Int I; } Mask;
+			Mask.I = VectorIntCompareGE(L, R);
+
+			FSimd4Selector Out;
+			VectorStoreAligned(Mask.V, Out.V);
+			return Out;
+		}
+
+		FORCEINLINE FSimd4Selector SimdGreater(const FSimd4Int32& InL, const FSimd4Int32& InR)
+		{
+			VectorRegister4i L = VectorIntLoadAligned(InL.V);
+			VectorRegister4i R = VectorIntLoadAligned(InR.V);
+
+			union { VectorRegister4Float V; VectorRegister4Int I; } Mask;
+			Mask.I = VectorIntCompareGT(L, R);
+
+			FSimd4Selector Out;
+			VectorStoreAligned(Mask.V, Out.V);
+			return Out;
+		}
+
+		FORCEINLINE FSimd4Selector SimdGreater(const FSimd4Realf& InL, const FSimd4Realf& InR)
+		{
+			VectorRegister4f L = VectorLoadAligned(InL.V);
+			VectorRegister4f R = VectorLoadAligned(InR.V);
+
+			VectorRegister4f Mask = VectorCompareGT(L, R);
+
+			FSimd4Selector Out;
+			VectorStoreAligned(Mask, Out.V);
+			return Out;
+		}
+
+		FORCEINLINE FSimd4Selector SimdLess(const FSimd4Int32& InL, const FSimd4Int32& InR)
+		{
+			VectorRegister4i L = VectorIntLoadAligned(InL.V);
+			VectorRegister4i R = VectorIntLoadAligned(InR.V);
+
+			VectorRegister4i Mask = VectorIntCompareLT(L, R);
+
+			FSimd4Selector Out;
+			VectorIntStoreAligned(Mask, Out.V);
+			return Out;
+		}
+
+		FORCEINLINE FSimd4Selector SimdLess(const FSimd4Realf& InL, const FSimd4Realf& InR)
+		{
+			VectorRegister4f L = VectorLoadAligned(InL.V);
+			VectorRegister4f R = VectorLoadAligned(InR.V);
+
+			VectorRegister4f Mask = VectorCompareLT(L, R);
+
+			FSimd4Selector Out;
+			VectorStoreAligned(Mask, Out.V);
+			return Out;
+		}
+
+		FORCEINLINE FSimd4Realf SimdSelect(const FSimd4Selector& InSelector, const FSimd4Realf& InL, const FSimd4Realf& InR)
+		{
+			VectorRegister4f L = VectorLoadAligned(InL.V);
+			VectorRegister4f R = VectorLoadAligned(InR.V);
+			VectorRegister4f Mask = VectorLoadAligned(InSelector.V);
+
+			VectorRegister4f Selected = VectorSelect(Mask, L, R);
+
+			FSimd4Realf Out;
+			VectorStoreAligned(Selected, Out.V);
+			return Out;
+		}
+
+		///////////////////////////////////////////////////////////////////////
+		//
+		// 4-wide Math operations
+		//
+		///////////////////////////////////////////////////////////////////////
+
+		FORCEINLINE FSimd4Realf SimdNegate(const FSimd4Realf& InL)
+		{
+			VectorRegister4f L = VectorLoadAligned(InL.V);
+
+			VectorRegister4f Neg = VectorNegate(L);
+
+			FSimd4Realf Out;
+			VectorStoreAligned(Neg, Out.V);
+			return Out;
+		}
+
+		FORCEINLINE FSimd4Realf SimdAdd(const FSimd4Realf& InL, const FSimd4Realf& InR)
+		{
+			VectorRegister4f L = VectorLoadAligned(InL.V);
+			VectorRegister4f R = VectorLoadAligned(InR.V);
+
+			VectorRegister4f Sum = VectorAdd(L, R);
+
+			FSimd4Realf Out;
+			VectorStoreAligned(Sum, Out.V);
+			return Out;
+		}
+
+		FORCEINLINE FSimd4Vec3f SimdAdd(const FSimd4Vec3f& L, const FSimd4Vec3f& R)
+		{
+			VectorRegister4f LX = VectorLoadAligned(L.VX);
+			VectorRegister4f LY = VectorLoadAligned(L.VY);
+			VectorRegister4f LZ = VectorLoadAligned(L.VZ);
+
+			VectorRegister4f RX = VectorLoadAligned(R.VX);
+			VectorRegister4f RY = VectorLoadAligned(R.VY);
+			VectorRegister4f RZ = VectorLoadAligned(R.VZ);
+
+			FSimd4Vec3f Out;
+			VectorStoreAligned(VectorAdd(LX, RX), Out.VX);
+			VectorStoreAligned(VectorAdd(LY, RY), Out.VY);
+			VectorStoreAligned(VectorAdd(LZ, RZ), Out.VZ);
+			return Out;
+		}
+
+		FORCEINLINE FSimd4Realf SimdSubtract(const FSimd4Realf& InL, const FSimd4Realf& InR)
+		{
+			VectorRegister4f L = VectorLoadAligned(InL.V);
+			VectorRegister4f R = VectorLoadAligned(InR.V);
+
+			VectorRegister4f Diff = VectorSubtract(L, R);
+
+			FSimd4Realf Out;
+			VectorStoreAligned(Diff, Out.V);
+			return Out;
+		}
+
+		FORCEINLINE FSimd4Vec3f SimdSubtract(const FSimd4Vec3f& L, const FSimd4Vec3f& R)
+		{
+			VectorRegister4f LX = VectorLoadAligned(L.VX);
+			VectorRegister4f LY = VectorLoadAligned(L.VY);
+			VectorRegister4f LZ = VectorLoadAligned(L.VZ);
+
+			VectorRegister4f RX = VectorLoadAligned(R.VX);
+			VectorRegister4f RY = VectorLoadAligned(R.VY);
+			VectorRegister4f RZ = VectorLoadAligned(R.VZ);
+
+			FSimd4Vec3f Out;
+			VectorStoreAligned(VectorSubtract(LX, RX), Out.VX);
+			VectorStoreAligned(VectorSubtract(LY, RY), Out.VY);
+			VectorStoreAligned(VectorSubtract(LZ, RZ), Out.VZ);
+			return Out;
+		}
+
+		FORCEINLINE FSimd4Realf SimdMultiply(const FSimd4Realf& InL, const FSimd4Realf& InR)
+		{
+			VectorRegister4f L = VectorLoadAligned(InL.V);
+			VectorRegister4f R = VectorLoadAligned(InR.V);
+
+			VectorRegister4f Product = VectorMultiply(L, R);
+
+			FSimd4Realf Out;
+			VectorStoreAligned(Product, Out.V);
+			return Out;
+		}
+
+		FORCEINLINE FSimd4Vec3f SimdMultiply(const FSimd4Vec3f& L, const FSimd4Realf& InR)
+		{
+			VectorRegister4f LX = VectorLoadAligned(L.VX);
+			VectorRegister4f LY = VectorLoadAligned(L.VY);
+			VectorRegister4f LZ = VectorLoadAligned(L.VZ);
+			VectorRegister4f R = VectorLoadAligned(InR.V);
+
+			FSimd4Vec3f Out;
+			VectorStoreAligned(VectorMultiply(LX, R), Out.VX);
+			VectorStoreAligned(VectorMultiply(LY, R), Out.VY);
+			VectorStoreAligned(VectorMultiply(LZ, R), Out.VZ);
+			return Out;
+		}
+
+		FORCEINLINE FSimd4Vec3f SimdMultiply(const FSimd4Realf& L, const FSimd4Vec3f& R)
+		{
+			return SimdMultiply(R, L);
+		}
+
+		FORCEINLINE FSimd4Vec3f SimdMultiply(const FSimd4Vec3f& L, const FSimd4Vec3f& R)
+		{
+			VectorRegister4f LX = VectorLoadAligned(L.VX);
+			VectorRegister4f LY = VectorLoadAligned(L.VY);
+			VectorRegister4f LZ = VectorLoadAligned(L.VZ);
+
+			VectorRegister4f RX = VectorLoadAligned(R.VX);
+			VectorRegister4f RY = VectorLoadAligned(R.VY);
+			VectorRegister4f RZ = VectorLoadAligned(R.VZ);
+
+			FSimd4Vec3f Out;
+			VectorStoreAligned(VectorMultiply(LX, RX), Out.VX);
+			VectorStoreAligned(VectorMultiply(LY, RY), Out.VY);
+			VectorStoreAligned(VectorMultiply(LZ, RZ), Out.VZ);
+			return Out;
+		}
+
+		FORCEINLINE FSimd4Vec3f SimdMultiplyAdd(const FSimd4Vec3f& L, const FSimd4Vec3f& R, const FSimd4Vec3f& Acc)
+		{
+			VectorRegister4f LX = VectorLoadAligned(L.VX);
+			VectorRegister4f LY = VectorLoadAligned(L.VY);
+			VectorRegister4f LZ = VectorLoadAligned(L.VZ);
+
+			VectorRegister4f RX = VectorLoadAligned(R.VX);
+			VectorRegister4f RY = VectorLoadAligned(R.VY);
+			VectorRegister4f RZ = VectorLoadAligned(R.VZ);
+
+			VectorRegister4f AX = VectorLoadAligned(Acc.VX);
+			VectorRegister4f AY = VectorLoadAligned(Acc.VY);
+			VectorRegister4f AZ = VectorLoadAligned(Acc.VZ);
+
+			FSimd4Vec3f Out;
+			VectorStoreAligned(VectorMultiplyAdd(LX, RX, AX), Out.VX);
+			VectorStoreAligned(VectorMultiplyAdd(LY, RY, AY), Out.VY);
+			VectorStoreAligned(VectorMultiplyAdd(LZ, RZ, AZ), Out.VZ);
+			return Out;
+		}
+
+		FORCEINLINE FSimd4Realf SimdDivide(const FSimd4Realf& InL, const FSimd4Realf& InR)
+		{
+			VectorRegister4f L = VectorLoadAligned(InL.V);
+			VectorRegister4f R = VectorLoadAligned(InR.V);
+
+			VectorRegister4f Ratio = VectorDivide(L, R);
+
+			FSimd4Realf Out;
+			VectorStoreAligned(Ratio, Out.V);
+			return Out;
+		}
+
+		FORCEINLINE FSimd4Vec3f SimdCrossProduct(const FSimd4Vec3f& L, const FSimd4Vec3f& R)
+		{
+			VectorRegister4f LX = VectorLoadAligned(L.VX);
+			VectorRegister4f LY = VectorLoadAligned(L.VY);
+			VectorRegister4f LZ = VectorLoadAligned(L.VZ);
+
+			VectorRegister4f RX = VectorLoadAligned(R.VX);
+			VectorRegister4f RY = VectorLoadAligned(R.VY);
+			VectorRegister4f RZ = VectorLoadAligned(R.VZ);
+
+			FSimd4Vec3f Out;
+			VectorStoreAligned(VectorSubtract(VectorMultiply(LY, RZ), VectorMultiply(LZ, RY)), Out.VX);
+			VectorStoreAligned(VectorSubtract(VectorMultiply(LZ, RX), VectorMultiply(LX, RZ)), Out.VY);
+			VectorStoreAligned(VectorSubtract(VectorMultiply(LX, RY), VectorMultiply(LY, RX)), Out.VZ);
+			return Out;
+		}
+
+		FORCEINLINE FSimd4Realf SimdDotProduct(const FSimd4Vec3f& L, const FSimd4Vec3f& R)
+		{
+			VectorRegister4f LX = VectorLoadAligned(L.VX);
+			VectorRegister4f LY = VectorLoadAligned(L.VY);
+			VectorRegister4f LZ = VectorLoadAligned(L.VZ);
+
+			VectorRegister4f RX = VectorLoadAligned(R.VX);
+			VectorRegister4f RY = VectorLoadAligned(R.VY);
+			VectorRegister4f RZ = VectorLoadAligned(R.VZ);
+
+			VectorRegister4f Dot = VectorMultiply(LX, RX);
+			Dot = VectorMultiplyAdd(LY, RY, Dot);
+			Dot = VectorMultiplyAdd(LZ, RZ, Dot);
+
+			FSimd4Realf Out;
+			VectorStoreAligned(Dot, Out.V);
+			return Out;
+		}
+
+		FORCEINLINE FSimd4Realf SimdSquare(const FSimd4Realf& InV)
+		{
+			VectorRegister4f V = VectorLoadAligned(InV.V);
+
+			VectorRegister4f Product = VectorMultiply(V, V);
+
+			FSimd4Realf Out;
+			VectorStoreAligned(Product, Out.V);
+			return Out;
+		}
+
+		FORCEINLINE FSimd4Realf SimdSqrt(const FSimd4Realf& InV)
+		{
+			VectorRegister4f V = VectorLoadAligned(InV.V);
+
+			VectorRegister4f Sqrt = VectorSqrt(V);
+
+			FSimd4Realf Out;
+			VectorStoreAligned(Sqrt, Out.V);
+			return Out;
+		}
+
+		FORCEINLINE FSimd4Realf SimdInvSqrt(const FSimd4Realf& InV)
+		{
+			VectorRegister4f V = VectorLoadAligned(InV.V);
+
+			VectorRegister4f Sqrt = VectorReciprocalSqrt(V);
+
+			FSimd4Realf Out;
+			VectorStoreAligned(Sqrt, Out.V);
+			return Out;
+		}
+
+		FORCEINLINE FSimd4Realf SimdMin(const FSimd4Realf& InL, const FSimd4Realf& InR)
+		{
+			VectorRegister4f L = VectorLoadAligned(InL.V);
+			VectorRegister4f R = VectorLoadAligned(InR.V);
+
+			VectorRegister4f Min = VectorMin(L, R);
+
+			FSimd4Realf Out;
+			VectorStoreAligned(Min, Out.V);
+			return Out;
+		}
+
+		FORCEINLINE FSimd4Realf SimdMax(const FSimd4Realf& InL, const FSimd4Realf& InR)
+		{
+			VectorRegister4f L = VectorLoadAligned(InL.V);
+			VectorRegister4f R = VectorLoadAligned(InR.V);
+
+			VectorRegister4f Max = VectorMax(L, R);
+
+			FSimd4Realf Out;
+			VectorStoreAligned(Max, Out.V);
+			return Out;
+		}
+
+	}
+}
\ No newline at end of file
diff --git a/Engine/Source/Runtime/Experimental/ChaosCore/Public/Chaos/SimdTypes.h b/Engine/Source/Runtime/Experimental/ChaosCore/Public/Chaos/SimdTypes.h
new file mode 100644
--- /dev/null
+++ b/Engine/Source/Runtime/Experimental/ChaosCore/Public/Chaos/SimdTypes.h
@@ -0,0 +1,249 @@
+// Copyright Epic Games, Inc. All Rights Reserved.
+#pragma once
+
+#include "Chaos/Real.h"
+#include "Chaos/Vector.h"
+
+namespace Chaos
+{
+	namespace Private
+	{
+		//
+		// WIP: Not ready for public use. Keep in the Private namespace for now...
+		// 
+		// NOTE: Only TNumLanes == 4 is supported for now
+		//
+
+		/**
+		* \file SimdTypes.h
+		* 
+		* This file contains a set of classes to help with SIMD data storage and 
+		* operations on that data.
+		* E.g., sets of N 3-vectors are stored as 3 arrays of N floats (one for
+		* the X elements, one for Y, one for Z).
+		* 
+		* Arranging data this way allows us to process N datasets for the cost of 1,
+		* and eliminates the need to call any horizontal vector operations (such as
+		* is required for the summation in a dot product of two vectors).
+		* 
+		* Examples:
+		*
+		*	For a real-world example see methods in FPBDCollisionSolverSimd like
+		*	SolvePositionNoFriction, SolvePositionWithFriction, and SolveVelocity.
+		* 
+		*	// Hopefully your data is already stored in SIMD lanes, but if not you
+		*	// can load it into SIMD registers with the factory functions and SetValue 
+		*	// methods. This should be avoided if possible.
+		* 
+		*	FSimd4Realf R0;
+		*	SimdR0.SetValue(0, -1);
+		*	SimdR0.SetValue(1, 0);
+		*	SimdR0.SetValue(2, 1);
+		*	SimdR0.SetValue(3, 2);
+		*	FSimd4Realf R1;
+		*	SimdR1.SetValue(0, 3);
+		*	SimdR1.SetValue(1, 4);
+		*	SimdR1.SetValue(2, 5);
+		*	SimdR1.SetValue(3, 6);
+		*
+		*	FSim4Selector IsR0GreaterThanZero = SimdGreater(R0, FSimd4Realf::Zero());
+		* 
+		*	// If it's likely that all lanes will agree on IsGreaterZero, it may be worth
+		*	// skipping the code block when possible. This will depend how expensive the
+		*	// branch is vs the code we skip and the probability of doing so.
+		*	FSimd4Realf R2 = FSimd4Realf::Zero();
+		*	if (SimdAnyTrue(IsR0GreaterThanZero))
+		*	{
+		*		R2 = SimdSelect(IsR0GreaterThanZero, R0, R1);	//	-> (3,4,1,2)
+		*	}
+		*/
+
+
+		/**
+		* The result of Simd comparison operations and used in SimdSelect.
+		* 
+		*	// Replace any non-positive values in R0 with the value from R1
+		*	FSim4Selector IsR0GreaterThanZero = SimdGreater(R0, FSimd4Realf::Zero());
+		*	R0 = SimdSelect(IsV0GreaterThanZero, R0, R1);
+		* 
+		*/
+		template<int TNumLanes>
+		struct TSimdSelector
+		{
+			alignas(UE_SSE_FLOAT_ALIGNMENT) float V[TNumLanes];
+
+			FORCEINLINE void SetValue(const int LaneIndex, bool B)
+			{
+				reinterpret_cast<uint32*>(V)[LaneIndex] = B ? 0xFFFFFFFF : 0;
+			}
+
+			FORCEINLINE bool GetValue(const int LaneIndex) const
+			{
+				// Can't do float comparison because its a NaN for true
+				return (reinterpret_cast<const uint32*>(V)[LaneIndex] != 0);
+			}
+
+			static TSimdSelector<TNumLanes> True();
+			static TSimdSelector<TNumLanes> False();
+		};
+
+		/**
+		* Used to store any value type in a TNumLanes-wide set. This is used primarily whenyou have 
+		* non-numeric data that varies per lane. E.g., a pointer to some per-lane shared data.
+		*/
+		template<typename T, int TNumLanes>
+		struct TSimdValue
+		{
+			using ValueType = T;
+			alignas(UE_SSE_FLOAT_ALIGNMENT) ValueType V[TNumLanes];
+
+			FORCEINLINE void SetValue(const int32 LaneIndex, const ValueType F)
+			{
+				V[LaneIndex] = F;
+			}
+
+			FORCEINLINE ValueType GetValue(const int32 LaneIndex) const
+			{
+				return V[LaneIndex];
+			}
+
+			FORCEINLINE void SetValues(const ValueType F)
+			{
+				for (int32 LaneIndex = 0; LaneIndex < TNumLanes; ++LaneIndex)
+				{
+					SetValue(LaneIndex, F);
+				}
+			}
+		};
+
+		/**
+		* A TNumLanes-wide set of integers.
+		*/
+		template<int TNumLanes>
+		struct TSimdInt32
+		{
+			alignas(UE_SSE_FLOAT_ALIGNMENT) int32 V[TNumLanes];
+
+			FORCEINLINE TSimdInt32()
+			{
+			}
+
+			FORCEINLINE void SetValue(const int32 LaneIndex, const int32 I)
+			{
+				V[LaneIndex] = I;
+			}
+
+			FORCEINLINE int32 GetValue(const int32 LaneIndex) const
+			{
+				return V[LaneIndex];
+			}
+
+			FORCEINLINE void SetValues(const int32 I)
+			{
+				*this = Make(I);
+			}
+
+			static TSimdInt32<TNumLanes> Make(const int32 I);
+			static TSimdInt32<TNumLanes> Zero();
+		};
+
+		/**
+		* A TNumLanes-wide set of single-precision floats.
+		*/
+		template<int TNumLanes>
+		struct TSimdRealf
+		{
+			alignas(UE_SSE_FLOAT_ALIGNMENT) float V[TNumLanes];
+
+			FORCEINLINE TSimdRealf()
+			{
+			}
+
+			FORCEINLINE void SetValue(const int32 LaneIndex, const float F)
+			{
+				V[LaneIndex] = F;
+			}
+
+			FORCEINLINE float GetValue(const int32 LaneIndex) const
+			{
+				return V[LaneIndex];
+			}
+
+			FORCEINLINE void SetValues(const float F)
+			{
+				*this = Make(F);
+			}
+
+			static TSimdRealf Make(const float F);
+			static TSimdRealf Zero();
+			static TSimdRealf One();
+		};
+
+		/**
+		* A TNumLanes-wide set of floating point 3-vectors. The vector elements
+		* are stored as 3 vectors of Xs, Ys and Zs.
+		*
+		*	FSimd4Vec3f C = SimdCrossProduct(A, B);
+		* 
+		*/
+		template<int TNumLanes>
+		struct TSimdVec3f
+		{
+			static_assert(TNumLanes == 4, "Other sizes not yet supported");
+
+			alignas(UE_SSE_FLOAT_ALIGNMENT) float VX[TNumLanes];
+			alignas(UE_SSE_FLOAT_ALIGNMENT) float VY[TNumLanes];
+			alignas(UE_SSE_FLOAT_ALIGNMENT) float VZ[TNumLanes];
+
+			FORCEINLINE TSimdVec3f()
+			{
+			}
+
+			FORCEINLINE TSimdVec3f(const FVec3f& V)
+			{
+				SetValues(V);
+			}
+
+			FORCEINLINE void SetValue(const int32 LaneIndex, const FVec3f& V)
+			{
+				VX[LaneIndex] = V.X;
+				VY[LaneIndex] = V.Y;
+				VZ[LaneIndex] = V.Z;
+			};
+
+			FORCEINLINE FVec3f GetValue(const int32 LaneIndex) const
+			{
+				return FVec3f(VX[LaneIndex], VY[LaneIndex], VZ[LaneIndex]);
+			}
+
+			FORCEINLINE void SetValues(const FVec3f& V)
+			{
+				*this = Make(V);
+			}
+
+			FORCEINLINE static TSimdVec3f Make(const FVec3f& V)
+			{
+				TSimdVec3f Out;
+				for (int32 LaneIndex = 0; LaneIndex < TNumLanes; ++LaneIndex)
+				{
+					Out.SetValue(LaneIndex, V);
+				}
+				return Out;
+			}
+
+			FORCEINLINE static TSimdVec3f Make(const FVec3f& V0, const FVec3f& V1, const FVec3f& V2, const FVec3f& V3)
+			{
+				TSimdVec3f Out;
+				if (TNumLanes == 4)
+				{
+					Out.SetValue(0, V0);
+					Out.SetValue(1, V1);
+					Out.SetValue(2, V2);
+					Out.SetValue(3, V3);
+
+				}
+				return Out;
+			}
+		};
+	}
+}
\ No newline at end of file
