Fix CIS in SimdTypes.h

diff --git a/Engine/Source/Runtime/Experimental/ChaosCore/Public/Chaos/SimdTypes.h b/Engine/Source/Runtime/Experimental/ChaosCore/Public/Chaos/SimdTypes.h
--- a/Engine/Source/Runtime/Experimental/ChaosCore/Public/Chaos/SimdTypes.h
+++ b/Engine/Source/Runtime/Experimental/ChaosCore/Public/Chaos/SimdTypes.h
@@ -14,6 +14,8 @@ namespace Chaos
 		// NOTE: Only TNumLanes == 4 is supported for now
 		//
 
+		static const int SimdAlignment = 16;
+
 		/**
 		* \file SimdTypes.h
 		* 
@@ -70,7 +72,7 @@ namespace Chaos
 		template<int TNumLanes>
 		struct TSimdSelector
 		{
-			alignas(UE_SSE_FLOAT_ALIGNMENT) float V[TNumLanes];
+			alignas(SimdAlignment) float V[TNumLanes];
 
 			FORCEINLINE void SetValue(const int LaneIndex, bool B)
 			{
@@ -95,7 +97,7 @@ namespace Chaos
 		struct TSimdValue
 		{
 			using ValueType = T;
-			alignas(UE_SSE_FLOAT_ALIGNMENT) ValueType V[TNumLanes];
+			alignas(SimdAlignment) ValueType V[TNumLanes];
 
 			FORCEINLINE void SetValue(const int32 LaneIndex, const ValueType F)
 			{
@@ -122,7 +124,7 @@ namespace Chaos
 		template<int TNumLanes>
 		struct TSimdInt32
 		{
-			alignas(UE_SSE_FLOAT_ALIGNMENT) int32 V[TNumLanes];
+			alignas(SimdAlignment) int32 V[TNumLanes];
 
 			FORCEINLINE TSimdInt32()
 			{
@@ -153,7 +155,7 @@ namespace Chaos
 		template<int TNumLanes>
 		struct TSimdRealf
 		{
-			alignas(UE_SSE_FLOAT_ALIGNMENT) float V[TNumLanes];
+			alignas(SimdAlignment) float V[TNumLanes];
 
 			FORCEINLINE TSimdRealf()
 			{
@@ -191,9 +193,9 @@ namespace Chaos
 		{
 			static_assert(TNumLanes == 4, "Other sizes not yet supported");
 
-			alignas(UE_SSE_FLOAT_ALIGNMENT) float VX[TNumLanes];
-			alignas(UE_SSE_FLOAT_ALIGNMENT) float VY[TNumLanes];
-			alignas(UE_SSE_FLOAT_ALIGNMENT) float VZ[TNumLanes];
+			alignas(SimdAlignment) float VX[TNumLanes];
+			alignas(SimdAlignment) float VY[TNumLanes];
+			alignas(SimdAlignment) float VZ[TNumLanes];
 
 			FORCEINLINE TSimdVec3f()
 			{
