abandon after level on the geometry collection should make sure collision is removed from the server. add an option to make sure collision is removed from the client as well.

diff --git a/Engine/Source/Runtime/Engine/Private/PhysicsEngine/Experimental/PhysScene_Chaos.cpp b/Engine/Source/Runtime/Engine/Private/PhysicsEngine/Experimental/PhysScene_Chaos.cpp
--- a/Engine/Source/Runtime/Engine/Private/PhysicsEngine/Experimental/PhysScene_Chaos.cpp
+++ b/Engine/Source/Runtime/Engine/Private/PhysicsEngine/Experimental/PhysScene_Chaos.cpp
@@ -1772,17 +1772,23 @@ void FPhysScene_Chaos::OnSyncBodies(Chaos::FPhysicsSolverBase* Solver)
 			// Don't pass in anything here so we don't end up locking anything because we can assume the scene is already locked.
 			FLockedWritePhysicsObjectExternalInterface Interface = FPhysicsObjectExternalInterface::LockWrite({});
 
+			auto ShouldGCParticleBeInAccelerationStructure = [&Interface](Chaos::FPhysicsObjectHandle Handle)
+			{
+				// It's possible to be an enabled particle and not qualify for query collisions if the GC particle has been replication abandoned. 
+				return !Interface->AreAllDisabled({ &Handle, 1 }) && Interface->AreAllShapesQueryEnabled({ &Handle, 1 });
+			};
+
 			TArray<Chaos::FPhysicsObjectHandle> ActiveHandles = Proxy->GetAllPhysicsObjects().FilterByPredicate(
-				[&Interface](Chaos::FPhysicsObjectHandle Handle)
+				[&ShouldGCParticleBeInAccelerationStructure](Chaos::FPhysicsObjectHandle Handle)
 				{
-					return !Interface->AreAllDisabled({ &Handle, 1 });
+					return ShouldGCParticleBeInAccelerationStructure(Handle);
 				}
 			);
 
 			TArray<Chaos::FPhysicsObjectHandle> DisabledHandles = Proxy->GetAllPhysicsObjects().FilterByPredicate(
-				[&Interface](Chaos::FPhysicsObjectHandle Handle)
+				[&ShouldGCParticleBeInAccelerationStructure](Chaos::FPhysicsObjectHandle Handle)
 				{
-					return Interface->AreAllDisabled({ &Handle, 1 });
+					return !ShouldGCParticleBeInAccelerationStructure(Handle);
 				}
 			);
 			Interface->AddToSpatialAcceleration(ActiveHandles, GetSpacialAcceleration());
diff --git a/Engine/Source/Runtime/Experimental/Chaos/Private/Chaos/PhysicsObjectInterface.cpp b/Engine/Source/Runtime/Experimental/Chaos/Private/Chaos/PhysicsObjectInterface.cpp
--- a/Engine/Source/Runtime/Experimental/Chaos/Private/Chaos/PhysicsObjectInterface.cpp
+++ b/Engine/Source/Runtime/Experimental/Chaos/Private/Chaos/PhysicsObjectInterface.cpp
@@ -377,6 +377,39 @@ namespace Chaos
 		return bDisabled;
 	}
 
+	template<EThreadContext Id>
+	bool FReadPhysicsObjectInterface<Id>::AreAllShapesQueryEnabled(TArrayView<FPhysicsObjectHandle> InObjects)
+	{
+		// This operation isn't particularly defined on the physics thread?
+		if constexpr (Id == EThreadContext::Internal)
+		{
+			return false;
+		}
+
+		if (InObjects.IsEmpty())
+		{
+			return false;
+		}
+
+		for (FPhysicsObjectHandle Object : InObjects)
+		{
+			if (Object)
+			{
+				if (TThreadParticle<Id>* Particle = Object->GetParticle<Id>())
+				{
+					for (const TUniquePtr<FPerShapeData>& ShapeData : Particle->ShapesArray())
+					{
+						if (!ShapeData->GetCollisionData().bQueryCollision)
+						{
+							return false;
+						}
+					}
+				}
+			}
+		}
+		return true;
+	}
+
 	template<EThreadContext Id>
 	float FReadPhysicsObjectInterface<Id>::GetMass(TArrayView<FPhysicsObjectHandle> InObjects)
 	{
@@ -615,6 +648,50 @@ namespace Chaos
 		}
 	}
 
+	template<EThreadContext Id>
+	void FWritePhysicsObjectInterface<Id>::UpdateShapeCollisionFlags(TArrayView<FPhysicsObjectHandle> InObjects, bool bSimCollision, bool bQueryCollision)
+	{
+		for (FPhysicsObjectHandle Object : InObjects)
+		{
+			if (!Object)
+			{
+				continue;
+			}
+
+			if (TThreadParticle<Id>* Particle = Object->GetParticle<Id>())
+			{
+				for (const TUniquePtr<Chaos::FPerShapeData>& ShapeData : Particle->ShapesArray())
+				{
+					FCollisionData Data = ShapeData->GetCollisionData();
+					Data.bSimCollision = bSimCollision;
+					Data.bQueryCollision = bQueryCollision;
+					ShapeData->SetCollisionData(Data);
+				}
+			}
+		}
+	}
+
+	template<EThreadContext Id>
+	void FWritePhysicsObjectInterface<Id>::UpdateShapeFilterData(TArrayView<FPhysicsObjectHandle> InObjects, const FCollisionFilterData& QueryData, const FCollisionFilterData& SimData)
+	{
+		for (FPhysicsObjectHandle Object : InObjects)
+		{
+			if (!Object)
+			{
+				continue;
+			}
+
+			if (TThreadParticle<Id>* Particle = Object->GetParticle<Id>())
+			{
+				for (const TUniquePtr<Chaos::FPerShapeData>& ShapeData : Particle->ShapesArray())
+				{
+					ShapeData->SetQueryData(QueryData);
+					ShapeData->SetSimData(SimData);
+				}
+			}
+		}
+	}
+
 	void FPhysicsObjectInterface::SetName(FPhysicsObjectHandle Object, const FName& InName)
 	{
 		if (!Object)
diff --git a/Engine/Source/Runtime/Experimental/Chaos/Public/Chaos/PhysicsObjectInterface.h b/Engine/Source/Runtime/Experimental/Chaos/Public/Chaos/PhysicsObjectInterface.h
--- a/Engine/Source/Runtime/Experimental/Chaos/Public/Chaos/PhysicsObjectInterface.h
+++ b/Engine/Source/Runtime/Experimental/Chaos/Public/Chaos/PhysicsObjectInterface.h
@@ -73,6 +73,7 @@ namespace Chaos
 		bool AreAllRigidBody(TArrayView<FPhysicsObjectHandle> InObjects);
 		bool AreAllDynamic(TArrayView<FPhysicsObjectHandle> InObjects);
 		bool AreAllDisabled(TArrayView<FPhysicsObjectHandle> InObjects);
+		bool AreAllShapesQueryEnabled(TArrayView<FPhysicsObjectHandle> InObjects);
 		float GetMass(TArrayView<FPhysicsObjectHandle> InObjects);
 		FBox GetBounds(TArrayView<FPhysicsObjectHandle> InObjects);
 		FBox GetWorldBounds(TArrayView<FPhysicsObjectHandle> InObjects);
@@ -97,6 +98,9 @@ namespace Chaos
 		void WakeUp(TArrayView<FPhysicsObjectHandle> InObjects);
 		void AddForce(TArrayView<FPhysicsObjectHandle> InObjects, const FVector& Force, bool bInvalidate);
 		void AddTorque(TArrayView<FPhysicsObjectHandle> InObjects, const FVector& Torque, bool bInvalidate);
+		
+		void UpdateShapeCollisionFlags(TArrayView<FPhysicsObjectHandle> InObjects, bool bSimCollision, bool bQueryCollision);
+		void UpdateShapeFilterData(TArrayView<FPhysicsObjectHandle> InObjects, const FCollisionFilterData& QueryData, const FCollisionFilterData& SimData);
 
 		template<typename TPayloadType, typename T, int d>
 		void AddToSpatialAcceleration(TArrayView<FPhysicsObjectHandle> InObjects, ISpatialAcceleration<TPayloadType, T, d>* SpatialAcceleration)
diff --git a/Engine/Source/Runtime/Experimental/GeometryCollectionEngine/Private/GeometryCollection/GeometryCollectionComponent.cpp b/Engine/Source/Runtime/Experimental/GeometryCollectionEngine/Private/GeometryCollection/GeometryCollectionComponent.cpp
--- a/Engine/Source/Runtime/Experimental/GeometryCollectionEngine/Private/GeometryCollection/GeometryCollectionComponent.cpp
+++ b/Engine/Source/Runtime/Experimental/GeometryCollectionEngine/Private/GeometryCollection/GeometryCollectionComponent.cpp
@@ -11,6 +11,7 @@
 #include "ChaosStats.h"
 #include "ComponentRecreateRenderStateContext.h"
 #include "Components/BoxComponent.h"
+#include "Engine/CollisionProfile.h"
 #include "Engine/Engine.h"
 #include "Engine/InstancedStaticMesh.h"
 #include "Field/FieldSystemComponent.h"
@@ -317,6 +318,7 @@ UGeometryCollectionComponent::UGeometryCollectionComponent(const FObjectInitiali
 #endif
 	, bEnableReplication(false)
 	, bEnableAbandonAfterLevel(true)
+	, AbandonedCollisionProfileName(UCollisionProfile::CustomCollisionProfileName)
 	, ReplicationAbandonClusterLevel_DEPRECATED(0)
 	, ReplicationAbandonAfterLevel(0)
 	, bRenderStateDirty(true)
@@ -1525,7 +1527,7 @@ void UGeometryCollectionComponent::UpdateRepData()
 
 		bool bClustersChanged = false;
 
-		const FPBDRigidsSolver* Solver = PhysicsProxy->GetSolver<Chaos::FPBDRigidsSolver>();
+		FPBDRigidsSolver* Solver = PhysicsProxy->GetSolver<Chaos::FPBDRigidsSolver>();
 		const FRigidClustering& RigidClustering = Solver->GetEvolution()->GetRigidClustering();
 
 		const TManagedArray<int32>* InitialLevels = PhysicsProxy->GetPhysicsCollection().FindAttribute<int32>("InitialLevel", FGeometryCollection::TransformGroup);
@@ -1533,6 +1535,7 @@ void UGeometryCollectionComponent::UpdateRepData()
 
 		//see if we have any new clusters that are enabled
 		TSet<FPBDRigidClusteredParticleHandle*> Processed;
+
 		for (FPBDRigidClusteredParticleHandle* Particle : PhysicsProxy->GetParticles())
 		{
 			// Particle can be null if we have embedded geometry 
@@ -1607,12 +1610,13 @@ void UGeometryCollectionComponent::UpdateRepData()
 						if (!Root->Disabled())
 						{
 							Solver->GetEvolution()->DisableParticle(Root);
+							Solver->GetParticles().MarkTransientDirtyParticle(Root);
 						}
 					}
 				}
 			}
 		}
-		
+
 		INC_DWORD_STAT_BY(STAT_GCReplicatedFractures, RepData.OneOffActivated.Num());
 		
 		//build up clusters to replicate and compare with previous frame
@@ -2679,32 +2683,35 @@ void UGeometryCollectionComponent::RegisterAndInitializePhysicsProxy()
 
 	// If we're replicating we need some extra setup - check netmode as we don't need this for standalone runtime where we aren't going to network the component
 	// IMPORTANT this need to happen after the object is registered so this will guarantee that the particles are properly created by the time the callback below gets called
-	if (GetIsReplicated())
+	if (GetIsReplicated() && PhysicsProxy->GetReplicationMode() == FGeometryCollectionPhysicsProxy::EReplicationMode::Client)
 	{
 		// Client side : geometry collection children of parents below the rep level need to be infinitely strong so that client cannot break it 
 		if (Chaos::FPhysicsSolver* CurrSolver = GetSolver(*this))
 		{
 			CurrSolver->EnqueueCommandImmediate([Proxy = PhysicsProxy, AbandonAfterLevel = ReplicationAbandonAfterLevel, EnableAbandonAfterLevel = bEnableAbandonAfterLevel]()
 				{
-					if (Proxy->GetReplicationMode() == FGeometryCollectionPhysicsProxy::EReplicationMode::Client)
+					// As we're not in control we make it so our simulated proxy cannot break clusters
+					// We have to set the strain to a high value but be below the max for the data type
+					// so releasing on authority demand works
+					for (Chaos::FPBDRigidClusteredParticleHandle* ParticleHandle : Proxy->GetParticles())
 					{
-						// As we're not in control we make it so our simulated proxy cannot break clusters
-						// We have to set the strain to a high value but be below the max for the data type
-						// so releasing on authority demand works
-						for (Chaos::FPBDRigidClusteredParticleHandle* ParticleHandle : Proxy->GetParticles())
+						if (ParticleHandle)
 						{
-							if (ParticleHandle)
+							const int32 Level = EnableAbandonAfterLevel ? ComputeParticleLevel(ParticleHandle) : -1;
+							if (Level <= AbandonAfterLevel + 1)	//we only replicate up until level X, but it means we should replicate the breaking event of level X+1 (but not X+1's positions)
 							{
-								const int32 Level = EnableAbandonAfterLevel ? ComputeParticleLevel(ParticleHandle) : -1;
-								if (Level <= AbandonAfterLevel + 1)	//we only replicate up until level X, but it means we should replicate the breaking event of level X+1 (but not X+1's positions)
-								{
-									ParticleHandle->SetMaximumInternalStrain();
-								}
+								ParticleHandle->SetMaximumInternalStrain();
 							}
 						}
 					}
 				});
 		}
+
+		// We also need to make sure collision is disabled on the game thread side of things if abandoned particles shouldn't have collision.
+		if (bEnableAbandonAfterLevel)
+		{
+			LoadCollisionProfileOnAbandonedParticles();
+		}
 	}
 
 	// We need to add the geometry collection into the external acceleration structure so that it's immediately available for queries instead of waiting for the sync from the physics thread (which could take awhile).
@@ -2720,6 +2727,54 @@ void UGeometryCollectionComponent::RegisterAndInitializePhysicsProxy()
 	SetAsyncPhysicsTickEnabled(GetIsReplicated());
 }
 
+void UGeometryCollectionComponent::SetAbandonedParticleCollisionProfileName(FName CollisionProfile)
+{
+	if (!bEnableAbandonAfterLevel)
+	{
+		return;
+	}
+
+	AbandonedCollisionProfileName = CollisionProfile;
+	LoadCollisionProfileOnAbandonedParticles();
+}
+
+void UGeometryCollectionComponent::LoadCollisionProfileOnAbandonedParticles()
+{
+	FCollisionResponseTemplate Template;
+	if (!PhysicsProxy || !UCollisionProfile::Get()->GetProfileTemplate(AbandonedCollisionProfileName, Template))
+	{
+		return;
+	}
+
+	Chaos::Facades::FCollectionHierarchyFacade HierarchyFacade(*RestCollection->GetGeometryCollection());
+
+	AActor* Owner = GetOwner();
+	const uint32 ActorID = Owner ? Owner->GetUniqueID() : 0;
+	const uint32 CompID = GetUniqueID();
+
+	FCollisionFilterData AbandonedQueryFilter = InitialQueryFilter;
+	FCollisionFilterData AbandonedSimFilter = InitialSimFilter;
+	CreateShapeFilterData(Template.ObjectType, BodyInstance.GetMaskFilter(), ActorID, Template.ResponseToChannels, CompID, INDEX_NONE, AbandonedQueryFilter, AbandonedSimFilter, BodyInstance.bUseCCD, bNotifyCollisions, false, false);
+
+	// Maintain parity with the rest of the geometry collection filters.
+	AbandonedQueryFilter.Word3 |= (EPDF_SimpleCollision | EPDF_ComplexCollision);
+	AbandonedSimFilter.Word3 |= (EPDF_SimpleCollision | EPDF_ComplexCollision);
+
+	TArray<Chaos::FPhysicsObjectHandle> PhysicsObjects = PhysicsProxy->GetAllPhysicsObjects();
+	FLockedWritePhysicsObjectExternalInterface Interface = FPhysicsObjectExternalInterface::LockWrite(PhysicsObjects);
+
+	for (int32 ParticleIndex = 0; ParticleIndex < PhysicsProxy->GetNumParticles(); ++ParticleIndex)
+	{
+		const int32 Level = HierarchyFacade.GetInitialLevel(ParticleIndex);
+		if (Level >= ReplicationAbandonAfterLevel + 1)
+		{
+			TArrayView<Chaos::FPhysicsObjectHandle> ParticleView{ &PhysicsObjects[ParticleIndex], 1 };
+			Interface->UpdateShapeCollisionFlags(ParticleView, CollisionEnabledHasPhysics(Template.CollisionEnabled), CollisionEnabledHasQuery(Template.CollisionEnabled));
+			Interface->UpdateShapeFilterData(ParticleView, AbandonedQueryFilter, AbandonedSimFilter);
+		}
+	}
+}
+
 #if WITH_EDITORONLY_DATA
 const FDamageCollector* UGeometryCollectionComponent::GetRunTimeDataCollector() const
 {
diff --git a/Engine/Source/Runtime/Experimental/GeometryCollectionEngine/Public/GeometryCollection/GeometryCollectionComponent.h b/Engine/Source/Runtime/Experimental/GeometryCollectionEngine/Public/GeometryCollection/GeometryCollectionComponent.h
--- a/Engine/Source/Runtime/Experimental/GeometryCollectionEngine/Public/GeometryCollection/GeometryCollectionComponent.h
+++ b/Engine/Source/Runtime/Experimental/GeometryCollectionEngine/Public/GeometryCollection/GeometryCollectionComponent.h
@@ -651,6 +651,10 @@ public:
 	UFUNCTION(BlueprintCallable, Category = "ChaosPhysics")
 	void SetEnableDamageFromCollision(bool bValue);
 
+
+	UFUNCTION(BlueprintCallable, Category = "ChaosPhysics")
+	void SetAbandonedParticleCollisionProfileName(FName CollisionProfile);
+
 	/** API for getting at geometry collection data */
 	int32 GetNumElements(FName Group) const;
 
@@ -1081,6 +1085,13 @@ protected:
 	UPROPERTY(EditAnywhere, BlueprintReadWrite, Category = Network)
 	bool bEnableAbandonAfterLevel;
 
+	/**
+	 * Whether abandoned particles on the client should continue to have collision (i.e.
+	 * still be in the external/internal acceleration structure).
+	 */
+	UPROPERTY(EditAnywhere, BlueprintReadWrite, BlueprintSetter= SetAbandonedParticleCollisionProfileName, Category = Network)
+	FName AbandonedCollisionProfileName;
+
 	/**
 	 * If replicating - the cluster level to stop sending corrections for geometry collection chunks.
 	 * recommended for smaller leaf levels when the size of the objects means they are no longer
@@ -1201,6 +1212,8 @@ private:
 	
 	void BuildInitialFilterData();
 
+	void LoadCollisionProfileOnAbandonedParticles();
+
 	/** The clusters we need to replicate */
 	TUniquePtr<TSet<Chaos::FPBDRigidClusteredParticleHandle*>> ClustersToRep;
 
