Fixed non unity/pch compile errors

diff --git a/Engine/Source/Runtime/Experimental/Chaos/Public/Chaos/Evolution/ABTestingConstraintContainerSolver.h b/Engine/Source/Runtime/Experimental/Chaos/Public/Chaos/Evolution/ABTestingConstraintContainerSolver.h
--- a/Engine/Source/Runtime/Experimental/Chaos/Public/Chaos/Evolution/ABTestingConstraintContainerSolver.h
+++ b/Engine/Source/Runtime/Experimental/Chaos/Public/Chaos/Evolution/ABTestingConstraintContainerSolver.h
@@ -1,8 +1,8 @@
 // Copyright Epic Games, Inc. All Rights Reserved.
 #pragma once
 
-#include "Chaos/ConstraintHandle.h"
-
+#include "Chaos/Evolution/SolverBodyContainer.h"
+#include "Chaos/Evolution/SolverConstraintContainer.h"
 #include <functional>
 
 #ifndef CHAOS_ABTEST_CONSTRAINTSOLVER_ENABLED
diff --git a/Engine/Source/Runtime/Experimental/ChaosCore/Public/Chaos/SimdTypes.h b/Engine/Source/Runtime/Experimental/ChaosCore/Public/Chaos/SimdTypes.h
--- a/Engine/Source/Runtime/Experimental/ChaosCore/Public/Chaos/SimdTypes.h
+++ b/Engine/Source/Runtime/Experimental/ChaosCore/Public/Chaos/SimdTypes.h
@@ -1,8 +1,7 @@
 // Copyright Epic Games, Inc. All Rights Reserved.
 #pragma once
 
-#include "Chaos/Real.h"
-#include "Chaos/Vector.h"
+#include "Chaos/Core.h"
 
 namespace Chaos
 {
