FReadPhysicsObjectInterface<Id>::AreAllShapesQueryEnabled should use the same code path for both external and internal contexts

diff --git a/Engine/Source/Runtime/Experimental/Chaos/Private/Chaos/PhysicsObjectInterface.cpp b/Engine/Source/Runtime/Experimental/Chaos/Private/Chaos/PhysicsObjectInterface.cpp
--- a/Engine/Source/Runtime/Experimental/Chaos/Private/Chaos/PhysicsObjectInterface.cpp
+++ b/Engine/Source/Runtime/Experimental/Chaos/Private/Chaos/PhysicsObjectInterface.cpp
@@ -380,12 +380,6 @@ namespace Chaos
 	template<EThreadContext Id>
 	bool FReadPhysicsObjectInterface<Id>::AreAllShapesQueryEnabled(TArrayView<FPhysicsObjectHandle> InObjects)
 	{
-		// This operation isn't particularly defined on the physics thread?
-		if constexpr (Id == EThreadContext::Internal)
-		{
-			return false;
-		}
-
 		if (InObjects.IsEmpty())
 		{
 			return false;
