chaos events dispatch should not iterate through all the listeners even when there are no events to dispatch

diff --git a/Engine/Source/Runtime/Experimental/Chaos/Public/EventManager.h b/Engine/Source/Runtime/Experimental/Chaos/Public/EventManager.h
--- a/Engine/Source/Runtime/Experimental/Chaos/Public/EventManager.h
+++ b/Engine/Source/Runtime/Experimental/Chaos/Public/EventManager.h
@@ -212,9 +212,15 @@ namespace Chaos
 		 */
 		virtual void DispatchConsumerData()
 		{
+			const PayloadType* Buffer = EventBuffer.Get()->GetConsumerBuffer();
+			if (IsEventDataEmpty(Buffer))
+			{
+				return;
+			}
+
 			for (FEventHandlerPtr Handler : HandlerArray)
 			{
-				Handler->HandleEvent(EventBuffer.Get()->GetConsumerBuffer());
+				Handler->HandleEvent(Buffer);
 			}
 		}
 
diff --git a/Engine/Source/Runtime/Experimental/Chaos/Public/EventsData.h b/Engine/Source/Runtime/Experimental/Chaos/Public/EventsData.h
--- a/Engine/Source/Runtime/Experimental/Chaos/Public/EventsData.h
+++ b/Engine/Source/Runtime/Experimental/Chaos/Public/EventsData.h
@@ -197,4 +197,39 @@ namespace Chaos
 		FAllCrumblingData CrumblingData;
 		FIndicesByPhysicsProxy PhysicsProxyToCrumblingIndices;
 	};
+
+	template<typename PayloadType>
+	bool IsEventDataEmpty(const PayloadType* Buffer)
+	{
+		ensure(Buffer);
+
+		if constexpr (std::is_same_v<PayloadType, FCollisionEventData>)
+		{
+			return Buffer->CollisionData.AllCollisionsArray.IsEmpty();
+		}
+		else if constexpr (std::is_same_v<PayloadType, FBreakingEventData>)
+		{
+			return Buffer->BreakingData.AllBreakingsArray.IsEmpty();
+		}
+		else if constexpr (std::is_same_v<PayloadType, FTrailingEventData>)
+		{
+			return Buffer->TrailingData.AllTrailingsArray.IsEmpty();
+		}
+		else if constexpr (std::is_same_v<PayloadType, FRemovalEventData>)
+		{
+			return Buffer->RemovalData.AllRemovalArray.IsEmpty();
+		}
+		else if constexpr (std::is_same_v<PayloadType, FSleepingEventData>)
+		{
+			return Buffer->SleepingData.IsEmpty();
+		}
+		else if constexpr (std::is_same_v<PayloadType, FCrumblingEventData>)
+		{
+			return Buffer->CrumblingData.AllCrumblingsArray.IsEmpty();
+		}
+		else
+		{
+			return false;
+		}
+	}
 }
