use ReadRef to get property from single particle physics proxy in DebugName()

diff --git a/Engine/Source/Runtime/Experimental/Chaos/Public/PhysicsProxy/SingleParticlePhysicsProxy.h b/Engine/Source/Runtime/Experimental/Chaos/Public/PhysicsProxy/SingleParticlePhysicsProxy.h
--- a/Engine/Source/Runtime/Experimental/Chaos/Public/PhysicsProxy/SingleParticlePhysicsProxy.h
+++ b/Engine/Source/Runtime/Experimental/Chaos/Public/PhysicsProxy/SingleParticlePhysicsProxy.h
@@ -248,7 +248,7 @@ public:
 	const TSharedPtr<FImplicitObject, ESPMode::ThreadSafe>& SharedGeometryLowLevel() const { return ReadRef([](auto* Ptr) -> const auto& { return Ptr->SharedGeometryLowLevel(); });}
 
 #if CHAOS_DEBUG_NAME
-	const TSharedPtr<FString, ESPMode::ThreadSafe>& DebugName() const { return Read([](auto* Ptr) { return Ptr->DebugName(); }); }
+	const TSharedPtr<FString, ESPMode::ThreadSafe>& DebugName() const { return ReadRef([](auto* Ptr) -> const auto& { return Ptr->DebugName(); }); }
 	void SetDebugName(const TSharedPtr<FString, ESPMode::ThreadSafe>& InDebugName) { Write([&InDebugName](auto* Ptr) { Ptr->SetDebugName(InDebugName); }); }
 #endif
 
