Add const solver accessor to SimCallbackObject

diff --git a/Engine/Source/Runtime/Experimental/Chaos/Public/Chaos/SimCallbackObject.h b/Engine/Source/Runtime/Experimental/Chaos/Public/Chaos/SimCallbackObject.h
--- a/Engine/Source/Runtime/Experimental/Chaos/Public/Chaos/SimCallbackObject.h
+++ b/Engine/Source/Runtime/Experimental/Chaos/Public/Chaos/SimCallbackObject.h
@@ -101,6 +101,8 @@ public:
 
 	FPhysicsSolverBase* GetSolver() { return Solver; }
 
+	const FPhysicsSolverBase* GetSolver() const { return Solver; }
+
 	// Rewind API
 	virtual int32 TriggerRewindIfNeeded_Internal(int32 LastCompletedStep) { ensure(false); return INDEX_NONE; }
 	virtual void ApplyCorrections_Internal(int32 PhysicsStep, FSimCallbackInput* Input) { ensure(false); }
