Iris - Store whether an object may be destroyed from the remote or not.

diff --git a/Engine/Source/Runtime/Experimental/Iris/Core/Private/Iris/ReplicationSystem/ReplicationReader.cpp b/Engine/Source/Runtime/Experimental/Iris/Core/Private/Iris/ReplicationSystem/ReplicationReader.cpp
--- a/Engine/Source/Runtime/Experimental/Iris/Core/Private/Iris/ReplicationSystem/ReplicationReader.cpp
+++ b/Engine/Source/Runtime/Experimental/Iris/Core/Private/Iris/ReplicationSystem/ReplicationReader.cpp
@@ -478,7 +478,10 @@ void FReplicationReader::ReadObject(FNetSerializationContext& Context)
 			goto ErrorHandling;
 		}
 
-		InternalIndex = NetRefHandleManager->GetInternalIndex(NetRefHandle);		 
+		InternalIndex = NetRefHandleManager->GetInternalIndex(NetRefHandle);
+		FNetRefHandleManager::FReplicatedObjectData& ObjectData = NetRefHandleManager->GetReplicatedObjectDataNoCheck(InternalIndex);
+		ObjectData.bAllowDestroyInstanceFromRemote = EnumHasAnyFlags(CreateResult.Flags, EReplicationBridgeCreateNetRefHandleResultFlags::AllowDestroyInstanceFromRemote);
+
 		FReplicatedObjectInfo& ObjectInfo = StartReplication(InternalIndex);
 
 		ObjectInfo.bIsDeltaCompressionEnabled = bIsDeltaCompressed;
