Iris
- Fixed testfail if Tear of and destroy flags was set at the same time.

diff --git a/Engine/Source/Runtime/Experimental/Iris/Core/Private/Iris/ReplicationSystem/ReplicationReader.cpp b/Engine/Source/Runtime/Experimental/Iris/Core/Private/Iris/ReplicationSystem/ReplicationReader.cpp
--- a/Engine/Source/Runtime/Experimental/Iris/Core/Private/Iris/ReplicationSystem/ReplicationReader.cpp
+++ b/Engine/Source/Runtime/Experimental/Iris/Core/Private/Iris/ReplicationSystem/ReplicationReader.cpp
@@ -302,13 +302,13 @@ void FReplicationReader::EndReplication(uint32 InternalIndex, bool bTearOff, boo
 		Attachments.DropAllAttachments(ENetObjectAttachmentType::Normal, InternalIndex);
 
 		EReplicationBridgeDestroyInstanceReason DestroyReason = EReplicationBridgeDestroyInstanceReason::DoNotDestroy;
-		if (bDestroyInstance)
+		if (bTearOff)
 		{
-			DestroyReason = EReplicationBridgeDestroyInstanceReason::Destroy;
+			DestroyReason = EReplicationBridgeDestroyInstanceReason::TearOff;
 		}
-		else if (bTearOff)
+		else if (bDestroyInstance)
 		{
-			DestroyReason = EReplicationBridgeDestroyInstanceReason::TearOff;
+			DestroyReason = EReplicationBridgeDestroyInstanceReason::Destroy;
 		}
 		EReplicationBridgeDestroyInstanceFlags DestroyFlags = (Data.bAllowDestroyInstanceFromRemote ? EReplicationBridgeDestroyInstanceFlags::AllowDestroyInstanceFromRemote : EReplicationBridgeDestroyInstanceFlags::None);
 		ReplicationBridge->DestroyNetObjectFromRemote(Data.RefHandle, DestroyReason, DestroyFlags);
