Chaos : Add gravity groups to rigid solver

diff --git a/Engine/Source/Runtime/Experimental/Chaos/Private/Chaos/PBDRigidsEvolutionGBF.cpp b/Engine/Source/Runtime/Experimental/Chaos/Private/Chaos/PBDRigidsEvolutionGBF.cpp
--- a/Engine/Source/Runtime/Experimental/Chaos/Private/Chaos/PBDRigidsEvolutionGBF.cpp
+++ b/Engine/Source/Runtime/Experimental/Chaos/Private/Chaos/PBDRigidsEvolutionGBF.cpp
@@ -464,7 +464,8 @@ void FPBDRigidsEvolutionGBF::AdvanceOneTimeStepImpl(const FReal Dt, const FSubSt
 		PreApplyCallback();
 	}
 
-	CollisionConstraints.SetGravity(GetGravityForces().GetAcceleration());
+	// todo(chaos) : we are using the main gravity ( index 0 ) we should revise this to account for the various gravities based on the constraint ? 
+	CollisionConstraints.SetGravity(GetGravityForces().GetAcceleration(0));
 
 	// Use the cache to update particles in islands that do not require resimming (not desynched)
 	{
diff --git a/Engine/Source/Runtime/Experimental/Chaos/Public/Chaos/ParticleHandle.h b/Engine/Source/Runtime/Experimental/Chaos/Public/Chaos/ParticleHandle.h
--- a/Engine/Source/Runtime/Experimental/Chaos/Public/Chaos/ParticleHandle.h
+++ b/Engine/Source/Runtime/Experimental/Chaos/Public/Chaos/ParticleHandle.h
@@ -1121,6 +1121,16 @@ public:
 		PBDRigidParticles->ControlFlags(ParticleIdx).SetGravityEnabled(bEnabled);
 	}
 
+	inline int32 GravityGroupIndex() const
+	{
+		return ControlFlags().GetGravityGroupIndex();
+	}
+
+	inline void SetGravityGroupIndex(int32 GravityGroupIndex)
+	{
+		PBDRigidParticles->ControlFlags(ParticleIdx).SetGravityGroupIndex(GravityGroupIndex);
+	}
+
 	inline bool CCDEnabled() const
 	{
 		return ControlFlags().GetCCDEnabled();
diff --git a/Engine/Source/Runtime/Experimental/Chaos/Public/Chaos/PerParticleGravity.h b/Engine/Source/Runtime/Experimental/Chaos/Public/Chaos/PerParticleGravity.h
--- a/Engine/Source/Runtime/Experimental/Chaos/Public/Chaos/PerParticleGravity.h
+++ b/Engine/Source/Runtime/Experimental/Chaos/Public/Chaos/PerParticleGravity.h
@@ -10,43 +10,52 @@ namespace Chaos
 	class FPerParticleGravity : public FPerParticleRule
 	{
 	public:
+		static constexpr uint32 MaxGravityGroups = 8;
+
 		FPerParticleGravity()
-			: MAcceleration(FVec3(0, 0, (FReal)-980.665)) {}
-		FPerParticleGravity(const FVec3& G)
-			: MAcceleration(G) {}
-		FPerParticleGravity(const FVec3& Direction, const FReal Magnitude)
-			: MAcceleration(Direction * Magnitude) {}
-		virtual ~FPerParticleGravity() {}
+		{
+			SetAllGravities(FVec3(0, 0, (FReal)-980.665));
+		}
 
-		// TODO: Remove this - we should no longer be using indices directly.
-		//       This has been kept around for cloth, which uses it in
-		//       PBDEvolution.
-		template<class T_PARTICLES>
-		inline void ApplyHelper(T_PARTICLES& InParticles, const FReal Dt, const int Index) const
+		FPerParticleGravity(const FVec3& G)
 		{
-			InParticles.Acceleration(Index) += MAcceleration;
+			SetAllGravities(G);
 		}
-		inline void Apply(FPBDParticles& InParticles, const FReal Dt, const int Index) const override //-V762
+		FPerParticleGravity(const FVec3& Direction, const FReal Magnitude)
 		{
-			ApplyHelper(InParticles, Dt, Index);
+			SetAllGravities(Direction * Magnitude);
 		}
+		virtual ~FPerParticleGravity() {}
 
 		void Apply(TTransientPBDRigidParticleHandle<FReal, 3>& Particle, const FReal Dt) const override //-V762
 		{
 			if(Particle.GravityEnabled())
 			{
-				Particle.Acceleration() += MAcceleration;
+				Particle.Acceleration() += MAccelerations[Particle.GravityGroupIndex()];
 			}
 		}
 
-		void SetAcceleration(const FVec3& Acceleration)
-		{ MAcceleration = Acceleration; }
+		void SetAcceleration(const FVec3& Acceleration, int32 GroupIndex)
+		{ 
+			MAccelerations[GroupIndex] = Acceleration;
+		}
 
-		const FVec3& GetAcceleration() const
-		{ return MAcceleration; }
+		const FVec3& GetAcceleration(int32 GroupIndex) const
+		{
+			return MAccelerations[GroupIndex];
+		}
+
+	private:
+		void SetAllGravities(const FVec3& Gravity)
+		{
+			for (int32 Index = 0; Index < MaxGravityGroups; Index++)
+			{
+				MAccelerations[Index] = Gravity;
+			}
+		}
 
 	private:
-		FVec3 MAcceleration;
+		FVec3 MAccelerations[MaxGravityGroups];
 	};
 
 	template<class T, int d>
diff --git a/Engine/Source/Runtime/Experimental/Chaos/Public/Chaos/RigidParticleControlFlags.h b/Engine/Source/Runtime/Experimental/Chaos/Public/Chaos/RigidParticleControlFlags.h
--- a/Engine/Source/Runtime/Experimental/Chaos/Public/Chaos/RigidParticleControlFlags.h
+++ b/Engine/Source/Runtime/Experimental/Chaos/Public/Chaos/RigidParticleControlFlags.h
@@ -30,6 +30,14 @@ namespace Chaos
 		bool GetGravityEnabled() const { return Flags.bGravityEnabled; }
 		FRigidParticleControlFlags& SetGravityEnabled(const bool bEnabled) { Flags.bGravityEnabled = bEnabled; return *this; }
 
+		int32 GetGravityGroupIndex() const { return Flags.GravityGroupIndex; }
+		FRigidParticleControlFlags& SetGravityGroupIndex(const int32 GravityGroupIndex) 
+		{ 
+			ensure(GravityGroupIndex < 8);
+			Flags.GravityGroupIndex = static_cast<FStorage>(GravityGroupIndex);
+			return *this; 
+		}
+
 		bool GetCCDEnabled() const { return Flags.bCCDEnabled; }
 		FRigidParticleControlFlags& SetCCDEnabled(const bool bEnabled) { Flags.bCCDEnabled = bEnabled; return *this; }
 
@@ -60,6 +68,7 @@ namespace Chaos
 			FStorage bOneWayInteractionEnabled : 1;
 			FStorage bMaxDepenetrationVelocityOverrideEnabled : 1;
 			FStorage bInertiaConditioningEnabled : 1;
+			FStorage GravityGroupIndex : 3;
 			// Add new properties above this line
 			// Change FStorage typedef if we exceed the max bits
 		};
diff --git a/Engine/Source/Runtime/PhysicsCore/Private/ChaosScene.cpp b/Engine/Source/Runtime/PhysicsCore/Private/ChaosScene.cpp
--- a/Engine/Source/Runtime/PhysicsCore/Private/ChaosScene.cpp
+++ b/Engine/Source/Runtime/PhysicsCore/Private/ChaosScene.cpp
@@ -304,7 +304,8 @@ void FChaosSceneSimCallback::OnPreSimulate_Internal()
 {
 	if(const FChaosSceneCallbackInput* Input = GetConsumerInput_Internal())
 	{
-		static_cast<Chaos::FPBDRigidsSolver*>(GetSolver())->GetEvolution()->GetGravityForces().SetAcceleration(Input->Gravity);
+		// the "main" gravity is index 0
+		static_cast<Chaos::FPBDRigidsSolver*>(GetSolver())->GetEvolution()->GetGravityForces().SetAcceleration(Input->Gravity, 0);
 	}
 }
 
