fix MSVC static analysis of possibly dereferencing null pointer in Chaos::IsEventDataEmpty

diff --git a/Engine/Source/Runtime/Experimental/Chaos/Public/EventsData.h b/Engine/Source/Runtime/Experimental/Chaos/Public/EventsData.h
--- a/Engine/Source/Runtime/Experimental/Chaos/Public/EventsData.h
+++ b/Engine/Source/Runtime/Experimental/Chaos/Public/EventsData.h
@@ -201,7 +201,10 @@ namespace Chaos
 	template<typename PayloadType>
 	bool IsEventDataEmpty(const PayloadType* Buffer)
 	{
-		ensure(Buffer);
+		if (!Buffer)
+		{
+			return false;
+		}
 
 		if constexpr (std::is_same_v<PayloadType, FCollisionEventData>)
 		{
