Geometry collection : add ability to set the gravity group index

diff --git a/Engine/Source/Runtime/Experimental/Chaos/Private/PhysicsProxy/GeometryCollectionPhysicsProxy.cpp b/Engine/Source/Runtime/Experimental/Chaos/Private/PhysicsProxy/GeometryCollectionPhysicsProxy.cpp
--- a/Engine/Source/Runtime/Experimental/Chaos/Private/PhysicsProxy/GeometryCollectionPhysicsProxy.cpp
+++ b/Engine/Source/Runtime/Experimental/Chaos/Private/PhysicsProxy/GeometryCollectionPhysicsProxy.cpp
@@ -1244,6 +1244,7 @@ void FGeometryCollectionPhysicsProxy::InitializeBodiesPT(Chaos::FPBDRigidsSolver
 			if (Handle)
 			{
 				Handle->SetGravityEnabled(bEnableGravity);
+				Handle->SetGravityGroupIndex(Parameters.GravityGroupIndex);
 				Handle->SetCCDEnabled(Parameters.UseCCD);
 				Handle->SetInertiaConditioningEnabled(Parameters.UseInertiaConditioning);
 				Handle->SetLinearEtherDrag(Parameters.LinearDamping);
diff --git a/Engine/Source/Runtime/Experimental/Chaos/Public/GeometryCollection/GeometryCollectionSimulationCoreTypes.h b/Engine/Source/Runtime/Experimental/Chaos/Public/GeometryCollection/GeometryCollectionSimulationCoreTypes.h
--- a/Engine/Source/Runtime/Experimental/Chaos/Public/GeometryCollection/GeometryCollectionSimulationCoreTypes.h
+++ b/Engine/Source/Runtime/Experimental/Chaos/Public/GeometryCollection/GeometryCollectionSimulationCoreTypes.h
@@ -200,6 +200,7 @@ struct FSimulationParameters
 		, bGenerateCrumblingData(false)
 		, bGenerateCrumblingChildrenData(false)
 		, EnableGravity(true)
+		, GravityGroupIndex(0)
 		, UseInertiaConditioning(true)
 		, UseCCD(false)
 		, LinearDamping(0.01f)
@@ -249,6 +250,7 @@ struct FSimulationParameters
 		, bGenerateCrumblingChildrenData(Other.bGenerateCrumblingChildrenData)
 		, Shared(Other.Shared)
 		, EnableGravity(Other.EnableGravity)
+		, GravityGroupIndex(Other.GravityGroupIndex)
 		, UseInertiaConditioning(Other.UseInertiaConditioning)
 		, UseCCD(Other.UseCCD)
 		, LinearDamping(Other.LinearDamping)
@@ -321,6 +323,7 @@ struct FSimulationParameters
 	FSharedSimulationParameters Shared;
 
 	bool EnableGravity;
+	int32 GravityGroupIndex;
 	bool UseInertiaConditioning;
 	bool UseCCD;
 	float LinearDamping;
diff --git a/Engine/Source/Runtime/Experimental/GeometryCollectionEngine/Private/GeometryCollection/GeometryCollectionComponent.cpp b/Engine/Source/Runtime/Experimental/GeometryCollectionEngine/Private/GeometryCollection/GeometryCollectionComponent.cpp
--- a/Engine/Source/Runtime/Experimental/GeometryCollectionEngine/Private/GeometryCollection/GeometryCollectionComponent.cpp
+++ b/Engine/Source/Runtime/Experimental/GeometryCollectionEngine/Private/GeometryCollection/GeometryCollectionComponent.cpp
@@ -286,6 +286,7 @@ UGeometryCollectionComponent::UGeometryCollectionComponent(const FObjectInitiali
 	, ChaosSolverActor(nullptr)
 	, InitializationState(ESimulationInitializationState::Unintialized)
 	, ObjectType(EObjectStateTypeEnum::Chaos_Object_Dynamic)
+	, GravityGroupIndex(0)
 	, bForceMotionBlur()
 	, EnableClustering(true)
 	, ClusterGroupIndex(0)
@@ -2630,6 +2631,7 @@ void UGeometryCollectionComponent::RegisterAndInitializePhysicsProxy()
 		SimulationParameters.bGenerateCrumblingData = bNotifyCrumblings;
 		SimulationParameters.bGenerateCrumblingChildrenData = bCrumblingEventIncludesChildren;
 		SimulationParameters.EnableGravity = BodyInstance.bEnableGravity;
+		SimulationParameters.GravityGroupIndex = GravityGroupIndex;
 		SimulationParameters.UseInertiaConditioning = BodyInstance.IsInertiaConditioningEnabled();
 		SimulationParameters.UseCCD = BodyInstance.bUseCCD;
 		SimulationParameters.LinearDamping = BodyInstance.LinearDamping;
diff --git a/Engine/Source/Runtime/Experimental/GeometryCollectionEngine/Public/GeometryCollection/GeometryCollectionComponent.h b/Engine/Source/Runtime/Experimental/GeometryCollectionEngine/Public/GeometryCollection/GeometryCollectionComponent.h
--- a/Engine/Source/Runtime/Experimental/GeometryCollectionEngine/Public/GeometryCollection/GeometryCollectionComponent.h
+++ b/Engine/Source/Runtime/Experimental/GeometryCollectionEngine/Public/GeometryCollection/GeometryCollectionComponent.h
@@ -718,6 +718,9 @@ public:
 	UPROPERTY(EditAnywhere, BlueprintReadWrite, Category = "ChaosPhysics|General")
 	EObjectStateTypeEnum ObjectType;
 
+	UPROPERTY(EditAnywhere, BlueprintReadWrite, Category = "ChaosPhysics|General")
+	int32 GravityGroupIndex;
+
 	/** If ForceMotionBlur is on, motion blur will always be active, even if the GeometryCollection is at rest. */
 	UPROPERTY(EditAnywhere, BlueprintReadWrite, Category = "ChaosPhysics|General")
 	bool bForceMotionBlur;
@@ -745,7 +748,7 @@ public:
 	FGeometryCollectionDamagePropagationData DamagePropagationData;
 
 	/** Whether or not collisions against this geometry collection will apply strain which could cause the geometry collection to fracture. */
-	UPROPERTY(EditDefaultsOnly, BlueprintReadOnly, Category = "ChaosPhysics|Damage")
+	UPROPERTY(EditAnywhere, BlueprintReadWrite, BlueprintSetter=SetEnableDamageFromCollision, Category = "ChaosPhysics|Damage")
 	bool bEnableDamageFromCollision;
 
 	/** Allow removal on sleep for the instance if the rest collection has it enabled */
