Add FindClosestPointAndBary to Chaos Tetrahedron.

diff --git a/Engine/Source/Programs/HeadlessChaos/Private/HeadlessChaosTestImplicits.cpp b/Engine/Source/Programs/HeadlessChaos/Private/HeadlessChaosTestImplicits.cpp
--- a/Engine/Source/Programs/HeadlessChaos/Private/HeadlessChaosTestImplicits.cpp
+++ b/Engine/Source/Programs/HeadlessChaos/Private/HeadlessChaosTestImplicits.cpp
@@ -256,6 +256,10 @@ namespace ChaosTest {
 		EXPECT_VECTOR_NEAR(Bary, FVector(.25, .25, .25), 0.001);
 		TVec3<double> Surf = Tet.ProjectToSurface(Tris, Pt);
 		EXPECT_VECTOR_NEAR(Surf, Tris[3].GetCentroid(), 0.001);
+		TVec4<double> ClosestBary;
+		TVec3<double> ClosestPoint = Tet.FindClosestPointAndBary(Pt, ClosestBary);
+		EXPECT_VECTOR_NEAR(ClosestPoint, Pt, 0.001);
+		EXPECT_VECTOR_NEAR(ClosestBary, TVec4<double>(0.25, 0.25, 0.25, 0.25), 0.001);
 
 		// Point
 		Pt = Tet[0];
@@ -267,6 +271,10 @@ namespace ChaosTest {
 		EXPECT_VECTOR_NEAR(Bary, FVector(1, 0, 0), 0.001);
 		Surf = Tet.ProjectToSurface(Tris, Pt);
 		EXPECT_VECTOR_NEAR(Surf, Pt, 0.001);
+		ClosestPoint = Tet.FindClosestPointAndBary(Pt, ClosestBary);
+		EXPECT_VECTOR_NEAR(ClosestPoint, Pt, 0.001);
+		EXPECT_VECTOR_NEAR(ClosestBary, TVec4<double>(1, 0, 0, 0), 0.001);
+
 
 		Pt[0] -= 0.1;
 		Hit = Tet.Inside(Pt, 0.001);
@@ -274,8 +282,19 @@ namespace ChaosTest {
 		RobustHit = Tet.RobustInside(Pt, -0.001);
 		EXPECT_FALSE(RobustHit);
 		Surf = Tet.ProjectToSurface(Tris, Pt);
+		ClosestPoint = Tet.FindClosestPointAndBary(Pt, ClosestBary);
 		Pt[0] += 0.1;
 		EXPECT_VECTOR_NEAR(Surf, Pt, 0.001);
+		EXPECT_VECTOR_NEAR(ClosestPoint, Tet[0], 0.001);
+		EXPECT_VECTOR_NEAR(ClosestBary, TVec4<double>(1, 0, 0, 0), 0.001);
+
+		Pt -= FVec3(0.1);
+		Surf = Tet.ProjectToSurface(Tris, Pt);
+		ClosestPoint = Tet.FindClosestPointAndBary(Pt, ClosestBary);
+		Pt += FVec3(0.1);
+		EXPECT_VECTOR_NEAR(Surf, Pt, 0.001);
+		EXPECT_VECTOR_NEAR(ClosestPoint, Tet[0], 0.001);
+		EXPECT_VECTOR_NEAR(ClosestBary, TVec4<double>(1, 0, 0, 0), 0.001);
 
 		// Edge
 		Pt = TVec3<double>(0.5, 0, 0);
@@ -287,6 +306,9 @@ namespace ChaosTest {
 		EXPECT_VECTOR_NEAR(Bary, FVector(.5, .5, 0), 0.001);
 		Surf = Tet.ProjectToSurface(Tris, Pt);
 		EXPECT_VECTOR_NEAR(Surf, Pt, 0.001);
+		ClosestPoint = Tet.FindClosestPointAndBary(Pt, ClosestBary);
+		EXPECT_VECTOR_NEAR(ClosestPoint, Pt, 0.001);
+		EXPECT_VECTOR_NEAR(ClosestBary, TVec4<double>(.5, .5, 0, 0), 0.001);
 
 		Pt[1] -= 0.1;
 		Hit = Tet.Inside(Pt, 0.001);
@@ -294,8 +316,21 @@ namespace ChaosTest {
 		RobustHit = Tet.RobustInside(Pt, -0.001);
 		EXPECT_FALSE(RobustHit);
 		Surf = Tet.ProjectToSurface(Tris, Pt);
+		ClosestPoint = Tet.FindClosestPointAndBary(Pt, ClosestBary);
 		Pt[1] += 0.1;
 		EXPECT_VECTOR_NEAR(Surf, Pt, 0.001);
+		EXPECT_VECTOR_NEAR(ClosestPoint, Pt, 0.001);
+		EXPECT_VECTOR_NEAR(ClosestBary, TVec4<double>(.5, .5, 0, 0), 0.001);
+
+		// This is a test case where ProjectToSurface will fail but FindClosestPointAndBary will succeed
+		Pt = TVec3<double>(1, 0, 0);
+		Pt += TVec3<double>(.1, -.05, 0);
+		//Surf = Tet.ProjectToSurface(Tris, Pt);
+		ClosestPoint = Tet.FindClosestPointAndBary(Pt, ClosestBary);
+		Pt -= TVec3<double>(.1, -.05, 0);
+		//EXPECT_VECTOR_NEAR(Surf, Pt, 0.001);
+		EXPECT_VECTOR_NEAR(ClosestPoint, Pt, 0.001);
+		EXPECT_VECTOR_NEAR(ClosestBary, TVec4<double>(0, 1, 0, 0), 0.001);
 
 		// Face
 		Pt = Tris[0].GetCentroid();
@@ -307,6 +342,9 @@ namespace ChaosTest {
 		EXPECT_VECTOR_NEAR(Bary, FVector(.3333, .3333, .3333), 0.001);
 		Surf = Tet.ProjectToSurface(Tris, Pt);
 		EXPECT_VECTOR_NEAR(Surf, Pt, 0.001);
+		ClosestPoint = Tet.FindClosestPointAndBary(Pt, ClosestBary);
+		EXPECT_VECTOR_NEAR(ClosestPoint, Pt, 0.001);
+		EXPECT_VECTOR_NEAR(ClosestBary, TVec4<double>(.3333, .3333, .3333, 0), 0.001);
 
 		Pt += Tris[0].GetNormal() * 0.1;
 		Hit = Tet.Inside(Pt, 0.001);
@@ -314,8 +352,11 @@ namespace ChaosTest {
 		RobustHit = Tet.RobustInside(Pt, -0.001);
 		EXPECT_FALSE(RobustHit);
 		Surf = Tet.ProjectToSurface(Tris, Pt);
+		ClosestPoint = Tet.FindClosestPointAndBary(Pt, ClosestBary);
 		Pt -= Tris[0].GetNormal() * 0.1;
 		EXPECT_VECTOR_NEAR(Surf, Pt, 0.001);
+		EXPECT_VECTOR_NEAR(ClosestPoint, Pt, 0.001);
+		EXPECT_VECTOR_NEAR(ClosestBary, TVec4<double>(.3333, .3333, .3333, 0), 0.001);
 
 		// Bounding Volume Hierarchy
 		// Put the center of the tet at the origin, and sweep point tests across it.
diff --git a/Engine/Source/Runtime/Experimental/Chaos/Public/Chaos/Plane.h b/Engine/Source/Runtime/Experimental/Chaos/Public/Chaos/Plane.h
--- a/Engine/Source/Runtime/Experimental/Chaos/Public/Chaos/Plane.h
+++ b/Engine/Source/Runtime/Experimental/Chaos/Public/Chaos/Plane.h
@@ -367,19 +367,21 @@ TVector<T, 2> ComputeBarycentricInPlane(const TVector<T, d>& P0, const TVector<T
 }
 
 template<typename T, int d>
-const TVector<T, d> FindClosestPointOnLineSegment(const TVector<T, d>& P0, const TVector<T, d>& P1, const TVector<T, d>& P)
+const TVector<T, d> FindClosestPointAndAlphaOnLineSegment(const TVector<T, d>& P0, const TVector<T, d>& P1, const TVector<T, d>& P, T& OutAlpha)
 {
 	const TVector<T, d> P10 = P1 - P0;
 	const TVector<T, d> PP0 = P - P0;
 	const T Proj = TVector<T, d>::DotProduct(P10, PP0);
 	if (Proj < (T)0) //first check we're not behind
 	{
+		OutAlpha = (T)0;
 		return P0;
 	}
 
 	const T Denom2 = P10.SizeSquared();
 	if (Denom2 < (T)1e-4)
 	{
+		OutAlpha = (T)0;
 		return P0;
 	}
 
@@ -387,12 +389,20 @@ const TVector<T, d> FindClosestPointOnLineSegment(const TVector<T, d>& P0, const
 	const T NormalProj = Proj / Denom2;
 	if (NormalProj > (T)1) //too far forward
 	{
+		OutAlpha = (T)1;
 		return P1;
 	}
 
+	OutAlpha = NormalProj;
 	return P0 + NormalProj * P10; //somewhere on the line
 }
 
+template<typename T, int d>
+const TVector<T, d> FindClosestPointOnLineSegment(const TVector<T, d>& P0, const TVector<T, d>& P1, const TVector<T, d>& P)
+{
+	T AlphaUnused;
+	return FindClosestPointAndAlphaOnLineSegment(P0, P1, P, AlphaUnused);
+}
 
 template<typename T, int d>
 TVector<T, d> FindClosestPointOnTriangle(const TVector<T, d>& ClosestPointOnPlane, const TVector<T, d>& P0, const TVector<T, d>& P1, const TVector<T, d>& P2, const TVector<T, d>& P)
diff --git a/Engine/Source/Runtime/Experimental/Chaos/Public/Chaos/Tetrahedron.h b/Engine/Source/Runtime/Experimental/Chaos/Public/Chaos/Tetrahedron.h
--- a/Engine/Source/Runtime/Experimental/Chaos/Public/Chaos/Tetrahedron.h
+++ b/Engine/Source/Runtime/Experimental/Chaos/Public/Chaos/Tetrahedron.h
@@ -5,8 +5,10 @@
 #include "Chaos/AABB.h"
 #include "Chaos/Core.h"
 #include "Chaos/Triangle.h"
+#include "Chaos/Plane.h"
 #include "Chaos/Math/Poisson.h" // 3x3 row major matrix functions
 #include "GenericPlatform/GenericPlatformMath.h"
+#include <limits>
 
 
 namespace Chaos {
@@ -190,6 +192,8 @@ namespace Chaos {
 				(Sign(V1) != Sign(V2) || Sign(V2) != Sign(V3) || Sign(V3) != Sign(V4)));
 		}
 
+		// Note: this method will fail to project to the actual surface if Location is outside and the closest point is
+		// on an edge with an acute angle such that Location is outside both connected faces. FindClosestPointAndBary will work in this case.
 		TVec3<T> ProjectToSurface(const TArray<TTriangle<T>>& Tris, const TVec3<T>& Location) const
 		{
 			if (Inside(Tris, Location)) 
@@ -231,6 +235,93 @@ namespace Chaos {
 			}
 		}
 
+		// Tolerance should be small negative number or zero. It is used to compare barycentric coordinate values
+		TVec3<T> FindClosestPointAndBary(const TVec3<T>& Location, TVec4<T>& OutBary, const T Tolerance = 0) const
+		{
+			const TVec3<T> TetWeights = GetFirstThreeBarycentricCoordinates(Location);
+			if (TetWeights[0] >= Tolerance && TetWeights[1] >= Tolerance && TetWeights[2] >= Tolerance && TetWeights[0] + TetWeights[1] + TetWeights[2] <= 1 + Tolerance)
+			{
+				// Inside tetrahedron
+				OutBary = TVec4<T>(TetWeights[0], TetWeights[1], TetWeights[2], static_cast<T>(1.0) - TetWeights[0] - TetWeights[1] - TetWeights[2]);
+				checkSlow(TVec3<T>::DistSquared(Location, GetPointFromBarycentricCoordinates(OutBary)) < static_cast<T>(UE_SMALL_NUMBER));
+				return Location;
+			}
+			
+			// Test closest point is inside a face
+			const bool IsInverted = GetSignedVolume() <= 0;
+			static const TVec3<int32> InvertedTriangleIndices[4] =
+			{
+				TVec3<int32>(0,1,2),
+				TVec3<int32>(0,3,1),
+				TVec3<int32>(0,2,3),
+				TVec3<int32>(1,3,2)
+			};
+			static const TVec3<int32> NonInvertedTriangleIndices[4] =
+			{
+				TVec3<int32>(0,2,1),
+				TVec3<int32>(0,1,3),
+				TVec3<int32>(0,3,2),
+				TVec3<int32>(1,2,3)
+			};
+			const TVec3<int32>* FaceIndices = IsInverted ? InvertedTriangleIndices : NonInvertedTriangleIndices;
+			for (int32 FaceId = 0; FaceId < 4; ++FaceId)
+			{
+				const TVec3<int32>& TriIndices = FaceIndices[FaceId];
+				const TTriangle<T> Tri(X[TriIndices[0]], X[TriIndices[1]], X[TriIndices[2]]);
+				const TPlaneConcrete<T, 3> Plane = Tri.GetPlane(0);
+				const T Dist = TVec3<T>::DotProduct(Location - Plane.X(), Plane.Normal());
+				if (Dist >= 0)
+				{
+					// Point is on front-face side of tri.
+					const TVec3<T> ClosestPointOnPlane = Location - Dist * Plane.Normal();
+					const TVec2<T> TriBary = ComputeBarycentricInPlane(Tri[0], Tri[1], Tri[2], ClosestPointOnPlane);
+					if (TriBary[0] >= Tolerance && TriBary[1] >= Tolerance && TriBary[0] + TriBary[1] <= 1 + Tolerance)
+					{
+						OutBary = TVec4<T>(static_cast<T>(0.));
+						OutBary[TriIndices[1]] = TriBary[0];
+						OutBary[TriIndices[2]] = TriBary[1];
+						OutBary[TriIndices[0]] = static_cast<T>(1.) - TriBary[0] - TriBary[1];
+						checkSlow(TVec3<T>::DistSquared(ClosestPointOnPlane, GetPointFromBarycentricCoordinates(OutBary)) < static_cast<T>(UE_SMALL_NUMBER));
+						return ClosestPointOnPlane;
+					}
+				}
+			}
+
+			// Closest point is on an edge/vertex
+			static const TVec2<int32> EdgeIndices[6] =
+			{
+				TVec2<int32>(0,1),
+				TVec2<int32>(0,2),
+				TVec2<int32>(0,3),
+				TVec2<int32>(1,2),
+				TVec2<int32>(1,3),
+				TVec2<int32>(2,3)
+			};
+			TVec3<T> ClosestEdgePoint;
+			T ClosestEdgeAlpha(0);
+			int32 ClosestEdgeIndex = INDEX_NONE;
+			T ClosestEdgeDistSq = std::numeric_limits<T>::max();
+			for (int32 EdgeIndex = 0; EdgeIndex < 6; ++EdgeIndex)
+			{
+				T Alpha;
+				const TVec3<T> ClosestPoint = FindClosestPointAndAlphaOnLineSegment(X[EdgeIndices[EdgeIndex][0]], X[EdgeIndices[EdgeIndex][1]], Location, Alpha);
+				const T DistSq = TVec3<T>::DistSquared(Location, ClosestPoint);
+				if (DistSq < ClosestEdgeDistSq)
+				{
+					ClosestEdgeDistSq = DistSq;
+					ClosestEdgePoint = ClosestPoint;
+					ClosestEdgeAlpha = Alpha;
+					ClosestEdgeIndex = EdgeIndex;
+				}
+			}
+			checkSlow(ClosestEdgeIndex != INDEX_NONE);
+
+			OutBary = TVec4<T>(static_cast<T>(0.));
+			OutBary[EdgeIndices[ClosestEdgeIndex][0]] = static_cast<T>(1.) - ClosestEdgeAlpha;
+			OutBary[EdgeIndices[ClosestEdgeIndex][1]] = ClosestEdgeAlpha;
+			checkSlow(TVec3<T>::DistSquared(ClosestEdgePoint, GetPointFromBarycentricCoordinates(OutBary)) < static_cast<T>(UE_SMALL_NUMBER));
+			return ClosestEdgePoint;
+		}
 
 	private:
 		static T TripleProduct(const TVec3<T>& U, const TVec3<T>& V, const TVec3<T>& W)
