Fix box overlap edge case

diff --git a/Engine/Source/Programs/HeadlessChaos/Private/HeadlessChaosTestOverlap.cpp b/Engine/Source/Programs/HeadlessChaos/Private/HeadlessChaosTestOverlap.cpp
--- a/Engine/Source/Programs/HeadlessChaos/Private/HeadlessChaosTestOverlap.cpp
+++ b/Engine/Source/Programs/HeadlessChaos/Private/HeadlessChaosTestOverlap.cpp
@@ -739,6 +739,34 @@ namespace ChaosTest
 			{
 				// Box test
 				const TBox<FReal, 3> Box = TBox<FReal, 3>({ -1.0, -1.0, -2.0 }, { 1.0, 1.0, 2.0 });
+				{
+					// Consistency test
+					for (FReal X = 0.0; X < 10.0; X += 1.0)
+					{
+						for (FReal Y = 0.0; Y < 10.0; Y += 1.0)
+						{
+							for (FReal Z = 0.0; Z < 10.0; Z += 1.0)
+							{
+								const FRigidTransform3 QueryTM(FVec3(X, Y, Z), FQuat(FVec3(1.0, 0.0, 0.0), 3.1415926 / 3.0));
+								bool bResult = TriangleMesh->OverlapGeom(Box, QueryTM, 0.0);
+								bool bResultMTD = TriangleMesh->OverlapGeom(Box, QueryTM, 0.0, &MTDInfo);
+								EXPECT_EQ(bResult, bResultMTD);
+
+								const FRigidTransform3 QueryIdTM(FVec3(X, Y, Z), FQuat::Identity);
+								bResult = TriangleMesh->OverlapGeom(Box, QueryIdTM, 0.0);
+								bResultMTD = TriangleMesh->OverlapGeom(Box, QueryIdTM, 0.0, &MTDInfo);
+								EXPECT_EQ(bResult, bResultMTD);
+							}
+						}
+					}
+				}
+				{
+					const FRigidTransform3 QueryTM(FVec3(7.0, 0.0, 10.0), FQuat(FVec3(0.0, 1.0, 0.0), 3.1415926 / 4.0));
+					bool bResult = TriangleMesh->OverlapGeom(Box, QueryTM, 0.0, nullptr);
+					EXPECT_EQ(bResult, false);
+					bool bResultMTD = TriangleMesh->OverlapGeom(Box, QueryTM, 0.0, &MTDInfo);
+					EXPECT_EQ(bResult, bResultMTD);
+				}
 				{
 					const FRigidTransform3 QueryTM(FVec3(0.0, 0.0, 0.0), FQuat::Identity);
 					bool bResult = TriangleMesh->OverlapGeom(Box, QueryTM, 0.0);
diff --git a/Engine/Source/Runtime/Experimental/Chaos/Private/Chaos/Collision/TriangleOverlap.cpp b/Engine/Source/Runtime/Experimental/Chaos/Private/Chaos/Collision/TriangleOverlap.cpp
--- a/Engine/Source/Runtime/Experimental/Chaos/Private/Chaos/Collision/TriangleOverlap.cpp
+++ b/Engine/Source/Runtime/Experimental/Chaos/Private/Chaos/Collision/TriangleOverlap.cpp
@@ -344,65 +344,59 @@ namespace Chaos
 		YAxisHalfExtent = VectorMultiply(YAxis, YHalfExtent);
 		ZAxisHalfExtent = VectorMultiply(ZAxis, ZHalfExtent);
 
-	}
-	namespace
-	{
-		FORCEINLINE_DEBUGGABLE bool HasToComputeEdge(const VectorRegister4Float& BoxNormal1, const VectorRegister4Float& BoxNormal2, const VectorRegister4Float& BoxEdge, const VectorRegister4Float& TriNormal, const VectorRegister4Float& TriEdge)
-		{
-			const FRealSingle TriNormBoxEdge = VectorDot3Scalar(TriNormal, BoxEdge);		// TriNormal | BoxEdge
-			const FRealSingle BoxNorm1TriEdge = VectorDot3Scalar(BoxNormal1, TriEdge);		// BoxNormalA | TriEdge
-			const FRealSingle BoxNorm2TriEdge = VectorDot3Scalar(BoxNormal2, TriEdge);		// BoxNormalB | TriEdge
+		const VectorRegister4Float PosPoint = VectorAdd(VectorAdd(XAxisHalfExtent, YAxisHalfExtent), ZAxisHalfExtent);
 
-			return ((BoxNorm1TriEdge * BoxNorm2TriEdge) < 0.0f) && (((TriNormBoxEdge * BoxNorm2TriEdge) > 0.0f || (TriNormBoxEdge * BoxNorm1TriEdge) > 0.0f));
+		VectorRegister4Float OtherPoints[3];
+		OtherPoints[0] = VectorAdd(VectorAdd(VectorNegate(XAxisHalfExtent), YAxisHalfExtent), ZAxisHalfExtent);
+		OtherPoints[1] = VectorAdd(VectorAdd(XAxisHalfExtent, VectorNegate(YAxisHalfExtent)), ZAxisHalfExtent);
+		OtherPoints[2] = VectorAdd(VectorAdd(XAxisHalfExtent, YAxisHalfExtent), VectorNegate(ZAxisHalfExtent));
+
+		for (int32 i = 0; i < 3; i++)
+		{
+			Edges[i] = VectorSubtract(OtherPoints[i], PosPoint);
 		}
+
 	}
 
-	FORCEINLINE_DEBUGGABLE bool FBoxSimd::ComputeEdge(const VectorRegister4Float& Normal, const VectorRegister4Float& PlaneVertex, const VectorRegister4Float& Edge, const VectorRegister4Float& Centroid) const
+	FORCEINLINE_DEBUGGABLE bool FBoxSimd::ComputeEdge(const VectorRegister4Float& Normal, const VectorRegister4Float& TriangleVertex, const VectorRegister4Float& Edge, const VectorRegister4Float& Centroid) const
 	{
-		const VectorRegister4Float PlaneNormal = VectorNormalize(VectorCross(Normal, Edge));
-		const VectorRegister4Float XComp = VectorSelect(VectorCompareGT(VectorDot3(PlaneNormal, XAxisHalfExtent), VectorZeroFloat()), XAxisHalfExtent, VectorNegate(XAxisHalfExtent));
-		const VectorRegister4Float YComp = VectorSelect(VectorCompareGT(VectorDot3(PlaneNormal, YAxisHalfExtent), VectorZeroFloat()), YAxisHalfExtent, VectorNegate(YAxisHalfExtent));
-		const VectorRegister4Float ZComp = VectorSelect(VectorCompareGT(VectorDot3(PlaneNormal, ZAxisHalfExtent), VectorZeroFloat()), ZAxisHalfExtent, VectorNegate(ZAxisHalfExtent));
+		const VectorRegister4Float PlaneNormal = VectorNormalize(VectorCross(Edge, Normal));
+		VectorRegister4Float XComp = VectorSelect(VectorCompareGT(VectorDot3(PlaneNormal, XAxisHalfExtent), VectorZeroFloat()), VectorNegate(XAxisHalfExtent), XAxisHalfExtent);
+		VectorRegister4Float YComp = VectorSelect(VectorCompareGT(VectorDot3(PlaneNormal, YAxisHalfExtent), VectorZeroFloat()), VectorNegate(YAxisHalfExtent), YAxisHalfExtent);
+		VectorRegister4Float ZComp = VectorSelect(VectorCompareGT(VectorDot3(PlaneNormal, ZAxisHalfExtent), VectorZeroFloat()), VectorNegate(ZAxisHalfExtent), ZAxisHalfExtent);
 
 		const VectorRegister4Float LocalClosest = VectorAdd(VectorAdd(XComp, YComp), ZComp);
 		const VectorRegister4Float ClosestPoint = VectorAdd(LocalClosest, Position);
 
-		FRealSingle Dist = VectorDot3Scalar(VectorSubtract(PlaneVertex, ClosestPoint), PlaneNormal);
-		if (Dist > 0.0)
+		FRealSingle DistClos = VectorDot3Scalar(VectorSubtract(ClosestPoint, TriangleVertex), PlaneNormal);
+		if (DistClos > 0.0)
 		{
 			return false;
 		}
 
 		// Triangle edge vs box edges 
-		VectorRegister4Float OtherClosests[3];
-		OtherClosests[0] = VectorAdd(VectorAdd(VectorNegate(XComp), YComp), ZComp);
-		OtherClosests[1] = VectorAdd(VectorAdd(XComp, VectorNegate(YComp)), ZComp);
-		OtherClosests[2] = VectorAdd(VectorAdd(XComp, YComp), VectorNegate(ZComp));
-
-		VectorRegister4Float BoxEdges[3];
-		VectorRegister4Float BoxEdgeNormals[3];
-		for (int32 i = 0; i < 3; i++)
+		for (int32 i = 0; i < EdgeNum; i++)
 		{
-			const VectorRegister4Float OtherClosest = VectorAdd(OtherClosests[i], Position);
-			BoxEdges[i] = VectorSubtract(OtherClosest, ClosestPoint);
-			BoxEdgeNormals[i] = VectorNegate(BoxEdges[i]);
-		}
+			VectorRegister4Float Axis = VectorCross(Edges[i], Edge);
+			const FRealSingle Sign = VectorDot3Scalar(VectorSubtract(TriangleVertex, Centroid), Axis);
+			if (Sign < 0.0f)
+			{
+				Axis = VectorNegate(Axis);
+			}
 
-		for (int32 i = 0; i < 3; i++)
-		{
-			if (HasToComputeEdge(BoxEdgeNormals[(i + 1) % 3], BoxEdgeNormals[(i + 2) % 3], BoxEdges[i], Normal, Edge))
+			const VectorRegister4Float XCompCurrent = VectorSelect(VectorCompareGT(VectorDot3(Axis, XAxisHalfExtent), VectorZeroFloat()), VectorNegate(XAxisHalfExtent), XAxisHalfExtent);
+			const VectorRegister4Float YCompCurrent = VectorSelect(VectorCompareGT(VectorDot3(Axis, YAxisHalfExtent), VectorZeroFloat()), VectorNegate(YAxisHalfExtent), YAxisHalfExtent);
+			const VectorRegister4Float ZCompCurrent = VectorSelect(VectorCompareGT(VectorDot3(Axis, ZAxisHalfExtent), VectorZeroFloat()), VectorNegate(ZAxisHalfExtent), ZAxisHalfExtent);
+
+			VectorRegister4Float LocalBoxVertex = VectorAdd(VectorAdd(XCompCurrent, YCompCurrent), ZCompCurrent);
+			const VectorRegister4Float BoxVertex = VectorAdd(LocalBoxVertex, Position);
+
+			const VectorRegister4Float SepVector = VectorSubtract(BoxVertex, TriangleVertex);
+			FRealSingle SeparationMin = VectorDot3Scalar(SepVector, Axis);
+
+			if (SeparationMin > UE_KINDA_SMALL_NUMBER)
 			{
-				VectorRegister4Float Axis = VectorCross(BoxEdges[i], Edge);
-				const FRealSingle Sign = VectorDot3Scalar(VectorSubtract(PlaneVertex, Centroid), Axis);
-				if (Sign < 0.0f)
-				{
-					Axis = VectorNegate(Axis);
-				}
-				const FRealSingle Separation = VectorDot3Scalar(VectorSubtract(ClosestPoint, PlaneVertex), Axis);
-				if (Separation > UE_KINDA_SMALL_NUMBER)
-				{
-					return false;
-				}
+				return false;
 			}
 		}
 		return true;
@@ -566,6 +560,18 @@ namespace Chaos
 		HalfExtents = VectorLoadFloat3(&HalfExtentsf.X);
 	}
 
+	namespace
+	{
+		FORCEINLINE_DEBUGGABLE bool HasToComputeEdge(const VectorRegister4Float& BoxNormal1, const VectorRegister4Float& BoxNormal2, const VectorRegister4Float& BoxEdge, const VectorRegister4Float& TriNormal, const VectorRegister4Float& TriEdge)
+		{
+			const FRealSingle TriNormBoxEdge = VectorDot3Scalar(TriNormal, BoxEdge);		// TriNormal | BoxEdge
+			const FRealSingle BoxNorm1TriEdge = VectorDot3Scalar(BoxNormal1, TriEdge);		// BoxNormalA | TriEdge
+			const FRealSingle BoxNorm2TriEdge = VectorDot3Scalar(BoxNormal2, TriEdge);		// BoxNormalB | TriEdge
+
+			return ((BoxNorm1TriEdge * BoxNorm2TriEdge) < 0.0f) && (((TriNormBoxEdge * BoxNorm2TriEdge) > 0.0f || (TriNormBoxEdge * BoxNorm1TriEdge) > 0.0f));
+		}
+	}
+
 	FORCEINLINE_DEBUGGABLE bool FAABBSimd::ComputeEdgeOverlap(const VectorRegister4Float& TriangleEdge, const VectorRegister4Float& TriangleVertex, const VectorRegister4Float& Centroid, const VectorRegister4Float& Normal, const VectorRegister4Float& LocalClosest) const
 	{
 		// Triangle edge vs box edges 
diff --git a/Engine/Source/Runtime/Experimental/Chaos/Private/Chaos/Collision/TriangleOverlap.h b/Engine/Source/Runtime/Experimental/Chaos/Private/Chaos/Collision/TriangleOverlap.h
--- a/Engine/Source/Runtime/Experimental/Chaos/Private/Chaos/Collision/TriangleOverlap.h
+++ b/Engine/Source/Runtime/Experimental/Chaos/Private/Chaos/Collision/TriangleOverlap.h
@@ -39,6 +39,8 @@ namespace Chaos
 		VectorRegister4Float XAxisHalfExtent;
 		VectorRegister4Float YAxisHalfExtent;
 		VectorRegister4Float ZAxisHalfExtent;
+		static constexpr int EdgeNum = 3;
+		VectorRegister4Float Edges[EdgeNum];
 	};
 
 	class FAABBSimd
