Chaos - Physics thread particle sweep functionality
- SetParticleTransformSwept

diff --git a/Engine/Source/Programs/HeadlessChaos/Private/HeadlessChaosTestMove.cpp b/Engine/Source/Programs/HeadlessChaos/Private/HeadlessChaosTestMove.cpp
new file mode 100644
--- /dev/null
+++ b/Engine/Source/Programs/HeadlessChaos/Private/HeadlessChaosTestMove.cpp
@@ -0,0 +1,128 @@
+// Copyright Epic Games, Inc. All Rights Reserved.
+#include "HeadlessChaos.h"
+#include "HeadlessChaosTestEvolution.h"
+
+namespace ChaosTest
+{
+	using namespace Chaos;
+
+	GTEST_TEST(MoveTests, TestMoveNoSweep)
+	{
+		const int32 NumTicks = 100;
+		const FReal Dt = 1.0/30.0;
+		const FReal DistanceTolerance = 0.1f;
+
+		const FVec3 FloorSize = FVec3(10000, 10000, 100);
+		const FVec3 BoxPos = FVec3(0,0,200);
+		const FRotation3 BoxRot = FRotation3::FromIdentity();
+		const FVec3 BoxSize = FVec3(100);
+		const FReal BoxMass = 100;
+
+		FEvolutionTest Test;
+
+		FGeometryParticleHandle* Floor = Test.AddParticleFloor(FloorSize);
+		FPBDRigidParticleHandle* Box = Test.AddParticleBox(BoxPos, BoxRot, BoxSize, BoxMass);
+
+		// Drop the box on the ground and wait until it sleeps
+		Test.TickUntilSleep();
+		EXPECT_TRUE(Box->IsSleeping());
+		EXPECT_NEAR(Box->X().Z, 0.5 * BoxSize.Z, DistanceTolerance);
+
+		// Move the box horizontally and up into the air
+		const FVec3 NewBoxPos = Box->X() + FVec3(100, 0, 100);
+		Test.GetEvolution().SetParticleTransformSwept(Box, NewBoxPos, BoxRot, false);
+
+		// The box should be at its new location and awake
+		EXPECT_FALSE(Box->IsSleeping());
+		EXPECT_NEAR(Box->X().X, NewBoxPos.X, DistanceTolerance);
+		EXPECT_NEAR(Box->X().Y, NewBoxPos.Y, DistanceTolerance);
+		EXPECT_NEAR(Box->X().Z, NewBoxPos.Z, DistanceTolerance);
+
+		// The box should fall to the ground again
+		Test.TickUntilSleep();
+		EXPECT_TRUE(Box->IsSleeping());
+		EXPECT_NEAR(Box->X().X, NewBoxPos.X, DistanceTolerance);
+		EXPECT_NEAR(Box->X().Z, 0.5 * BoxSize.Z, DistanceTolerance);
+	}
+
+	GTEST_TEST(MoveTests, TestMoveWithSweepNoHit)
+	{
+		const int32 NumTicks = 100;
+		const FReal Dt = 1.0 / 30.0;
+		const FReal DistanceTolerance = 0.1f;
+
+		const FVec3 FloorSize = FVec3(10000, 10000, 100);
+		const FVec3 BoxPos = FVec3(0, 0, 200);
+		const FRotation3 BoxRot = FRotation3::FromIdentity();
+		const FVec3 BoxSize = FVec3(100);
+		const FReal BoxMass = 100;
+
+		FEvolutionTest Test;
+
+		FGeometryParticleHandle* Floor = Test.AddParticleFloor(FloorSize);
+		FPBDRigidParticleHandle* Box = Test.AddParticleBox(BoxPos, BoxRot, BoxSize, BoxMass);
+
+		// Drop the box on the ground and wait until it sleeps
+		Test.TickUntilSleep();
+		EXPECT_TRUE(Box->IsSleeping());
+		EXPECT_NEAR(Box->X().Z, 0.5 * BoxSize.Z, DistanceTolerance);
+
+		// Move the box horizontally and up into the air
+		const FVec3 NewBoxPos = Box->X() + FVec3(100, 0, 100);
+		Test.GetEvolution().SetParticleTransformSwept(Box, NewBoxPos, BoxRot, false);
+
+		// The box should be at its new location and awake
+		EXPECT_FALSE(Box->IsSleeping());
+		EXPECT_NEAR(Box->X().X, NewBoxPos.X, DistanceTolerance);
+		EXPECT_NEAR(Box->X().Y, NewBoxPos.Y, DistanceTolerance);
+		EXPECT_NEAR(Box->X().Z, NewBoxPos.Z, DistanceTolerance);
+
+		// The box should fall to the ground again
+		Test.TickUntilSleep();
+		EXPECT_TRUE(Box->IsSleeping());
+		EXPECT_NEAR(Box->X().X, NewBoxPos.X, DistanceTolerance);
+		EXPECT_NEAR(Box->X().Z, 0.5 * BoxSize.Z, DistanceTolerance);
+	}
+
+	GTEST_TEST(MoveTests, TestMoveWithSweepBlockingHit)
+	{
+		const int32 NumTicks = 100;
+		const FReal Dt = 1.0 / 30.0;
+		const FReal DistanceTolerance = 0.1f;
+
+		const FVec3 FloorSize = FVec3(10000, 10000, 100);
+		const FVec3 BoxPos = FVec3(0, 0, 200);
+		const FRotation3 BoxRot = FRotation3::FromIdentity();
+		const FVec3 BoxSize = FVec3(100);
+		const FReal BoxMass = 100;
+
+		const FVec3 BoxTargetPos = FVec3(100, 0, 100 + 0.5 * BoxSize.Z);
+
+		FEvolutionTest Test;
+
+		FGeometryParticleHandle* Floor = Test.AddParticleFloor(FloorSize);
+		FPBDRigidParticleHandle* Box = Test.AddParticleBox(BoxPos, BoxRot, BoxSize, BoxMass);
+		FPBDRigidParticleHandle* BlockingBox = Test.AddParticleBox(BoxTargetPos + 0.5 * BoxSize, BoxRot, BoxSize, 0);
+
+		// Drop the box on the ground and wait until it sleeps
+		Test.TickUntilSleep();
+		EXPECT_TRUE(Box->IsSleeping());
+		EXPECT_NEAR(Box->X().Z, 0.5 * BoxSize.Z, DistanceTolerance);
+
+		// Move the box horizontally and up into the air
+		Test.GetEvolution().SetParticleTransformSwept(Box, BoxTargetPos, BoxRot, false);
+
+		// The box should have been stopped on its way to its new location
+		EXPECT_FALSE(Box->IsSleeping());
+		EXPECT_NEAR(Box->X().X, 50, DistanceTolerance);
+		EXPECT_NEAR(Box->X().Y, 0, DistanceTolerance);
+		EXPECT_NEAR(Box->X().Z, 50 + 0.5 * BoxSize.Z, DistanceTolerance);
+
+		// The box should fall to the ground again
+		Test.TickUntilSleep();
+		EXPECT_TRUE(Box->IsSleeping());
+		EXPECT_NEAR(Box->X().X, 50, DistanceTolerance);
+		EXPECT_NEAR(Box->X().Z, 0.5 * BoxSize.Z, DistanceTolerance);
+	}
+
+}
\ No newline at end of file
diff --git a/Engine/Source/Programs/HeadlessChaos/Private/HeadlessChaosTestUtility.cpp b/Engine/Source/Programs/HeadlessChaos/Private/HeadlessChaosTestUtility.cpp
--- a/Engine/Source/Programs/HeadlessChaos/Private/HeadlessChaosTestUtility.cpp
+++ b/Engine/Source/Programs/HeadlessChaos/Private/HeadlessChaosTestUtility.cpp
@@ -164,12 +164,11 @@ namespace ChaosTest {
 		Particle->P() = Particle->X();
 		Particle->Q() = Particle->R();
 
-		// Inertia assumes cube so is imcorrect for rectangular boxes
-		FReal MaxScale = Scale.GetMax();
-		FReal ScaleSq = MaxScale * MaxScale;
+		// Assume unit mass - this gets scaled externally by the actual mass
+		const FVec3 ScaleSq = Scale * Scale;
 		Particle->M() = 1.0;
 		Particle->InvM() = 1.0;
-		Particle->I() = TVec3<FRealSingle>(ScaleSq / 6.0);
+		Particle->I() = (1.0 / 12.0) * TVec3<FRealSingle>(ScaleSq.Y + ScaleSq.Z, ScaleSq.X + ScaleSq.Z, ScaleSq.X + ScaleSq.Y);
 		Particle->InvI() = TVec3<FRealSingle>(6.0 / ScaleSq);
 
 		Particle->SetDynamicGeometry(MakeUnique<TBox<FReal, 3>>(-Scale / 2.0, Scale / 2.0, Margin));
@@ -186,8 +185,7 @@ namespace ChaosTest {
 		Particle->CollisionParticles()->X(CollisionIndex++) = FVec3(-Scale[0] / 2.0, +Scale[1] / 2.0, +Scale[2] / 2.0);
 		Particle->CollisionParticles()->X(CollisionIndex++) = FVec3(+Scale[0] / 2.0, +Scale[1] / 2.0, +Scale[2] / 2.0);
 
-		// This is needed for calculating contacts (Bounds are bigger than they need to be, even allowing for rotation)
-		Particle->SetLocalBounds(FAABB3(FVec3(-MaxScale), FVec3(MaxScale)));
+		Particle->SetLocalBounds(Particle->Geometry()->BoundingBox());
 		Particle->UpdateWorldSpaceState(FRigidTransform3::Identity, FVec3(0));
 		Particle->SetHasBounds(true);
 
@@ -408,9 +406,8 @@ namespace ChaosTest {
 	void InitStaticParticleBox(FGeometryParticleHandle* Particle, const FVec3& Scale, TArray<TVec3<int32>>* OutElements)
 	{
 		Particle->X() = FVec3(0, 0, 0);
-		Particle->R() = FRotation3::MakeFromEuler(FVec3(0, 0, 0)).GetNormalized();
+		Particle->R() = FRotation3::FromIdentity();
 
-		check(Scale.X == Scale.Y && Scale.X == Scale.Z);
 		FReal ScaleSq = Scale.X * Scale.X;
 
 		Particle->SetDynamicGeometry(MakeUnique<TBox<FReal, 3>>(-Scale / 2.0, Scale / 2.0));
diff --git a/Engine/Source/Programs/HeadlessChaos/Public/HeadlessChaosTestEvolution.h b/Engine/Source/Programs/HeadlessChaos/Public/HeadlessChaosTestEvolution.h
new file mode 100644
--- /dev/null
+++ b/Engine/Source/Programs/HeadlessChaos/Public/HeadlessChaosTestEvolution.h
@@ -0,0 +1,198 @@
+// Copyright Epic Games, Inc. All Rights Reserved.
+#pragma once
+
+#include "Chaos/Box.h"
+#include "Chaos/Collision/PBDCollisionConstraint.h"
+#include "Chaos/PBDCollisionConstraints.h"
+#include "Chaos/PBDRigidsEvolutionGBF.h"
+#include "Chaos/PBDRigidParticles.h"
+#include "Chaos/Sphere.h"
+#include "Chaos/Utilities.h"
+#include "HeadlessChaosTestUtility.h"
+
+namespace ChaosTest 
+{
+	/**
+	 * Base class for evolution tests
+	 */
+	class FEvolutionTest
+	{
+	public:
+
+		FEvolutionTest()
+			: SOAs(UniqueIndices)
+			, Evolution(SOAs, PhysicalMaterials)
+		{
+			PhysicalMaterial = MakeUnique<FChaosPhysicsMaterial>();
+			PhysicalMaterial->Friction = 0;
+			PhysicalMaterial->Restitution = 0;
+			PhysicalMaterial->DisabledLinearThreshold = 0;
+			PhysicalMaterial->DisabledAngularThreshold = 0;
+			EnableSleeping();
+
+			Dt = 1.0 / 30.0;
+		}
+
+		virtual ~FEvolutionTest()
+		{
+		}
+
+		FPBDRigidsEvolutionGBF& GetEvolution()
+		{
+			return Evolution;
+		}
+
+		FPBDCollisionConstraints& GetCollisionConstraints()
+		{
+			return Evolution.GetCollisionConstraints();
+		}
+
+		// Get a particle. The index is the order in which they were created by the AddParticleXXX methods
+		FGeometryParticleHandle* GetParticle(const int32 ParticleIndex)
+		{
+			return Particles[ParticleIndex];
+		}
+
+		// Set up the physical material to allow sleeping (sleeping is enabled by default)
+		void EnableSleeping()
+		{
+			// These are the defaults from PhysicalMaterial.cpp
+			PhysicalMaterial->SleepingLinearThreshold = 1;
+			PhysicalMaterial->SleepingAngularThreshold = 0.05;
+			PhysicalMaterial->SleepCounterThreshold = 4;
+			OnMaterialChanged();
+		}
+
+		// Set up the physical material to disable sleeping (sleeping is enabled by default)
+		void DisableSleeping()
+		{
+			// These are the defaults from PhysicalMaterial.cpp
+			PhysicalMaterial->SleepingLinearThreshold = 0;
+			PhysicalMaterial->SleepingAngularThreshold = 0;
+			PhysicalMaterial->SleepCounterThreshold = std::numeric_limits<int32>::max();
+			OnMaterialChanged();
+		}
+
+		void Tick()
+		{
+			Evolution.AdvanceOneTimeStep(Dt);
+			Evolution.EndFrame(Dt);
+		}
+
+		bool TickUntilSleep(const int32 MaxTicks = 1000)
+		{
+			int32 TickIndex = 0;
+			while (!AreAllParticlesAsleep() && (TickIndex++ < MaxTicks))
+			{
+				Tick();
+			}
+
+			return AreAllParticlesAsleep();
+		}
+
+		bool AreAllParticlesAsleep() const
+		{
+			for (FGeometryParticleHandle* Particle : Particles)
+			{
+				if (FPBDRigidParticleHandle* Rigid = Particle->CastToRigidParticle())
+				{
+					if (Rigid->IsDynamic() && !Rigid->IsSleeping())
+					{
+						return false;
+					}
+				}
+			}
+			return true;
+		}
+
+		// Create a dynamic or kinematic particle depending on the mass
+		FPBDRigidParticleHandle* AddParticleBox(const FVec3& Position, const FRotation3& Rotation, const FVec3& Size, FReal Mass)
+		{
+			FPBDRigidParticleHandle* Particle = AppendDynamicParticleBox(SOAs, Size);
+
+			if (Mass == 0)
+			{
+				Evolution.SetParticleObjectState(Particle, EObjectStateType::Kinematic);
+			}
+
+			ResetParticle(Particle, Position, Rotation, FVec3(0), FVec3(0));
+
+			if(Particle->ObjectState() == EObjectStateType::Dynamic)
+			{
+				// NOTE: particle was created with unit mass
+				Particle->M() = Particle->M() * Mass;
+				Particle->I() = Particle->I() * Mass;
+				Particle->InvM() = Particle->InvM() * ((FReal)1 / Mass);
+				Particle->InvI() = Particle->InvI() * ((FReal)1 / Mass);
+			}
+			Evolution.SetPhysicsMaterial(Particle, MakeSerializable(PhysicalMaterial));
+
+			Evolution.EnableParticle(Particle);
+
+			Particles.Add(Particle);
+
+			return Particle;
+		}
+
+		// Add a static particle, e.g., for the floor
+		FGeometryParticleHandle* AddStaticParticleBox(const FVec3& Position, const FRotation3& Rotation, const FVec3& Size)
+		{
+			FGeometryParticleHandle* Particle = AppendStaticParticleBox(SOAs, Size);
+
+			ResetParticle(Particle, Position, Rotation, FVec3(0), FVec3(0));
+
+			Evolution.SetPhysicsMaterial(Particle, MakeSerializable(PhysicalMaterial));
+
+			Evolution.EnableParticle(Particle);
+
+			Particles.Add(Particle);
+
+			return Particle;
+		}
+
+		// Add a static particle for the floor consisting of a box with the specified size and the upper surface at Z=0
+		FGeometryParticleHandle* AddParticleFloor(const FVec3& Size)
+		{
+			return AddStaticParticleBox(FVec3(0, 0, -0.5 * Size.Z), FRotation3::FromIdentity(), Size);
+		}
+
+		void ResetParticle(FGeometryParticleHandle* Particle, const FVec3& Position, const FRotation3& Rotation, const FVec3& Velocity, const FVec3& AngularVelocity)
+		{
+			Particle->X() = Position;
+			Particle->R() = Rotation;
+			if (FKinematicGeometryParticleHandle* KinParticle = Particle->CastToKinematicParticle())
+			{
+				KinParticle->V() = Velocity;
+				KinParticle->W() = AngularVelocity;
+			}
+			if (FPBDRigidParticleHandle* PBDParticle = Particle->CastToRigidParticle())
+			{
+				PBDParticle->P() = Position;
+				PBDParticle->Q() = Rotation;
+			}
+		}
+
+		void OnMaterialChanged()
+		{
+			// When we change the material so any collisions using it must reacquire materials properties
+			GetCollisionConstraints().GetConstraintAllocator().VisitCollisions(
+				[](FPBDCollisionConstraint& Collision)
+				{
+					Collision.ClearMaterialProperties();
+					return ECollisionVisitorResult::Continue;
+				});
+		}
+
+		// Solver state
+		FParticleUniqueIndicesMultithreaded UniqueIndices;
+		FPBDRigidsSOAs SOAs;
+		FPBDRigidsEvolutionGBF Evolution;
+		TUniquePtr<FChaosPhysicsMaterial> PhysicalMaterial;
+		THandleArray<FChaosPhysicsMaterial> PhysicalMaterials;
+
+		TArray<FGeometryParticleHandle*> Particles;
+
+		FReal Dt;
+	};
+
+}
\ No newline at end of file
diff --git a/Engine/Source/Runtime/Experimental/Chaos/Private/Chaos/Collision/SimSweep.cpp b/Engine/Source/Runtime/Experimental/Chaos/Private/Chaos/Collision/SimSweep.cpp
new file mode 100644
--- /dev/null
+++ b/Engine/Source/Runtime/Experimental/Chaos/Private/Chaos/Collision/SimSweep.cpp
@@ -0,0 +1,29 @@
+// Copyright Epic Games, Inc. All Rights Reserved.
+#include "Chaos/Collision/SimSweep.h"
+#include "Chaos/ISpatialAcceleration.h"
+
+namespace Chaos
+{
+	namespace Private
+	{
+		bool SimSweepParticleFirstHit(
+			ISpatialAcceleration<FAccelerationStructureHandle, FReal, 3>* SpatialAcceleration,
+			FIgnoreCollisionManager* InIgnoreCollisionManager,
+			const FGeometryParticleHandle* SweptParticle,
+			const FVec3& StartPos,
+			const FRotation3& Rot,
+			const FVec3& Dir,
+			const FReal Length,
+			FSimSweepParticleHit& OutHit,
+			const FReal InHitDistanceEqualTolerance)
+		{
+			FSimSweepParticleFilterBroadPhase ParticleFilter(InIgnoreCollisionManager);
+			FSimSweepShapeFilterNarrowPhase ShapeFilter;
+			FSimSweepCollectorFirstHit HitCollector(InHitDistanceEqualTolerance, OutHit);
+
+			SimSweepParticle(SpatialAcceleration, SweptParticle, StartPos, Rot, Dir, Length, ParticleFilter, ShapeFilter, HitCollector);
+
+			return OutHit.IsHit();
+		}
+	}
+}
\ No newline at end of file
diff --git a/Engine/Source/Runtime/Experimental/Chaos/Private/Chaos/PBDRigidsEvolutionGBF.cpp b/Engine/Source/Runtime/Experimental/Chaos/Private/Chaos/PBDRigidsEvolutionGBF.cpp
--- a/Engine/Source/Runtime/Experimental/Chaos/Private/Chaos/PBDRigidsEvolutionGBF.cpp
+++ b/Engine/Source/Runtime/Experimental/Chaos/Private/Chaos/PBDRigidsEvolutionGBF.cpp
@@ -1,5 +1,6 @@
 // Copyright Epic Games, Inc. All Rights Reserved.
 #include "Chaos/PBDRigidsEvolutionGBF.h"
+#include "Chaos/Collision/SimSweep.h"
 #include "Chaos/Defines.h"
 #include "Chaos/Evolution/SolverBodyContainer.h"
 #include "Chaos/Framework/Parallel.h"
@@ -936,6 +937,46 @@ void FPBDRigidsEvolutionGBF::DestroyTransientConstraints(FGeometryParticleHandle
 	}
 }
 
+CHAOS_API void FPBDRigidsEvolutionGBF::SetParticleTransform(FGeometryParticleHandle* InParticle, const FVec3& InPos, const FRotation3& InRot, const bool bIsTeleport)
+{
+	const FVec3 PrevX = InParticle->X();
+	const FRotation3 PrevR = InParticle->R();
+
+	FGenericParticleHandle(InParticle)->SetTransform(InPos, InRot);
+
+	OnParticleMoved(InParticle, PrevX, PrevR, bIsTeleport);
+}
+
+
+void FPBDRigidsEvolutionGBF::SetParticleTransformSwept(FGeometryParticleHandle* InParticle, const FVec3& InPos, const FRotation3& InRot, const bool bIsTeleport)
+{
+	// If the rotation has changed, we need to adjust the initial sweep position to account for the fact that the
+	// particle will have moved in a straight line through its center of mass. This is important when we have shapes
+	// that are not centerd on their center of mass (either it's multi-shape body, or the user modified the CoM).
+	const FVec3 CoMOffset = FConstGenericParticleHandle(InParticle)->CenterOfMass();
+	const FVec3 EndCoM = InPos + InRot * CoMOffset;
+	const FVec3 StartCoM = InParticle->X() + InParticle->R() * CoMOffset;
+	FVec3 SweepDir = EndCoM - StartCoM;
+	const FReal SweepLength = SweepDir.SafeNormalize();
+	const FVec3 SweepStart = InPos - SweepDir * SweepLength;
+
+	FReal TOI = FReal(1.0);
+	if (SweepLength > UE_SMALL_NUMBER)
+	{
+		Private::FSimSweepParticleHit Hit;
+		if (Private::SimSweepParticleFirstHit(GetSpatialAcceleration(), &GetBroadPhase().GetIgnoreCollisionManager(), InParticle, SweepStart, InRot, SweepDir, SweepLength, Hit))
+		{
+			TOI = Hit.HitTOI;
+		}
+	}
+
+	if (TOI > 0)
+	{
+		const FVec3 NewPos = FMath::Lerp(SweepStart, InPos, TOI);
+		SetParticleTransform(InParticle, NewPos, InRot, bIsTeleport);
+	}
+}
+
 void FPBDRigidsEvolutionGBF::OnParticleMoved(FGeometryParticleHandle* InParticle, const FVec3& PrevX, const FRotation3& PrevR, const bool bIsTeleport)
 {
 	// When a particle is moved, we need to tell the collisions because they cache friction state and 
@@ -962,6 +1003,11 @@ void FPBDRigidsEvolutionGBF::OnParticleMoved(FGeometryParticleHandle* InParticle
 				const FReal SmoothRate = FMath::Clamp(CVars::SmoothedPositionLerpRate, 0.0f, 1.0f);
 				Rigid->VSmooth() = FMath::Lerp(Rigid->VSmooth(), Rigid->V() + DV, SmoothRate);
 			}
+
+			if (Rigid->IsSleeping())
+			{
+				SetParticleObjectState(Rigid, EObjectStateType::Dynamic);
+			}
 		}
 	}
 }
diff --git a/Engine/Source/Runtime/Experimental/Chaos/Public/Chaos/Collision/CollisionConstraintAllocator.h b/Engine/Source/Runtime/Experimental/Chaos/Public/Chaos/Collision/CollisionConstraintAllocator.h
--- a/Engine/Source/Runtime/Experimental/Chaos/Public/Chaos/Collision/CollisionConstraintAllocator.h
+++ b/Engine/Source/Runtime/Experimental/Chaos/Public/Chaos/Collision/CollisionConstraintAllocator.h
@@ -483,7 +483,7 @@ namespace Chaos
 			}
 
 			/**
-			 * @brief Iterate over all collisions, including sleeping ones
+			 * @brief Iterate over all collisions (read-only), including sleeping ones
 			 * Visitor signature: ECollisionVisitorResult(const FPBDCollisionConstraint&)
 			*/
 			template<typename TLambda>
@@ -498,6 +498,22 @@ namespace Chaos
 				}
 			}
 
+			/**
+			 * @brief Iterate over all collisions (write-enabled), including sleeping ones
+			 * Visitor signature: ECollisionVisitorResult(FPBDCollisionConstraint&)
+			*/
+			template<typename TLambda>
+			void VisitCollisions(const TLambda& Visitor) const
+			{
+				for (const FParticlePairMidPhasePtr& MidPhase : ParticlePairMidPhases)
+				{
+					if (MidPhase->VisitCollisions(Visitor) == ECollisionVisitorResult::Stop)
+					{
+						return;
+					}
+				}
+			}
+
 			/**
 			 * @brief Iterate over all midphases
 			 * Visitor signature: ECollisionVisitorResult(FParticlePairMidPhase&)
diff --git a/Engine/Source/Runtime/Experimental/Chaos/Public/Chaos/Collision/CollisionConstraintFlags.h b/Engine/Source/Runtime/Experimental/Chaos/Public/Chaos/Collision/CollisionConstraintFlags.h
--- a/Engine/Source/Runtime/Experimental/Chaos/Public/Chaos/Collision/CollisionConstraintFlags.h
+++ b/Engine/Source/Runtime/Experimental/Chaos/Public/Chaos/Collision/CollisionConstraintFlags.h
@@ -14,9 +14,6 @@
 
 namespace Chaos
 {
-	// @todo(chaos): remove these when FIgnoreCollisionManager is split into Core and Proxy classes (see todo comments below)
-	class FPBDRigidsSolver;
-
 	/**
 	 * @brief Flags for user-control over per-particle collision behaviour
 	 * @todo(chaos): move SmoothEdgeCollisions to FRigidParticleControlFlags and remove ECollisionConstraintFlags
diff --git a/Engine/Source/Runtime/Experimental/Chaos/Public/Chaos/Collision/CollisionFilter.h b/Engine/Source/Runtime/Experimental/Chaos/Public/Chaos/Collision/CollisionFilter.h
--- a/Engine/Source/Runtime/Experimental/Chaos/Public/Chaos/Collision/CollisionFilter.h
+++ b/Engine/Source/Runtime/Experimental/Chaos/Public/Chaos/Collision/CollisionFilter.h
@@ -3,12 +3,12 @@
 
 // HEADER_UNIT_SKIP - Internal
 
-#include "CoreMinimal.h"
 #include "Chaos/Core.h"
+#include "Chaos/Collision/CollisionConstraintFlags.h"
 #include "Chaos/CollisionFilterData.h"
 #include "Chaos/GeometryParticles.h"
-#include "Chaos/ImplicitObject.h"
-#include "Chaos/Particles.h"
+#include "Chaos/ImplicitObjectType.h"
+#include "Chaos/ParticleHandle.h"
 
 namespace Chaos
 {
@@ -16,29 +16,105 @@ namespace Chaos
 	enum { NumExtraFilterBits = 6 };
 	enum { NumCollisionChannelBits = 5 };
 
-	inline bool IsFilterValid(const FCollisionFilterData& Filter)
+	/**
+	 * Check whether the two particles need to be considered in the broadphase.
+	 * @return true if the particle pair should be considered by the broadphase
+	 */
+	inline bool ParticlePairBroadPhaseFilter(
+		const FGeometryParticleHandle* Particle1,
+		const FGeometryParticleHandle* Particle2,
+		const FIgnoreCollisionManager* IgnoreCollisionManager)
 	{
-		return Filter.Word0 || Filter.Word1 || Filter.Word2 || Filter.Word3;
+		if (Particle1 == Particle2)
+		{
+			return false;
+		}
+
+		bool bIsKinematic1 = true;
+		bool bUseIgnoreCollisionManager1 = false;
+		bool bDisabled1 = false;
+		int32 CollisionGroup1 = 0;
+		const FPBDRigidParticleHandle* Rigid1 = Particle1->CastToRigidParticle();
+		if (Rigid1 != nullptr)
+		{
+			bIsKinematic1 = Rigid1->IsKinematic();
+			bUseIgnoreCollisionManager1 = Rigid1->UseIgnoreCollisionManager();
+			bDisabled1 = Rigid1->Disabled();
+			CollisionGroup1 = Rigid1->CollisionGroup();
+		}
+
+		bool bIsKinematic2 = true;
+		bool bUseIgnoreCollisionManager2 = false;
+		bool bDisabled2 = false;
+		int32 CollisionGroup2 = 0;
+		const FPBDRigidParticleHandle* Rigid2 = Particle2->CastToRigidParticle();
+		if (Rigid2 != nullptr)
+		{
+			bIsKinematic2 = Rigid2->IsKinematic();
+			bUseIgnoreCollisionManager2 = Rigid2->UseIgnoreCollisionManager();
+			bDisabled2 = Rigid2->Disabled();
+			CollisionGroup2 = Rigid2->CollisionGroup();
+		}
+
+		// @todo(chaos): This should not be happening if the disabled particles are removed from the active particles list, but GeometryCollection may leave them there
+		//check(!bDisabled2);
+		if (bDisabled1 || bDisabled2)
+		{
+			return false;
+		}
+
+		// At least one particle needs to be dynamic to generate a collision response
+		if (bIsKinematic1 && bIsKinematic2)
+		{
+			return false;
+		}
+		check((Rigid1 != nullptr) || (Rigid2 != nullptr));
+
+		// CollisionGroups are used by geometry collections for high-level collision filtering
+		// CollisionGroup == 0 : Collide_With_Everything
+		// CollisionGroup == INDEX_NONE : Disabled collisions
+		// CollisionGroup1 != CollisionGroup2 : Disabled collisions (if other conditions not met)
+		if (CollisionGroup1 == INDEX_NONE || CollisionGroup2 == INDEX_NONE)
+		{
+			return false;
+		}
+		if ((CollisionGroup1 != 0) && (CollisionGroup2 != 0) && (CollisionGroup1 != CollisionGroup2))
+		{
+			return false;
+		}
+
+		// Is this particle interaction governed by the IgnoreCollisionManager? If so, check to see if interaction is allowed
+		if ((bUseIgnoreCollisionManager1 || bUseIgnoreCollisionManager2) && (IgnoreCollisionManager != nullptr))
+		{
+			if (IgnoreCollisionManager->IgnoresCollision(Particle1, Particle2))
+			{
+				return false;
+			}
+		}
+
+		return true;
 	}
 
-	//inline uint32 GetChaosCollisionChannel(uint32 Word3)
-	//{
-	//	uint32 ChannelMask = (Word3 << NumExtraFilterBits) >> (32 - NumCollisionChannelBits);
-	//	return (uint32)ChannelMask;
-	//}
 
-	//inline uint32 GetChaosCollisionChannelAndExtraFilter(uint32 Word3, FMaskFilter& OutMaskFilter)
-	//{
-	//	uint32 ChannelMask = GetChaosCollisionChannel(Word3);
-	//	OutMaskFilter = Word3 >> (32 - NumExtraFilterBits);
-	//	return (uint32)ChannelMask;
-	//}
+	/**
+	* A filter set to all zeroes means we do not filter
+	*/
+	inline bool IsFilterValid(const FCollisionFilterData& Filter)
+	{
+		return Filter.Word0 || Filter.Word1 || Filter.Word2 || Filter.Word3;
+	}
 
+	/**
+	* Does the shape collide with anything at all in the Sim? (else it is query-only)
+	*/
 	inline bool FilterHasSimEnabled(const FPerShapeData* Shape)
 	{
 		return (!Shape || (Shape->GetSimEnabled() && IsFilterValid(Shape->GetSimData())));
 	}
 
+	/**
+	* Sim Collision filter (i.e., not Query collision filter)
+	*/
 	inline bool DoCollide(EImplicitObjectType Implicit0Type, const FPerShapeData* Shape0, EImplicitObjectType Implicit1Type, const FPerShapeData* Shape1)
 	{
 		//
@@ -100,7 +176,17 @@ namespace Chaos
 			}
 		}
 
-
 		return true;
 	}
+
+	/**
+	 * Check whether the two particles need to be considered in the broadphase
+	 * NOTE: Implicit0Type and Implicit1Type must be the inner type of the implicit, 
+	 * with any decorators removed (scaled, instanced, transformed).
+	 * @return true if the shape pair should be considered by the narrowphase
+	 */
+	inline bool ShapePairNarrowPhaseFilter(EImplicitObjectType Implicit0Type, const FPerShapeData* Shape0, EImplicitObjectType Implicit1Type, const FPerShapeData* Shape1)
+	{
+		return DoCollide(Implicit0Type, Shape0, Implicit1Type, Shape1);
+	}
 }
\ No newline at end of file
diff --git a/Engine/Source/Runtime/Experimental/Chaos/Public/Chaos/Collision/ParticlePairMidPhase.h b/Engine/Source/Runtime/Experimental/Chaos/Public/Chaos/Collision/ParticlePairMidPhase.h
--- a/Engine/Source/Runtime/Experimental/Chaos/Public/Chaos/Collision/ParticlePairMidPhase.h
+++ b/Engine/Source/Runtime/Experimental/Chaos/Public/Chaos/Collision/ParticlePairMidPhase.h
@@ -241,12 +241,14 @@ namespace Chaos
 
 		/**
 		 * @brief Visit all the collisions
+		 * @param Visitor functor with signature ECollisionVisitorResult(FPBDCollisionConstaint& Constraint)
 		*/
 		template<typename TLambda>
 		ECollisionVisitorResult VisitCollisions(const TLambda& Visitor, const bool bOnlyActive = true);
 
 		/**
 		 * @brief Visit all the collisions
+		 * @param Visitor functor with signature ECollisionVisitorResult(const FPBDCollisionConstaint& Constraint)
 		*/
 		template<typename TLambda>
 		ECollisionVisitorResult VisitConstCollisions(const TLambda& Visitor, const bool bOnlyActive = true) const;
@@ -395,6 +397,7 @@ namespace Chaos
 		 * @brief Call a lambda on each active collision constraint
 		 * This includes sleeping constraints, but skips constraints that are were not used on the last awake tick
 		 * but are still kept around as an optimization.
+		 * @param Visitor functor with signature ECollisionVisitorResult(FPBDCollisionConstaint& Constraint)
 		*/
 		template<typename TLambda>
 		ECollisionVisitorResult VisitCollisions(const TLambda& Visitor, const bool bOnlyActive = true);
@@ -402,6 +405,7 @@ namespace Chaos
 
 		/**
 		 * @brief Call a lambda on each active collision constraint
+		 * @param Visitor functor with signature ECollisionVisitorResult(const FPBDCollisionConstaint& Constraint)
 		 */
 		template<typename TLambda>
 		ECollisionVisitorResult VisitConstCollisions(const TLambda& Visitor, const bool bOnlyActive = true) const;
diff --git a/Engine/Source/Runtime/Experimental/Chaos/Public/Chaos/Collision/SimSweep.h b/Engine/Source/Runtime/Experimental/Chaos/Public/Chaos/Collision/SimSweep.h
new file mode 100644
--- /dev/null
+++ b/Engine/Source/Runtime/Experimental/Chaos/Public/Chaos/Collision/SimSweep.h
@@ -0,0 +1,417 @@
+// Copyright Epic Games, Inc. All Rights Reserved.
+#pragma once
+
+#include "Chaos/Collision/CollisionFilter.h"
+#include "Chaos/GeometryQueries.h"
+#include "Chaos/ISpatialAcceleration.h"
+#include "Chaos/ParticleHandle.h"
+
+namespace Chaos
+{
+	namespace Private
+	{
+		/**
+		* Results from a SimSweepParticle.
+		* NOTE: Only HitParticle is guaranteed to be initialized. If HitParticle is null, all other properties are undefined.
+		*/
+		class FSimSweepParticleHit
+		{
+		public:
+			FSimSweepParticleHit()
+				: HitParticle(nullptr)
+			{
+			}
+
+			void Init()
+			{
+				HitParticle = nullptr;
+			}
+
+			bool IsHit() const
+			{
+				return HitParticle != nullptr;
+			}
+
+			// The particle that we hit, or null
+			const FGeometryParticleHandle* HitParticle;
+
+			// The shape that we hit (or uninitialized)
+			const FPerShapeData* HitShape;
+
+			// The geometry that we hit (or uninitialized)
+			const FImplicitObject* HitGeometry;
+
+			// The contact position (or uninitialized)
+			FVec3 HitPosition;
+
+			// The contact normal for HitDistance >= 0, or MTD direction for HitDistance < 0 (or uninitialized)
+			FVec3 HitNormal;
+
+			// The normal of the face we hit, if HitFaceIndex is set (or uninitialized)
+			FVec3 HitFaceNormal;
+
+			// The distance along the sweep at the hit, or the seperation if negative (i.e., -Penetration) (or uninitialized)
+			FReal HitDistance;
+
+			// Time of impact [0,1] (or uninitialized)
+			FReal HitTOI;
+
+			// The face index of the shape we hit, or INDEX_NONE (or uninitialized)
+			int32 HitFaceIndex;
+		};
+
+		/**
+		* A spatial acceleration visitor that forwards callbacks to a functor
+		*/
+		template<typename TVisitor>
+		class TSimSweepSQVisitor : public ISpatialVisitor<FAccelerationStructureHandle, FReal>
+		{
+		public:
+			using FVisitorData = TSpatialVisitorData<FAccelerationStructureHandle>;
+
+			TSimSweepSQVisitor(const TVisitor& InVisitor)
+				: Visitor(InVisitor)
+			{
+			}
+
+			virtual bool Overlap(const FVisitorData& Instance) override final
+			{
+				return CallVisitor(Instance);
+			}
+
+			virtual bool Sweep(const FVisitorData& Instance, FQueryFastData& CurData) override final
+			{
+				return CallVisitor(Instance);
+			}
+
+			virtual bool Raycast(const FVisitorData& Instance, FQueryFastData& CurData) override final
+			{
+				return CallVisitor(Instance);
+			}
+		private:
+			bool CallVisitor(const FVisitorData& Instance)
+			{
+				FGeometryParticleHandle* Particle = Instance.Payload.GetGeometryParticleHandle_PhysicsThread();
+				Visitor(Particle);
+				return true;
+			}
+
+			const TVisitor& Visitor;
+		};
+
+		/**
+		* Sweep SweptParticle against OtherParticle. The provided filter policy can be used to reject shape pairs.
+		* The collector policy chooses which sweep results(s) to keep.
+		*
+		* @tparam TShapeFilter A custom filter functor with the signature bool(const FPerShapeData* SweptShape, const FImplicitObject* SweptImplicit, const FPerShapeData* OtherShape, const FImplicitObject* OtherImplicit)
+		* used to accept/reject shapes pairs. It should return true if the pair should be considered for a sweep check.
+		* @see FSimSweepShapeFilterNarrowPhase which provides the standard narrowphase filter used in collision detection.
+		*
+		* @tparam THitCollector A hit collector functor with signature bool(const FVec3& Dir, const FReal Length, const FSimSweepParticleHit& InHit).
+		* It should return true to continue checking further objects, or false to stop any further sweeping.
+		* @see FSimSweepCollectorFirstHit which keeps the hit with the lowest time of impact and the most opposing normal if there is a tie.
+		*
+		* @param SweptParticle The particle to be swept
+		* @param OtherParticle The particle to be swept against
+		* @param StartPos The start position of the particle to be swept
+		* @param Rot The rotation to use for the sweep
+		* @param Dir The direction of the sweep (must be normalized)
+		* @param Length The sweep length
+		* @param ShapeFilter The filter to use to accept/reject shape pairs within the particle pairs
+		* @param HitCollector Hits are passed to the collector one at a time in a random order (i.e., they are not sorted by TOI)
+		* 
+		* @note This function does not reset the HitCollector.
+		*/
+		template<typename TShapeFilter, typename THitCollector>
+		void SimSweepParticlePair(
+			const FGeometryParticleHandle* SweptParticle,
+			const FGeometryParticleHandle* OtherParticle,
+			const FVec3& StartPos,
+			const FRotation3& Rot,
+			const FVec3& Dir,
+			const FReal Length,
+			TShapeFilter& ShapeFilter,
+			THitCollector& HitCollector)
+		{
+			// The particle transforms
+			const FRigidTransform3 SweptWorldTransform = FRigidTransform3(StartPos, Rot);
+			const FRigidTransform3 OtherWorldTransform = OtherParticle->GetTransformXR();
+
+			// Storage for the next hit result in the sweep (we only need one at a time)
+			FSimSweepParticleHit Hit;
+			Hit.HitParticle = OtherParticle;
+
+			// Visit all shape pairs and run a sweep (if they pass a filter check)
+			for (const TUniquePtr<FPerShapeData>& SweptShape : SweptParticle->ShapesArray())
+			{
+				const FImplicitObject* SweptImplicit = SweptShape->GetLeafGeometry();
+				const FRigidTransform3 SweptShapeWorldTransform = FRigidTransform3(SweptShape->GetLeafRelativeTransform()) * SweptWorldTransform;
+
+				for (const TUniquePtr<FPerShapeData>& OtherShape : OtherParticle->ShapesArray())
+				{
+					const FImplicitObject* OtherImplicit = OtherShape->GetLeafGeometry();
+
+					// Check the filter to see if we want to  process this shape pair
+					const bool bSweepShape = ShapeFilter(SweptShape.Get(), SweptImplicit, OtherShape.Get(), OtherImplicit);
+
+					if (bSweepShape)
+					{
+						const FRigidTransform3 OtherShapeWorldTransform = FRigidTransform3(OtherShape->GetLeafRelativeTransform()) * OtherWorldTransform;
+
+						const bool bSweepHit = SweepQuery(
+							*OtherImplicit,
+							OtherShapeWorldTransform,
+							*SweptImplicit,
+							SweptShapeWorldTransform,
+							Dir,
+							Length,
+							Hit.HitDistance,
+							Hit.HitPosition,
+							Hit.HitNormal,
+							Hit.HitFaceIndex,
+							Hit.HitFaceNormal,
+							0/*Thickness*/,
+							true/*bComputeMTD*/);
+
+						// If we have a hit, pass it to the collector
+						if (bSweepHit)
+						{
+							// Fill in the hit data that wasn't filled in by SweepQuery
+							Hit.HitTOI = (Hit.HitDistance > 0) ? (Hit.HitDistance / Length) : 0;
+
+							HitCollector(Dir, Length, Hit);
+						}
+					}
+				}
+			}
+		}
+
+		/**
+		* Sweep a particle against all other particles that it may hit. The particles are visited in a semi-random order, depending on how
+		* the particles are held in the spatial acceleration. The provided filter policies can be used to reject particle pairs or shape pairs.
+		* The collector policy chooses which sweep result(s) to keep.
+		* 
+		* @tparam TParticleFilter A custom filter functor with the signature bool(const FGeometryParticleHandle* SweptParticle, const FGeometryParticleHandle* OtherParticle)
+		* used to accept/reject particle pairs. It should return true if the particle pair should be considered for a sweep check.
+		* @see FSimSweepParticleFilterBroadPhase which provides the standard broadphase filter used in collision detection.
+		*
+		* @tparam TShapeFilter A custom filter functor with the signature bool(const FPerShapeData* SweptShape, const FImplicitObject* SweptImplicit, const FPerShapeData* OtherShape, const FImplicitObject* OtherImplicit)
+		* used to accept/reject shapes pairs. It should return true if the pair should be considered for a sweep check.
+		* @see FSimSweepShapeFilterNarrowPhase which provides the standard narrowphase filter used in collision detection.
+		*
+		* @tparam THitCollector A hit collector functor with signature bool(const FVec3& Dir, const FReal Length, const FSimSweepParticleHit& InHit). 
+		* It should return true to continue checking further objects, or false to stop any further sweeping. 
+		* @see FSimSweepCollectorFirstHit which keeps the hit with the lowest time of impact and the most opposing normal if there is a tie.
+		* 
+		* @param SweptParticle The particle to be swept
+		* @param StartPos The start position of the particle to be swept
+		* @param Rot The rotation to use for the sweep
+		* @param Dir The direction of the sweep (must be normalized)
+		* @param Length The sweep length
+		* @param ParticleFilter The filter to use to accept/reject particle pairs
+		* @param ShapeFilter The filter to use to accept/reject shape pairs within the particle pairs
+		* @param HitCollector Hits are passed to the collector one at a time in a semi-random order (i.e., they are not sorted by TOI)
+		*/
+		template<typename TParticleFilter, typename TShapeFilter, typename THitCollector>
+		void SimSweepParticle(
+			ISpatialAcceleration<FAccelerationStructureHandle, FReal, 3>* SpatialAcceleration,
+			const FGeometryParticleHandle* SweptParticle,
+			const FVec3& StartPos, 
+			const FRotation3& Rot,
+			const FVec3& Dir,
+			const FReal Length,
+			TParticleFilter& ParticleFilter, 
+			TShapeFilter& ShapeFilter,
+			THitCollector& HitCollector)
+		{
+			// Cannot sweep infinite particles (they are all static anyway)
+			if (!SweptParticle->HasBounds())
+			{
+				return;
+			}
+
+			// @todo(chaos): do we need to generate a hit for no sweep?
+			if (Length < UE_SMALL_NUMBER)
+			{
+				return;
+			}
+
+			// Functor: Process a particle encountered by the sweep through the spatial acceleration
+			// by sweeping all shapes of the swept particle against all shapes of the other particle.
+			const auto& ProcessParticlePair = 
+				[SweptParticle, &StartPos, &Dir, &Length, &Rot, &ParticleFilter, &ShapeFilter, &HitCollector]
+				(const FGeometryParticleHandle* OtherParticle)
+				-> void
+				{
+					// We cannot hit ourself
+					if (OtherParticle != SweptParticle)
+					{
+						// Check the particle filter to see if we should process this particle pair
+						if (ParticleFilter(SweptParticle, OtherParticle))
+						{
+							// Do the sweep
+							SimSweepParticlePair(SweptParticle, OtherParticle, StartPos, Rot, Dir, Length, ShapeFilter, HitCollector);
+						}
+					}
+				};
+
+			// Calculate the ray start/end and extents for the spatial acceleration sweep
+			// (the sweep through the spatial acceleration requires a ray centered on the bounds)
+			const FAABB3 TransformedParticleBounds = SweptParticle->Geometry()->CalculateTransformedBounds(FRigidTransform3(StartPos, Rot));
+			const FVec3 BoundsExtents = FReal(0.5) * TransformedParticleBounds.Extents();
+			const FVec3 BoundsStartPos = TransformedParticleBounds.Center();
+
+			// Sweep our bounds through the spatial acceleration structure and then sweep the shapes of our particle
+			// against those of any particles that our bounds sweep encounters.
+			TSimSweepSQVisitor SpatialAcclerationVisitor(ProcessParticlePair);
+			SpatialAcceleration->Sweep(BoundsStartPos, Dir, Length, BoundsExtents, SpatialAcclerationVisitor);
+		}
+
+
+		/**
+		* The default particle pair filter for sim sweeps that uses the same test as in the broadphase.
+		* I.e., sweep will be stopped by any object that the swept particle can collide with.
+		*/
+		class FSimSweepParticleFilterBroadPhase
+		{
+		public:
+			FSimSweepParticleFilterBroadPhase(FIgnoreCollisionManager* InIgnoreCollisionManager)
+				: IgnoreCollisionManager(InIgnoreCollisionManager)
+			{
+			}
+
+			bool operator()(const FGeometryParticleHandle* SweptParticle, const FGeometryParticleHandle* OtherParticle)
+			{
+				return ParticlePairBroadPhaseFilter(SweptParticle, OtherParticle, IgnoreCollisionManager);
+			}
+
+		private:
+			FIgnoreCollisionManager* IgnoreCollisionManager;
+		};
+
+		/**
+		* The default shape pair filter for sim sweeps that uses the same test as in the narrowphase.
+		* I.e., sweep will be stopped by any object that the swept particle can collide with.
+		*/
+		class FSimSweepShapeFilterNarrowPhase
+		{
+		public:
+			bool operator()(const FPerShapeData* SweptShape, const FImplicitObject* SweptImplicit, const FPerShapeData* OtherShape, const FImplicitObject* OtherImplicit)
+			{
+				const EImplicitObjectType SweptImplicitType = SweptImplicit ? GetInnerType(SweptImplicit->GetType()) : ImplicitObjectType::Unknown;
+				const EImplicitObjectType OtherImplicitType = OtherImplicit ? GetInnerType(OtherImplicit->GetType()) : ImplicitObjectType::Unknown;
+
+				return ShapePairNarrowPhaseFilter(SweptImplicitType, SweptShape, OtherImplicitType, OtherShape);
+			}
+		};
+
+		/**
+		* A hit collector for SimSweepParticle that just keeps the first hit, except for initial overlaps where the
+		* sweep is already moving us away from the contact. When there are multiple hits at the same TOI/initial overlap
+		* we take the hit with the most opposing normal.
+		*/
+		class FSimSweepCollectorFirstHit
+		{
+		public:
+			/**
+			* @param InHitDistanceEqualTolerance Sweep hits within this distance are assumed to be at equal distance, in which 
+			* case we take the hit with the most-opposing normal
+			* 
+			* @param OutHit a reference to a hit structure that will contain the results.
+			*/
+			FSimSweepCollectorFirstHit(const FReal InHitDistanceEqualTolerance, FSimSweepParticleHit& OutFirstHit)
+				: FirstHit(OutFirstHit)
+				, HitDistanceEqualTolerance(InHitDistanceEqualTolerance)
+			{
+				Init();
+			}
+
+			/**
+			* Keep the first hit and continue looking for more hits
+			*/
+			bool operator()(const FVec3& Dir, const FReal Length, const FSimSweepParticleHit& InHit)
+			{
+				ProcessHit(Dir, Length, InHit);
+				return true;
+			}
+
+			void Init()
+			{
+				FirstHit.Init();
+			}
+
+			bool IsHit() const
+			{
+				return FirstHit.IsHit();
+			}
+
+			const FSimSweepParticleHit& GetFirstHit() const
+			{
+				check(IsHit());
+				return FirstHit;
+			}
+
+		private:
+			void ProcessHit(const FVec3& Dir, const FReal Length, const FSimSweepParticleHit& InHit)
+			{
+				if (!FirstHit.IsHit() || (InHit.HitDistance < FirstHit.HitDistance + HitDistanceEqualTolerance))
+				{
+					// If the distance is very close to the existing hit, keep the one with the most opposing normal.
+					if (FirstHit.IsHit() && FMath::IsNearlyEqual(InHit.HitDistance, FirstHit.HitDistance))
+					{
+						const FReal ExistingDotNormal = FVec3::DotProduct(Dir, FirstHit.HitNormal);
+						const FReal NewDotNormal = FVec3::DotProduct(Dir, InHit.HitNormal);
+						if (NewDotNormal < ExistingDotNormal)
+						{
+							FirstHit = InHit;
+						}
+						return;
+					}
+
+					// If this is a closer hit we will take it unless we are already moving out of contact.
+					// NOTE: we will only ever get an outward pointing normal for initial overlaps (HitDistance <= 0).
+					if ((InHit.HitDistance > 0) || (FVec3::DotProduct(InHit.HitNormal, Dir) < 0))
+					{
+						FirstHit = InHit;
+						return;
+					}
+				}
+			}
+
+			FSimSweepParticleHit& FirstHit;
+			FReal HitDistanceEqualTolerance;
+		};
+
+		/**
+		* Sweep a particle to find the first impact with another particle that the particle would normally collide with 
+		* according to the sim filter data on the particles.
+		* 
+		* @param SpatialAcceleration The spatial acceleration structure used in the broadphase
+		* @param IgnoreCollisionManager The ignore collision manager used in the main broadphase (or null)
+		* @param SweptParticle The particle to be swept
+		* @param StartPos The start position of the particle to be swept
+		* @param Rot The rotation to use for the sweep
+		* @param Dir The direction of the sweep (must be normalized)
+		* @param Length The sweep length
+		* @param OutHit Output set to the first hit if there was one, otherwise invalidated.
+		* 
+		* @return true if there was a hit, otherwise false
+		* 
+		* @note This call SimSweepParticle with the FSimSweepParticleFilterCollidable and FSimSweepShapeFilterCollidable filter policies,
+		* and the FSimSweepCollectorFirstHit hit collector. If these policies are not what you want, you can create your own policies 
+		* and/or collector (e.g., to collect all hits) and use SimSweepParticle directly.
+		* 
+		*/
+		extern bool SimSweepParticleFirstHit(
+			ISpatialAcceleration<FAccelerationStructureHandle, FReal, 3>* SpatialAcceleration,
+			FIgnoreCollisionManager* IgnoreCollisionManager,
+			const FGeometryParticleHandle* SweptParticle,
+			const FVec3& StartPos, 
+			const FRotation3& Rot,
+			const FVec3& Dir,
+			const FReal Length,
+			FSimSweepParticleHit& OutHit,
+			const FReal InHitDistanceEqualTolerance = UE_KINDA_SMALL_NUMBER);
+	}
+
+}
\ No newline at end of file
diff --git a/Engine/Source/Runtime/Experimental/Chaos/Public/Chaos/Collision/SpatialAccelerationBroadPhase.h b/Engine/Source/Runtime/Experimental/Chaos/Public/Chaos/Collision/SpatialAccelerationBroadPhase.h
--- a/Engine/Source/Runtime/Experimental/Chaos/Public/Chaos/Collision/SpatialAccelerationBroadPhase.h
+++ b/Engine/Source/Runtime/Experimental/Chaos/Public/Chaos/Collision/SpatialAccelerationBroadPhase.h
@@ -2,8 +2,8 @@
 #pragma once
 
 #include "Chaos/Collision/CollisionConstraintAllocator.h"
-#include "Chaos/Collision/CollisionConstraintFlags.h"
 #include "Chaos/Collision/CollisionContext.h"
+#include "Chaos/Collision/CollisionFilter.h"
 #include "Chaos/Collision/StatsData.h"
 #include "Chaos/ISpatialAccelerationCollection.h"
 #include "Chaos/ParticleHandleFwd.h"
@@ -23,14 +23,20 @@ namespace Chaos
 
 	namespace Private
 	{
+
 		/**
 		 * Check whether the two particles are allowed to collide, and also whether we should flip their order.
 		 */
-		inline bool ParticlePairCollisionAllowed(const FGeometryParticleHandle* Particle1, const FGeometryParticleHandle* Particle2, const FIgnoreCollisionManager& IgnoreCollisionManager, const bool bIsResimming,  bool& bOutSwapOrder)
+		inline bool ParticlePairCollisionAllowed(
+			const FGeometryParticleHandle* Particle1, 
+			const FGeometryParticleHandle* Particle2, 
+			const FIgnoreCollisionManager& IgnoreCollisionManager, 
+			const bool bIsResimming,  
+			bool& bOutSwapOrder)
 		{
 			bOutSwapOrder = false;
 
-			if (Particle1 == Particle2)
+			if (!ParticlePairBroadPhaseFilter(Particle1, Particle2, &IgnoreCollisionManager))
 			{
 				return false;
 			}
@@ -38,71 +44,23 @@ namespace Chaos
 			bool bIsKinematic1 = true;
 			bool bIsDynamicAwake1 = false;
 			bool bIsDynamicAsleep1 = false;
-			bool bUseIgnoreCollisionManager1 = false;
-			bool bDisabled1 = false;
-			int32 CollisionGroup1 = 0;
 			const FPBDRigidParticleHandle* Rigid1 = Particle1->CastToRigidParticle();
 			if (Rigid1 != nullptr)
 			{
 				bIsKinematic1 = Rigid1->IsKinematic();
 				bIsDynamicAsleep1 = !bIsKinematic1 && Rigid1->IsSleeping();
 				bIsDynamicAwake1 = !bIsKinematic1 && !bIsDynamicAsleep1;
-				bUseIgnoreCollisionManager1 = Rigid1->UseIgnoreCollisionManager();
-				bDisabled1 = Rigid1->Disabled();
-				CollisionGroup1 = Rigid1->CollisionGroup();
 			}
 
 			bool bIsKinematic2 = true;
 			bool bIsDynamicAwake2 = false;
 			bool bIsDynamicAsleep2 = false;
-			bool bUseIgnoreCollisionManager2 = false;
-			bool bDisabled2 = false;
-			int32 CollisionGroup2 = 0;
 			const FPBDRigidParticleHandle* Rigid2 = Particle2->CastToRigidParticle();
 			if (Rigid2 != nullptr)
 			{
 				bIsKinematic2 = Rigid2->IsKinematic();
 				bIsDynamicAsleep2 = !bIsKinematic2 && Rigid2->IsSleeping();
 				bIsDynamicAwake2 = !bIsKinematic2 && !bIsDynamicAsleep2;
-				bUseIgnoreCollisionManager2 = Rigid2->UseIgnoreCollisionManager();
-				bDisabled2 = Rigid2->Disabled();
-				CollisionGroup2 = Rigid2->CollisionGroup();
-			}
-
-			// @todo(chaos): This should not be happening if the disabled particles are removed from the active particles list, but GeometryCollection may leave them there
-			//check(!bDisabled2);
-			if (bDisabled1 || bDisabled2)
-			{
-				return false;
-			}
-
-			// At least one particle needs to be dynamic to generate a collision response
-			if (bIsKinematic1 && bIsKinematic2)
-			{
-				return false;
-			}
-			check((Rigid1 != nullptr) || (Rigid2 != nullptr));
-
-			// Is this particle interaction governed by the IgnoreCollisionManager? If so, check to see if interaction is allowed
-			if (bUseIgnoreCollisionManager1 || bUseIgnoreCollisionManager2)
-			{
-				if (IgnoreCollisionManager.IgnoresCollision(Particle1, Particle2))
-				{
-					return false;
-				}
-			}
-
-			// CollisionGroups are used by geometry collections for high-level collision filtering
-			// CollisionGroup == 0 : Collide_With_Everything
-			// CollisionGroup == INDEX_NONE : Disabled collisions
-			// CollisionGroup1 != CollisionGroup2 : Disabled collisions (if other conditions not met)
-			if (CollisionGroup1 == INDEX_NONE || CollisionGroup2 == INDEX_NONE)
-			{
-				return false;
-			}
-			if ((CollisionGroup1 != 0) && (CollisionGroup2 != 0) && (CollisionGroup1 != CollisionGroup2))
-			{
-				return false;
 			}
 
 			// In both cases (resim or not) we will generate (1) dynamic-(sleeping,kinematic(steady+moving),static) pairs + (2) sleeping-moving kinematic ones
diff --git a/Engine/Source/Runtime/Experimental/Chaos/Public/Chaos/GeometryQueries.h b/Engine/Source/Runtime/Experimental/Chaos/Public/Chaos/GeometryQueries.h
--- a/Engine/Source/Runtime/Experimental/Chaos/Public/Chaos/GeometryQueries.h
+++ b/Engine/Source/Runtime/Experimental/Chaos/Public/Chaos/GeometryQueries.h
@@ -119,6 +119,7 @@ namespace Chaos
 		return false;
 	}
 
+	// @todo(chaos): This does not handle Unions
 	template <typename SweptGeometry>
 	bool SweepQuery(const FImplicitObject& A, const FRigidTransform3& ATM, const SweptGeometry& B, const FRigidTransform3& BTM, const FVec3& Dir, const FReal Length, FReal& OutTime, FVec3& OutPosition, FVec3& OutNormal, int32& OutFaceIndex, FVec3& OutFaceNormal, const FReal Thickness, const bool bComputeMTD)
 	{
@@ -248,4 +249,17 @@ namespace Chaos
 
 		return bResult;
 	}
+
+
+	// @todo(chaos): This does not handle Unions
+	inline bool SweepQuery(const FImplicitObject& A, const FRigidTransform3& ATM, const FImplicitObject& B, const FRigidTransform3& BTM, const FVec3& Dir, const FReal Length, FReal& OutTime, FVec3& OutPosition, FVec3& OutNormal, int32& OutFaceIndex, FVec3& OutFaceNormal, const FReal Thickness, const bool bComputeMTD)
+	{
+		return Chaos::Utilities::CastHelper(B, BTM,
+			[&A, &ATM, &Dir, &Length, &OutTime, &OutPosition, &OutNormal, &OutFaceIndex, &OutFaceNormal, &Thickness, &bComputeMTD]
+			(const auto& BInner, const FTransform& BInnerTM) -> bool
+			{
+				return SweepQuery(A, ATM, BInner, BInnerTM, Dir, Length, OutTime, OutPosition, OutNormal, OutFaceIndex, OutFaceNormal, Thickness, bComputeMTD);
+			});
+	}
+
 }
diff --git a/Engine/Source/Runtime/Experimental/Chaos/Public/Chaos/ImplicitObject.h b/Engine/Source/Runtime/Experimental/Chaos/Public/Chaos/ImplicitObject.h
--- a/Engine/Source/Runtime/Experimental/Chaos/Public/Chaos/ImplicitObject.h
+++ b/Engine/Source/Runtime/Experimental/Chaos/Public/Chaos/ImplicitObject.h
@@ -5,6 +5,7 @@
 #include "Chaos/Serializable.h"
 #include "Chaos/Core.h"
 #include "Chaos/ImplicitFwd.h"
+#include "Chaos/ImplicitObjectType.h"
 #include "Chaos/AABB.h"
 
 #ifndef TRACK_CHAOS_GEOMETRY
@@ -30,66 +31,6 @@ using FAABB3 = TAABB<FReal, 3>;
 using FParticles = TParticles<FReal, 3>;
 using FSphere = TSphere<FReal, 3>; // warning: code assumes that FImplicitObjects with type ImplicitObjectType::Sphere are FSpheres, but all TSpheres will think they have ImplicitObjectType::Sphere.
 
-namespace ImplicitObjectType
-{
-	enum
-	{
-		//Note: add entries in order to avoid serialization issues (but before IsInstanced)
-		Sphere = 0, // warning: code assumes that this is an FSphere, but all TSpheres will think this is their type.
-		Box,
-		Plane,
-		Capsule,
-		Transformed,
-		Union,
-		LevelSet,
-		Unknown,
-		Convex,
-		TaperedCylinder,
-		Cylinder,
-		TriangleMesh,
-		HeightField,
-		DEPRECATED_Scaled,	//needed for serialization of existing data
-		Triangle,
-		UnionClustered,
-		TaperedCapsule,
-
-		//Add entries above this line for serialization
-		IsInstanced = 1 << 6,
-		IsScaled = 1 << 7
-	};
-}
-
-using EImplicitObjectType = uint8;	//see ImplicitObjectType
-
-FORCEINLINE bool IsInstanced(EImplicitObjectType Type)
-{
-	return (Type & ImplicitObjectType::IsInstanced) != 0;
-}
-
-FORCEINLINE bool IsScaled(EImplicitObjectType Type)
-{
-	return (Type & ImplicitObjectType::IsScaled) != 0;
-}
-
-FORCEINLINE EImplicitObjectType GetInnerType(EImplicitObjectType Type)
-{
-	return Type & (~(ImplicitObjectType::IsScaled | ImplicitObjectType::IsInstanced));
-}
-
-
-namespace EImplicitObject
-{
-	enum Flags
-	{
-		IsConvex = 1,
-		HasBoundingBox = 1 << 1,
-		DisableCollisions = 1 << 2
-	};
-
-	inline const int32 FiniteConvex = IsConvex | HasBoundingBox;
-}
-
-
 
 template<class T, int d, bool bSerializable>
 struct TImplicitObjectPtrStorage
diff --git a/Engine/Source/Runtime/Experimental/Chaos/Public/Chaos/ImplicitObjectType.h b/Engine/Source/Runtime/Experimental/Chaos/Public/Chaos/ImplicitObjectType.h
new file mode 100644
--- /dev/null
+++ b/Engine/Source/Runtime/Experimental/Chaos/Public/Chaos/ImplicitObjectType.h
@@ -0,0 +1,65 @@
+// Copyright Epic Games, Inc. All Rights Reserved.
+#pragma once
+
+#include "Chaos/Core.h"
+
+namespace Chaos
+{
+	namespace ImplicitObjectType
+	{
+		enum
+		{
+			//Note: add entries in order to avoid serialization issues (but before IsInstanced)
+			Sphere = 0, // warning: code assumes that this is an FSphere, but all TSpheres will think this is their type.
+			Box,
+			Plane,
+			Capsule,
+			Transformed,
+			Union,
+			LevelSet,
+			Unknown,
+			Convex,
+			TaperedCylinder,
+			Cylinder,
+			TriangleMesh,
+			HeightField,
+			DEPRECATED_Scaled,	//needed for serialization of existing data
+			Triangle,
+			UnionClustered,
+			TaperedCapsule,
+
+			//Add entries above this line for serialization
+			IsInstanced = 1 << 6,
+			IsScaled = 1 << 7
+		};
+	}
+
+	using EImplicitObjectType = uint8;	//see ImplicitObjectType
+
+	FORCEINLINE bool IsInstanced(EImplicitObjectType Type)
+	{
+		return (Type & ImplicitObjectType::IsInstanced) != 0;
+	}
+
+	FORCEINLINE bool IsScaled(EImplicitObjectType Type)
+	{
+		return (Type & ImplicitObjectType::IsScaled) != 0;
+	}
+
+	FORCEINLINE EImplicitObjectType GetInnerType(EImplicitObjectType Type)
+	{
+		return Type & (~(ImplicitObjectType::IsScaled | ImplicitObjectType::IsInstanced));
+	}
+
+	namespace EImplicitObject
+	{
+		enum Flags
+		{
+			IsConvex = 1,
+			HasBoundingBox = 1 << 1,
+			DisableCollisions = 1 << 2
+		};
+
+		inline const int32 FiniteConvex = IsConvex | HasBoundingBox;
+	}
+}
\ No newline at end of file
diff --git a/Engine/Source/Runtime/Experimental/Chaos/Public/Chaos/PBDRigidsEvolutionGBF.h b/Engine/Source/Runtime/Experimental/Chaos/Public/Chaos/PBDRigidsEvolutionGBF.h
--- a/Engine/Source/Runtime/Experimental/Chaos/Public/Chaos/PBDRigidsEvolutionGBF.h
+++ b/Engine/Source/Runtime/Experimental/Chaos/Public/Chaos/PBDRigidsEvolutionGBF.h
@@ -130,17 +130,21 @@ namespace Chaos
 
 		/**
 		 * User has moved a particle
+		 * Does not change velocity.
+		 * Will wake the particle if this is a move (i.e., bIsTeleport is false)
 		 */
-		CHAOS_API void SetParticleTransform(FGeometryParticleHandle* InParticle, const FVec3& InPos, const FRotation3& InRot, const bool bIsTeleport)
-		{
-			const FVec3 PrevX = InParticle->X();
-			const FRotation3 PrevR = InParticle->R();
-
-			FGenericParticleHandle(InParticle)->SetTransform(InPos, InRot);
+		CHAOS_API void SetParticleTransform(FGeometryParticleHandle* InParticle, const FVec3& InPos, const FRotation3& InRot, const bool bIsTeleport);
 
-			OnParticleMoved(InParticle, PrevX, PrevR, bIsTeleport);
-		}
+		/**
+		 * Move a particle to a new location with a sweep and stop and the first opposing contact.
+		 * Does not change velocity.
+		 * Will wake the particle if this is a move (i.e., bIsTeleport is false)
+		 */
+		CHAOS_API virtual void SetParticleTransformSwept(FGeometryParticleHandle* InParticle, const FVec3& InPos, const FRotation3& InRot, const bool bIsTeleport);
 
+		/**
+		 * Called when a particle is moved. We need to reset some friction properties, sleeping properties, etc
+		 */
 		CHAOS_API void OnParticleMoved(FGeometryParticleHandle* InParticle, const FVec3& PrevX, const FRotation3& PrevR, const bool bIsTeleport);
 
 
