Fix SAT with AABB Triangle
#prefligh 63d802ab31334253e50a25e5

diff --git a/Engine/Source/Runtime/Experimental/Chaos/Private/Chaos/Collision/TriangleOverlap.cpp b/Engine/Source/Runtime/Experimental/Chaos/Private/Chaos/Collision/TriangleOverlap.cpp
--- a/Engine/Source/Runtime/Experimental/Chaos/Private/Chaos/Collision/TriangleOverlap.cpp
+++ b/Engine/Source/Runtime/Experimental/Chaos/Private/Chaos/Collision/TriangleOverlap.cpp
@@ -378,11 +378,8 @@ namespace Chaos
 		for (int32 i = 0; i < EdgeNum; i++)
 		{
 			VectorRegister4Float Axis = VectorCross(Edges[i], Edge);
-			const FRealSingle Sign = VectorDot3Scalar(VectorSubtract(TriangleVertex, Centroid), Axis);
-			if (Sign < 0.0f)
-			{
-				Axis = VectorNegate(Axis);
-			}
+			const VectorRegister4Float Sign = VectorDot3(VectorSubtract(TriangleVertex, Centroid), Axis);
+			Axis = VectorSelect(VectorCompareLT(Sign, VectorZeroFloat()), VectorNegate(Axis), Axis);
 
 			const VectorRegister4Float XCompCurrent = VectorSelect(VectorCompareGT(VectorDot3(Axis, XAxisHalfExtent), VectorZeroFloat()), VectorNegate(XAxisHalfExtent), XAxisHalfExtent);
 			const VectorRegister4Float YCompCurrent = VectorSelect(VectorCompareGT(VectorDot3(Axis, YAxisHalfExtent), VectorZeroFloat()), VectorNegate(YAxisHalfExtent), YAxisHalfExtent);
@@ -560,51 +557,23 @@ namespace Chaos
 		HalfExtents = VectorLoadFloat3(&HalfExtentsf.X);
 	}
 
-	namespace
+	FORCEINLINE_DEBUGGABLE bool FAABBSimd::ComputeEdgeOverlap(const VectorRegister4Float& TriangleEdge, const VectorRegister4Float& TriangleVertex, const VectorRegister4Float& Centroid) const
 	{
-		FORCEINLINE_DEBUGGABLE bool HasToComputeEdge(const VectorRegister4Float& BoxNormal1, const VectorRegister4Float& BoxNormal2, const VectorRegister4Float& BoxEdge, const VectorRegister4Float& TriNormal, const VectorRegister4Float& TriEdge)
-		{
-			const FRealSingle TriNormBoxEdge = VectorDot3Scalar(TriNormal, BoxEdge);		// TriNormal | BoxEdge
-			const FRealSingle BoxNorm1TriEdge = VectorDot3Scalar(BoxNormal1, TriEdge);		// BoxNormalA | TriEdge
-			const FRealSingle BoxNorm2TriEdge = VectorDot3Scalar(BoxNormal2, TriEdge);		// BoxNormalB | TriEdge
-
-			return ((BoxNorm1TriEdge * BoxNorm2TriEdge) < 0.0f) && (((TriNormBoxEdge * BoxNorm2TriEdge) > 0.0f || (TriNormBoxEdge * BoxNorm1TriEdge) > 0.0f));
-		}
-	}
-
-	FORCEINLINE_DEBUGGABLE bool FAABBSimd::ComputeEdgeOverlap(const VectorRegister4Float& TriangleEdge, const VectorRegister4Float& TriangleVertex, const VectorRegister4Float& Centroid, const VectorRegister4Float& Normal, const VectorRegister4Float& LocalClosest) const
-	{
-		// Triangle edge vs box edges 
-		const VectorRegister4Float ClosestPoint = VectorAdd(LocalClosest, Position);
-
-		VectorRegister4Float OtherClosests[3];
-		OtherClosests[0] = VectorBitwiseXor(LocalClosest, SignX);
-		OtherClosests[1] = VectorBitwiseXor(LocalClosest, SignY);
-		OtherClosests[2] = VectorBitwiseXor(LocalClosest, SignZ);
-		VectorRegister4Float BoxEdges[3];
-		VectorRegister4Float BoxEdgeNormals[3];
 		for (int32 i = 0; i < 3; i++)
 		{
-			const VectorRegister4Float OtherClosest = VectorAdd(OtherClosests[i], Position);
-			BoxEdges[i] = VectorSubtract(OtherClosest, ClosestPoint);
-			BoxEdgeNormals[i] = VectorBitwiseXor(SignBit, BoxEdges[i]);
-		}
+			VectorRegister4Float Axis = VectorCross(Edges[i], TriangleEdge);
+			const VectorRegister4Float Sign = VectorDot3(VectorSubtract(TriangleVertex, Centroid), Axis);
+			Axis = VectorSelect(VectorCompareLT(Sign, VectorZeroFloat()), VectorBitwiseXor(SignBit, Axis), Axis);
 
-		for (int32 i = 0; i < 3; i++)
-		{
-			if (HasToComputeEdge(BoxEdgeNormals[(i + 1) % 3], BoxEdgeNormals[(i + 2) % 3], BoxEdges[i], Normal, TriangleEdge))
+			// Take the opposite (xor) vector and extract sign (and), result in a (not and)
+			const VectorRegister4Float AxisSigns = VectorBitwiseNotAnd(Axis, SignBit);
+			const VectorRegister4Float LocalClosest = VectorBitwiseOr(AxisSigns, HalfExtents);
+			const VectorRegister4Float ClosestPoint = VectorAdd(LocalClosest, Position);
+
+			const FRealSingle ScaledSeparation = VectorDot3Scalar(VectorSubtract(ClosestPoint, TriangleVertex), Axis);
+			if (ScaledSeparation > UE_KINDA_SMALL_NUMBER)
 			{
-				VectorRegister4Float Axis = VectorCross(BoxEdges[i], TriangleEdge);
-				const FRealSingle Sign = VectorDot3Scalar(VectorSubtract(TriangleVertex, Centroid), Axis);
-				if (Sign < 0.0f)
-				{
-					Axis = VectorBitwiseXor(SignBit, Axis);
-				}
-				const FRealSingle ScaledSeparation = VectorDot3Scalar(VectorSubtract(ClosestPoint, TriangleVertex), Axis);
-				if (ScaledSeparation > UE_KINDA_SMALL_NUMBER)
-				{
-					return false;
-				}
+				return false;
 			}
 		}
 		return true;
@@ -689,14 +658,14 @@ namespace Chaos
 			return false;
 		}
 
-		if (!ComputeEdgeOverlap(AB, A, Centroid, Normal, ABLocalClosest))
+		if (!ComputeEdgeOverlap(AB, A, Centroid))
 		{
 			return false;
 		}
-		if (!ComputeEdgeOverlap(BC, B, Centroid, Normal, BCLocalClosest))
+		if (!ComputeEdgeOverlap(BC, B, Centroid))
 		{
 			return false;
 		}
-		return ComputeEdgeOverlap(CA, C, Centroid, Normal, CALocalClosest);
+		return ComputeEdgeOverlap(CA, C, Centroid);
 	}
 }
diff --git a/Engine/Source/Runtime/Experimental/Chaos/Private/Chaos/Collision/TriangleOverlap.h b/Engine/Source/Runtime/Experimental/Chaos/Private/Chaos/Collision/TriangleOverlap.h
--- a/Engine/Source/Runtime/Experimental/Chaos/Private/Chaos/Collision/TriangleOverlap.h
+++ b/Engine/Source/Runtime/Experimental/Chaos/Private/Chaos/Collision/TriangleOverlap.h
@@ -39,7 +39,7 @@ namespace Chaos
 		VectorRegister4Float XAxisHalfExtent;
 		VectorRegister4Float YAxisHalfExtent;
 		VectorRegister4Float ZAxisHalfExtent;
-		static constexpr int EdgeNum = 3;
+		static constexpr int32 EdgeNum = 3;
 		VectorRegister4Float Edges[EdgeNum];
 	};
 
@@ -55,7 +55,7 @@ namespace Chaos
 
 	private:
 		void Initialize(const FRigidTransform3& Transform, const FVec3f& HalfExtentsf);
-		bool ComputeEdgeOverlap(const VectorRegister4Float& TriangleEdge, const VectorRegister4Float& TriangleVertex, const VectorRegister4Float& Centroid, const VectorRegister4Float& Normal, const VectorRegister4Float& LocalClosest) const;
+		bool ComputeEdgeOverlap(const VectorRegister4Float& TriangleEdge, const VectorRegister4Float& TriangleVertex, const VectorRegister4Float& Centroid) const;
 
 		VectorRegister4Float Position;
 		VectorRegister4Float HalfExtents;
@@ -63,6 +63,12 @@ namespace Chaos
 		const static VectorRegister4Float SignX;
 		const static VectorRegister4Float SignY;
 		const static VectorRegister4Float SignZ;
+		static constexpr int32 EdgeNum = 3;
+		constexpr static VectorRegister4Float Edges[EdgeNum] = {
+			MakeVectorRegisterFloatConstant(1.f, 0.f, 0.f, 0.f),
+			MakeVectorRegisterFloatConstant(0.f, 1.f, 0.f, 0.f),
+			MakeVectorRegisterFloatConstant(0.f, 0.f, 1.f, 0.f)
+		};
 	};
 }
 
