Iris - Allow all members of LifetimeConditionals states to be disabled.

diff --git a/Engine/Source/Runtime/Experimental/Iris/Core/Private/Iris/ReplicationSystem/Conditionals/ReplicationConditionals.cpp b/Engine/Source/Runtime/Experimental/Iris/Core/Private/Iris/ReplicationSystem/Conditionals/ReplicationConditionals.cpp
--- a/Engine/Source/Runtime/Experimental/Iris/Core/Private/Iris/ReplicationSystem/Conditionals/ReplicationConditionals.cpp
+++ b/Engine/Source/Runtime/Experimental/Iris/Core/Private/Iris/ReplicationSystem/Conditionals/ReplicationConditionals.cpp
@@ -164,24 +164,12 @@ void FReplicationConditionals::InitPropertyCustomConditions(FInternalNetRefIndex
 			const FRepChangedPropertyTracker* Tracker = ChangedPropertyTracker.Get();
 			for (const FReplicationStateMemberLifetimeConditionDescriptor& MemberLifeTimeConditionDescriptor : MakeArrayView(StateDescriptor->MemberLifetimeConditionDescriptors, StateDescriptor->MemberCount))
 			{
-				if (MemberLifeTimeConditionDescriptor.Condition == ELifetimeCondition::COND_Custom)
+				const SIZE_T MemberIndex = &MemberLifeTimeConditionDescriptor - StateDescriptor->MemberLifetimeConditionDescriptors;
+				const uint16 RepIndex = StateDescriptor->MemberProperties[MemberIndex]->RepIndex;
+				if (!Tracker->IsParentActive(RepIndex))
 				{
-					const SIZE_T MemberIndex = &MemberLifeTimeConditionDescriptor - StateDescriptor->MemberLifetimeConditionDescriptors;
-					const uint16 RepIndex = StateDescriptor->MemberProperties[MemberIndex]->RepIndex;
-					if (!Tracker->IsParentActive(RepIndex))
-					{
-						const FReplicationStateMemberChangeMaskDescriptor& MemberChangeMaskDescriptor = StateDescriptor->MemberChangeMaskDescriptors[MemberIndex];
-						ConditionalChangeMask.ClearBits(MemberChangeMaskDescriptor.BitOffset, MemberChangeMaskDescriptor.BitCount);
-					}
-				}
-				else
-				{
-					const SIZE_T MemberIndex = &MemberLifeTimeConditionDescriptor - StateDescriptor->MemberLifetimeConditionDescriptors;
-					const uint16 RepIndex = StateDescriptor->MemberProperties[MemberIndex]->RepIndex;
-					if (!Tracker->IsParentActive(RepIndex))
-					{
-						UE_LOG(LogIris, Warning, TEXT("Trying to change non-existing custom conditional for RepIndex %u in protocol %s"), RepIndex, ToCStr(Protocol->DebugName));
-					}
+					const FReplicationStateMemberChangeMaskDescriptor& MemberChangeMaskDescriptor = StateDescriptor->MemberChangeMaskDescriptors[MemberIndex];
+					ConditionalChangeMask.ClearBits(MemberChangeMaskDescriptor.BitOffset, MemberChangeMaskDescriptor.BitCount);
 				}
 			}
 		}
