Iris - Detect if object becomes dirty during PreReplication

diff --git a/Engine/Source/Runtime/Experimental/Iris/Core/Private/Iris/ReplicationSystem/ObjectReplicationBridge.cpp b/Engine/Source/Runtime/Experimental/Iris/Core/Private/Iris/ReplicationSystem/ObjectReplicationBridge.cpp
--- a/Engine/Source/Runtime/Experimental/Iris/Core/Private/Iris/ReplicationSystem/ObjectReplicationBridge.cpp
+++ b/Engine/Source/Runtime/Experimental/Iris/Core/Private/Iris/ReplicationSystem/ObjectReplicationBridge.cpp
@@ -752,6 +752,9 @@ void UObjectReplicationBridge::PreUpdateAndPollImpl(FNetRefHandle Handle)
 		if (const FReplicationInstanceProtocol* InstanceProtocol = ObjectData.InstanceProtocol)
 		{
 			const EReplicationInstanceProtocolTraits InstanceTraits = InstanceProtocol->InstanceTraits;
+			const bool bNeedsPoll = EnumHasAnyFlags(InstanceTraits, EReplicationInstanceProtocolTraits::NeedsPoll);
+			bool bIsDirtyObject = DirtyObjects.GetBit(InternalObjectIndex);
+
 			// Call per-instance PreUpdate function
 			if (PreUpdateInstanceFunction && EnumHasAnyFlags(InstanceTraits, EReplicationInstanceProtocolTraits::NeedsPreSendUpdate))
 			{
@@ -761,15 +764,24 @@ void UObjectReplicationBridge::PreUpdateAndPollImpl(FNetRefHandle Handle)
 #endif
 				(*PreUpdateInstanceFunction)(ObjectData.RefHandle, ReplicatedInstances[InternalObjectIndex], this);
 				++Stats.PreUpdatedObjectCount;
+
+				// Pre update may dirty push based properties. Detect it.
+				if ((bNeedsPoll & !bIsDirtyObject))
+				{
+					bIsDirtyObject = FGlobalDirtyNetObjectTracker::IsNetObjectStateDirty(ObjectData.NetHandle);
+					if (bIsDirtyObject)
+					{
+						DirtyObjects.SetBit(bIsDirtyObject);
+					}
+				}
 			}
 
-			const bool bIsDirtyObject = DirtyObjects.GetBit(InternalObjectIndex);
 			const bool bIsNewInScope = !PrevScopableObjects.GetBit(InternalObjectIndex);
 			const bool bIsGCAffectedObject = GarbageCollectionAffectedObjects.GetBit(InternalObjectIndex);			
 			GarbageCollectionAffectedObjects.ClearBit(InternalObjectIndex);
 
 			// Early out if the instance does not require polling
-			if (!EnumHasAnyFlags(InstanceTraits, EReplicationInstanceProtocolTraits::NeedsPoll))
+			if (!bNeedsPoll)
 			{
 				return;
 			}
diff --git a/Engine/Source/Runtime/Net/Core/Private/Net/Core/DirtyNetObjectTracker/GlobalDirtyNetObjectTracker.cpp b/Engine/Source/Runtime/Net/Core/Private/Net/Core/DirtyNetObjectTracker/GlobalDirtyNetObjectTracker.cpp
--- a/Engine/Source/Runtime/Net/Core/Private/Net/Core/DirtyNetObjectTracker/GlobalDirtyNetObjectTracker.cpp
+++ b/Engine/Source/Runtime/Net/Core/Private/Net/Core/DirtyNetObjectTracker/GlobalDirtyNetObjectTracker.cpp
@@ -38,6 +38,16 @@ void FGlobalDirtyNetObjectTracker::MarkNetObjectStateDirty(FNetHandle NetHandle)
 	}
 }
 
+bool FGlobalDirtyNetObjectTracker::IsNetObjectStateDirty(FNetHandle NetHandle)
+{
+	if (Instance)
+	{
+		return Instance->DirtyObjects.Find(NetHandle) != nullptr;
+	}
+
+	return false;
+}
+
 FGlobalDirtyNetObjectTracker::FPollHandle FGlobalDirtyNetObjectTracker::CreatePoller()
 {
 	if (Instance)
diff --git a/Engine/Source/Runtime/Net/Core/Public/Net/Core/DirtyNetObjectTracker/GlobalDirtyNetObjectTracker.h b/Engine/Source/Runtime/Net/Core/Public/Net/Core/DirtyNetObjectTracker/GlobalDirtyNetObjectTracker.h
--- a/Engine/Source/Runtime/Net/Core/Public/Net/Core/DirtyNetObjectTracker/GlobalDirtyNetObjectTracker.h
+++ b/Engine/Source/Runtime/Net/Core/Public/Net/Core/DirtyNetObjectTracker/GlobalDirtyNetObjectTracker.h
@@ -50,6 +50,9 @@ public:
 	/** Marks an object as dirty. Assumes Init() has been called. */
 	NETCORE_API static void MarkNetObjectStateDirty(FNetHandle);
 
+	/** Returns true if an object has been marked as dirty since the last reset. Assumes Init() has been called. */
+	NETCORE_API static bool IsNetObjectStateDirty(FNetHandle);
+
 	/** Create a poller which is assumed to call PollDirtyNetObjects(). Assumes Init() has been called. */
 	NETCORE_API static FPollHandle CreatePoller();
 
