Geometry collection : Add ISM pool subsystem to allow auto assignment of ISMPool
- ISMPool subsystem creates the ISMPool actor and return it when asked for
- Fix root proxy mesh not being properly loaded before sent to the ISM Pool
- Fix visiblity of the geometry component being overriden if ISMPool is not being used

diff --git a/Engine/Source/Runtime/Experimental/GeometryCollectionEngine/Private/GeometryCollection/GeometryCollectionComponent.cpp b/Engine/Source/Runtime/Experimental/GeometryCollectionEngine/Private/GeometryCollection/GeometryCollectionComponent.cpp
--- a/Engine/Source/Runtime/Experimental/GeometryCollectionEngine/Private/GeometryCollection/GeometryCollectionComponent.cpp
+++ b/Engine/Source/Runtime/Experimental/GeometryCollectionEngine/Private/GeometryCollection/GeometryCollectionComponent.cpp
@@ -30,6 +30,7 @@
 #include "GeometryCollection/GeometryCollectionUtility.h"
 #include "GeometryCollection/GeometryCollectionISMPoolActor.h"
 #include "GeometryCollection/GeometryCollectionISMPoolComponent.h"
+#include "GeometryCollection/GeometryCollectionISMPoolSubSystem.h"
 #include "Math/Sphere.h"
 #include "Modules/ModuleManager.h"
 #include "Net/Core/PushModel/PushModel.h"
@@ -317,6 +318,8 @@ UGeometryCollectionComponent::UGeometryCollectionComponent(const FObjectInitiali
 	, bEnableRunTimeDataCollection(false)
 	, RunTimeDataCollectionGuid(FGuid::NewGuid())
 #endif
+	, ISMPool(nullptr)
+	, bAutoAssignISMPool(false)
 	, bEnableReplication(false)
 	, bEnableAbandonAfterLevel(true)
 	, AbandonedCollisionProfileName(UCollisionProfile::CustomCollisionProfileName)
@@ -2375,6 +2378,23 @@ void UGeometryCollectionComponent::AsyncPhysicsTickComponent(float DeltaTime, fl
 
 void UGeometryCollectionComponent::OnRegister()
 {
+	// important : we shoudl assign this as soon as possible to avoid the scene proxy from being created
+	AssignedISMPool = nullptr;
+	if (bChaos_GC_UseISMPool && GetWorld()->IsGameWorld())
+	{
+		if (bAutoAssignISMPool)
+		{
+			if (UGeometryCollectionISMPoolSubSystem* ISMPoolSubSystem = UWorld::GetSubsystem<UGeometryCollectionISMPoolSubSystem>(GetWorld()))
+			{
+				AssignedISMPool = ISMPoolSubSystem->FindISMPoolActor(*this);
+			}
+		}
+		else
+		{
+			AssignedISMPool = ISMPool;
+		}
+	}
+
 	//UE_LOG(UGCC_LOG, Log, TEXT("GeometryCollectionComponent[%p]::OnRegister()[%p]"), this,RestCollection );
 	ResetDynamicCollection();
 
@@ -2394,6 +2414,12 @@ void UGeometryCollectionComponent::OnRegister()
 	Super::OnRegister();
 }
 
+void UGeometryCollectionComponent::OnUnregister()
+{
+	Super::OnUnregister();
+	AssignedISMPool = nullptr;
+}
+
 void UGeometryCollectionComponent::ResetDynamicCollection()
 {
 	bool bCreateDynamicCollection = true;
@@ -4065,7 +4091,7 @@ void UGeometryCollectionComponent::InitializeEmbeddedGeometry()
 
 bool UGeometryCollectionComponent::CanUseISMPool() const 
 {
-	return bChaos_GC_UseISMPool && ISMPool && GetWorld()->IsGameWorld();
+	return bChaos_GC_UseISMPool && AssignedISMPool && GetWorld()->IsGameWorld();
 }
 
 void UGeometryCollectionComponent::RegisterToISMPool()
@@ -4074,7 +4100,7 @@ void UGeometryCollectionComponent::RegisterToISMPool()
 
 	if (CanUseISMPool())
 	{
-		if (UGeometryCollectionISMPoolComponent* ISMPoolComp = ISMPool->GetISMPoolComp())
+		if (UGeometryCollectionISMPoolComponent* ISMPoolComp = AssignedISMPool->GetISMPoolComp())
 		{
 			bool bCanRenderComponent = true;
 			if (RestCollection)
@@ -4136,7 +4162,7 @@ void UGeometryCollectionComponent::RegisterToISMPool()
 
 				// root proxy if available 
 				// TODO : if ISM pool is not available : uses a standard static mesh component
-				if (UObject* RootMeshProxyObject = RestCollection->RootProxy.ResolveObject())
+				if (UObject* RootMeshProxyObject = RestCollection->RootProxy.TryLoad())
 				{
 					if (UStaticMesh* RootMeshProxy = Cast<UStaticMesh>(RootMeshProxyObject))
 					{
@@ -4159,23 +4185,23 @@ void UGeometryCollectionComponent::RegisterToISMPool()
 
 void UGeometryCollectionComponent::UnregisterFromISMPool()
 {
-	if (ISMPool)
+	if (AssignedISMPool)
 	{
-		if (UGeometryCollectionISMPoolComponent* ISMPoolComp = ISMPool->GetISMPoolComp())
+		if (UGeometryCollectionISMPoolComponent* ISMPoolComp = AssignedISMPool->GetISMPoolComp())
 		{
 			ISMPoolComp->DestroyMeshGroup(ISMPoolMeshGroupIndex);
 			ISMPoolMeshGroupIndex = INDEX_NONE;
 			ISMPoolRootProxyMeshId = INDEX_NONE;
 		}
+		SetVisibility(true);
 	}
-	SetVisibility(true);
 }
 
 void UGeometryCollectionComponent::RefreshISMPoolInstances()
 {
 	if (CanUseISMPool())
 	{
-		if (UGeometryCollectionISMPoolComponent* ISMPoolComp = ISMPool->GetISMPoolComp())
+		if (UGeometryCollectionISMPoolComponent* ISMPoolComp = AssignedISMPool->GetISMPoolComp())
 		{
 			if (RestCollection)
 			{
diff --git a/Engine/Source/Runtime/Experimental/GeometryCollectionEngine/Private/GeometryCollection/GeometryCollectionISMPoolSubSystem.cpp b/Engine/Source/Runtime/Experimental/GeometryCollectionEngine/Private/GeometryCollection/GeometryCollectionISMPoolSubSystem.cpp
new file mode 100644
--- /dev/null
+++ b/Engine/Source/Runtime/Experimental/GeometryCollectionEngine/Private/GeometryCollection/GeometryCollectionISMPoolSubSystem.cpp
@@ -0,0 +1,52 @@
+// Copyright Epic Games, Inc. All Rights Reserved.
+
+#include "GeometryCollection/GeometryCollectionISMPoolSubSystem.h"
+#include "GeometryCollection/GeometryCollectionISMPoolActor.h"
+#include "Engine/World.h"
+
+//----------------------------------------------------------------------//
+// UGeometryCollectionISMPoolSubSystem
+//----------------------------------------------------------------------//
+UGeometryCollectionISMPoolSubSystem::UGeometryCollectionISMPoolSubSystem()
+	: ISMPoolActor(nullptr)
+{
+}
+
+void UGeometryCollectionISMPoolSubSystem::Initialize(FSubsystemCollectionBase& Collection)
+{
+	Super::Initialize(Collection);
+	Collection.InitializeDependency<UGeometryCollectionISMPoolSubSystem>();
+}
+
+void UGeometryCollectionISMPoolSubSystem::Deinitialize()
+{
+	if (ISMPoolActor)
+	{
+		if (UWorld* World = GetWorld())
+		{
+			GetWorld()->DestroyActor(ISMPoolActor);
+		}
+	}
+	Super::Deinitialize();
+}
+
+void UGeometryCollectionISMPoolSubSystem::OnWorldBeginPlay(UWorld& InWorld)
+{
+	Super::OnWorldBeginPlay(InWorld);
+}
+
+AGeometryCollectionISMPoolActor* UGeometryCollectionISMPoolSubSystem::FindISMPoolActor(const USceneComponent& Requester)
+{
+	// on demand creation of the actor 
+	// very simple logic for now, we can extend that in the future to return a specific actor based on the criteria
+	if (!ISMPoolActor)
+	{
+		// we keep it as transient to avoid accumulation of those actors in saved levels
+		FActorSpawnParameters Params;
+		Params.ObjectFlags = EObjectFlags::RF_DuplicateTransient | EObjectFlags::RF_Transient;
+		ISMPoolActor = GetWorld()->SpawnActor<AGeometryCollectionISMPoolActor>(Params);
+	}
+	return ISMPoolActor;
+}
+
+
diff --git a/Engine/Source/Runtime/Experimental/GeometryCollectionEngine/Public/GeometryCollection/GeometryCollectionComponent.h b/Engine/Source/Runtime/Experimental/GeometryCollectionEngine/Public/GeometryCollection/GeometryCollectionComponent.h
--- a/Engine/Source/Runtime/Experimental/GeometryCollectionEngine/Public/GeometryCollection/GeometryCollectionComponent.h
+++ b/Engine/Source/Runtime/Experimental/GeometryCollectionEngine/Public/GeometryCollection/GeometryCollectionComponent.h
@@ -483,6 +483,7 @@ public:
 	//~ Begin UPrimitiveComponent Interface.
 	virtual FPrimitiveSceneProxy* CreateSceneProxy() override;
 	virtual void OnRegister() override;
+	virtual void OnUnregister() override;
 	virtual FBodyInstance* GetBodyInstance(FName BoneName = NAME_None, bool bGetWelded = true, int32 Index = -1) const override;
 	virtual void SetNotifyRigidBodyCollision(bool bNewNotifyRigidBodyCollision) override;
 	virtual bool CanEditSimulatePhysics() override;
@@ -1048,6 +1049,10 @@ protected:
 	UPROPERTY(EditAnywhere, BlueprintReadWrite, Category = "ChaosPhysics|Rendering", meta = (DisplayName = "ISM Pool"))
 	TObjectPtr<AGeometryCollectionISMPoolActor> ISMPool;
 
+	UPROPERTY(EditAnywhere, BlueprintReadWrite, Category = "ChaosPhysics|Rendering")
+	bool bAutoAssignISMPool;
+
+	TObjectPtr<AGeometryCollectionISMPoolActor> AssignedISMPool = nullptr;
 	int32 ISMPoolMeshGroupIndex = INDEX_NONE;
 	int32 ISMPoolRootProxyMeshId = INDEX_NONE;
 
diff --git a/Engine/Source/Runtime/Experimental/GeometryCollectionEngine/Public/GeometryCollection/GeometryCollectionISMPoolSubSystem.h b/Engine/Source/Runtime/Experimental/GeometryCollectionEngine/Public/GeometryCollection/GeometryCollectionISMPoolSubSystem.h
new file mode 100644
--- /dev/null
+++ b/Engine/Source/Runtime/Experimental/GeometryCollectionEngine/Public/GeometryCollection/GeometryCollectionISMPoolSubSystem.h
@@ -0,0 +1,37 @@
+// Copyright Epic Games, Inc. All Rights Reserved.
+
+#pragma once
+
+#include "Subsystems/WorldSubsystem.h"
+#include "GeometryCollectionISMPoolSubSystem.generated.h"
+
+class AGeometryCollectionISMPoolActor;
+
+/**
+* A subsystem managing StateTree assets in Mass
+*/
+UCLASS()
+class GEOMETRYCOLLECTIONENGINE_API UGeometryCollectionISMPoolSubSystem : public UWorldSubsystem
+{
+	GENERATED_BODY()
+
+public:
+
+	UGeometryCollectionISMPoolSubSystem();
+
+	// USubsystem BEGIN
+	virtual void Initialize(FSubsystemCollectionBase& Collection) override;
+	virtual void Deinitialize() override;
+	// USubsystem END
+
+	// UWorldSubsystem BEGIN
+	virtual void OnWorldBeginPlay(UWorld& InWorld) override;
+	// UWorldSubsystem END
+
+	AGeometryCollectionISMPoolActor* FindISMPoolActor(const USceneComponent& Requester);
+
+protected:
+
+	/** for now we only use one ISMPool actor per world, but we could extend the system to manage many more and return the right one based on  search criteria */
+	TObjectPtr<AGeometryCollectionISMPoolActor> ISMPoolActor;
+};
