potentially fix an issue where the geometry collection async physics tick was being called in the editor on PostLoad incorrectly

diff --git a/Engine/Source/Runtime/Experimental/GeometryCollectionEngine/Private/GeometryCollection/GeometryCollectionComponent.cpp b/Engine/Source/Runtime/Experimental/GeometryCollectionEngine/Private/GeometryCollection/GeometryCollectionComponent.cpp
--- a/Engine/Source/Runtime/Experimental/GeometryCollectionEngine/Private/GeometryCollection/GeometryCollectionComponent.cpp
+++ b/Engine/Source/Runtime/Experimental/GeometryCollectionEngine/Private/GeometryCollection/GeometryCollectionComponent.cpp
@@ -48,6 +48,7 @@
 #if WITH_EDITOR
 #include "AssetToolsModule.h"
 #include "Editor.h"
+#include "UObject/UObjectThreadContext.h"
 #endif
 
 #include "PhysicsEngine/BodySetup.h"
@@ -2362,6 +2363,19 @@ void UGeometryCollectionComponent::TickComponent(float DeltaTime, enum ELevelTic
 
 void UGeometryCollectionComponent::AsyncPhysicsTickComponent(float DeltaTime, float SimTime)
 {
+	if (!PhysicsProxy || !PhysicsProxy->IsInitializedOnPhysicsThread())
+	{
+		return;
+	}
+
+#if WITH_EDITOR
+	// A bit of a hack because the super async physics tick will crash if this is true.
+	if (FUObjectThreadContext::Get().IsRoutingPostLoad)
+	{
+		return;
+	}
+#endif
+
 	Super::AsyncPhysicsTickComponent(DeltaTime, SimTime);
 
 	// using net mode for now as using local role seemed to cause other issues at initialization time
