Add Poisson Solve on a Uniform Grid to Chaos

diff --git a/Engine/Source/Runtime/Experimental/Chaos/Public/Chaos/Math/Krylov.h b/Engine/Source/Runtime/Experimental/Chaos/Public/Chaos/Math/Krylov.h
--- a/Engine/Source/Runtime/Experimental/Chaos/Public/Chaos/Math/Krylov.h
+++ b/Engine/Source/Runtime/Experimental/Chaos/Public/Chaos/Math/Krylov.h
@@ -17,7 +17,8 @@ void LanczosCG(
 	const TArray<T>& b,
 	const int max_it, 
 	const T res = 1e-4, 
-	bool check_residual = false) 
+	bool check_residual = false,
+	int min_parallel_batch_size = 1000) 
 {
 	auto dotProduct = [](const TArray<T>& x, const TArray<T>& y) 
 	{
@@ -28,37 +29,37 @@ void LanczosCG(
 		{
 			result += x[i] * y[i];
 		}/*,
-		x.Num() < 1000);*/
+		x.Num() < min_parallel_batch_size);*/
 		return result;
 	};
 
-	auto AXPY = [](TArray<T>& y, T a, const TArray<T>& x) 
+	auto AXPY = [min_parallel_batch_size](TArray<T>& y, T a, const TArray<T>& x) 
 	{
 		checkfSlow(x.Num() == y.Num(), TEXT("LanczosCG: trying to take a linear combination of vectors with different sizes."));
 		PhysicsParallelFor(x.Num(), [&](const int32 i)
 		{
 			y[i] += a * x[i];
 	    },
-		x.Num() < 1000); // Single threaded for less than 1k operations.
+		x.Num() < min_parallel_batch_size); // Single threaded for less than min_parallel_batch_size operations.
 	};
 
-	auto scale = [](TArray<T>& y, T a) 
+	auto scale = [min_parallel_batch_size](TArray<T>& y, T a)
 	{
 		PhysicsParallelFor(y.Num(), [&](const int32 i)
 		{
 			y[i] *= a;
 	    },
-		y.Num() < 1000); // Single threaded for less than 1k operations.
+		y.Num() < min_parallel_batch_size); // Single threaded for less than min_parallel_batch_size operations.
 	};
 
-	auto set = [](TArray<T>& y, const TArray<T>& x) 
+	auto set = [min_parallel_batch_size](TArray<T>& y, const TArray<T>& x)
 	{
 		checkfSlow(x.Num() == y.Num(), TEXT("LanczosCG: trying to set to vectors with different sizes."));
 		PhysicsParallelFor(x.Num(), [&](const int32 i)
 		{
 			y[i] = x[i];
 		},
-		x.Num() < 1000); // Single threaded for less than 1k operations.
+		x.Num() < min_parallel_batch_size); // Single threaded for less than 1k operations.
 	};
 
 	LanczosCG<T>(multiplyA, dotProduct, AXPY, scale, set, x, b, max_it, res, check_residual);
diff --git a/Engine/Source/Runtime/Experimental/Chaos/Public/Chaos/Math/Poisson.h b/Engine/Source/Runtime/Experimental/Chaos/Public/Chaos/Math/Poisson.h
--- a/Engine/Source/Runtime/Experimental/Chaos/Public/Chaos/Math/Poisson.h
+++ b/Engine/Source/Runtime/Experimental/Chaos/Public/Chaos/Math/Poisson.h
@@ -4,6 +4,7 @@
 //#include "ChaosFlesh/FleshCollectionUtility.h"
 #include "Chaos/Framework/Parallel.h"
 #include "Chaos/Math/Krylov.h"
+#include "Chaos/UniformGrid.h"
 #include "Chaos/Vector.h"
 #include "Math/Matrix.h"
 #include "Math/Vector.h"
@@ -567,5 +568,88 @@ ComputeFiberField(
 	}//);
 }
 
+// 9 point laplacian with dirichlet boundaries
+template<class TV, class T, bool NodalValues = false>
+void Laplacian(const TUniformGrid<T, 3>& UniformGrid,
+	const TArray<TV>& U,
+	TArray<TV>& Lu)
+{
+	const TVec3<int32> Counts = NodalValues ? UniformGrid.NodeCounts() : UniformGrid.Counts();
+
+	const int32 XStride = Counts[1] * Counts[2];
+	const int32 YStride = Counts[2];
+	constexpr int32 ZStride = 1;
+	const Chaos::TVector<TV, 3> OneOverDxSq = Chaos::TVector<TV, 3>(1) / (UniformGrid.Dx() * UniformGrid.Dx());
+
+	Fill(Lu, TV(0));
+
+	for (int32 I = 1; I < Counts.X - 1; ++I)
+	{
+		for (int32 J = 1; J < Counts.Y - 1; ++J)
+		{
+			for (int32 K = 1; K < Counts.Z - 1; ++K)
+			{
+				const int32 FlatIndex = I * XStride + J * YStride + K * ZStride;
+				Lu[FlatIndex] +=
+					(U[FlatIndex + XStride] - (T)2. * U[FlatIndex] + U[FlatIndex - XStride]) * OneOverDxSq[0] +
+					(U[FlatIndex + YStride] - (T)2. * U[FlatIndex] + U[FlatIndex - YStride]) * OneOverDxSq[1] +
+					(U[FlatIndex + ZStride] - (T)2. * U[FlatIndex] + U[FlatIndex - ZStride]) * OneOverDxSq[2];
+			}
+		}
+	}
+}
+
+// Solve Poisson on a Uniform Grid (with standard 9-point Laplacian).
+// Weights is a flattened array (using same indexing as UniformGrid and ArrayND) and the output solution.
+// InConstrainedNodes is a list of flattened index nodes that are boundary conditions. ConstrainedWeights is a parallel array of the constrained values.
+template<class TV, class T, bool NodalValues = false>
+void PoissonSolve(const TArray<int32>& InConstrainedNodes,
+	const TArray<TV>& ConstrainedWeights,
+	const TUniformGrid<T, 3>& UniformGrid,
+	const int32 MaxItCG,
+	const TV CGTol,
+	TArray<TV>& Weights,
+	bool bCheckResidual = false,
+	int32 MinParallelBatchSize = 1000)
+{
+	auto ProjectBCs = [&InConstrainedNodes](TArray<TV>& U)
+	{
+		for (int32 i = 0; i < InConstrainedNodes.Num(); i++)
+		{
+			U[InConstrainedNodes[i]] = (T)0.;
+		}
+	};
+
+	auto MultiplyLaplacian = [&ProjectBCs, &UniformGrid](TArray<TV>& LU, const TArray<TV>& U)
+	{
+		Laplacian<TV, T, NodalValues>(
+			UniformGrid,
+			U,
+			LU);
+		ProjectBCs(LU);
+	};
+
+	Fill(Weights, TV(0));
+
+	TArray<TV> InitialGuess;
+	InitialGuess.Init((TV)0., Weights.Num());
+	for (int32 i = 0; i < InConstrainedNodes.Num(); i++)
+	{
+		InitialGuess[InConstrainedNodes[i]] = ConstrainedWeights[i];
+	}
+	TArray<TV> MinusResidual;
+	MinusResidual.Init((TV)0., Weights.Num());
+	MultiplyLaplacian(MinusResidual, InitialGuess);
+
+	TArray<TV> MinusDw;
+	MinusDw.Init((TV)0., Weights.Num());
+	Chaos::LanczosCG(MultiplyLaplacian, MinusDw, MinusResidual, MaxItCG, CGTol, bCheckResidual, MinParallelBatchSize);
+	
+	for (int32 i = 0; i < InitialGuess.Num(); i++)
+	{
+		Weights[i] = InitialGuess[i] - MinusDw[i];
+	}
+}
+
 
-} // namespace ChaosFlesh
+} // namespace Chaos
