disabling a geometry collection particle from the physics proxy should also mark it dirty so it'll get processed for removal from SQ

diff --git a/Engine/Source/Runtime/Experimental/Chaos/Private/PhysicsProxy/GeometryCollectionPhysicsProxy.cpp b/Engine/Source/Runtime/Experimental/Chaos/Private/PhysicsProxy/GeometryCollectionPhysicsProxy.cpp
--- a/Engine/Source/Runtime/Experimental/Chaos/Private/PhysicsProxy/GeometryCollectionPhysicsProxy.cpp
+++ b/Engine/Source/Runtime/Experimental/Chaos/Private/PhysicsProxy/GeometryCollectionPhysicsProxy.cpp
@@ -1657,6 +1657,7 @@ void FGeometryCollectionPhysicsProxy::DisableParticles_External(TArray<int32>&&
 				for (int32 TransformIdx : IndicesToDisable)
 				{
 					RBDSolver->GetEvolution()->DisableParticleWithRemovalEvent(SolverParticleHandles[TransformIdx]);
+					RBDSolver->GetParticles().MarkTransientDirtyParticle(SolverParticleHandles[TransformIdx]);
 				}
 			});
 	}
