Dataflow - Added a custom serialization virtual method for nodes.

diff --git a/Engine/Source/Runtime/Experimental/Dataflow/Core/Public/Dataflow/DataflowNode.h b/Engine/Source/Runtime/Experimental/Dataflow/Core/Public/Dataflow/DataflowNode.h
--- a/Engine/Source/Runtime/Experimental/Dataflow/Core/Public/Dataflow/DataflowNode.h
+++ b/Engine/Source/Runtime/Experimental/Dataflow/Core/Public/Dataflow/DataflowNode.h
@@ -114,6 +114,9 @@ struct DATAFLOWCORE_API FDataflowNode
 	static const FLinearColor DefaultNodeTitleColor;
 	static const FLinearColor DefaultNodeBodyTintColor;
 
+	/** Override this method to provide custom serialization for this node. */
+	virtual void Serialize(FArchive& Ar) {}
+
 	//
 	//  Struct Support
 	//
@@ -280,7 +283,8 @@ public:																				\
 	virtual void SerializeInternal(FArchive& Ar) override {							\
 		UScriptStruct* const Struct = TYPE::StaticStruct();							\
 		Struct->SerializeTaggedProperties(Ar, (uint8*)this,							\
-		Struct, nullptr);}															\
+		Struct, nullptr);															\
+		Serialize(Ar);}																\
 	virtual FName GetDisplayName() const override { return TYPE::StaticDisplay(); }	\
 	virtual FName GetCategory() const override { return TYPE::StaticCategory(); }	\
 	virtual FString GetTags() const override { return TYPE::StaticTags(); }			\
