Fixed GC component implicit generation calculating incorrect sizes for box clusters.

We generate box cluster implicits by wrapping a bounding box around all child verts during DDC data generation, however we were appending each child's verts to a list, then transforming the verts and adding them all to the bounds. This effectively means that for child 0 of a cluster of n nodes we would apply the cluster-to-child transform of every child 1...n to the verts of child 0 - leading to wildly incorrect bounds after calculation. This was also leading to the cluster immediately breaking on the physics thread if simulating imparting non-physical impulses to the collection

diff --git a/Engine/Source/Runtime/Experimental/Chaos/Private/PhysicsProxy/GeometryCollectionPhysicsProxy.cpp b/Engine/Source/Runtime/Experimental/Chaos/Private/PhysicsProxy/GeometryCollectionPhysicsProxy.cpp
--- a/Engine/Source/Runtime/Experimental/Chaos/Private/PhysicsProxy/GeometryCollectionPhysicsProxy.cpp
+++ b/Engine/Source/Runtime/Experimental/Chaos/Private/PhysicsProxy/GeometryCollectionPhysicsProxy.cpp
@@ -3879,6 +3879,9 @@ void FGeometryCollectionPhysicsProxy::InitializeSharedCollisionStructures(
 									ChildMassToClusterMass.TransformPosition(MassSpaceParticles.X(VertIdx));
 								InstanceBoundingBox += MassSpaceParticles.X(VertIdx);
 							}
+
+							// Reset vert set for next child so we don't repeatedly transform verts and over extend our calculated bounds
+							VertsAdded.Reset();
 						}
 					}
 				} // tmp scope
diff --git a/Engine/Source/Runtime/Experimental/ChaosCore/Public/Chaos/Core.h b/Engine/Source/Runtime/Experimental/ChaosCore/Public/Chaos/Core.h
--- a/Engine/Source/Runtime/Experimental/ChaosCore/Public/Chaos/Core.h
+++ b/Engine/Source/Runtime/Experimental/ChaosCore/Public/Chaos/Core.h
@@ -51,7 +51,7 @@ namespace Chaos
 	// NOTE: if you get a merge conflict on the GUID, you must replace it with a new GUID - do not accept the source or target
 	// or you will likely get DDC version conflicts resulting in crashes during load.
 	// Core version string for Chaos data. Any DDC builder dependent on Chaos for serialization should depend on this version
-	inline const TCHAR* const ChaosVersionGUID = TEXT("FC60014F-0799-4BFD-B01C-44C1DEE81F32");
+	inline const TCHAR* const ChaosVersionGUID = TEXT("E2C28A68-EE02-4BED-9343-8659AF804349");
 
 	inline FString GetChaosVersionStringInner()
 	{
