Iris
- Added internal support for having properties sharing changemask bit

diff --git a/Engine/Source/Runtime/Experimental/Iris/Core/Private/Iris/ReplicationState/ReplicationStateDescriptorBuilder.cpp b/Engine/Source/Runtime/Experimental/Iris/Core/Private/Iris/ReplicationState/ReplicationStateDescriptorBuilder.cpp
--- a/Engine/Source/Runtime/Experimental/Iris/Core/Private/Iris/ReplicationState/ReplicationStateDescriptorBuilder.cpp
+++ b/Engine/Source/Runtime/Experimental/Iris/Core/Private/Iris/ReplicationState/ReplicationStateDescriptorBuilder.cpp
@@ -47,6 +47,7 @@ namespace UE::Net::Private
 static const FName PropertyNetSerializerRegistry_NAME_ReplicationID("ReplicationID");
 static const FName PropertyNetSerializerRegistry_NAME_ChangeMaskStorage("ChangeMaskStorage");
 static const FName ReplicationStateDescriptorBuilder_NAME_NetCullDistanceSquared("NetCullDistanceSquared");
+static const FName ReplicationStateDescriptorBuilder_NAME_RoleGroup("RoleGroup");
 
 static TAutoConsoleVariable<int32> CVarIrisUseNativeFastArray(TEXT("net.Iris.UseNativeFastArray"), 1, TEXT("Enable or disable IrisNativeFastArray."));
 
@@ -118,6 +119,7 @@ public:
 		EMemberPropertyTraits Traits;											// Traits for the property
 		ELifetimeCondition ReplicationCondition;								// Condition for when and where the property should be replicated
 		uint16 ChangeMaskBits;													// Bits needed to track dirtiness for this property. Init state properties use no bits.
+		FName ChangeMaskGroupName;												// We allow members to specify a changemask group to facilitate sharing of the same changemask bit
 		FMemoryLayoutUtil::FSizeAndAlignment ExternalSizeAndAlignment;			// If we do not have a property we need to pass this in explicitly
 	};
 
@@ -934,14 +936,36 @@ void FPropertyReplicationStateDescriptorBuilder::BuildMemberChangeMaskDescriptor
 	FReplicationStateMemberChangeMaskDescriptor* CurrentMemberChangeMaskDescriptor = MemberChangeMaskDescriptors;
 	const bool bIsInitState = Context.IsInitState();
 
+	struct FChangeMaskGroupInfo
+	{
+		FName ChangeMaskGroupName;
+		FReplicationStateMemberChangeMaskDescriptor ChangeMaskDescriptor;
+	};
+	TArray<FChangeMaskGroupInfo, TInlineAllocator<16>> ChangemaskGroups;
+
 	for (const FMemberProperty& Info : Members)
 	{
 		const uint32 BitCount = (bIsInitState ? 0U : Info.ChangeMaskBits);
+		const FName ChangeMaskGroupName = Info.ChangeMaskGroupName;
 
-		CurrentMemberChangeMaskDescriptor->BitCount = BitCount;
-		CurrentMemberChangeMaskDescriptor->BitOffset = (bIsInitState ? 0U : CurrentBitOffset);
+		if (ChangeMaskGroupName.IsNone())
+		{
+			CurrentMemberChangeMaskDescriptor->BitCount = BitCount;
+			CurrentMemberChangeMaskDescriptor->BitOffset = (bIsInitState ? 0U : CurrentBitOffset);
+			CurrentBitOffset += BitCount;
+		}
+		else if (FChangeMaskGroupInfo* SharedChangeMaskGroupInfo = ChangemaskGroups.FindByPredicate([ChangeMaskGroupName](const FChangeMaskGroupInfo& GroupInfo) { return GroupInfo.ChangeMaskGroupName == ChangeMaskGroupName; }))
+		{
+			*CurrentMemberChangeMaskDescriptor = SharedChangeMaskGroupInfo->ChangeMaskDescriptor;
+		}
+		else
+		{
+			CurrentMemberChangeMaskDescriptor->BitCount = BitCount;
+			CurrentMemberChangeMaskDescriptor->BitOffset = (bIsInitState ? 0U : CurrentBitOffset);
+			CurrentBitOffset += BitCount;
 
-		CurrentBitOffset += BitCount;
+			ChangemaskGroups.Add({ ChangeMaskGroupName, *CurrentMemberChangeMaskDescriptor });
+		}
 		++CurrentMemberChangeMaskDescriptor;
 	}
 
@@ -1843,6 +1867,12 @@ void FPropertyReplicationStateDescriptorBuilder::GetSerializerTraits(FMemberProp
 	OutMemberProperty.Traits |= EnumHasAnyFlags(NetSerializer->Traits, ENetSerializerTraits::HasConnectionSpecificSerialization) ? EMemberPropertyTraits::HasConnectionSpecificSerialization : EMemberPropertyTraits::None;
 	OutMemberProperty.Traits |= EnumHasAnyFlags(NetSerializer->Traits, ENetSerializerTraits::UseSerializerIsEqual) ? EMemberPropertyTraits::UseSerializerIsEqual : EMemberPropertyTraits::None;
 	OutMemberProperty.Traits |= GetFastArrayPropertyTraits(NetSerializer, Property, bAllowFastArrayWithExtraProperties);
+
+	// Tag Role/RemoteRole to share the same changemask as the serializer operates on multiple properties
+	if (NetSerializer == &UE_NET_GET_SERIALIZER(FNetRoleNetSerializer))
+	{
+		OutMemberProperty.ChangeMaskGroupName = ReplicationStateDescriptorBuilder_NAME_RoleGroup;
+	}
 }
 
 bool FPropertyReplicationStateDescriptorBuilder::IsSupportedProperty(FMemberProperty& OutMemberProperty, const FProperty* Property, const FIsSupportedPropertyParams& Parameters)
