Constraints: per axis filter

diff --git a/Engine/Source/Runtime/AnimationCore/Public/Constraint.h b/Engine/Source/Runtime/AnimationCore/Public/Constraint.h
--- a/Engine/Source/Runtime/AnimationCore/Public/Constraint.h
+++ b/Engine/Source/Runtime/AnimationCore/Public/Constraint.h
@@ -58,48 +58,48 @@ struct ANIMATIONCORE_API FFilterOptionPerAxis
 		, bZ(true)
 	{}
 
-	void FilterVector(FVector& Input, float ResetValue = 0.f) const
+	void FilterVector(FVector& Input, const FVector& ResetValue = FVector::ZeroVector) const
 	{
 		if (!bX)
 		{
-			Input.X = ResetValue;
+			Input.X = ResetValue.X;
 		}
 
 		if (!bY)
 		{
-			Input.Y = ResetValue;
+			Input.Y = ResetValue.Y;
 		}
 
 		if (!bZ)
 		{
-			Input.Z = ResetValue;
+			Input.Z = ResetValue.Z;
 		}
 	}
 
-	void FilterQuat(FQuat& Input) const
+	void FilterQuat(FQuat& Input, const FQuat& ResetValue = FQuat::Identity) const
 	{
 		FRotator Rotator = Input.Rotator();
 
-		FilterRotator(Rotator);
+		FilterRotator(Rotator, ResetValue.Rotator());
 
 		Input = Rotator.Quaternion();
 	}
 
-	void FilterRotator(FRotator& Input) const
+	void FilterRotator(FRotator& Input, const FRotator& ResetValue = FRotator::ZeroRotator) const
 	{
 		if (!bX)
 		{
-			Input.Roll = 0.f;
+			Input.Roll = ResetValue.Roll;
 		}
 
 		if (!bY)
 		{
-			Input.Pitch = 0.f;
+			Input.Pitch = ResetValue.Pitch;
 		}
 
 		if (!bZ)
 		{
-			Input.Yaw = 0.f;
+			Input.Yaw = ResetValue.Yaw;
 		}
 	}
 
@@ -151,7 +151,7 @@ struct ANIMATIONCORE_API FTransformFilter
 		Input.SetRotation(Rotation);
 
 		FVector Scale3D = Input.GetScale3D();
-		ScaleFilter.FilterVector(Scale3D, 1.f);
+		ScaleFilter.FilterVector(Scale3D, FVector::OneVector);
 		Input.SetScale3D(Scale3D);
 	}
 
@@ -166,7 +166,7 @@ struct ANIMATIONCORE_API FTransformFilter
 		Input.Rotation = Rotation;
 
 		FVector Scale = Input.Scale;
-		ScaleFilter.FilterVector(Scale, 1.f);
+		ScaleFilter.FilterVector(Scale, FVector::OneVector);
 		Input.Scale = Scale;
 	}
 };
diff --git a/Engine/Source/Runtime/Experimental/Animation/Constraints/Private/TransformConstraint.cpp b/Engine/Source/Runtime/Experimental/Animation/Constraints/Private/TransformConstraint.cpp
--- a/Engine/Source/Runtime/Experimental/Animation/Constraints/Private/TransformConstraint.cpp
+++ b/Engine/Source/Runtime/Experimental/Animation/Constraints/Private/TransformConstraint.cpp
@@ -662,11 +662,16 @@ FConstraintTickFunction::ConstraintFunction UTickableTranslationConstraint::GetF
 
 		const FVector ParentTranslation = GetParentGlobalTransform().GetLocation();
 		FTransform Transform = GetChildGlobalTransform();
+		const FVector ChildTranslation = Transform.GetLocation();
+		
 		FVector NewTranslation = (!bMaintainOffset) ? ParentTranslation : ParentTranslation + OffsetTranslation;
 		if (ClampedWeight < 1.0f - KINDA_SMALL_NUMBER)
 		{
-			NewTranslation = FMath::Lerp<FVector>(Transform.GetLocation(), NewTranslation, ClampedWeight);
+			NewTranslation = FMath::Lerp<FVector>(ChildTranslation, NewTranslation, ClampedWeight);
 		}
+
+		AxisFilter.FilterVector(NewTranslation, ChildTranslation);
+		
 		Transform.SetLocation(NewTranslation);
 			
 		SetChildGlobalTransform(Transform);
@@ -794,12 +799,16 @@ FConstraintTickFunction::ConstraintFunction UTickableRotationConstraint::GetFunc
 		
 		const FQuat ParentRotation = GetParentGlobalTransform().GetRotation();
 		FTransform Transform = GetChildGlobalTransform();
+		const FQuat ChildRotation = Transform.GetRotation();
 
 		FQuat NewRotation = (!bMaintainOffset) ? ParentRotation : ParentRotation * OffsetRotation;
 		if (ClampedWeight < 1.0f - KINDA_SMALL_NUMBER)
 		{
-			NewRotation = FQuat::Slerp(Transform.GetRotation(), NewRotation, ClampedWeight);
+			NewRotation = FQuat::Slerp(ChildRotation, NewRotation, ClampedWeight);
 		}
+
+		AxisFilter.FilterQuat(NewRotation, ChildRotation);
+		
 		Transform.SetRotation(NewRotation);
 		
 		SetChildGlobalTransform(Transform);
@@ -931,11 +940,16 @@ FConstraintTickFunction::ConstraintFunction UTickableScaleConstraint::GetFunctio
 		
 		const FVector ParentScale = GetParentGlobalTransform().GetScale3D();
 		FTransform Transform = GetChildGlobalTransform();
+		const FVector ChildScale = Transform.GetScale3D();
+		
 		FVector NewScale = (!bMaintainOffset) ? ParentScale : ParentScale * OffsetScale;
 		if (ClampedWeight < 1.0f - KINDA_SMALL_NUMBER)
 		{
-			NewScale = FMath::Lerp<FVector>(Transform.GetScale3D(), NewScale, ClampedWeight);
+			NewScale = FMath::Lerp<FVector>(ChildScale, NewScale, ClampedWeight);
 		}
+
+		AxisFilter.FilterVector(NewScale, ChildScale);
+		
 		Transform.SetScale3D(NewScale);
 			
 		SetChildGlobalTransform(Transform);
@@ -1058,17 +1072,31 @@ FConstraintTickFunction::ConstraintFunction UTickableParentConstraint::GetFuncti
 			return;
 		}
 
-		auto LerpTransform = [ClampedWeight](const FTransform& InTransform, FTransform& InTransformToBeSet)
+		auto LerpAndFilterTransform = [ClampedWeight](
+			const FTransform& InTransform, FTransform& OutTransform, const FTransformFilter& InFilter)
 		{
+			const FVector Location = InTransform.GetLocation();
+			const FQuat Rotation = InTransform.GetRotation();
+			const FVector Scale = InTransform.GetScale3D();
+
+			FVector NewLocation = OutTransform.GetLocation();
+			FQuat NewRotation = OutTransform.GetRotation();
+			FVector NewScale = OutTransform.GetScale3D();
+			
 			if (ClampedWeight < 1.0f - KINDA_SMALL_NUMBER)
 			{
-				InTransformToBeSet.SetLocation(
-					FMath::Lerp<FVector>(InTransform.GetLocation(), InTransformToBeSet.GetLocation(), ClampedWeight));
-				InTransformToBeSet.SetRotation(
-					FQuat::Slerp(InTransform.GetRotation(), InTransformToBeSet.GetRotation(), ClampedWeight));
-				InTransformToBeSet.SetScale3D(
-					FMath::Lerp<FVector>(InTransform.GetScale3D(), InTransformToBeSet.GetScale3D(), ClampedWeight));
+				NewLocation = FMath::Lerp<FVector>(Location, NewLocation, ClampedWeight);
+				NewRotation = FQuat::Slerp(Rotation, NewRotation, ClampedWeight);
+				NewScale = FMath::Lerp<FVector>(Scale, NewScale, ClampedWeight);
 			}
+
+			InFilter.TranslationFilter.FilterVector(NewLocation, Location);
+			InFilter.RotationFilter.FilterQuat(NewRotation, Rotation);
+			InFilter.ScaleFilter.FilterVector(NewScale, Scale);
+
+			OutTransform.SetLocation(NewLocation);
+			OutTransform.SetRotation(NewRotation);
+			OutTransform.SetScale3D(NewScale);
 		};
 
 		const FTransform ParentTransform = GetParentGlobalTransform();
@@ -1076,7 +1104,7 @@ FConstraintTickFunction::ConstraintFunction UTickableParentConstraint::GetFuncti
 		FTransform TargetTransform = (!bMaintainOffset) ? ParentTransform : OffsetTransform * ParentTransform;
 		//apply weight if needed
 		const FTransform ChildGlobalTransform = GetChildGlobalTransform();
-		LerpTransform(ChildGlobalTransform, TargetTransform);
+		LerpAndFilterTransform(ChildGlobalTransform, TargetTransform, TransformFilter);
 
 		//remove scale?
 		if (!bScaling)
diff --git a/Engine/Source/Runtime/Experimental/Animation/Constraints/Public/TransformConstraint.h b/Engine/Source/Runtime/Experimental/Animation/Constraints/Public/TransformConstraint.h
--- a/Engine/Source/Runtime/Experimental/Animation/Constraints/Public/TransformConstraint.h
+++ b/Engine/Source/Runtime/Experimental/Animation/Constraints/Public/TransformConstraint.h
@@ -198,6 +198,10 @@ protected:
 	UPROPERTY(EditAnywhere, BlueprintReadWrite, Category="Offset", meta=(EditCondition="bMaintainOffset"))
 	FVector OffsetTranslation = FVector::ZeroVector;
 
+	/** Defines which translation axis is constrained. */
+	UPROPERTY(BlueprintReadWrite, Category = "Axis Filter")
+	FFilterOptionPerAxis AxisFilter;
+
 #if WITH_EDITOR
 public:
 	// UObject interface
@@ -241,6 +245,10 @@ protected:
 	UPROPERTY(EditAnywhere, BlueprintReadWrite, Category="Offset", meta=(EditCondition="bMaintainOffset"))
 	FQuat OffsetRotation = FQuat::Identity;
 
+	/** Defines which rotation axis is constrained. */
+	UPROPERTY(BlueprintReadWrite, Category = "Axis Filter")
+	FFilterOptionPerAxis AxisFilter;
+
 #if WITH_EDITOR
 public:
 	// UObject interface
@@ -283,6 +291,10 @@ protected:
 	UPROPERTY(EditAnywhere, BlueprintReadWrite, Category="Offset", meta=(EditCondition="bMaintainOffset"))
 	FVector OffsetScale = FVector::OneVector;
 
+	/** Defines which scale axis is constrained. */
+	UPROPERTY(BlueprintReadWrite, Category = "Axis Filter")
+	FFilterOptionPerAxis AxisFilter;
+
 #if WITH_EDITOR
 public:
 	// UObject interface
@@ -335,6 +347,10 @@ protected:
 	UPROPERTY(EditAnywhere, BlueprintReadWrite, Category="Properties")
 	bool bScaling = false;
 
+	/** Defines which translation/rotation/scale axis are constrained. */
+	UPROPERTY(BlueprintReadWrite, Category = "Axis Filter") 
+	FTransformFilter TransformFilter;
+
 #if WITH_EDITOR
 public:
 	// UObject interface
