Add Skinned Level Set Physics type.

This adds a new Chaos Implicit Object TWeightedLatticeImplicitObject which can wrap any (bounded) rigid implicit object with a weighted lattice. The lattice deforms via Linear Blend Skinning and is used to approximate deformation of the wrapped rigid implicit object.

A new physics asset FKShapeElem type has been added: the FKSkinnedLevelSet which holds a TWeightedLatticeImplicitObject<FLevelSet>. The PhysicsAssetEditor has been updated to generate this for skeletal meshes. Skinned level sets can be generated from a subset of bones in the skeletal mesh. Weights are transferred by Poisson diffusing weights on the surface to the volumetric lattice points.

There is limited debug draw and the asset generation workflow likely will need iteration.

diff --git a/Engine/Source/Developer/PhysicsUtilities/PhysicsUtilities.Build.cs b/Engine/Source/Developer/PhysicsUtilities/PhysicsUtilities.Build.cs
--- a/Engine/Source/Developer/PhysicsUtilities/PhysicsUtilities.Build.cs
+++ b/Engine/Source/Developer/PhysicsUtilities/PhysicsUtilities.Build.cs
@@ -13,6 +13,7 @@ public class PhysicsUtilities : ModuleRules
 				"CoreUObject",
 				"Engine",
 				"PhysicsCore",
+				"Chaos"
 			}
 		);
 
diff --git a/Engine/Source/Developer/PhysicsUtilities/Private/LevelSetHelpers.cpp b/Engine/Source/Developer/PhysicsUtilities/Private/LevelSetHelpers.cpp
new file mode 100644
--- /dev/null
+++ b/Engine/Source/Developer/PhysicsUtilities/Private/LevelSetHelpers.cpp
@@ -0,0 +1,188 @@
+// Copyright Epic Games, Inc. All Rights Reserved.
+
+#include "LevelSetHelpers.h"
+
+#include "Chaos/Levelset.h"
+#include "DynamicMesh/DynamicMesh3.h"
+#include "Engine/SkeletalMesh.h"
+#include "Implicit/SweepingMeshSDF.h"
+#include "Implicit/Solidify.h"
+#include "PhysicsEngine/BodySetup.h"
+#include "PhysicsEngine/LevelSetElem.h"
+
+
+// Copied from CollisionGeometryConversion.h (Plugins/Runtime/MeshModelingToolset)
+static void CreateLevelSetElement(const FVector3d& GridOrigin, const UE::Geometry::FDenseGrid3f& Grid, float CellSize, FKLevelSetElem& LevelSetOut)
+{
+	const UE::Geometry::FVector3i& GridDims = Grid.GetDimensions();
+	TArray<double> OutGridValues;
+	OutGridValues.Init(0.0, GridDims[0] * GridDims[1] * GridDims[2]);
+
+	for (int I = 0; I < GridDims[0]; ++I)
+	{
+		for (int J = 0; J < GridDims[1]; ++J)
+		{
+			for (int K = 0; K < GridDims[2]; ++K)
+			{
+				// account for different ordering in Geometry::TDenseGrid3 vs Chaos::TUniformGrid
+				const int InBufferIndex = I + GridDims[0] * (J + GridDims[1] * K);
+				const int OutBufferIndex = K + GridDims[2] * (J + GridDims[1] * I);
+				OutGridValues[OutBufferIndex] = Grid[InBufferIndex];
+			}
+		}
+	}
+
+	const FTransform3d GridTransform = FTransform(GridOrigin);
+	FTransform ChaosTransform = GridTransform;
+	ChaosTransform.AddToTranslation(-0.5 * CellSize * FVector::One());		// account for grid origin being at cell center vs cell corner
+	LevelSetOut.BuildLevelSet(ChaosTransform, OutGridValues, FIntVector(GridDims[0], GridDims[1], GridDims[2]), CellSize);
+}
+
+static void CreateLevelSetElement(const FVector3d& GridOrigin, const UE::Geometry::FDenseGrid3f& Grid, float CellSize, TUniquePtr<Chaos::FLevelSet>& LevelSetOut)
+{
+	const UE::Geometry::FVector3i& GridDims = Grid.GetDimensions();
+	TArray<double> OutGridValues;
+	OutGridValues.Init(0.0, GridDims[0] * GridDims[1] * GridDims[2]);
+
+	for (int I = 0; I < GridDims[0]; ++I)
+	{
+		for (int J = 0; J < GridDims[1]; ++J)
+		{
+			for (int K = 0; K < GridDims[2]; ++K)
+			{
+				// account for different ordering in Geometry::TDenseGrid3 vs Chaos::TUniformGrid
+				const int InBufferIndex = I + GridDims[0] * (J + GridDims[1] * K);
+				const int OutBufferIndex = K + GridDims[2] * (J + GridDims[1] * I);
+				OutGridValues[OutBufferIndex] = Grid[InBufferIndex];
+			}
+		}
+	}
+
+	const FVector3d GridMin = GridOrigin - 0.5 * CellSize * FVector::One();		// account for grid origin being at cell center vs cell corner
+	const FVector3d GridMax = GridMin + CellSize * FVector3d((double)GridDims[0], (double)GridDims[1], (double)GridDims[2]);
+
+	const Chaos::TVec3<int32> ChaosLSGridDims(GridDims[0], GridDims[1], GridDims[2]);
+	Chaos::TUniformGrid<Chaos::FReal, 3> LevelSetGrid(GridMin, GridMax, ChaosLSGridDims);
+	Chaos::TArrayND<Chaos::FReal, 3> Phi(ChaosLSGridDims, OutGridValues);
+
+	LevelSetOut = MakeUnique<Chaos::FLevelSet>(MoveTemp(LevelSetGrid), MoveTemp(Phi), 0);
+}
+
+// Copied from FMeshSimpleShapeApproximation::Generate_LevelSets (Plugins/Runtime/GeometryProcessing)
+template<typename FLevelSetElemType>
+bool CreateLevelSetForMeshInternal(const UE::Geometry::FDynamicMesh3& InMesh, int32 InLevelSetGridResolution, FLevelSetElemType& OutElement)
+{
+	using UE::Geometry::FDynamicMesh3;
+	using UE::Geometry::TSweepingMeshSDF;
+
+	constexpr int32 NumNarrowBandCells = 2;
+	constexpr int32 NumExpandCells = 1;
+
+	// Inside SDF.Compute(), extra grid cell are added to account for the narrow band ("NarrowBandMaxDistance") as well as an 
+	// outer buffer ("ExpandBounds"). So here we adjust the cell size to make the final output resolution closer to the user-specified resolution.
+	const int32 LevelSetGridResolution = FMath::Max(1, InLevelSetGridResolution - 2 * (NumNarrowBandCells + NumExpandCells));
+
+	const UE::Geometry::FAxisAlignedBox3d Bounds = InMesh.GetBounds();
+	const double CellSize = Bounds.MaxDim() / LevelSetGridResolution;
+	const double ExpandBounds = NumExpandCells * CellSize;
+
+	UE::Geometry::TMeshAABBTree3<FDynamicMesh3> Spatial(&InMesh);
+
+	// Input mesh is likely open, so solidify it first
+	UE::Geometry::TFastWindingTree<FDynamicMesh3> FastWinding(&Spatial);
+	UE::Geometry::TImplicitSolidify<FDynamicMesh3> Solidify(&InMesh, &Spatial, &FastWinding);
+	Solidify.ExtendBounds = ExpandBounds;
+	Solidify.MeshCellSize = CellSize;
+	Solidify.WindingThreshold = 0.5;
+	Solidify.SurfaceSearchSteps = 3;
+	Solidify.bSolidAtBoundaries = true;
+
+	const UE::Geometry::FDynamicMesh3 SolidMesh = FDynamicMesh3(&Solidify.Generate());
+	Spatial.SetMesh(&SolidMesh, true);
+
+	TSweepingMeshSDF<FDynamicMesh3> SDF;
+	SDF.Mesh = &SolidMesh;
+	SDF.Spatial = &Spatial;
+	SDF.ComputeMode = TSweepingMeshSDF<FDynamicMesh3>::EComputeModes::FullGrid;
+	SDF.CellSize = (float)CellSize;
+	SDF.NarrowBandMaxDistance = NumNarrowBandCells * CellSize;
+	SDF.ExactBandWidth = FMath::CeilToInt32(SDF.NarrowBandMaxDistance / CellSize);
+	SDF.ExpandBounds = FVector3d(ExpandBounds);
+
+	if (SDF.Compute(Bounds))
+	{
+		CreateLevelSetElement((FVector3d)SDF.GridOrigin, SDF.Grid, SDF.CellSize, OutElement);
+		return true;
+	}
+
+	return false;
+}
+
+namespace LevelSetHelpers
+{
+bool CreateLevelSetForBone(UBodySetup* BodySetup, const TArray<FVector3f>& InVertices, const TArray<uint32>& InIndices, uint32 InResolution)
+{
+	check(BodySetup != NULL);
+
+	// Validate input by checking bounding box
+	FBox VertBox(ForceInit);
+	for (const FVector3f& Vert : InVertices)
+	{
+		VertBox += (FVector)Vert;
+	}
+
+	// If box is invalid, or the largest dimension is less than 1 unit, or smallest is less than 0.1, skip trying to generate collision
+	if (VertBox.IsValid == 0 || VertBox.GetSize().GetMax() < 1.f || VertBox.GetSize().GetMin() < 0.1f || InIndices.Num() == 0)
+	{
+		return false;
+	}
+
+	TRACE_CPUPROFILER_EVENT_SCOPE(MeshToLevelSet)
+
+	// Clean out old hulls
+	BodySetup->RemoveSimpleCollision();
+
+	UE::Geometry::FDynamicMesh3 DynamicMesh;
+	CreateDynamicMesh(InVertices, InIndices, DynamicMesh);
+
+	FKLevelSetElem LevelSetElement;
+	const bool bOK = CreateLevelSetForMesh(DynamicMesh, InResolution, LevelSetElement);
+
+	if (!bOK)
+	{
+		return false;
+	}
+
+	BodySetup->AggGeom.LevelSetElems.Add(LevelSetElement);
+	BodySetup->InvalidatePhysicsData(); // update GUID
+	return true;
+}
+
+void CreateDynamicMesh(const TArray<FVector3f>& InVertices, const TArray<uint32>& InIndices, UE::Geometry::FDynamicMesh3& OutMesh)
+{
+	OutMesh.Clear();
+	for (const FVector3f& Vertex : InVertices)
+	{
+		OutMesh.AppendVertex(FVector3d(Vertex));
+	}
+
+	check(InIndices.Num() % 3 == 0);
+	const int32 NumTriangles = InIndices.Num() / 3;
+	for (int32 TriangleID = 0; TriangleID < NumTriangles; ++TriangleID)
+	{
+		const UE::Geometry::FIndex3i Triangle(InIndices[3 * TriangleID], InIndices[3 * TriangleID + 1], InIndices[3 * TriangleID + 2]);
+		OutMesh.AppendTriangle(Triangle);
+	}
+}
+
+bool CreateLevelSetForMesh(const UE::Geometry::FDynamicMesh3& InMesh, int32 InLevelSetGridResolution, FKLevelSetElem& OutElement)
+{
+	return CreateLevelSetForMeshInternal(InMesh, InLevelSetGridResolution, OutElement);
+}
+
+bool CreateLevelSetForMesh(const UE::Geometry::FDynamicMesh3& InMesh, int32 InLevelSetGridResolution, TUniquePtr<Chaos::FLevelSet>& OutElement)
+{
+	return CreateLevelSetForMeshInternal(InMesh, InLevelSetGridResolution, OutElement);
+}
+
+} // namespace LevelSetHelpers
diff --git a/Engine/Source/Developer/PhysicsUtilities/Private/LevelSetHelpers.h b/Engine/Source/Developer/PhysicsUtilities/Private/LevelSetHelpers.h
new file mode 100644
--- /dev/null
+++ b/Engine/Source/Developer/PhysicsUtilities/Private/LevelSetHelpers.h
@@ -0,0 +1,21 @@
+// Copyright Epic Games, Inc. All Rights Reserved.
+
+#pragma once
+
+struct FKLevelSetElem;
+namespace UE::Geometry
+{
+class FDynamicMesh3;
+}
+namespace Chaos
+{
+class FLevelSet;
+}
+
+namespace LevelSetHelpers
+{
+bool CreateLevelSetForBone(class UBodySetup* BodySetup, const TArray<FVector3f>& InVertices, const TArray<uint32>& InIndices, uint32 InResolution);
+void CreateDynamicMesh(const TArray<FVector3f>& InVertices, const TArray<uint32>& InIndices, UE::Geometry::FDynamicMesh3& OutMesh);
+bool CreateLevelSetForMesh(const UE::Geometry::FDynamicMesh3& InMesh, int32 InLevelSetGridResolution, FKLevelSetElem& OutElement);
+bool CreateLevelSetForMesh(const UE::Geometry::FDynamicMesh3& InMesh, int32 InLevelSetGridResolution, TUniquePtr<Chaos::FLevelSet>& OutElement);
+}
diff --git a/Engine/Source/Developer/PhysicsUtilities/Private/PhysicsAssetUtils.cpp b/Engine/Source/Developer/PhysicsUtilities/Private/PhysicsAssetUtils.cpp
--- a/Engine/Source/Developer/PhysicsUtilities/Private/PhysicsAssetUtils.cpp
+++ b/Engine/Source/Developer/PhysicsUtilities/Private/PhysicsAssetUtils.cpp
@@ -11,16 +11,18 @@
 #include "PhysicsEngine/BoxElem.h"
 #include "PhysicsEngine/SphereElem.h"
 #include "PhysicsEngine/SphylElem.h"
+#include "PhysicsEngine/SkinnedLevelSetElem.h"
 #include "Animation/SkeletalMeshActor.h"
 #include "Components/SkeletalMeshComponent.h"
+#include "Rendering/SkeletalMeshRenderData.h"
 #include "PhysicsEngine/PhysicsAsset.h"
 #include "PhysicsEngine/PhysicsConstraintTemplate.h"
 #include "PreviewScene.h"
 #include "Misc/ScopedSlowTask.h"
 #include "SkinnedBoneTriangleCache.h"
-#include "DynamicMesh/DynamicMesh3.h"		// For mesh to level set
-#include "Implicit/SweepingMeshSDF.h"		// For mesh to level set
-#include "Implicit/Solidify.h"		// For mesh to level set
+#include "SkinnedLevelSetBuilder.h"
+#include "LevelSetHelpers.h"
+#include "Chaos/Levelset.h"
 
 namespace FPhysicsAssetUtils
 {
@@ -118,6 +120,86 @@ void AddInfoToParentInfo(const FTransform& LocalToParentTM, const FBoneVertInfo&
 	}
 }
 
+void CreateWeightedLevelSetBody(FSkinnedLevelSetBuilder& Builder, UBodySetup* BodySetup, UPhysicsAsset* PhysAsset, FName RootBoneName, const FPhysAssetCreateParams& Params)
+{
+	check(BodySetup->BoneName == RootBoneName);
+
+	BodySetup->RemoveSimpleCollision();
+	BodySetup->AggGeom.SkinnedLevelSetElems.Add(Builder.CreateSkinnedLevelSetElem());
+	BodySetup->InvalidatePhysicsData(); // update GUID
+}
+
+int32 CalculateCommonRootBoneIndex(const FReferenceSkeleton& RefSkeleton, const TArray<int32>& BoneIndices)
+{
+	// This is copied from ClothSimulationModel
+
+	// Starts at root
+	int32 RootBoneIndex = 0;
+
+	// List of valid paths to the root bone from each bone
+	TArray<TArray<int32>> PathsToRoot;
+	PathsToRoot.Reserve(BoneIndices.Num());
+
+	for (const int32 BoneIndex : BoneIndices)
+	{
+		TArray<int32>& Path = PathsToRoot.AddDefaulted_GetRef();
+
+		int32 CurrentBone = BoneIndex;
+		Path.Add(CurrentBone);
+
+		while (CurrentBone != 0 && CurrentBone != INDEX_NONE)
+		{
+			CurrentBone = RefSkeleton.GetParentIndex(CurrentBone);
+			Path.Add(CurrentBone);
+		}
+	}
+
+	// Paths are from leaf->root, we want the other way
+	for (TArray<int32>& Path : PathsToRoot)
+	{
+		Algo::Reverse(Path);
+	}
+
+	// the last common bone in all is the root
+	const int32 NumPaths = PathsToRoot.Num();
+	if (NumPaths > 0)
+	{
+		TArray<int32>& FirstPath = PathsToRoot[0];
+
+		const int32 FirstPathSize = FirstPath.Num();
+		for (int32 PathEntryIndex = 0; PathEntryIndex < FirstPathSize; ++PathEntryIndex)
+		{
+			const int32 CurrentQueryIndex = FirstPath[PathEntryIndex];
+			bool bValidRoot = true;
+
+			for (int32 PathIndex = 1; PathIndex < NumPaths; ++PathIndex)
+			{
+				if (!PathsToRoot[PathIndex].Contains(CurrentQueryIndex))
+				{
+					bValidRoot = false;
+					break;
+				}
+			}
+
+			if (bValidRoot)
+			{
+				RootBoneIndex = CurrentQueryIndex;
+			}
+			else
+			{
+				// Once we fail to find a valid root we're done.
+				break;
+			}
+		}
+	}
+	else
+	{
+		// Just use the root
+		RootBoneIndex = 0;
+	}
+	return RootBoneIndex;
+}
+
 bool CreateFromSkeletalMeshInternal(UPhysicsAsset* PhysicsAsset, USkeletalMesh* SkelMesh, const FPhysAssetCreateParams& Params, const FSkinnedBoneTriangleCache& TriangleCache)
 {
 	// For each bone, get the vertices most firmly attached to it.
@@ -133,8 +215,9 @@ bool CreateFromSkeletalMeshInternal(UPhysicsAsset* PhysicsAsset, USkeletalMesh*
 	//If not, add bone to parent for possible merge
 
 	const TArray<FTransform> LocalPose = SkelMesh->GetRefSkeleton().GetRefBonePose();
-	TMap<int32, FBoneVertInfo> BoneToMergedBones;
-	TMap<FName, TArray<FName>> BoneNameToMergedParent;
+	typedef TArray<int32> FMergedBoneIndices;
+
+	TMap<int32, TPair<FMergedBoneIndices,FBoneVertInfo>> BoneToMergedBones;
 	const int32 NumBones = Infos.Num();
 
 	TArray<float> MergedSizes;
@@ -150,15 +233,19 @@ bool CreateFromSkeletalMeshInternal(UPhysicsAsset* PhysicsAsset, USkeletalMesh*
 			if(ParentIndex != INDEX_NONE)
 			{
 				MergedSizes[ParentIndex] += MyMergedSize;
-				FBoneVertInfo& ParentMergedBones = BoneToMergedBones.FindOrAdd(ParentIndex);	//Add this bone to its parent merged bones
+				TPair<FMergedBoneIndices,FBoneVertInfo>& ParentMergedBones = BoneToMergedBones.FindOrAdd(ParentIndex);	//Add this bone to its parent merged bones
+				ParentMergedBones.Get<0>().Add(BoneIdx);
+
 				const FTransform LocalTM = LocalPose[BoneIdx];
 
-				AddInfoToParentInfo(LocalTM, Infos[BoneIdx], ParentMergedBones);
+				AddInfoToParentInfo(LocalTM, Infos[BoneIdx], ParentMergedBones.Get<1>());
 
-				if(FBoneVertInfo* MyMergedBones = BoneToMergedBones.Find(BoneIdx))
+				if(TPair<FMergedBoneIndices, FBoneVertInfo>* MyMergedBones = BoneToMergedBones.Find(BoneIdx))
 				{
 					//make sure any bones merged to this bone get merged into the parent
-					AddInfoToParentInfo(LocalTM, *MyMergedBones, ParentMergedBones);
+					ParentMergedBones.Get<0>().Append(MyMergedBones->Get<0>());
+					AddInfoToParentInfo(LocalTM, MyMergedBones->Get<1>(), ParentMergedBones.Get<1>());
+
 					BoneToMergedBones.Remove(BoneIdx);
 				}
 			}
@@ -189,33 +276,76 @@ bool CreateFromSkeletalMeshInternal(UPhysicsAsset* PhysicsAsset, USkeletalMesh*
 		}
 	}
 
-	FScopedSlowTask SlowTask((float)NumBones * 2, FText(), IsInGameThread());
-	if (IsInGameThread())
-	{
-		SlowTask.MakeDialog();
-	}
-
-	// Finally, iterate through all the bones and create bodies when needed
-	for (int32 BoneIndex = 0; BoneIndex < NumBones; BoneIndex++)
+	auto ShouldMakeBone = [&Params, &MergedSizes, ForcedRootBoneIndex](int32 BoneIndex)
 	{
-		// Determine if we should create a physics body for this bone
-		bool bMakeBone = false;
-
 		// If desired - make a body for EVERY bone
 		if (Params.bBodyForAll)
 		{
-			bMakeBone = true;
+			return true;
 		}
 		else if (MergedSizes[BoneIndex] > Params.MinBoneSize)
 		{
 			// If bone is big enough - create physics.
-			bMakeBone = true;
+			return true;
 		}
-		else if(BoneIndex == ForcedRootBoneIndex)
+		else if (BoneIndex == ForcedRootBoneIndex)
 		{
 			// If the bone is a forced root body we must create they body no matter how small
-			bMakeBone = true;
+			return true;
+		}
+
+		return false;
+	};
+
+	int32 RootBoneIndex = 0;
+	if (Params.GeomType == EFG_SkinnedLevelSet)
+	{
+		// Calculate common root bone from all merged bones
+		if (!Params.bBodyForAll)
+		{
+			TArray<int32> MergedBones;
+			MergedBones.Reserve(NumBones);
+			for (int32 BoneIndex = 0; BoneIndex < NumBones; ++BoneIndex)
+			{
+				if (ShouldMakeBone(BoneIndex))
+				{
+					MergedBones.Add(BoneIndex);
+				}
+			}
+			RootBoneIndex = CalculateCommonRootBoneIndex(SkelMesh->GetRefSkeleton(), MergedBones);
+		}
+	}
+
+	FSkinnedLevelSetBuilder LatticeBuilder(*SkelMesh, TriangleCache, RootBoneIndex);
+	if (Params.GeomType == EFG_SkinnedLevelSet)
+	{
+		check(RootBoneIndex != INDEX_NONE);
+		TArray<int32> AllBoneIndices;
+		const int32 RawBoneNum = SkelMesh->GetRefSkeleton().GetRawBoneNum();
+		AllBoneIndices.Reserve(RawBoneNum - RootBoneIndex);
+		for (int32 BoneIndex = RootBoneIndex; BoneIndex < RawBoneNum; ++BoneIndex)
+		{
+			AllBoneIndices.Emplace(BoneIndex);
 		}
+		TArray<uint32> OrigIndicesUnused;
+		if (!LatticeBuilder.InitializeSkinnedLevelset(Params, AllBoneIndices, OrigIndicesUnused))
+		{
+			return false;
+		}
+	}
+
+	// Finally, iterate through all the bones and create bodies when needed
+
+	FScopedSlowTask SlowTask((float)NumBones * 2, FText(), IsInGameThread());
+	if (IsInGameThread())
+	{
+		SlowTask.MakeDialog();
+	}
+
+	for (int32 BoneIndex = 0; BoneIndex < NumBones; BoneIndex++)
+	{
+		// Determine if we should create a physics body for this bone
+		const bool bMakeBone = ShouldMakeBone(BoneIndex);
 
 		if (bMakeBone)
 		{
@@ -227,77 +357,97 @@ bool CreateFromSkeletalMeshInternal(UPhysicsAsset* PhysicsAsset, USkeletalMesh*
 				SlowTask.EnterProgressFrame(1.0f, FText::Format(NSLOCTEXT("PhysicsAssetEditor", "ResetCollsionStepInfo", "Generating collision for {0}"), FText::FromName(BoneName)));
 			}
 
-			int32 NewBodyIndex = CreateNewBody(PhysicsAsset, BoneName, Params);
-			UBodySetup* NewBodySetup = PhysicsAsset->SkeletalBodySetups[NewBodyIndex];
-			check(NewBodySetup->BoneName == BoneName);
 
 			//Construct the info - in the case of merged bones we append all the data
 			FBoneVertInfo Info = Infos[BoneIndex];
-			if(const FBoneVertInfo* MergedBones = BoneToMergedBones.Find(BoneIndex))
+			FMergedBoneIndices BoneIndices;
+			BoneIndices.Add(BoneIndex);
+			if (const TPair<FMergedBoneIndices, FBoneVertInfo>* MergedBones = BoneToMergedBones.Find(BoneIndex))
 			{
 				//Don't need to convert into parent space since this was already done
-				Info.Normals.Append(MergedBones->Normals);
-				Info.Positions.Append(MergedBones->Positions);
+				BoneIndices.Append(MergedBones->Get<0>());
+				Info.Normals.Append(MergedBones->Get<1>().Normals);
+				Info.Positions.Append(MergedBones->Get<1>().Positions);
 			}
 
-			// Fill in collision info for this bone.
-			const bool bSuccess = CreateCollisionFromBoneInternal(NewBodySetup, SkelMesh, BoneIndex, Params, Info, TriangleCache);
-			if(bSuccess)
+			if (Params.GeomType == EFG_SkinnedLevelSet)
 			{
-				// create joint to parent body
-				if (Params.bCreateConstraints)
-				{
-					int32 ParentIndex = BoneIndex;
-					int32 ParentBodyIndex = INDEX_NONE;
-					FName ParentName;
+				LatticeBuilder.AddBoneInfluence(BoneIndex, BoneIndices);
+			}
+			else
+			{
+				const int32 NewBodyIndex = CreateNewBody(PhysicsAsset, BoneName, Params);
+				UBodySetup* NewBodySetup = PhysicsAsset->SkeletalBodySetups[NewBodyIndex];
+				check(NewBodySetup->BoneName == BoneName);
 
-					do
+				// Fill in collision info for this bone.
+				const bool bSuccess = CreateCollisionFromBoneInternal(NewBodySetup, SkelMesh, BoneIndex, Params, Info, TriangleCache);
+				if (bSuccess)
+				{
+					// create joint to parent body
+					if (Params.bCreateConstraints)
 					{
-						//Travel up the hierarchy to find a parent which has a valid body
-						ParentIndex = SkelMesh->GetRefSkeleton().GetParentIndex(ParentIndex);
-						if(ParentIndex != INDEX_NONE)
+						int32 ParentIndex = BoneIndex;
+						int32 ParentBodyIndex = INDEX_NONE;
+						FName ParentName;
+
+						do
 						{
-							ParentName = SkelMesh->GetRefSkeleton().GetBoneName(ParentIndex);
-							ParentBodyIndex = PhysicsAsset->FindBodyIndex(ParentName);
-						}
-						else
+							//Travel up the hierarchy to find a parent which has a valid body
+							ParentIndex = SkelMesh->GetRefSkeleton().GetParentIndex(ParentIndex);
+							if (ParentIndex != INDEX_NONE)
+							{
+								ParentName = SkelMesh->GetRefSkeleton().GetBoneName(ParentIndex);
+								ParentBodyIndex = PhysicsAsset->FindBodyIndex(ParentName);
+							}
+							else
+							{
+								//no more parents so just stop
+								break;
+							}
+
+						} while (ParentBodyIndex == INDEX_NONE);
+
+						if (ParentBodyIndex != INDEX_NONE)
 						{
-							//no more parents so just stop
-							break;
-						}
-						
-					}while(ParentBodyIndex == INDEX_NONE);
+							//found valid parent body so create joint
+							int32 NewConstraintIndex = CreateNewConstraint(PhysicsAsset, BoneName);
+							UPhysicsConstraintTemplate* CS = PhysicsAsset->ConstraintSetup[NewConstraintIndex];
 
-					if (ParentBodyIndex != INDEX_NONE)
-					{
-						//found valid parent body so create joint
-						int32 NewConstraintIndex = CreateNewConstraint(PhysicsAsset, BoneName);
-						UPhysicsConstraintTemplate* CS = PhysicsAsset->ConstraintSetup[NewConstraintIndex];
-						
-						// set angular constraint mode
-						CS->DefaultInstance.SetAngularSwing1Motion(Params.AngularConstraintMode);
-						CS->DefaultInstance.SetAngularSwing2Motion(Params.AngularConstraintMode);
-						CS->DefaultInstance.SetAngularTwistMotion(Params.AngularConstraintMode);
-
-						// Place joint at origin of child
-						CS->DefaultInstance.ConstraintBone1 = BoneName;
-						CS->DefaultInstance.ConstraintBone2 = ParentName;
-						CS->DefaultInstance.SnapTransformsToDefault(EConstraintTransformComponentFlags::All, PhysicsAsset);
-
-						CS->SetDefaultProfile(CS->DefaultInstance);
-
-						// Disable collision between constrained bodies by default.
-						PhysicsAsset->DisableCollision(NewBodyIndex, ParentBodyIndex);
+							// set angular constraint mode
+							CS->DefaultInstance.SetAngularSwing1Motion(Params.AngularConstraintMode);
+							CS->DefaultInstance.SetAngularSwing2Motion(Params.AngularConstraintMode);
+							CS->DefaultInstance.SetAngularTwistMotion(Params.AngularConstraintMode);
+
+							// Place joint at origin of child
+							CS->DefaultInstance.ConstraintBone1 = BoneName;
+							CS->DefaultInstance.ConstraintBone2 = ParentName;
+							CS->DefaultInstance.SnapTransformsToDefault(EConstraintTransformComponentFlags::All, PhysicsAsset);
+
+							CS->SetDefaultProfile(CS->DefaultInstance);
+
+							// Disable collision between constrained bodies by default.
+							PhysicsAsset->DisableCollision(NewBodyIndex, ParentBodyIndex);
+						}
 					}
 				}
-			}
-			else
-			{
-				DestroyBody(PhysicsAsset, NewBodyIndex);
+				else
+				{
+					DestroyBody(PhysicsAsset, NewBodyIndex);
+				}
 			}
 		}
 	}
 
+	if (Params.GeomType == EFG_SkinnedLevelSet)
+	{
+		// Finish Building WeightedLattice
+		const FName RootBoneName = SkelMesh->GetRefSkeleton().GetBoneName(RootBoneIndex);
+		const int32 NewBodyIndex = CreateNewBody(PhysicsAsset, RootBoneName, Params);
+		UBodySetup* NewBodySetup = PhysicsAsset->SkeletalBodySetups[NewBodyIndex];
+		CreateWeightedLevelSetBody(LatticeBuilder, NewBodySetup, PhysicsAsset, RootBoneName, Params);
+	}
+
 	//Go through and ensure any overlapping bodies are marked as disable collision
 	FPreviewScene TmpScene;
 	UWorld* TmpWorld = TmpScene.GetWorld();
@@ -348,7 +498,7 @@ bool CreateFromSkeletalMesh(UPhysicsAsset* PhysicsAsset, USkeletalMesh* SkelMesh
 
 	FSkinnedBoneTriangleCache TriangleCache(*SkelMesh, Params);
 
-	if (Params.GeomType == EFG_SingleConvexHull || Params.GeomType == EFG_MultiConvexHull || Params.GeomType == EFG_LevelSet)
+	if (Params.GeomType == EFG_SingleConvexHull || Params.GeomType == EFG_MultiConvexHull || Params.GeomType == EFG_LevelSet || Params.GeomType == EFG_SkinnedLevelSet)
 	{
 		TriangleCache.BuildCache();
 	}
@@ -443,144 +593,6 @@ FVector ComputeEigenVector(const FMatrix& A)
 	return Bk.GetSafeNormal();
 }
 
-
-namespace LevelSetHelpers
-{
-	// Copied from CollisionGeometryConversion.h (Plugins/Runtime/MeshModelingToolset)
-	void CreateLevelSetElement(const FTransform3d& GridTransform, const UE::Geometry::FDenseGrid3f& Grid, float CellSize, FKLevelSetElem& LevelSetOut)
-	{
-		const UE::Geometry::FVector3i& GridDims = Grid.GetDimensions();
-		TArray<double> OutGridValues;
-		OutGridValues.Init(0.0, GridDims[0] * GridDims[1] * GridDims[2]);
-
-		for (int I = 0; I < GridDims[0]; ++I)
-		{
-			for (int J = 0; J < GridDims[1]; ++J)
-			{
-				for (int K = 0; K < GridDims[2]; ++K)
-				{
-					// account for different ordering in Geometry::TDenseGrid3 vs Chaos::TUniformGrid
-					const int InBufferIndex = I + GridDims[0] * (J + GridDims[1] * K);
-					const int OutBufferIndex = K + GridDims[2] * (J + GridDims[1] * I);
-					OutGridValues[OutBufferIndex] = Grid[InBufferIndex];
-				}
-			}
-		}
-
-		FTransform ChaosTransform = GridTransform;
-		ChaosTransform.AddToTranslation(-0.5 * CellSize * FVector::One());		// account for grid origin being at cell center vs cell corner
-		LevelSetOut.BuildLevelSet(ChaosTransform, OutGridValues, FIntVector(GridDims[0], GridDims[1], GridDims[2]), CellSize);
-	}
-
-	// Copied from FMeshSimpleShapeApproximation::Generate_LevelSets (Plugins/Runtime/GeometryProcessing)
-	bool CreateLevelSetForMesh(const UE::Geometry::FDynamicMesh3& InMesh, int32 InLevelSetGridResolution, FKLevelSetElem& OutElement)
-	{
-		using UE::Geometry::FDynamicMesh3;
-		using UE::Geometry::TSweepingMeshSDF;
-
-		constexpr int32 NumNarrowBandCells = 2;
-		constexpr int32 NumExpandCells = 1;
-
-		// Inside SDF.Compute(), extra grid cell are added to account for the narrow band ("NarrowBandMaxDistance") as well as an 
-		// outer buffer ("ExpandBounds"). So here we adjust the cell size to make the final output resolution closer to the user-specified resolution.
-		const int32 LevelSetGridResolution = FMath::Max(1, InLevelSetGridResolution - 2*(NumNarrowBandCells + NumExpandCells));
-
-		const UE::Geometry::FAxisAlignedBox3d Bounds = InMesh.GetBounds();
-		const double CellSize = Bounds.MaxDim() / LevelSetGridResolution;
-		const double ExpandBounds = NumExpandCells * CellSize;
-
-		UE::Geometry::TMeshAABBTree3<FDynamicMesh3> Spatial(&InMesh);
-
-		// Input mesh is likely open, so solidify it first
-		UE::Geometry::TFastWindingTree<FDynamicMesh3> FastWinding(&Spatial);
-		UE::Geometry::TImplicitSolidify<FDynamicMesh3> Solidify(&InMesh, &Spatial, &FastWinding);
-		Solidify.ExtendBounds = ExpandBounds;
-		Solidify.MeshCellSize = CellSize;
-		Solidify.WindingThreshold = 0.5;
-		Solidify.SurfaceSearchSteps = 3;
-		Solidify.bSolidAtBoundaries = true;
-
-		const UE::Geometry::FDynamicMesh3 SolidMesh = FDynamicMesh3(&Solidify.Generate());
-		Spatial.SetMesh(&SolidMesh, true);
-
-		TSweepingMeshSDF<FDynamicMesh3> SDF;
-		SDF.Mesh = &SolidMesh;
-		SDF.Spatial = &Spatial;
-		SDF.ComputeMode = TSweepingMeshSDF<FDynamicMesh3>::EComputeModes::FullGrid;
-		SDF.CellSize = (float)CellSize;
-		SDF.NarrowBandMaxDistance = NumNarrowBandCells * CellSize;
-		SDF.ExactBandWidth = FMath::CeilToInt32(SDF.NarrowBandMaxDistance / CellSize);
-		SDF.ExpandBounds = FVector3d(ExpandBounds);
-
-		if (SDF.Compute(Bounds))
-		{
-			const FTransform3d GridTransform = FTransform((FVector3d)SDF.GridOrigin);
-			CreateLevelSetElement(GridTransform, SDF.Grid, SDF.CellSize, OutElement);
-			return true;
-		}
-
-		return false;
-	}
-
-	void CreateDynamicMesh(const TArray<FVector3f>& InVertices, const TArray<uint32>& InIndices, UE::Geometry::FDynamicMesh3& OutMesh)
-	{
-		OutMesh.Clear();
-		for (const FVector3f& Vertex : InVertices)
-		{
-			OutMesh.AppendVertex(FVector3d(Vertex));
-		}
-
-		check(InIndices.Num() % 3 == 0);
-		const int32 NumTriangles = InIndices.Num() / 3;
-		for (int32 TriangleID = 0; TriangleID < NumTriangles; ++TriangleID)
-		{
-			const UE::Geometry::FIndex3i Triangle(InIndices[3 * TriangleID], InIndices[3 * TriangleID + 1], InIndices[3 * TriangleID + 2]);
-			OutMesh.AppendTriangle(Triangle);
-		}
-	}
-
-	bool CreateLevelSetForBone(UBodySetup* BodySetup, const TArray<FVector3f>& InVertices, const TArray<uint32>& InIndices, uint32 InResolution)
-	{
-		check(BodySetup != NULL);
-
-		// Validate input by checking bounding box
-		FBox VertBox(ForceInit);
-		for (const FVector3f& Vert : InVertices)
-		{
-			VertBox += (FVector)Vert;
-		}
-
-		// If box is invalid, or the largest dimension is less than 1 unit, or smallest is less than 0.1, skip trying to generate collision
-		if (VertBox.IsValid == 0 || VertBox.GetSize().GetMax() < 1.f || VertBox.GetSize().GetMin() < 0.1f || InIndices.Num() == 0)
-		{
-			return false;
-		}
-
-		TRACE_CPUPROFILER_EVENT_SCOPE(MeshToLevelSet)
-
-		// Clean out old hulls
-		BodySetup->RemoveSimpleCollision();
-
-		UE::Geometry::FDynamicMesh3 DynamicMesh;
-		CreateDynamicMesh(InVertices, InIndices, DynamicMesh);
-
-		FKLevelSetElem LevelSetElement;
-		const bool bOK = CreateLevelSetForMesh(DynamicMesh, InResolution, LevelSetElement);
-
-		if (!bOK)
-		{
-			return false; 
-		}
-
-		BodySetup->AggGeom.LevelSetElems.Add(LevelSetElement);
-		BodySetup->InvalidatePhysicsData(); // update GUID
-		return true;
-	}
-
-}
-
-
-
 bool CreateCollisionFromBoneInternal(UBodySetup* bs, USkeletalMesh* skelMesh, int32 BoneIndex, const FPhysAssetCreateParams& Params, const FBoneVertInfo& Info, const FSkinnedBoneTriangleCache& TriangleCache)
 {
 #if WITH_EDITOR
@@ -812,6 +824,191 @@ bool CreateCollisionFromBones(UBodySetup* bs, USkeletalMesh* skelMesh, const TAr
 	return bAllSuccessful;
 }
 
+bool CreateCollisionsFromBones(UPhysicsAsset* PhysicsAsset, USkeletalMesh* SkelMesh, const TArray<int32>& BodyIndices, const FPhysAssetCreateParams& Params, const TArray<FBoneVertInfo>& Info, TArray<int32>& OutSuccessfulBodyIndices)
+{
+	check(SkelMesh);
+	check(PhysicsAsset);
+
+	FSkinnedBoneTriangleCache TriangleCache(*SkelMesh, Params);
+	if (Params.GeomType == EFG_SingleConvexHull || Params.GeomType == EFG_MultiConvexHull || Params.GeomType == EFG_LevelSet || Params.GeomType == EFG_SkinnedLevelSet)
+	{
+		TriangleCache.BuildCache();
+	}
+
+	OutSuccessfulBodyIndices.Reset();
+	bool bAllSuccessful = true;
+
+	// Destroying bodies causes the indices to change. Use names instead.
+	TArray<FName> BodySetupBoneNames;
+	BodySetupBoneNames.Reserve(BodyIndices.Num());
+	TArray<FName> SubBoneNames;
+	for (int32 Index = 0; Index < BodyIndices.Num(); ++Index)
+	{
+		const UBodySetup* BodySetup = PhysicsAsset->SkeletalBodySetups[BodyIndices[Index]];
+		BodySetupBoneNames.Add(BodySetup->BoneName);
+
+		// Get sub-bones for any skinned levelsets
+		for (const FKSkinnedLevelSetElem& SkinnedLevelSetElem : BodySetup->AggGeom.SkinnedLevelSetElems)
+		{
+			if (const Chaos::TWeightedLatticeImplicitObject<Chaos::FLevelSet>* SkinnedLevelSet = SkinnedLevelSetElem.GetWeightedLevelSet().Get())
+			{
+				SubBoneNames.Reserve(SubBoneNames.Num() + SkinnedLevelSet->GetUsedBones().Num());
+				for (const FName& SubBoneName : SkinnedLevelSet->GetUsedBones())
+				{
+					SubBoneNames.AddUnique(SubBoneName);
+				}
+			}
+		}
+	}
+
+	int32 RootBoneIndex = INDEX_NONE;
+	TArray<int32> AllBoneIndices;
+	if (Params.GeomType == EFG_SkinnedLevelSet)
+	{
+		AllBoneIndices.Reserve(BodySetupBoneNames.Num());
+		for (int32 Index = 0; Index < BodySetupBoneNames.Num(); ++Index)
+		{
+			const int32 BoneIndex = SkelMesh->GetRefSkeleton().FindBoneIndex(BodySetupBoneNames[Index]);
+			AllBoneIndices.Add(BoneIndex);
+		}
+		RootBoneIndex = CalculateCommonRootBoneIndex(SkelMesh->GetRefSkeleton(), AllBoneIndices);
+	}
+
+	FSkinnedLevelSetBuilder LatticeBuilder(*SkelMesh, TriangleCache, RootBoneIndex);
+	if (Params.GeomType == EFG_SkinnedLevelSet)
+	{
+		TArray<uint32> OrigIndices;
+		if (!LatticeBuilder.InitializeSkinnedLevelset(Params, AllBoneIndices, OrigIndices))
+		{
+			return false;
+		}
+
+		if (Params.VertWeight == EVW_AnyWeight)
+		{
+			// Add additional sub-bones that contribute to the vertices
+			TSet<int32> SubBoneSet;
+			LatticeBuilder.GetInfluencingBones(OrigIndices, SubBoneSet);
+
+			for (int32 AddedSubBone : SubBoneSet)
+			{
+				SubBoneNames.AddUnique(SkelMesh->GetRefSkeleton().GetBoneName(AddedSubBone));
+			}
+		}
+	}
+
+	FScopedSlowTask SlowTask((float)BodySetupBoneNames.Num() + (float)SubBoneNames.Num());
+	if (IsInGameThread())
+	{
+		SlowTask.MakeDialog();
+	}
+
+
+	TArray<FName> SuccessfulBodySetupNames;
+	for (const FName& BoneName : BodySetupBoneNames)
+	{
+		if (IsInGameThread())
+		{
+			SlowTask.EnterProgressFrame(1.0f, FText::Format(NSLOCTEXT("PhysicsAssetEditor", "ResetCollsionStepInfo", "Generating collision for {0}"), FText::FromName(BoneName)));
+		}
+
+		const int32 BodyIndex = PhysicsAsset->FindBodyIndex(BoneName);
+		check(BodyIndex != INDEX_NONE);
+
+		UBodySetup* BodySetup = PhysicsAsset->SkeletalBodySetups[BodyIndex];
+		BodySetup->Modify();
+		check(BodySetup);
+		const int32 BoneIndex = SkelMesh->GetRefSkeleton().FindBoneIndex(BodySetup->BoneName);
+		check(BoneIndex != INDEX_NONE);
+
+		if (Params.GeomType == EFG_SkinnedLevelSet)
+		{
+			if (BoneIndex != RootBoneIndex)
+			{
+				// This bone is now getting merged into the RootBoneIndex's body
+				DestroyBody(PhysicsAsset, BodyIndex);
+			}
+
+			TArray<int32> BoneIndexArray;
+			BoneIndexArray.Add(BoneIndex);
+			LatticeBuilder.AddBoneInfluence(BoneIndex, BoneIndexArray);
+		}
+		else
+		{
+			if (CreateCollisionFromBoneInternal(BodySetup, SkelMesh, BoneIndex, Params, Info[BoneIndex], TriangleCache))
+			{
+				SuccessfulBodySetupNames.AddUnique(BoneName);
+			}
+			else
+			{
+				bAllSuccessful = false;
+				DestroyBody(PhysicsAsset, BodyIndex);
+			}
+		}
+	}
+
+	for (const FName& BoneName : SubBoneNames)
+	{
+		if (IsInGameThread())
+		{
+			SlowTask.EnterProgressFrame(1.0f, FText::Format(NSLOCTEXT("PhysicsAssetEditor", "ResetCollsionStepInfo", "Generating collision for {0}"), FText::FromName(BoneName)));
+		}
+
+		if (BodySetupBoneNames.Find(BoneName) != INDEX_NONE)
+		{
+			// We've already dealt with this above.
+			continue;
+		}
+
+		if (Params.GeomType == EFG_SkinnedLevelSet)
+		{
+			const int32 BoneIndex = SkelMesh->GetRefSkeleton().FindBoneIndex(BoneName);
+			TArray<int32> BoneIndexArray;
+			BoneIndexArray.Add(BoneIndex);
+			LatticeBuilder.AddBoneInfluence(BoneIndex, BoneIndexArray);
+		}
+		else
+		{
+			if (PhysicsAsset->FindBodyIndex(BoneName) != INDEX_NONE)
+			{
+				// We already have an existing bone. Don't replace it.
+				continue;
+			}
+			const int32 NewBodyIndex = CreateNewBody(PhysicsAsset, BoneName, Params);
+			UBodySetup* NewBodySetup = PhysicsAsset->SkeletalBodySetups[NewBodyIndex];
+			check(NewBodySetup->BoneName == BoneName);
+			const int32 BoneIndex = SkelMesh->GetRefSkeleton().FindBoneIndex(NewBodySetup->BoneName);
+			if (CreateCollisionFromBoneInternal(NewBodySetup, SkelMesh, BoneIndex, Params, Info[BoneIndex], TriangleCache))
+			{
+				SuccessfulBodySetupNames.AddUnique(BoneName);
+			}
+			else
+			{
+				bAllSuccessful = false;
+				DestroyBody(PhysicsAsset, NewBodyIndex);
+			}
+		}
+	}
+
+	OutSuccessfulBodyIndices.Reserve(SuccessfulBodySetupNames.Num());
+	for (const FName& BoneName : SuccessfulBodySetupNames)
+	{
+		const int32 BodyIndex = PhysicsAsset->FindBodyIndex(BoneName);
+		check(BodyIndex != INDEX_NONE);
+		OutSuccessfulBodyIndices.Add(BodyIndex);
+	}
+
+	if (Params.GeomType == EFG_SkinnedLevelSet)
+	{
+		// Finish Building WeightedLattice
+		const FName RootBoneName = SkelMesh->GetRefSkeleton().GetBoneName(RootBoneIndex);
+		const int32 BodyIndex = CreateNewBody(PhysicsAsset, RootBoneName, Params);
+		UBodySetup* BodySetup = PhysicsAsset->SkeletalBodySetups[BodyIndex];
+		CreateWeightedLevelSetBody(LatticeBuilder, BodySetup, PhysicsAsset, RootBoneName, Params);
+		OutSuccessfulBodyIndices.Add(BodyIndex);
+	}
+	return bAllSuccessful;
+}
+
 void WeldBodies(UPhysicsAsset* PhysAsset, int32 BaseBodyIndex, int32 AddBodyIndex, USkeletalMeshComponent* SkelComp)
 {
 	if(BaseBodyIndex == INDEX_NONE || AddBodyIndex == INDEX_NONE)
diff --git a/Engine/Source/Developer/PhysicsUtilities/Private/SkinnedBoneTriangleCache.cpp b/Engine/Source/Developer/PhysicsUtilities/Private/SkinnedBoneTriangleCache.cpp
--- a/Engine/Source/Developer/PhysicsUtilities/Private/SkinnedBoneTriangleCache.cpp
+++ b/Engine/Source/Developer/PhysicsUtilities/Private/SkinnedBoneTriangleCache.cpp
@@ -5,7 +5,7 @@
 #include "Rendering/SkeletalMeshModel.h"
 #include "Rendering/SkeletalMeshRenderData.h"
 
-FSkinnedBoneTriangleCache::FSkinnedBoneTriangleCache(USkeletalMesh& InSkeletalMesh, const FPhysAssetCreateParams& Params)
+FSkinnedBoneTriangleCache::FSkinnedBoneTriangleCache(const USkeletalMesh& InSkeletalMesh, const FPhysAssetCreateParams& Params)
 	: SkeletalMesh(InSkeletalMesh),
 	StaticLODModel(*SkeletalMesh.GetImportedModel()),
 	RenderData(*SkeletalMesh.GetResourceForRendering()),
@@ -24,19 +24,16 @@ void FSkinnedBoneTriangleCache::BuildCache()
 	BuildOwnedTrianglesSetForEachBone();
 }
 
-void FSkinnedBoneTriangleCache::GetVerticesAndIndicesForBone(const int32 BoneIndex, TArray<FVector3f>& OutVertexPositions, TArray<uint32>& OutIndices) const
+void FSkinnedBoneTriangleCache::GetVerticesAndIndicesForBoneInternal(const int32 TrianglesBoneIndex, const int32 TransformBoneIndex, TArray<FVector3f>& OutVertexPositions, TArray<uint32>& OutIndices, TArray<uint32>& OutLocalToSkinnedVertexIndex) const
 {
-	OutVertexPositions.Empty();
-	OutIndices.Empty();
-
-	const FTriangleArray* TriangleArrayPointer = BoneIndexToTriangles.Find(BoneIndex);
+	const FTriangleArray* TriangleArrayPointer = BoneIndexToTriangles.Find(TrianglesBoneIndex);
 	if ( !TriangleArrayPointer )
 	{
 		return;
 	}
 
 	const FTriangleArray& TrianglesForBone = *TriangleArrayPointer;
-	const FMatrix ComponentToBoneMatrix = BoneTransformMatrix(BoneIndex);
+	const FMatrix ComponentToBoneMatrix = BoneTransformMatrix(TransformBoneIndex);
 
 	TMap<FSkinnedVertexIndex, uint32> SkinnedVertIndexToOutputIndex;
 
@@ -55,6 +52,7 @@ void FSkinnedBoneTriangleCache::GetVerticesAndIndicesForBone(const int32 BoneInd
 			{
 				OutVertexPositions.Add(VertexPosition(VertIndex, ComponentToBoneMatrix));
 				SkinnedVertIndexToOutputIndex.Add(VertIndex, static_cast<uint32>(OutVertexPositions.Num() - 1));
+				OutLocalToSkinnedVertexIndex.Add(VertIndex);
 			}
 
 			// Map the skinned vert index to the index in our output array.
@@ -63,6 +61,25 @@ void FSkinnedBoneTriangleCache::GetVerticesAndIndicesForBone(const int32 BoneInd
 	}
 }
 
+void FSkinnedBoneTriangleCache::GetVerticesAndIndicesForBone(const int32 BoneIndex, TArray<FVector3f>& OutVertexPositions, TArray<uint32>& OutIndices) const
+{
+	OutVertexPositions.Reset();
+	OutIndices.Reset();
+	TArray<uint32> UnusedLocalToSkinnedVertexIndex;
+	return GetVerticesAndIndicesForBoneInternal(BoneIndex, BoneIndex, OutVertexPositions, OutIndices, UnusedLocalToSkinnedVertexIndex);
+}
+
+void FSkinnedBoneTriangleCache::GetVerticesAndIndicesForBones(const int32 PrimaryBoneIndex, const TArray<int32>& BoneIndices, TArray<FVector3f>& OutVertexPositions, TArray<uint32>& OutIndices, TArray<uint32>& OutLocalToSkinnedVertexIndex) const
+{
+	OutVertexPositions.Reset();
+	OutIndices.Reset();
+	OutLocalToSkinnedVertexIndex.Reset();
+	for (const int32 BoneIndex : BoneIndices)
+	{
+		GetVerticesAndIndicesForBoneInternal(BoneIndex, PrimaryBoneIndex, OutVertexPositions, OutIndices, OutLocalToSkinnedVertexIndex);
+	}
+}
+
 void FSkinnedBoneTriangleCache::BuildInfluencedIndexSetForEachBone()
 {
 	const uint32 TotalVertices = VertexBuffer.GetNumVertices();
diff --git a/Engine/Source/Developer/PhysicsUtilities/Private/SkinnedBoneTriangleCache.h b/Engine/Source/Developer/PhysicsUtilities/Private/SkinnedBoneTriangleCache.h
--- a/Engine/Source/Developer/PhysicsUtilities/Private/SkinnedBoneTriangleCache.h
+++ b/Engine/Source/Developer/PhysicsUtilities/Private/SkinnedBoneTriangleCache.h
@@ -30,10 +30,12 @@
 class FSkinnedBoneTriangleCache
 {
 public:
-	explicit FSkinnedBoneTriangleCache(USkeletalMesh& InSkeletalMesh, const FPhysAssetCreateParams& Params);
+	explicit FSkinnedBoneTriangleCache(const USkeletalMesh& InSkeletalMesh, const FPhysAssetCreateParams& Params);
 
 	void BuildCache();
 	void GetVerticesAndIndicesForBone(const int32 BoneIndex, TArray<FVector3f>& OutVertexPositions, TArray<uint32>& OutIndices) const;
+	// Vertex positions will be transformed by PrimaryBoneIndex's transform matrix.
+	void GetVerticesAndIndicesForBones(const int32 PrimaryBoneIndex, const TArray<int32>& BoneIndices, TArray<FVector3f>& OutVertexPositions, TArray<uint32>& OutIndices, TArray<uint32>& OutLocalToSkinnedVertexIndex) const;
 
 private:
 	typedef int32 FBoneIndex;
@@ -72,12 +74,14 @@ private:
 
 	FVector3f VertexPosition(const FSkinnedVertexIndex VertIndex, const FMatrix& ComponentToBoneMatrix) const;
 
+	void GetVerticesAndIndicesForBoneInternal(const int32 TrianglesBoneIndex, const int32 TransformBoneIndex, TArray<FVector3f>& OutVertexPositions, TArray<uint32>& OutIndices, TArray<uint32>& OutLocalToSkinnedVertexIndex) const;
+
 	// Inputs
-	USkeletalMesh& SkeletalMesh;
+	const USkeletalMesh& SkeletalMesh;
 
 	// Computed from inputs
-	class FSkeletalMeshModel& StaticLODModel;
-	class FSkeletalMeshRenderData& RenderData;
+	const class FSkeletalMeshModel& StaticLODModel;
+	const class FSkeletalMeshRenderData& RenderData;
 	const class FPositionVertexBuffer& VertexBuffer;
 
 	/**
diff --git a/Engine/Source/Developer/PhysicsUtilities/Private/SkinnedLevelSetBuilder.cpp b/Engine/Source/Developer/PhysicsUtilities/Private/SkinnedLevelSetBuilder.cpp
new file mode 100644
--- /dev/null
+++ b/Engine/Source/Developer/PhysicsUtilities/Private/SkinnedLevelSetBuilder.cpp
@@ -0,0 +1,303 @@
+// Copyright Epic Games, Inc. All Rights Reserved.
+
+#include "SkinnedLevelSetBuilder.h"
+#include "SkinnedBoneTriangleCache.h"
+#include "LevelSetHelpers.h"
+
+#include "BoneWeights.h"
+#include "Chaos/ArrayND.h"
+#include "Chaos/Levelset.h"
+#include "Chaos/UniformGrid.h"
+#include "Chaos/Math/Poisson.h"
+#include "DynamicMesh/DynamicMesh3.h"
+#include "Engine/SkeletalMesh.h"
+#include "Implicit/SweepingMeshSDF.h"
+#include "Implicit/Solidify.h"
+#include "PhysicsEngine/BodySetup.h"
+#include "PhysicsEngine/SkinnedLevelSetElem.h"
+#include "Rendering/SkeletalMeshModel.h"
+#include "Rendering/SkeletalMeshRenderData.h"
+
+FSkinnedLevelSetBuilder::FSkinnedLevelSetBuilder(const USkeletalMesh& InSkeletalMesh, const FSkinnedBoneTriangleCache& InTriangleCache, const int32 InRootBoneIndex)
+	:SkeletalMesh(InSkeletalMesh)
+	, TriangleCache(InTriangleCache)
+	, RootBoneIndex(InRootBoneIndex)
+	, StaticLODModel(SkeletalMesh.GetImportedModel())
+	, RenderData(SkeletalMesh.GetResourceForRendering())
+{
+	check(StaticLODModel);
+	check(RenderData);
+}
+
+bool FSkinnedLevelSetBuilder::InitializeSkinnedLevelset(const FPhysAssetCreateParams& Params, const TArray<int32>& BoneIndices, TArray<uint32>& OutOrigIndices)
+{
+	check(SkeletalMesh.GetRefSkeleton().IsValidRawIndex(RootBoneIndex));
+
+	TArray<FVector3f> Positions;
+	TArray<uint32> Indices;
+	if (Params.VertWeight == EVW_AnyWeight)
+	{
+		// Want to use EVW_DominantWeight for initial levelset
+		// Build a different cache
+		FPhysAssetCreateParams ParamsCopy;
+		ParamsCopy.VertWeight = EVW_DominantWeight;
+		FSkinnedBoneTriangleCache DominantCache(SkeletalMesh, ParamsCopy);
+		DominantCache.BuildCache();
+		DominantCache.GetVerticesAndIndicesForBones(RootBoneIndex, BoneIndices, Positions, Indices, OutOrigIndices);
+	}
+	else
+	{
+
+		TriangleCache.GetVerticesAndIndicesForBones(RootBoneIndex, BoneIndices, Positions, Indices, OutOrigIndices);
+	}
+
+	if (Positions.Num())
+	{
+		UE::Geometry::FDynamicMesh3 DynamicMesh;
+		LevelSetHelpers::CreateDynamicMesh(Positions, Indices, DynamicMesh);
+
+		const bool bOK = LevelSetHelpers::CreateLevelSetForMesh(DynamicMesh, Params.LevelSetResolution, LevelSet);
+		if (bOK)
+		{
+			GenerateGrid(Params.LatticeResolution, LevelSet->BoundingBox());
+		}
+		else
+		{
+			FMessageLog EditorErrors("EditorErrors");
+			EditorErrors.Warning(NSLOCTEXT("PhysicsAssetUtils", "SkinnedLevelSetError", "An error occurred creating root skinned level set."));
+			EditorErrors.Open();
+		}
+		return bOK;
+	}
+	FMessageLog EditorErrors("EditorErrors");
+	EditorErrors.Warning(NSLOCTEXT("PhysicsAssetUtils", "LevelSetNoPositions", "Unable to create a level set for the given bone as there are no vertices associated with the bone."));
+	EditorErrors.Open();
+	return false;
+}
+
+void FSkinnedLevelSetBuilder::AddBoneInfluence(int32 PrimaryBoneIndex, const TArray<int32>& AllBonesForInfluence)
+{
+	// Need to generate weights for grid points corresponding with this bone.
+	// Will laplacian diffuse the weights from surface to bone interior.
+
+	// Get Triangles that correspond with these bones.
+	TArray<FVector3f> AllVerts;
+	TArray<uint32> AllIndices;
+	TArray<uint32> AllOrigIndices;
+	TriangleCache.GetVerticesAndIndicesForBones(RootBoneIndex, AllBonesForInfluence, AllVerts, AllIndices, AllOrigIndices);
+
+	// Strip any vertices that are outside the levelset.
+	TArray<FVector3f> Verts;
+	Verts.Reserve(AllVerts.Num());
+	TArray<uint32> Indices;
+	Indices.Reserve(AllIndices.Num());
+	TArray<uint32> OrigIndices;
+	OrigIndices.Reserve(AllOrigIndices.Num());
+	TArray<int32> AllVertToVerts;
+	AllVertToVerts.SetNumUninitialized(AllVerts.Num());
+	const double MaxInsidePhi = LevelSet->GetGrid().Dx().GetMax() * .5;
+
+	for (int32 AllVertIdx = 0; AllVertIdx < AllVerts.Num(); ++AllVertIdx)
+	{
+		if (LevelSet->SignedDistance(Chaos::FVec3(AllVerts[AllVertIdx])) < MaxInsidePhi)
+		{
+			const int32 VertIdx = Verts.Add(AllVerts[AllVertIdx]);
+			OrigIndices.Add(AllOrigIndices[AllVertIdx]);
+			AllVertToVerts[AllVertIdx] = VertIdx;
+		}
+		else
+		{
+			AllVertToVerts[AllVertIdx] = INDEX_NONE;
+		}
+	}
+	check(AllIndices.Num() % 3 == 0);
+	const int32 NumTriangles = AllIndices.Num() / 3;
+	for (int32 AllTriIdx = 0; AllTriIdx < NumTriangles; ++AllTriIdx)
+	{
+		const int32 VertIdx0 = AllVertToVerts[AllIndices[AllTriIdx * 3]];
+		const int32 VertIdx1 = AllVertToVerts[AllIndices[AllTriIdx * 3 + 1]];
+		const int32 VertIdx2 = AllVertToVerts[AllIndices[AllTriIdx * 3 + 2]];
+		if (VertIdx0 != INDEX_NONE && VertIdx1 != INDEX_NONE && VertIdx2 != INDEX_NONE)
+		{
+			Indices.Add(VertIdx0);
+			Indices.Add(VertIdx1);
+			Indices.Add(VertIdx2);
+		}
+	}
+	
+	const Chaos::TUniformGrid<double, 3>& LatticeGrid = GetGrid();
+	const Chaos::TUniformGrid<double, 3>& LevelsetGrid = LevelSet->GetGrid();
+	// Surface is almost certainly not closed, but we want to diffuse differently inside vs outside, so we
+	// need a closed surface. Using solidify to do this.
+	UE::Geometry::FDynamicMesh3 DynamicMesh;
+	LevelSetHelpers::CreateDynamicMesh(Verts, Indices, DynamicMesh);
+	UE::Geometry::TMeshAABBTree3<UE::Geometry::FDynamicMesh3> Spatial(&DynamicMesh);
+	UE::Geometry::TFastWindingTree<UE::Geometry::FDynamicMesh3> FastWinding(&Spatial);
+	UE::Geometry::TImplicitSolidify<UE::Geometry::FDynamicMesh3> Solidify(&DynamicMesh, &Spatial, &FastWinding);
+	constexpr int32 NumExpandCells = 1;
+	const double SolidifyCellSize = LevelsetGrid.Dx().Min();
+	Solidify.ExtendBounds = NumExpandCells * SolidifyCellSize;
+	Solidify.MeshCellSize = SolidifyCellSize;
+	Solidify.WindingThreshold = 0.5;
+	Solidify.SurfaceSearchSteps = 3;
+	Solidify.bSolidAtBoundaries = true;
+	const UE::Geometry::FDynamicMesh3 SolidMesh(&Solidify.Generate());
+	UE::Geometry::TMeshAABBTree3<UE::Geometry::FDynamicMesh3> SolidSpatial(&SolidMesh);
+	UE::Geometry::TFastWindingTree<UE::Geometry::FDynamicMesh3> SolidFastWinding(&SolidSpatial);
+	constexpr int32 OuterWeightExpandCells = 3; // Diffuse out this far outside mesh
+	constexpr int32 SurfaceExpandCells = 1; // Consider points this far from the surface of the mesh to be on the surface (dirichlet boundary)
+	UE::Geometry::IMeshSpatial::FQueryOptions QueryOptions;
+	QueryOptions.MaxDistance = (double)LatticeGrid.Dx().Max() * FMath::Max(OuterWeightExpandCells, SurfaceExpandCells);// Used for initial query against solidified mesh
+
+	const double OuterWeightExpandDistSq = FMath::Square((double)LatticeGrid.Dx().Max() * OuterWeightExpandCells);
+	const double SurfaceExpandDistSq = FMath::Square((double)LatticeGrid.Dx().Max() * SurfaceExpandCells);
+	TSet<int32> BonesSet(AllBonesForInfluence);
+	// Create sub-grid for the poisson solve.
+	UE::Geometry::FAxisAlignedBox3d SolidBBox = SolidSpatial.GetBoundingBox();
+	const Chaos::TVec3<int32> MinCornerCell = (LatticeGrid.Cell(SolidBBox.Min) - Chaos::TVec3<int32>(OuterWeightExpandCells)).ComponentwiseMax(Chaos::TVec3<int32>(0));
+	const Chaos::TVec3<int32> MaxCornerCell = (LatticeGrid.Cell(SolidBBox.Max) + Chaos::TVec3<int32>(OuterWeightExpandCells)).ComponentwiseMin(LatticeGrid.Counts() - Chaos::TVec3<int32>(1));
+	const Chaos::TUniformGrid<double, 3> SubGrid = LatticeGrid.SubGrid(MinCornerCell, MaxCornerCell);
+	// Get constrained nodes and weights for poisson solve
+	TArray<int32> ConstrainedNodes;
+	TArray<float> ConstrainedWeights;
+	const Chaos::TVec3<int32> NodeCounts = SubGrid.NodeCounts();
+	TSet<int32> OutsideWeighted;
+	for (int32 I = 0; I < NodeCounts.X; ++I)
+	{
+		for (int32 J = 0; J < NodeCounts.Y; ++J)
+		{
+			for (int32 K = 0; K < NodeCounts.Z; ++K)
+			{
+				const Chaos::TVector<int32, 3> Index(I, J, K);
+				const Chaos::TVector<double, 3> Location = SubGrid.Node(Index);
+				// Determine if close to surface of solid mesh.
+				double NearestDistSq = std::numeric_limits<double>::max();
+				const int32 NearestSolidTriangle = SolidSpatial.FindNearestTriangle(Location, NearestDistSq, QueryOptions);
+				if (NearestSolidTriangle == IndexConstants::InvalidID || NearestDistSq > SurfaceExpandCells)
+				{
+					// Not near surface.
+					// Determine if inside or outside
+					const bool IsInside = SolidFastWinding.IsInside(Location);
+					const bool IsGridBoundary = I == 0 || I == NodeCounts.X - 1 || J == 0 || J == NodeCounts.Y - 1 || K == 0 || K == NodeCounts.Z - 1;
+					if (!IsInside)
+					{
+						const bool IsOutsideOuterExpand = NearestSolidTriangle == IndexConstants::InvalidID || NearestDistSq > OuterWeightExpandDistSq;
+						if (IsOutsideOuterExpand)
+						{
+							// Outside the distance to diffuse this weight. Mark as dirichlet with weight = 0
+							ConstrainedNodes.Add(SubGrid.FlatIndex(Index, true));
+							ConstrainedWeights.Add(0.f);
+							continue;
+						}
+						else if (!IsGridBoundary)
+						{
+							OutsideWeighted.Add(SubGrid.FlatIndex(Index, true));
+							continue;
+						}
+					}
+					else if (!IsGridBoundary)
+					{
+						continue;
+					}
+				}
+				// Need to find closest triangle on original mesh to get weights. No distance limit since we might be in a hole that was filled
+				// by solidify. 
+				const int32 NearestTriangle = Spatial.FindNearestTriangle(Location, NearestDistSq);
+				// Get Barycentric coordinates for closest point
+				UE::Geometry::FTriangle3d Triangle;
+				DynamicMesh.GetTriVertices(NearestTriangle, Triangle.V[0], Triangle.V[1], Triangle.V[2]);
+				UE::Geometry::TDistPoint3Triangle3<double> TriQuery(Location, Triangle);
+				TriQuery.ComputeResult();
+				const UE::Geometry::FIndex3i& TriangleIndices = DynamicMesh.GetTriangle(NearestTriangle);
+				const float Weight = TriQuery.TriangleBaryCoords.X * GetWeightForIndices(BonesSet, OrigIndices[TriangleIndices.A]) +
+					TriQuery.TriangleBaryCoords.Y * GetWeightForIndices(BonesSet, OrigIndices[TriangleIndices.B]) + TriQuery.TriangleBaryCoords.Z * GetWeightForIndices(BonesSet, OrigIndices[TriangleIndices.C]);
+				ConstrainedNodes.Add(SubGrid.FlatIndex(Index, true));
+				ConstrainedWeights.Add(Weight);
+			}
+		}
+	}
+	TArray<float> Weights;
+	Weights.SetNumUninitialized(SubGrid.GetNumNodes());
+	constexpr int32 MaxIter = 10000;
+	constexpr float Res = 1e-4;
+	constexpr bool bCheckResidual = true;
+	constexpr int32 MinParallelBatchSize = 10000;
+	Chaos::PoissonSolve<float, double, true>(ConstrainedNodes, ConstrainedWeights, SubGrid, MaxIter, Res, Weights, bCheckResidual, MinParallelBatchSize);
+	for (int32 Idx = 0; Idx < Weights.Num(); ++Idx)
+	{
+		if (Weights[Idx] > UE_SMALL_NUMBER)
+		{
+			Chaos::TVec3<int32> SubIndex;
+			SubGrid.FlatToMultiIndex(Idx, SubIndex, true);
+			const int32 FullIndex = LatticeGrid.FlatIndex(SubIndex + MinCornerCell, true);
+			AddInfluence(FullIndex, PrimaryBoneIndex, Weights[Idx], OutsideWeighted.Contains(Idx));
+		}
+	}
+}
+
+FKSkinnedLevelSetElem FSkinnedLevelSetBuilder::CreateSkinnedLevelSetElem()
+{
+	FinalizeInfluences([this](int32 BoneIndex)
+		{
+			return SkeletalMesh.GetRefSkeleton().GetBoneName(BoneIndex);
+		},
+		[this](int32 BoneIndex)
+		{
+			if (BoneIndex == INDEX_NONE)
+			{
+				return FTransform(FMatrix(SkeletalMesh.GetRefBasesInvMatrix()[RootBoneIndex]));
+			}
+		return FTransform(FMatrix(SkeletalMesh.GetRefBasesInvMatrix()[BoneIndex]));
+		}
+		);
+	TUniquePtr<Chaos::TWeightedLatticeImplicitObject<Chaos::FLevelSet> > WeightedLevelSet = Generate(MoveTemp(LevelSet));
+	FKSkinnedLevelSetElem Ret;
+	Ret.SetWeightedLevelSet(MoveTemp(WeightedLevelSet));
+	return Ret;
+}
+
+float FSkinnedLevelSetBuilder::GetWeightForIndices(const TSet<int32>& BoneIndices, uint32 VertIndex) const
+{
+	int32 SectionIndex;
+	int32 SoftVertIndex;
+	RenderData->LODRenderData[0].GetSectionFromVertexIndex(VertIndex, SectionIndex, SoftVertIndex);
+	const FSkelMeshSection& Section = StaticLODModel->LODModels[0].Sections[SectionIndex];
+	const FSoftSkinVertex& SoftVert = Section.SoftVertices[SoftVertIndex];
+
+	float Weight = 0.f;
+	for (int32 InfluenceIndex = 0; InfluenceIndex < MAX_TOTAL_INFLUENCES; ++InfluenceIndex)
+	{
+		const FBoneIndexType BoneMapIndex = SoftVert.InfluenceBones[InfluenceIndex];
+		const int32 ActualBoneIndex = Section.BoneMap[BoneMapIndex];
+		if (BoneIndices.Contains(ActualBoneIndex))
+		{
+			// Use max weight for all bones.
+			Weight = FMath::Max(Weight, (float)SoftVert.InfluenceWeights[InfluenceIndex] / UE::AnimationCore::MaxRawBoneWeightFloat);
+		}
+	}
+
+	return Weight;
+}
+
+void FSkinnedLevelSetBuilder::GetInfluencingBones(const TArray<uint32>& SkinnedVertexIndices, TSet<int32>& Bones)
+{
+	for (uint32 SkinnedVertIndex : SkinnedVertexIndices)
+	{
+		int32 SectionIndex;
+		int32 SoftVertIndex;
+		RenderData->LODRenderData[0].GetSectionFromVertexIndex(SkinnedVertIndex, SectionIndex, SoftVertIndex);
+
+		const FSkelMeshSection& Section = StaticLODModel->LODModels[0].Sections[SectionIndex];
+		const FSoftSkinVertex& SoftVert = Section.SoftVertices[SoftVertIndex];
+		for (int32 InfluenceIndex = 0; InfluenceIndex < MAX_TOTAL_INFLUENCES; ++InfluenceIndex)
+		{
+			const uint16 InfluenceWeight = SoftVert.InfluenceWeights[InfluenceIndex];
+			if (InfluenceWeight >= 1)
+			{
+				const FBoneIndexType BoneMapIndex = SoftVert.InfluenceBones[InfluenceIndex];
+				const int32 ActualBoneIndex = Section.BoneMap[BoneMapIndex];
+				Bones.Add(ActualBoneIndex);
+			}
+		}
+	}
+}
\ No newline at end of file
diff --git a/Engine/Source/Developer/PhysicsUtilities/Private/SkinnedLevelSetBuilder.h b/Engine/Source/Developer/PhysicsUtilities/Private/SkinnedLevelSetBuilder.h
new file mode 100644
--- /dev/null
+++ b/Engine/Source/Developer/PhysicsUtilities/Private/SkinnedLevelSetBuilder.h
@@ -0,0 +1,43 @@
+// Copyright Epic Games, Inc. All Rights Reserved.
+
+#pragma once
+
+#include "Chaos/WeightedLatticeImplicitObject.h"
+
+class FSkinnedBoneTriangleCache;
+class FSkeletalMeshModel;
+class FSkeletalMeshRenderData;
+struct FPhysAssetCreateParams;
+namespace Chaos
+{
+class FLevelSet;
+}
+
+class FSkinnedLevelSetBuilder : public Chaos::FWeightedLatticeImplicitObjectBuilder
+{
+public:
+
+	FSkinnedLevelSetBuilder(const USkeletalMesh& InSkeletalMesh, const FSkinnedBoneTriangleCache& InTriangleCache, const int32 InRootBoneIndex);
+
+	bool InitializeSkinnedLevelset(const FPhysAssetCreateParams& Params, const TArray<int32>& BoneIndices, TArray<uint32>& OutOrigIndices);
+	void AddBoneInfluence(int32 PrimaryBoneIndex, const TArray<int32>& AllBonesForInfluence);
+
+	FKSkinnedLevelSetElem CreateSkinnedLevelSetElem();
+
+	void GetInfluencingBones(const TArray<uint32>& SkinnedVertexIndices, TSet<int32>& Bones);
+
+private:
+	// Input
+	const USkeletalMesh& SkeletalMesh;
+	const FSkinnedBoneTriangleCache& TriangleCache;
+	const int32 RootBoneIndex;
+
+	// Computed from Input
+	FSkeletalMeshModel* StaticLODModel;
+	const FSkeletalMeshRenderData* RenderData;
+
+	// Output
+	TUniquePtr<Chaos::FLevelSet> LevelSet;
+
+	float GetWeightForIndices(const TSet<int32>& BoneIndices, uint32 VertIndex) const;
+};
diff --git a/Engine/Source/Developer/PhysicsUtilities/Public/PhysicsAssetUtils.h b/Engine/Source/Developer/PhysicsUtilities/Public/PhysicsAssetUtils.h
--- a/Engine/Source/Developer/PhysicsUtilities/Public/PhysicsAssetUtils.h
+++ b/Engine/Source/Developer/PhysicsUtilities/Public/PhysicsAssetUtils.h
@@ -23,6 +23,7 @@ enum EPhysAssetFitGeomType : int
 	EFG_SingleConvexHull	UMETA(DisplayName="Single Convex Hull"),
 	EFG_MultiConvexHull		UMETA(DisplayName="Multi Convex Hull"),
 	EFG_LevelSet			UMETA(DisplayName="Level Set"),
+	EFG_SkinnedLevelSet		UMETA(DisplayName="Skinned Level Set"),
 };
 
 UENUM()
@@ -53,6 +54,7 @@ struct FPhysAssetCreateParams
 		HullCount = 4;
 		MaxHullVerts = 16;
 		LevelSetResolution = 8;
+		LatticeResolution = 8;
 	}
 
 	/** Bones that are shorter than this value will be ignored for body creation */
@@ -105,8 +107,13 @@ struct FPhysAssetCreateParams
 
 	/** When creating level sets, the grid resolution to use */
 	UPROPERTY(EditAnywhere, AdvancedDisplay, Category = "Body Creation", 
-		meta = (ClampMin = 1, UIMin = 10, UIMax = 100, ClampMax = 500, EditCondition = "GeomType == EPhysAssetFitGeomType::EFG_LevelSet"))
+		meta = (ClampMin = 1, UIMin = 10, UIMax = 100, ClampMax = 500, EditCondition = "GeomType == EPhysAssetFitGeomType::EFG_LevelSet || GeomType == EPhysAssetFitGeomType::EFG_SkinnedLevelSet"))
 	int32								LevelSetResolution;
+
+	/** When creating skinned level sets, the embedding grid resolution to use*/
+	UPROPERTY(EditAnywhere, AdvancedDisplay, Category = "Body Creation",
+		meta = (ClampMin = 1, UIMin = 10, UIMax = 100, ClampMax = 500, EditCondition = "GeomType == EPhysAssetFitGeomType::EFG_SkinnedLevelSet"))
+	int32								LatticeResolution;
 };
 
 class UPhysicsAsset;
@@ -132,7 +139,7 @@ namespace FPhysicsAssetUtils
 	 * 
 	 * @warning Certain physics geometry types, such as multi-convex hull, must recreate internal caches every time this function is called.
 	 * If you find you're calling this function repeatedly for different bone indices on the same mesh,
-	 * CreateFromSkeletalMesh or CreateCollisionFromBones will provide better performance.
+	 * CreateFromSkeletalMesh or CreateCollisionsFromBones will provide better performance.
 	 *
 	 * @param	bs					BodySetup to create the collision for
 	 * @param	skelMesh			The SkeletalMesh we create collision for
@@ -154,6 +161,18 @@ namespace FPhysicsAssetUtils
 	 */
 	PHYSICSUTILITIES_API bool CreateCollisionFromBones( UBodySetup* bs, USkeletalMesh* skelMesh, const TArray<int32>& BoneIndices, FPhysAssetCreateParams& Params, const FBoneVertInfo& Info );
 
+	/** Replaces any collision already in the  with an auto-generated ones using the parameters provided.
+	 *
+	 * @param	PhysicsAsset		The PhysicsAsset instance to update
+	 * @param	SkelMesh			The Skeletal Mesh to create the physics asset from
+	 * @param	BodyIndices			Indices of the existing BodySetups the collisions are created for
+	 * @param	Params				Additional parameters to control the creation
+	 * @param	Infos				The vertices to create the collisions for (parallel array to BodyIndices)
+	 * @param   OutSuccessfulBodyIndices Newly created BodySetup indices.
+	 * @return  Returns true if successfully created collision from all specified bodies
+	 */
+	PHYSICSUTILITIES_API bool CreateCollisionsFromBones(UPhysicsAsset* PhysicsAsset, USkeletalMesh* SkelMesh, const TArray<int32>& BodyIndices, const FPhysAssetCreateParams& Params, const TArray<FBoneVertInfo>& Infos, TArray<int32>& OutSuccessfulBodyIndices);
+
 	/**
 	 * Does a few things:
 	 * - add any collision primitives from body2 into body1 (adjusting the tm of each).
diff --git a/Engine/Source/Editor/Persona/Private/PhysicsAssetRenderUtils.cpp b/Engine/Source/Editor/Persona/Private/PhysicsAssetRenderUtils.cpp
--- a/Engine/Source/Editor/Persona/Private/PhysicsAssetRenderUtils.cpp
+++ b/Engine/Source/Editor/Persona/Private/PhysicsAssetRenderUtils.cpp
@@ -403,6 +403,7 @@ namespace PhysicsAssetRender
 				CollectDrawElements(DrawElements, EAggCollisionShape::Convex, AggGeom->ConvexElems, BoneTM, Scale, i, RenderSettings, PhysicsAsset, PDI, GetPrimitiveColor, GetPrimitiveMaterial, GetPrimitiveTransform, CreateHitProxy);
 				CollectDrawElements(DrawElements, EAggCollisionShape::TaperedCapsule, AggGeom->TaperedCapsuleElems, BoneTM, Scale, i, RenderSettings, PhysicsAsset, PDI, GetPrimitiveColor, GetPrimitiveMaterial, GetPrimitiveTransform, CreateHitProxy);
 				CollectDrawElements(DrawElements, EAggCollisionShape::LevelSet, AggGeom->LevelSetElems, BoneTM, Scale, i, RenderSettings, PhysicsAsset, PDI, GetPrimitiveColor, GetPrimitiveMaterial, GetPrimitiveTransform, CreateHitProxy);
+				CollectDrawElements(DrawElements, EAggCollisionShape::SkinnedLevelSet, AggGeom->SkinnedLevelSetElems, BoneTM, Scale, i, RenderSettings, PhysicsAsset, PDI, GetPrimitiveColor, GetPrimitiveMaterial, GetPrimitiveTransform, CreateHitProxy);
 
 				// Sort elements
 				DrawElements.Sort([](const DrawElement& DE1, const DrawElement& DE2) {
diff --git a/Engine/Source/Editor/PhysicsAssetEditor/PhysicsAssetEditor.Build.cs b/Engine/Source/Editor/PhysicsAssetEditor/PhysicsAssetEditor.Build.cs
--- a/Engine/Source/Editor/PhysicsAssetEditor/PhysicsAssetEditor.Build.cs
+++ b/Engine/Source/Editor/PhysicsAssetEditor/PhysicsAssetEditor.Build.cs
@@ -40,7 +40,8 @@ public class PhysicsAssetEditor : ModuleRules
 				"MeshUtilitiesCommon",
 				"ApplicationCore",
 				"EditorStyle",
-				"ToolWidgets"
+				"ToolWidgets",
+				"Chaos"
 			}
         );
 
diff --git a/Engine/Source/Editor/PhysicsAssetEditor/Private/PhysicsAssetEditor.cpp b/Engine/Source/Editor/PhysicsAssetEditor/Private/PhysicsAssetEditor.cpp
--- a/Engine/Source/Editor/PhysicsAssetEditor/Private/PhysicsAssetEditor.cpp
+++ b/Engine/Source/Editor/PhysicsAssetEditor/Private/PhysicsAssetEditor.cpp
@@ -2112,12 +2112,9 @@ void FPhysicsAssetEditor::ResetBoneCollision()
 
 	if(SharedData->SelectedBodies.Num() > 0)
 	{
-		TArray<int32> BodyIndices;
+		TArray<int32> SelectedBodyIndices;
 		const FScopedTransaction Transaction( LOCTEXT("ResetBoneCollision", "Reset Bone Collision") );
-
-		FScopedSlowTask SlowTask((float)SharedData->SelectedBodies.Num());
-		SlowTask.MakeDialog();
-		for(int32 i=0; i<SharedData->SelectedBodies.Num(); ++i)
+		for (int32 i = 0; i < SharedData->SelectedBodies.Num(); ++i)
 		{
 			int32 SelectedBodyIndex = SharedData->SelectedBodies[i].Index;
 			if (SharedData->PhysicsAsset->SkeletalBodySetups.IsValidIndex(SelectedBodyIndex) == false)
@@ -2125,24 +2122,16 @@ void FPhysicsAssetEditor::ResetBoneCollision()
 				continue;
 			}
 
-			UBodySetup* BodySetup = SharedData->PhysicsAsset->SkeletalBodySetups[SelectedBodyIndex];
-			check(BodySetup);
-			SlowTask.EnterProgressFrame(1.0f, FText::Format(LOCTEXT("ResetCollsionStepInfo", "Generating collision for {0}"), FText::FromName(BodySetup->BoneName)));
-			BodySetup->Modify();
+			SelectedBodyIndices.Add(SelectedBodyIndex);
+		}
 
-			int32 BoneIndex = EditorSkelMesh->GetRefSkeleton().FindBoneIndex(BodySetup->BoneName);
-			check(BoneIndex != INDEX_NONE);
+		TArray<int32> BodyIndices;
+		FPhysicsAssetUtils::CreateCollisionsFromBones(SharedData->PhysicsAsset, EditorSkelMesh, SelectedBodyIndices, NewBodyData,
+			NewBodyData.VertWeight == EVW_DominantWeight ? SharedData->DominantWeightBoneInfos : SharedData->AnyWeightBoneInfos, BodyIndices);
 
-			const FBoneVertInfo& UseVertInfo = NewBodyData.VertWeight == EVW_DominantWeight ? SharedData->DominantWeightBoneInfos[BoneIndex] : SharedData->AnyWeightBoneInfos[BoneIndex];
-			if(FPhysicsAssetUtils::CreateCollisionFromBone(BodySetup, EditorSkelMesh, BoneIndex, NewBodyData, UseVertInfo))
-			{
-				SharedData->AutoNameAllPrimitives(SelectedBodyIndex, NewBodyData.GeomType);
-				BodyIndices.AddUnique(SelectedBodyIndex);
-			}
-			else
-			{
-				FPhysicsAssetUtils::DestroyBody(SharedData->PhysicsAsset, SelectedBodyIndex);
-			}
+		for(const int32 BodyIndex : BodyIndices)
+		{
+			SharedData->AutoNameAllPrimitives(BodyIndex, NewBodyData.GeomType);
 		}
 
 		//deselect first
diff --git a/Engine/Source/Editor/PhysicsAssetEditor/Private/PhysicsAssetEditorSkeletalMeshComponent.cpp b/Engine/Source/Editor/PhysicsAssetEditor/Private/PhysicsAssetEditorSkeletalMeshComponent.cpp
--- a/Engine/Source/Editor/PhysicsAssetEditor/Private/PhysicsAssetEditorSkeletalMeshComponent.cpp
+++ b/Engine/Source/Editor/PhysicsAssetEditor/Private/PhysicsAssetEditorSkeletalMeshComponent.cpp
@@ -20,6 +20,8 @@
 #include "AnimPreviewInstance.h"
 #include "UObject/Package.h"
 #include "Styling/AppStyle.h"
+#include "Chaos/WeightedLatticeImplicitObject.h"
+#include "Chaos/Levelset.h"
 
 namespace
 {
@@ -230,6 +232,12 @@ FTransform UPhysicsAssetEditorSkeletalMeshComponent::GetPrimitiveTransform(const
 		PrimTM.ScaleTranslation(Scale3D);
 		return PrimTM * BoneTM;
 	}
+	else if (PrimType == EAggCollisionShape::SkinnedLevelSet)
+	{
+		FTransform PrimTM = ManTM * SharedBodySetup->AggGeom.SkinnedLevelSetElems[PrimIndex].GetTransform();
+		PrimTM.ScaleTranslation(Scale3D);
+		return PrimTM * BoneTM;
+	}
 
 	// Should never reach here
 	check(0);
@@ -353,6 +361,7 @@ void UPhysicsAssetEditorSkeletalMeshComponent::RefreshBoneTransforms(FActorCompo
 		FinalizeBoneTransform();
 		bNeedToFlipSpaceBaseBuffers = true;
 	}
+	UpdateSkinnedLevelSets();
 }
 
 void UPhysicsAssetEditorSkeletalMeshComponent::AddImpulseAtLocation(FVector Impulse, FVector Location, FName BoneName)
@@ -417,3 +426,50 @@ void UPhysicsAssetEditorSkeletalMeshComponent::CreateSimulationFloor(FBodyInstan
 		PhatPreviewInstance->CreateSimulationFloor(FloorBodyInstance, Transform);
 	}
 }
+
+void UPhysicsAssetEditorSkeletalMeshComponent::UpdateSkinnedLevelSets()
+{
+	UPhysicsAsset* const PhysicsAsset = GetPhysicsAsset();
+	if (!PhysicsAsset)
+	{
+		return;
+	}
+	for (int32 i = 0; i < PhysicsAsset->SkeletalBodySetups.Num(); ++i)
+	{
+		const int32 BoneIndex = GetBoneIndex(PhysicsAsset->SkeletalBodySetups[i]->BoneName);
+		if (BoneIndex != INDEX_NONE)
+		{
+			FKAggregateGeom* const AggGeom = &PhysicsAsset->SkeletalBodySetups[i]->AggGeom;
+			if (AggGeom)
+			{
+				for (FKSkinnedLevelSetElem& SkinnedLevelSet : AggGeom->SkinnedLevelSetElems)
+				{
+					if (SkinnedLevelSet.GetWeightedLevelSet().IsValid())
+					{
+						const TArray<FName>& UsedBoneNames = SkinnedLevelSet.GetWeightedLevelSet()->GetUsedBones();
+
+						const FTransform RootTransformInv = GetBoneTransform(BoneIndex, FTransform::Identity).Inverse();
+						TArray<FTransform> Transforms;
+						Transforms.SetNum(UsedBoneNames.Num());
+
+						for (int32 LocalIdx = 0; LocalIdx < UsedBoneNames.Num(); ++LocalIdx)
+						{
+							const int32 LocalBoneIndex = GetBoneIndex(UsedBoneNames[LocalIdx]);
+							if (LocalBoneIndex != INDEX_NONE)
+							{
+								const FTransform BoneTransformTimesRootTransformInv = GetBoneTransform(LocalBoneIndex, RootTransformInv);
+								Transforms[LocalIdx] = BoneTransformTimesRootTransformInv;
+							}
+							else
+							{
+								Transforms[LocalIdx] = RootTransformInv;
+							}
+						}
+
+						SkinnedLevelSet.GetWeightedLevelSet()->DeformPoints(Transforms);
+					}
+				}
+			}
+		}
+	}
+}
\ No newline at end of file
diff --git a/Engine/Source/Editor/PhysicsAssetEditor/Private/PhysicsAssetEditorSkeletalMeshComponent.h b/Engine/Source/Editor/PhysicsAssetEditor/Private/PhysicsAssetEditorSkeletalMeshComponent.h
--- a/Engine/Source/Editor/PhysicsAssetEditor/Private/PhysicsAssetEditorSkeletalMeshComponent.h
+++ b/Engine/Source/Editor/PhysicsAssetEditor/Private/PhysicsAssetEditorSkeletalMeshComponent.h
@@ -78,4 +78,7 @@ class UPhysicsAssetEditorSkeletalMeshComponent : public UDebugSkelMeshComponent
 
 public:
 	virtual bool CanOverrideCollisionProfile() const override { return false;  }
+
+private:
+	void UpdateSkinnedLevelSets();
 };
diff --git a/Engine/Source/Runtime/Engine/Classes/PhysicsEngine/AggregateGeom.h b/Engine/Source/Runtime/Engine/Classes/PhysicsEngine/AggregateGeom.h
--- a/Engine/Source/Runtime/Engine/Classes/PhysicsEngine/AggregateGeom.h
+++ b/Engine/Source/Runtime/Engine/Classes/PhysicsEngine/AggregateGeom.h
@@ -7,6 +7,7 @@
 #include "PhysicsEngine/ConvexElem.h"
 #include "PhysicsEngine/LevelSetElem.h"
 #include "PhysicsEngine/BoxElem.h"
+#include "PhysicsEngine/SkinnedLevelSetElem.h"
 #include "PhysicsEngine/SphereElem.h"
 #include "PhysicsEngine/SphylElem.h"
 #include "PhysicsEngine/TaperedCapsuleElem.h"
@@ -38,6 +39,9 @@ struct ENGINE_API FKAggregateGeom
 	UPROPERTY(EditAnywhere, editfixedsize, Category = "Aggregate Geometry", meta = (DisplayName = "Level Sets"))
 	TArray<FKLevelSetElem> LevelSetElems;
 
+	UPROPERTY(EditAnywhere, editfixedsize, Category = "Aggregate Geometry", meta = (DisplayName = "(Experimental) Skinned Level Sets"), Experimental)
+	TArray<FKSkinnedLevelSetElem> SkinnedLevelSetElems;
+
 	class FKConvexGeomRenderInfo* RenderInfo;
 
 	FKAggregateGeom()
@@ -60,12 +64,12 @@ struct ENGINE_API FKAggregateGeom
 
 	int32 GetElementCount() const
 	{
-		return SphereElems.Num() + SphylElems.Num() + BoxElems.Num() + ConvexElems.Num() + TaperedCapsuleElems.Num() + LevelSetElems.Num();
+		return SphereElems.Num() + SphylElems.Num() + BoxElems.Num() + ConvexElems.Num() + TaperedCapsuleElems.Num() + LevelSetElems.Num() + SkinnedLevelSetElems.Num();
 	}
 
 	int32 GetElementCount(EAggCollisionShape::Type Type) const;
 
-	SIZE_T GetAllocatedSize() const { return SphereElems.GetAllocatedSize() + SphylElems.GetAllocatedSize() + BoxElems.GetAllocatedSize() + ConvexElems.GetAllocatedSize() + TaperedCapsuleElems.GetAllocatedSize(); }
+	SIZE_T GetAllocatedSize() const { return SphereElems.GetAllocatedSize() + SphylElems.GetAllocatedSize() + BoxElems.GetAllocatedSize() + ConvexElems.GetAllocatedSize() + TaperedCapsuleElems.GetAllocatedSize() + LevelSetElems.GetAllocatedSize() + SkinnedLevelSetElems.GetAllocatedSize(); }
 
 	FKShapeElem* GetElement(const EAggCollisionShape::Type Type, const int32 Index)
 	{
@@ -83,6 +87,8 @@ struct ENGINE_API FKAggregateGeom
 			if (ensure(TaperedCapsuleElems.IsValidIndex(Index))) { return &TaperedCapsuleElems[Index]; }
 		case EAggCollisionShape::LevelSet:
 			if (ensure(LevelSetElems.IsValidIndex(Index))) { return &LevelSetElems[Index]; }
+		case EAggCollisionShape::SkinnedLevelSet:
+			if (ensure(SkinnedLevelSetElems.IsValidIndex(Index))) { return &SkinnedLevelSetElems[Index]; }
 		default:
 			ensure(false);
 			return nullptr;
@@ -103,6 +109,8 @@ struct ENGINE_API FKAggregateGeom
 		if (Index < TaperedCapsuleElems.Num()) { return &TaperedCapsuleElems[Index]; }
 		Index -= TaperedCapsuleElems.Num();
 		if (Index < LevelSetElems.Num()) { return &LevelSetElems[Index]; }
+		Index -= LevelSetElems.Num();
+		if (Index < SkinnedLevelSetElems.Num()) { return &SkinnedLevelSetElems[Index]; }
 		ensure(false);
 		return nullptr;
 	}
@@ -121,6 +129,8 @@ struct ENGINE_API FKAggregateGeom
 		if (Index < TaperedCapsuleElems.Num()) { return &TaperedCapsuleElems[Index]; }
 		Index -= TaperedCapsuleElems.Num();
 		if (Index < LevelSetElems.Num()) { return &LevelSetElems[Index]; }
+		Index -= LevelSetElems.Num();
+		if (Index < SkinnedLevelSetElems.Num()) { return &SkinnedLevelSetElems[Index]; }
 		ensure(false);
 		return nullptr;
 	}
@@ -133,6 +143,7 @@ struct ENGINE_API FKAggregateGeom
 		SphereElems.Empty();
 		TaperedCapsuleElems.Empty();
 		LevelSetElems.Empty();
+		SkinnedLevelSetElems.Empty();
 
 		FreeRenderInfo();
 	}
@@ -177,5 +188,6 @@ private:
 		ConvexElems = Other.ConvexElems;
 		TaperedCapsuleElems = Other.TaperedCapsuleElems;
 		LevelSetElems = Other.LevelSetElems;
+		SkinnedLevelSetElems = Other.SkinnedLevelSetElems;
 	}
 };
diff --git a/Engine/Source/Runtime/Engine/Classes/PhysicsEngine/ShapeElem.h b/Engine/Source/Runtime/Engine/Classes/PhysicsEngine/ShapeElem.h
--- a/Engine/Source/Runtime/Engine/Classes/PhysicsEngine/ShapeElem.h
+++ b/Engine/Source/Runtime/Engine/Classes/PhysicsEngine/ShapeElem.h
@@ -20,6 +20,7 @@ namespace EAggCollisionShape
 		Convex,
 		TaperedCapsule,
 		LevelSet,
+		SkinnedLevelSet,
 
 		Unknown
 	};
diff --git a/Engine/Source/Runtime/Engine/Classes/PhysicsEngine/SkinnedLevelSetElem.h b/Engine/Source/Runtime/Engine/Classes/PhysicsEngine/SkinnedLevelSetElem.h
new file mode 100644
--- /dev/null
+++ b/Engine/Source/Runtime/Engine/Classes/PhysicsEngine/SkinnedLevelSetElem.h
@@ -0,0 +1,71 @@
+// Copyright Epic Games, Inc. All Rights Reserved.
+
+#pragma once
+
+#include "UObject/ObjectMacros.h"
+#include "PhysicsEngine/ShapeElem.h"
+#include "SkinnedLevelSetElem.generated.h"
+
+namespace Chaos
+{
+class FLevelSet;
+template<typename TConcrete> class TWeightedLatticeImplicitObject;
+}
+
+USTRUCT()
+struct ENGINE_API FKSkinnedLevelSetElem : public FKShapeElem
+{
+	GENERATED_USTRUCT_BODY()
+
+	FKSkinnedLevelSetElem() :
+		FKShapeElem(EAggCollisionShape::SkinnedLevelSet)
+	{}
+
+	FKSkinnedLevelSetElem(const FKSkinnedLevelSetElem& Other)
+	{
+		CloneElem(Other);
+	}
+
+	const FKSkinnedLevelSetElem& operator=(const FKSkinnedLevelSetElem& Other)
+	{
+		CloneElem(Other);
+		return *this;
+	}
+
+	void SetWeightedLevelSet(TUniquePtr< Chaos::TWeightedLatticeImplicitObject<Chaos::FLevelSet>>&& InWeightedLevelSet);
+
+	FTransform GetTransform() const;
+
+	// Draw functions
+	virtual void DrawElemWire(class FPrimitiveDrawInterface* PDI, const FTransform& ElemTM, const float Scale, const FColor Color) const override;
+	virtual void DrawElemSolid(class FPrimitiveDrawInterface* PDI, const FTransform& ElemTM, const float Scale, const FMaterialRenderProxy* MaterialRenderProxy) const override;
+	void GetElemSolid(const FTransform& ElemTM, const FVector& Scale3D, const FMaterialRenderProxy* MaterialRenderProxy, int32 ViewIndex, class FMeshElementCollector& Collector) const;
+
+	FBox CalcAABB(const FTransform& BoneTM, const FVector& Scale3D) const;
+	FIntVector3 LevelSetGridResolution() const;
+	FIntVector3 LatticeGridResolution() const;
+
+	bool Serialize(FArchive& Ar);
+
+	const TSharedPtr<Chaos::TWeightedLatticeImplicitObject<Chaos::FLevelSet>>& GetWeightedLevelSet() const
+	{
+		return WeightedLevelSet;
+	}
+
+private:
+
+	TSharedPtr<Chaos::TWeightedLatticeImplicitObject<Chaos::FLevelSet>, ESPMode::ThreadSafe> WeightedLevelSet;
+
+	/** Helper function to safely copy instances of this shape*/
+	 void CloneElem(const FKSkinnedLevelSetElem& Other);
+};
+
+/* Enable our own serialization function to handle FKSkinnedLevelSetElem */
+template<>
+struct TStructOpsTypeTraits<FKSkinnedLevelSetElem> : public TStructOpsTypeTraitsBase2<FKSkinnedLevelSetElem>
+{
+	enum
+	{
+		WithSerializer = true
+	};
+};
diff --git a/Engine/Source/Runtime/Engine/Private/Collision/KAggregateGeom.cpp b/Engine/Source/Runtime/Engine/Private/Collision/KAggregateGeom.cpp
--- a/Engine/Source/Runtime/Engine/Private/Collision/KAggregateGeom.cpp
+++ b/Engine/Source/Runtime/Engine/Private/Collision/KAggregateGeom.cpp
@@ -18,6 +18,7 @@
 #include "Chaos/Convex.h"
 #include "Chaos/Levelset.h"
 #include "PhysicsEngine/TaperedCapsuleElem.h"
+#include "Chaos/WeightedLatticeImplicitObject.h"
 #if INTEL_ISPC
 #include "KAggregateGeom.ispc.generated.h"
 #endif
@@ -114,6 +115,11 @@ FBox FKAggregateGeom::CalcAABB(const FTransform& Transform) const
 		Box += LevelSetElems[i].CalcAABB(BoneTM, (FVector)Scale3D);
 	}
 
+	for (int32 i = 0; i < SkinnedLevelSetElems.Num(); i++)
+	{
+		Box += SkinnedLevelSetElems[i].CalcAABB(BoneTM, (FVector)Scale3D);
+	}
+
 	return Box;
 
 }
@@ -595,6 +601,58 @@ void FKLevelSetElem::GetZeroIsosurfaceGridCellFaces(TArray<FVector3f>& Vertices,
 	}
 }
 
+
+///////////////////////////////////////
+//////// FKSkinnedLevelSetElem ////////
+///////////////////////////////////////
+
+FBox FKSkinnedLevelSetElem::CalcAABB(const FTransform& BoneTM, const FVector& Scale3D) const
+{
+	FBox Box;
+
+	if (WeightedLevelSet.IsValid())
+	{
+		Box = FBox(WeightedLevelSet->BoundingBox().Min(), WeightedLevelSet->BoundingBox().Max());
+		const FTransform LocalToWorld = FTransform(FQuat::Identity, FVector::ZeroVector, Scale3D) * BoneTM;
+		return Box.TransformBy(LocalToWorld);
+	}
+
+	return Box;
+}
+
+FIntVector3 FKSkinnedLevelSetElem::LevelSetGridResolution() const
+{
+	if (WeightedLevelSet.IsValid())
+	{
+		const Chaos::TVec3<int32>& Dim = WeightedLevelSet->GetEmbeddedObject()->GetGrid().Counts();
+		return FIntVector3(Dim[0], Dim[1], Dim[2]);
+	}
+
+	return FIntVector3(0, 0, 0);
+}
+
+FIntVector3 FKSkinnedLevelSetElem::LatticeGridResolution() const
+{
+	if (WeightedLevelSet.IsValid())
+	{
+		const Chaos::TVec3<int32>& Dim = WeightedLevelSet->GetGrid().Counts();
+		return FIntVector3(Dim[0], Dim[1], Dim[2]);
+	}
+
+	return FIntVector3(0, 0, 0);
+}
+
+void FKSkinnedLevelSetElem::SetWeightedLevelSet(TUniquePtr< Chaos::TWeightedLatticeImplicitObject<Chaos::FLevelSet>>&& InWeightedLevelSet)
+{
+	WeightedLevelSet = TSharedPtr<Chaos::TWeightedLatticeImplicitObject<Chaos::FLevelSet>>(InWeightedLevelSet.Release());
+}
+
+FTransform FKSkinnedLevelSetElem::GetTransform() const
+{
+	return FTransform();
+}
+
+
 static const float DIST_COMPARE_THRESH = 0.1f;
 static const float DIR_COMPARE_THRESH = 0.0003f; // about 1 degree
 
@@ -921,3 +979,19 @@ bool FKLevelSetElem::Serialize(FArchive& Ar)
 	return true;
 }
 
+bool FKSkinnedLevelSetElem::Serialize(FArchive& Ar)
+{
+	if (Ar.IsLoading())
+	{
+		WeightedLevelSet = TSharedPtr<Chaos::TWeightedLatticeImplicitObject<Chaos::FLevelSet>>(new Chaos::TWeightedLatticeImplicitObject<Chaos::FLevelSet>());
+	}
+
+	if (WeightedLevelSet.IsValid())
+	{
+		Chaos::FChaosArchive ChaosAr(Ar);
+		WeightedLevelSet->Serialize(ChaosAr);
+	}
+
+	return true;
+}
+
diff --git a/Engine/Source/Runtime/Engine/Private/PhysicsEngine/BodySetup.cpp b/Engine/Source/Runtime/Engine/Private/PhysicsEngine/BodySetup.cpp
--- a/Engine/Source/Runtime/Engine/Private/PhysicsEngine/BodySetup.cpp
+++ b/Engine/Source/Runtime/Engine/Private/PhysicsEngine/BodySetup.cpp
@@ -2277,6 +2277,14 @@ void FKLevelSetElem::ScaleElem(FVector DeltaSize, float MinSize)
 	SetTransform(ScaledTransform);
 }
 
+
+/** Helper function to safely copy instances of this shape*/
+void FKSkinnedLevelSetElem::CloneElem(const FKSkinnedLevelSetElem& Other)
+{
+	Super::CloneElem(Other);
+	WeightedLevelSet = Other.WeightedLevelSet;
+}
+
 #if WITH_EDITOR
 
 bool FBodySetupObjectTextFactory::CanCreateClass(UClass* InObjectClass, bool& bOmitSubObjs) const
diff --git a/Engine/Source/Runtime/Engine/Private/PhysicsEngine/PhysDrawing.cpp b/Engine/Source/Runtime/Engine/Private/PhysicsEngine/PhysDrawing.cpp
--- a/Engine/Source/Runtime/Engine/Private/PhysicsEngine/PhysDrawing.cpp
+++ b/Engine/Source/Runtime/Engine/Private/PhysicsEngine/PhysDrawing.cpp
@@ -17,6 +17,9 @@
 #include "PhysicsEngine/PhysicsAsset.h"
 #include "PhysicsEngine/TaperedCapsuleElem.h"
 #include "StaticMeshResources.h"
+#include "Chaos/Levelset.h"
+#include "Chaos/UniformGrid.h"
+#include "Chaos/WeightedLatticeImplicitObject.h"
 
 static const int32 DrawCollisionSides = 32;
 static const int32 DrawConeLimitSides = 40;
@@ -842,6 +845,120 @@ void FKLevelSetElem::GetElemSolid(const FTransform& ElemTM, const FVector& Scale
 	MeshBuilder.GetMesh(ElemTM.ToMatrixWithScale(), MaterialRenderProxy, SDPG_World, false, false, ViewIndex, Collector);
 }
 
+/////////////////////////////////////////////////////////////////////////////////////
+// FKSkinnedLevelSetElem
+/////////////////////////////////////////////////////////////////////////////////////
+
+static void DrawSkinnedLevelSetLattice(class FPrimitiveDrawInterface* PDI, const FMatrix& LocalToWorld, const FColor Color, const Chaos::TWeightedLatticeImplicitObject<Chaos::FLevelSet>* WeightedLevelSet)
+{
+	const Chaos::TUniformGrid<double, 3>& LatticeGrid = WeightedLevelSet->GetGrid();
+	const Chaos::TArrayND<Chaos::FVec3, 3>& DeformedPoints = WeightedLevelSet->GetDeformedPoints();
+	const Chaos::TArrayND<bool, 3>& EmptyCells = WeightedLevelSet->GetEmptyCells();
+
+
+	// TODO: downres drawing lattice when res is too high. 
+	// Does two passes to count lines and reserve help?
+	const Chaos::TVec3<int32> CellCounts = LatticeGrid.Counts();
+	for (int32 I = 0; I < CellCounts.X; ++I)
+	{
+		for (int32 J = 0; J < CellCounts.Y; ++J)
+		{
+			for (int32 K = 0; K < CellCounts.Z; ++K)
+			{
+				if (!EmptyCells(I, J, K))
+				{
+					const FVector P000 = LocalToWorld.TransformPosition(FVector(DeformedPoints(I, J, K)));
+					const FVector P001 = LocalToWorld.TransformPosition(FVector(DeformedPoints(I, J, K + 1)));
+					const FVector P010 = LocalToWorld.TransformPosition(FVector(DeformedPoints(I, J + 1, K)));
+					const FVector P011 = LocalToWorld.TransformPosition(FVector(DeformedPoints(I, J + 1, K + 1)));
+					const FVector P100 = LocalToWorld.TransformPosition(FVector(DeformedPoints(I + 1, J, K)));
+					const FVector P101 = LocalToWorld.TransformPosition(FVector(DeformedPoints(I + 1, J, K + 1)));
+					const FVector P110 = LocalToWorld.TransformPosition(FVector(DeformedPoints(I + 1, J + 1, K)));
+					const FVector P111 = LocalToWorld.TransformPosition(FVector(DeformedPoints(I + 1, J + 1, K + 1)));
+
+					PDI->AddReserveLines(SDPG_World, 12);
+					PDI->DrawLine(P000, P001, Color, SDPG_World);
+					PDI->DrawLine(P000, P010, Color, SDPG_World);
+					PDI->DrawLine(P000, P100, Color, SDPG_World);
+					PDI->DrawLine(P001, P011, Color, SDPG_World);
+					PDI->DrawLine(P001, P101, Color, SDPG_World);
+					PDI->DrawLine(P010, P011, Color, SDPG_World);
+					PDI->DrawLine(P010, P110, Color, SDPG_World);
+					PDI->DrawLine(P011, P111, Color, SDPG_World);
+					PDI->DrawLine(P100, P101, Color, SDPG_World);
+					PDI->DrawLine(P100, P110, Color, SDPG_World);
+					PDI->DrawLine(P101, P111, Color, SDPG_World);
+					PDI->DrawLine(P110, P111, Color, SDPG_World);
+
+				}
+			}
+		}
+	}
+
+}
+
+void FKSkinnedLevelSetElem::DrawElemWire(class FPrimitiveDrawInterface* PDI, const FTransform& ElemTM, float Scale, const FColor Color) const
+{
+	if (WeightedLevelSet.IsValid())
+	{
+		DrawSkinnedLevelSetLattice(PDI, ElemTM.ToMatrixWithScale(), Color, WeightedLevelSet.Get());
+
+		// TODO grid cells
+	}
+}
+
+void FKSkinnedLevelSetElem::DrawElemSolid(class FPrimitiveDrawInterface* PDI, const FTransform& ElemTM, float Scale, const FMaterialRenderProxy* MaterialRenderProxy) const
+{
+	if (WeightedLevelSet.IsValid())
+	{
+		const FColor LatticeColor = FColor::Cyan; // TODO figure out what color to actually draw lattice
+		DrawSkinnedLevelSetLattice(PDI, ElemTM.ToMatrixWithScale(), LatticeColor, WeightedLevelSet.Get());
+
+		TArray<FVector3f> Vertices;
+		TArray<FIntVector> Tris;
+		const Chaos::FLevelSet* const LevelSet = WeightedLevelSet->GetEmbeddedObject();
+		LevelSet->GetZeroIsosurfaceGridCellFaces(Vertices, Tris);
+
+		FDynamicMeshBuilder MeshBuilder(PDI->View->GetFeatureLevel());
+		for (const FVector3f& V : Vertices)
+		{
+			MeshBuilder.AddVertex(FDynamicMeshVertex(FVector3f(WeightedLevelSet->GetDeformedPoint(Chaos::FVec3(V)))));
+		}
+		for (const FIntVector& T : Tris)
+		{
+			MeshBuilder.AddTriangle(T[0], T[1], T[2]);
+		}
+
+		MeshBuilder.Draw(PDI, ElemTM.ToMatrixWithScale(), MaterialRenderProxy, SDPG_World, 0.f);
+	}
+}
+
+void FKSkinnedLevelSetElem::GetElemSolid(const FTransform& ElemTM, const FVector& Scale3D, const FMaterialRenderProxy* MaterialRenderProxy, int32 ViewIndex, class FMeshElementCollector& Collector) const
+{
+	if (WeightedLevelSet.IsValid())
+	{
+		const FColor LatticeColor = FColor::Cyan; // TODO figure out what color to actually draw lattice
+		DrawSkinnedLevelSetLattice(Collector.GetPDI(ViewIndex), ElemTM.ToMatrixWithScale(), LatticeColor, WeightedLevelSet.Get());
+
+		TArray<FVector3f> Vertices;
+		TArray<FIntVector> Tris;
+		const Chaos::FLevelSet* const LevelSet = WeightedLevelSet->GetEmbeddedObject();
+		LevelSet->GetZeroIsosurfaceGridCellFaces(Vertices, Tris);
+
+		FDynamicMeshBuilder MeshBuilder(Collector.GetPDI(ViewIndex)->View->GetFeatureLevel());
+		for (const FVector3f& V : Vertices)
+		{
+			MeshBuilder.AddVertex(FDynamicMeshVertex(FVector3f(WeightedLevelSet->GetDeformedPoint(Chaos::FVec3(V)))));
+		}
+		for (const FIntVector& T : Tris)
+		{
+			MeshBuilder.AddTriangle(T[0], T[1], T[2]);
+		}
+
+		MeshBuilder.GetMesh(ElemTM.ToMatrixWithScale(), MaterialRenderProxy, SDPG_World, false, false, ViewIndex, Collector);
+	}
+}
+
 /////////////////////////////////////////////////////////////////////////////////////
 // FKConvexGeomRenderInfo
 /////////////////////////////////////////////////////////////////////////////////////
@@ -995,9 +1112,13 @@ void FKAggregateGeom::GetAggGeom(const FTransform& Transform, const FColor Color
 		ElemTM *= ParentTM;
 
 		if (bDrawSolid)
+		{
 			TaperedCapsuleElems[i].GetElemSolid(ElemTM, Scale3D, MatInst, ViewIndex, Collector);
+		}
 		else
+		{
 			TaperedCapsuleElems[i].DrawElemWire(Collector.GetPDI(ViewIndex), ElemTM, Scale3D, Color);
+		}
 	}
 
 	for (int32 i = 0; i < LevelSetElems.Num(); i++)
@@ -1006,9 +1127,25 @@ void FKAggregateGeom::GetAggGeom(const FTransform& Transform, const FColor Color
 		ElemTM *= Transform;
 
 		if (bDrawSolid)
+		{
 			LevelSetElems[i].GetElemSolid(ElemTM, Scale3D, MatInst, ViewIndex, Collector);
+		}
 		else
+		{
 			LevelSetElems[i].DrawElemWire(Collector.GetPDI(ViewIndex), ElemTM, 1.f, Color);
+		}
+	}
+
+	for (int32 i = 0; i < SkinnedLevelSetElems.Num(); i++)
+	{
+		if (bDrawSolid)
+		{
+			SkinnedLevelSetElems[i].GetElemSolid(Transform, Scale3D, MatInst, ViewIndex, Collector);
+		}
+		else
+		{
+			SkinnedLevelSetElems[i].DrawElemWire(Collector.GetPDI(ViewIndex), Transform, 1.f, Color);
+		}
 	}
 }
 
diff --git a/Engine/Source/Runtime/Experimental/Chaos/Private/Chaos/Levelset.cpp b/Engine/Source/Runtime/Experimental/Chaos/Private/Chaos/Levelset.cpp
--- a/Engine/Source/Runtime/Experimental/Chaos/Private/Chaos/Levelset.cpp
+++ b/Engine/Source/Runtime/Experimental/Chaos/Private/Chaos/Levelset.cpp
@@ -694,6 +694,28 @@ void FLevelSet::GetZeroIsosurfaceGridCellFaces(TArray<FVector3f>& Vertices, TArr
 	}
 }
 
+void FLevelSet::GetInteriorCells(TArray<TVec3<int32>>& InteriorCells, const FReal InteriorThreshold) const
+{
+	InteriorCells.Reset();
+
+	const TVector<int32, 3> Cells = MGrid.Counts();
+
+	for (int i = 0; i < Cells.X; ++i)
+	{
+		for (int j = 0; j < Cells.Y; ++j)
+		{
+			for (int k = 0; k < Cells.Z; ++k)
+			{
+				const FReal Value = MPhi(i, j, k);
+				if (Value < InteriorThreshold)
+				{
+					InteriorCells.Emplace(i, j, k);
+				}
+			}
+		}
+	}
+}
+
 void FLevelSet::ComputeConvexity(const TArray<TVec3<int32>>& InterfaceIndices)
 {
 	this->bIsConvex = true;
diff --git a/Engine/Source/Runtime/Experimental/Chaos/Private/Chaos/UniformGrid.cpp b/Engine/Source/Runtime/Experimental/Chaos/Private/Chaos/UniformGrid.cpp
--- a/Engine/Source/Runtime/Experimental/Chaos/Private/Chaos/UniformGrid.cpp
+++ b/Engine/Source/Runtime/Experimental/Chaos/Private/Chaos/UniformGrid.cpp
@@ -250,6 +250,17 @@ namespace Chaos
 		return Result;
 	}
 
+	template<class T>
+	TUniformGrid<T, 3> TUniformGrid<T, 3>::SubGrid(const TVector<int32, 3>& MinCell, const TVector<int32, 3>& MaxCell) const
+	{
+		TUniformGrid<T, 3> Result;
+		Result.MMinCorner = Node(MinCell);
+		Result.MMaxCorner = Node(MaxCell + TVector<int32, 3>(1, 1, 1));
+		Result.MCells = MaxCell - MinCell + TVector<int32, 3>(1, 1, 1);
+		Result.MDx = MDx;
+		return Result;
+	}
+
 	template<class T>
 	void TMPMGrid<T>::BaseNodeIndex(const TVector<T, 3>& X, TVector<int32, 3>& Index, TVector<T, 3>& Weights) const
 	{
@@ -419,4 +430,4 @@ template class Chaos::TMPMGrid<Chaos::FReal>;
 #endif
 
 template Chaos::FVec3 Chaos::TUniformGridBase<Chaos::FReal, 3>::LinearlyInterpolate<Chaos::FVec3>(const Chaos::TArrayND<Chaos::FVec3, 3>&, const Chaos::FVec3&) const;
-template CHAOS_API Chaos::FReal Chaos::TUniformGridBase<Chaos::FReal, 3>::LinearlyInterpolate<Chaos::FReal>(const TArrayND<Chaos::FReal, 3>&, const Chaos::FVec3&) const;
\ No newline at end of file
+template CHAOS_API Chaos::FReal Chaos::TUniformGridBase<Chaos::FReal, 3>::LinearlyInterpolate<Chaos::FReal>(const TArrayND<Chaos::FReal, 3>&, const Chaos::FVec3&) const;
diff --git a/Engine/Source/Runtime/Experimental/Chaos/Private/Chaos/WeightedLatticeImplicitObject.cpp b/Engine/Source/Runtime/Experimental/Chaos/Private/Chaos/WeightedLatticeImplicitObject.cpp
new file mode 100644
--- /dev/null
+++ b/Engine/Source/Runtime/Experimental/Chaos/Private/Chaos/WeightedLatticeImplicitObject.cpp
@@ -0,0 +1,793 @@
+// Copyright Epic Games, Inc. All Rights Reserved.
+#include "Chaos/WeightedLatticeImplicitObject.h"
+#include "Chaos/Levelset.h"
+#include "Chaos/Tetrahedron.h"
+
+namespace Chaos
+{
+
+FWeightedLatticeImplicitObject::FEmbeddingCoordinate::FEmbeddingCoordinate(const TVec3<int32>& InCellIndex, const FVec3& TrilinearCoordinate)
+	:CellIndex(InCellIndex)
+{
+	const FVec4 Coord4(TrilinearCoordinate.X, TrilinearCoordinate.Y, TrilinearCoordinate.Z, 1.);
+
+	if (((CellIndex.X + CellIndex.Y + CellIndex.Z) & 1) == 0)
+	{
+		static const FVec4 EvenPlaneTests[4] =
+		{
+			FVec4(-1., 1., -1., 0.),
+			FVec4(1., 1., 1., -2.),
+			FVec4(1., -1., -1., 0.),
+			FVec4(-1., -1., 1., 0.)
+		};
+
+		FReal PlaneTestValue = 0.;
+		for (LocalTetrahedron = 0; LocalTetrahedron < 4; ++LocalTetrahedron)
+		{
+			PlaneTestValue = Dot4(EvenPlaneTests[LocalTetrahedron], Coord4);
+			if (PlaneTestValue >= 0.)
+			{
+				break;
+			}
+		}
+		switch (LocalTetrahedron)
+		{
+		case 0:
+			BarycentricCoordinate = FVec3(PlaneTestValue, 1. - Coord4[1], Coord4[0]);
+			break;
+		case 1:
+			BarycentricCoordinate = FVec3(PlaneTestValue, 1. - Coord4[2], 1. - Coord4[1]);
+			break;
+		case 2:
+			BarycentricCoordinate = FVec3(PlaneTestValue, Coord4[2], Coord4[1]);
+			break;
+		case 3:
+			BarycentricCoordinate = FVec3(PlaneTestValue, 1. - Coord4[2], Coord4[1]);
+			break;
+		case 4:
+			BarycentricCoordinate = FVec3(-.5 * Dot4(EvenPlaneTests[3], Coord4), -.5 * Dot4(EvenPlaneTests[2], Coord4), -.5 * Dot4(EvenPlaneTests[1], Coord4));
+			break;
+		default:
+			check(false);
+		}
+	}
+	else
+	{
+		static const FVec4 OddPlaneTests[4] =
+		{
+			FVec4(-1., -1., -1., 1.),
+			FVec4(-1., 1., 1., -1.),
+			FVec4(1., -1., 1., -1.),
+			FVec4(1., 1., -1., -1.)
+		};
+		FReal PlaneTestValue = 0.;
+		for (LocalTetrahedron = 0; LocalTetrahedron < 4; ++LocalTetrahedron)
+		{
+			PlaneTestValue = Dot4(OddPlaneTests[LocalTetrahedron], Coord4);
+			if (PlaneTestValue >= 0.)
+			{
+				break;
+			}
+		}
+		switch (LocalTetrahedron)
+		{
+		case 0:
+			BarycentricCoordinate = FVec3(PlaneTestValue, Coord4[0], Coord4[1]);
+			break;
+		case 1:
+			BarycentricCoordinate = FVec3(PlaneTestValue, 1. - Coord4[2], Coord4[0]);
+			break;
+		case 2:
+			BarycentricCoordinate = FVec3(PlaneTestValue, Coord4[1], 1. - Coord4[2]);
+			break;
+		case 3:
+			BarycentricCoordinate = FVec3(PlaneTestValue, 1. - Coord4[0], 1. - Coord4[1]);
+			break;
+		case 4:
+			BarycentricCoordinate = FVec3(-.5 * Dot4(OddPlaneTests[1], Coord4), -.5 * Dot4(OddPlaneTests[0], Coord4), -.5 * Dot4(OddPlaneTests[2], Coord4));
+			break;
+		default:
+			check(false);
+		}
+	}
+}
+
+FMatrix FWeightedLatticeImplicitObject::FEmbeddingCoordinate::DeformationTransform(const TArrayND<FVec3, 3>& InDeformedPoints, const TUniformGrid<FReal, 3>& InGrid) const
+{
+	const TVec4<TVec3<int32>>& TetOffsets = TetrahedronOffsets();
+	auto TestOffsets = [this, &TetOffsets, &InGrid, &InDeformedPoints](const FMatrix& Transform)
+	{
+		constexpr FReal Tolerance = 1e-3;
+		for (int32 VertId = 0; VertId < 4; ++VertId)
+		{
+			const FVec3 Node = InGrid.Node(CellIndex + TetOffsets[VertId]);
+			const FVec3 TransformedPoint = Transform.TransformPosition(Node);
+			const FVec3 DeformedPoint = InDeformedPoints(CellIndex + TetOffsets[VertId]);
+			const FReal DiffSq = (TransformedPoint - DeformedPoint).SquaredLength();
+			check(DiffSq < Tolerance);
+		}
+
+	};
+
+	if (LocalTetrahedron < 4)
+	{
+		// These are the trirectangular tetrahedra (angles at corner are right angles)
+		const FVec3 Corner = InDeformedPoints(CellIndex + TetOffsets[0]);
+		const FVec3 Axis1 = InDeformedPoints(CellIndex + TetOffsets[1]) - Corner;
+		const FVec3 Axis2 = InDeformedPoints(CellIndex + TetOffsets[2]) - Corner;
+		const FVec3 Axis3 = InDeformedPoints(CellIndex + TetOffsets[3]) - Corner;
+
+		FMatrix Result;
+
+		if (((CellIndex.X + CellIndex.Y + CellIndex.Z) & 1) == 0)
+		{
+			switch (LocalTetrahedron)
+			{
+			case 0:
+				Result = FMatrix(Axis2 / InGrid.Dx().X, -Axis1 / InGrid.Dx().Y, Axis3 / InGrid.Dx().Z, FVec3(0, 0, 0));
+				break;
+			case 1:
+				Result = FMatrix(-Axis3 / InGrid.Dx().X, -Axis2 / InGrid.Dx().Y, -Axis1 / InGrid.Dx().Z, FVec3(0, 0, 0));
+				break;
+			case 2:
+				Result = FMatrix(-Axis3 / InGrid.Dx().X, Axis2 / InGrid.Dx().Y, Axis1 / InGrid.Dx().Z, FVec3(0, 0, 0));
+				break;
+			case 3:
+				Result = FMatrix(Axis3 / InGrid.Dx().X, Axis2 / InGrid.Dx().Y, -Axis1 / InGrid.Dx().Z, FVec3(0, 0, 0));
+				break;
+			default:
+				check(0);
+			}
+		}
+		else
+		{
+			switch (LocalTetrahedron)
+			{
+			case 0:
+				Result = FMatrix(Axis1 / InGrid.Dx().X, Axis2 / InGrid.Dx().Y, Axis3 / InGrid.Dx().Z, FVec3(0, 0, 0));
+				break;
+			case 1:
+				Result = FMatrix(Axis2 / InGrid.Dx().X, -Axis3 / InGrid.Dx().Y, -Axis1 / InGrid.Dx().Z, FVec3(0, 0, 0));
+				break;
+			case 2:
+				Result = FMatrix(-Axis3 / InGrid.Dx().X, Axis1 / InGrid.Dx().Y, -Axis2 / InGrid.Dx().Z, FVec3(0, 0, 0));
+				break;
+			case 3:
+				Result = FMatrix(-Axis1 / InGrid.Dx().X, -Axis2 / InGrid.Dx().Y, Axis3 / InGrid.Dx().Z, FVec3(0, 0, 0));
+				break;
+			default:
+				check(0);
+			}
+		}
+		Result.SetOrigin(Corner - Result.TransformVector(InGrid.Node(CellIndex + TetOffsets[0])));
+		TestOffsets(Result);
+		return Result;
+	}
+	else
+	{
+		const FVec3& P0 = InDeformedPoints(CellIndex + TetOffsets[0]);
+		const FVec3& P1 = InDeformedPoints(CellIndex + TetOffsets[1]);
+		const FVec3& P2 = InDeformedPoints(CellIndex + TetOffsets[2]);
+		const FVec3& P3 = InDeformedPoints(CellIndex + TetOffsets[3]);
+
+		// These are the regular tetrahedra. 
+		if (((CellIndex.X + CellIndex.Y + CellIndex.Z) & 1) == 0)
+		{
+			FMatrix Result(.5 * (P0 + P3 - P1 - P2) / InGrid.Dx().X, .5 * (P0 + P1 - P2 - P3) / InGrid.Dx().Y, .5 * (P1 + P3 - P0 - P2) / InGrid.Dx().Z, FVec3(0, 0, 0));
+			Result.SetOrigin(P2 - Result.TransformVector(InGrid.Node(CellIndex + TetOffsets[2])));
+			TestOffsets(Result);
+			return Result;
+		}
+		else
+		{
+			FMatrix Result(.5 * (P0 + P1 - P2 - P3) / InGrid.Dx().X, .5 * (P1 + P2 - P0 - P3) / InGrid.Dx().Y, .5 * (P1 + P3 - P0 - P2) / InGrid.Dx().Z, FVec3(0, 0, 0));
+			Result.SetOrigin(P1 - Result.TransformVector(InGrid.Node(CellIndex + TetOffsets[1])));
+			TestOffsets(Result);
+			return Result;
+		}
+	}
+}
+
+FWeightedLatticeImplicitObject::FWeightedLatticeImplicitObject(const FWeightedLatticeImplicitObject& Other)
+	:FImplicitObject(EImplicitObject::HasBoundingBox, Other.GetType())
+	, Grid(Other.Grid)
+	, UsedBones(Other.UsedBones)
+	, ReferenceRelativeTransforms(Other.ReferenceRelativeTransforms)
+	, SolverBoneIndices(Other.SolverBoneIndices)
+	, LocalBoundingBox(Other.LocalBoundingBox)
+{
+	BoneData.Copy(Other.BoneData);
+	DeformedPoints.Copy(Other.DeformedPoints);
+	EmptyCells.Copy(Other.EmptyCells);
+	if (!Other.bSpatialDirty)
+	{
+		UpdateSpatialHierarchy();
+	}
+	bSpatialDirty = Other.bSpatialDirty;
+}
+
+FWeightedLatticeImplicitObject::FWeightedLatticeImplicitObject(int32 Flags, EImplicitObjectType InType, TUniformGrid<FReal, 3>&& InGrid,
+	TArrayND<FWeightedLatticeInfluenceData, 3>&& InBoneData, TArray<FName>&& InUsedBones, TArray<FTransform>&& InReferenceRelativeTransforms)
+	: FImplicitObject(Flags, InType | ImplicitObjectType::IsWeightedLattice)
+	, Grid(MoveTemp(InGrid))
+	, BoneData(MoveTemp(InBoneData))
+	, UsedBones(MoveTemp(InUsedBones))
+	, ReferenceRelativeTransforms(MoveTemp(InReferenceRelativeTransforms))
+{
+	FinalizeConstruction();
+}
+
+FWeightedLatticeImplicitObject::FWeightedLatticeImplicitObject(FWeightedLatticeImplicitObject&& Other)
+	: FImplicitObject(EImplicitObject::HasBoundingBox, Other.GetType())
+	, Grid(MoveTemp(Other.Grid))
+	, BoneData(MoveTemp(Other.BoneData))
+	, UsedBones(MoveTemp(Other.UsedBones))
+	, ReferenceRelativeTransforms(MoveTemp(Other.ReferenceRelativeTransforms))
+	, EmptyCells(MoveTemp(Other.EmptyCells))
+	, LocalBoundingBox(MoveTemp(Other.LocalBoundingBox))
+	, DeformedPoints(MoveTemp(Other.DeformedPoints))
+{
+}
+
+void FWeightedLatticeImplicitObject::FinalizeConstruction()
+{
+	SetEmptyCells();
+	InitializeDeformedPoints();
+	bSpatialDirty = true;
+}
+
+void FWeightedLatticeImplicitObject::Serialize(FChaosArchive& Ar)
+{
+	FImplicitObject::SerializeImp(Ar);
+	Ar << Grid;
+	Ar << BoneData;
+	Ar << UsedBones;
+	Ar << ReferenceRootTransform;
+	Ar << ReferenceRelativeTransforms;
+	if (Ar.IsLoading())
+	{
+		FinalizeConstruction();
+	}
+}
+
+void FWeightedLatticeImplicitObject::DeformPoints(const TArray<FTransform>& RelativeTransforms)
+{
+	checkSlow(RelativeTransforms.Num() == ReferenceRelativeTransforms.Num());
+	TArray<FTransform> BoneTransforms;
+	BoneTransforms.SetNum(RelativeTransforms.Num());
+	for (int32 Index = 0; Index < RelativeTransforms.Num(); ++Index)
+	{
+		BoneTransforms[Index] = ReferenceRelativeTransforms[Index] * RelativeTransforms[Index];
+	}
+
+	const TVec3<int32> NodeCounts = Grid.NodeCounts();
+	LocalBoundingBox = TAABB<FReal, 3>::EmptyAABB();
+	for (int32 I = 0; I < NodeCounts.X; ++I)
+	{
+		for (int32 J = 0; J < NodeCounts.Y; ++J)
+		{
+			for (int32 K = 0; K < NodeCounts.Z; ++K)
+			{
+				const TVec3<int32> Index(I, J, K);
+				const FWeightedLatticeInfluenceData& InfluenceData = BoneData(Index);
+				if (InfluenceData.NumInfluences == 0)
+				{
+					LocalBoundingBox.GrowToInclude(DeformedPoints(Index));
+				}
+				else
+				{
+					FVec3 NewPos(0.);
+					const FVec3 UndeformedPos = Grid.Node(Index);
+					for (int32 InfluenceIdx = 0; InfluenceIdx < InfluenceData.NumInfluences; ++InfluenceIdx)
+					{
+						NewPos += InfluenceData.BoneWeights[InfluenceIdx] * BoneTransforms[InfluenceData.BoneIndices[InfluenceIdx]].TransformPosition(UndeformedPos);
+					}
+					bSpatialDirty = FVec3::DistSquared(DeformedPoints(Index), NewPos) > UE_SMALL_NUMBER * UE_SMALL_NUMBER;
+					DeformedPoints(Index) = NewPos;
+					LocalBoundingBox.GrowToInclude(NewPos);
+				}
+			}
+		}
+	}
+}
+
+FVec3 FWeightedLatticeImplicitObject::GetDeformedPoint(const FVec3& UndeformedPoint) const
+{
+	const TVec3<int32> CellIndex = Grid.Cell(UndeformedPoint);
+	if (EmptyCells(CellIndex))
+	{
+		return UndeformedPoint;
+	}
+
+	const FVec3 Alpha = ((UndeformedPoint - Grid.Node(CellIndex)) / Grid.Dx()).BoundToBox(FVec3(0.), FVec3(1.));
+	FEmbeddingCoordinate EmbeddingCoord(CellIndex, Alpha);
+	return EmbeddingCoord.DeformedPosition(DeformedPoints);
+}
+
+struct FWeightedLatticeBvEntry
+{
+	const FWeightedLatticeImplicitObject* Lattice;
+	TVec3<int32> CellIndex;
+
+	bool HasBoundingBox() const
+	{
+		return true;
+	}
+
+	Chaos::TAABB<FReal, 3> BoundingBox() const
+	{
+		Chaos::TAABB<FReal, 3> Bounds;
+		for (int32 I = 0; I < 2; ++I)
+		{
+			for (int32 J = 0; J < 2; ++J)
+			{
+				for (int32 K = 0; K < 2; ++K)
+				{
+					Bounds.GrowToInclude(Lattice->GetDeformedPoints()(CellIndex + TVec3<int32>(I, J, K)));
+				}
+			}
+		}
+		return Bounds;
+	}
+
+	template<typename TPayloadType>
+	TPayloadType GetPayload(int32 Idx) const
+	{
+		return Lattice->GetGrid().FlatIndex(CellIndex, false);
+	}
+};
+
+void FWeightedLatticeImplicitObject::UpdateSpatialHierarchy()
+{
+	if (!bSpatialDirty)
+	{
+		return;
+	}
+
+	TArray<FWeightedLatticeBvEntry> Entries;
+	Entries.Reserve(Grid.GetNumCells());
+	for (int32 I = 0; I < Grid.Counts().X; ++I)
+	{
+		for (int32 J = 0; J < Grid.Counts().Y; ++J)
+		{
+			for (int32 K = 0; K < Grid.Counts().Z; ++K)
+			{
+				if (!EmptyCells(I, J, K))
+				{
+					Entries.Add({ this, {I,J,K} });
+				}
+			}
+		}
+	}
+
+	Spatial.Initialize(Entries);
+
+	bSpatialDirty = false;
+}
+
+bool FWeightedLatticeImplicitObject::GetEmbeddingCoordinates(const FVec3& DeformedPoint, TArray<FEmbeddingCoordinate>& CoordinatesOut, bool bFindClosest) const
+{
+	check(!bSpatialDirty);
+	CoordinatesOut.Reset();
+	const bool bInLocalBoundingBox = LocalBoundingBox.Contains(DeformedPoint);
+	if (!bInLocalBoundingBox && !bFindClosest)
+	{
+		return false;
+	}
+
+	TArray<int32> IntersectingCells = bInLocalBoundingBox ? Spatial.FindAllIntersections(DeformedPoint) : TArray<int32>();
+	for (int32 IntersectingCell : IntersectingCells)
+	{
+		const TVec3<int32> CellIndex = Grid.GetIndex(IntersectingCell);
+		const TVec4<TVec3<int32>>* TestCellOffsets = ((CellIndex.X + CellIndex.Y + CellIndex.Z) & 1) == 0 ? FEmbeddingCoordinate::EvenIndexTetrahedraOffsets : FEmbeddingCoordinate::OddIndexTetrahedraOffsets;
+		for (int32 TestTetIndex = 0; TestTetIndex < 5; ++TestTetIndex)
+		{
+			const TTetrahedron<FReal> Tetrahedron(
+				DeformedPoints(CellIndex + TestCellOffsets[TestTetIndex][0]),
+				DeformedPoints(CellIndex + TestCellOffsets[TestTetIndex][1]),
+				DeformedPoints(CellIndex + TestCellOffsets[TestTetIndex][2]),
+				DeformedPoints(CellIndex + TestCellOffsets[TestTetIndex][3]));
+
+			const TVec3<FReal> Weights = Tetrahedron.GetFirstThreeBarycentricCoordinates(DeformedPoint);
+			constexpr FReal Tolerance = 0;
+			const bool IsInside = Weights[0] >= Tolerance && Weights[1] >= Tolerance && Weights[2] >= Tolerance && Weights[0] + Weights[1] + Weights[2] <= 1 + Tolerance;
+			if (IsInside)
+			{
+				CoordinatesOut.Emplace(CellIndex, TestTetIndex, Weights);
+			}
+		}
+	}
+
+	if (CoordinatesOut.Num() > 0 || !bFindClosest)
+	{
+		return CoordinatesOut.Num() > 0;
+	}
+
+	check(bFindClosest);
+	FReal IntersectionTestRadius = 0.;
+	
+	// Keep increasing intersection test radius until find something.
+	const FReal IntersectionTestRadiusBase = bInLocalBoundingBox ? 0. : LocalBoundingBox.SignedDistance(DeformedPoint);
+	FReal IntersectionTestRadiusExpand= Grid.Dx().GetMax();
+	constexpr int32 MaxIterativeQueries = 5;
+	int32 IterativeQueryCount = 0;
+	while (IntersectingCells.IsEmpty() && IterativeQueryCount < MaxIterativeQueries)
+	{
+		IntersectionTestRadiusExpand *= 2.;
+		IntersectionTestRadius = IntersectionTestRadiusBase + IntersectionTestRadiusExpand;
+
+		TAABB<FReal, 3> QueryBounds(DeformedPoint, DeformedPoint);
+		QueryBounds.Thicken(IntersectionTestRadius);
+
+		IntersectingCells = Spatial.FindAllIntersections(QueryBounds);	
+		++IterativeQueryCount;
+	}
+
+	const bool bDoBruteForceQuery = IntersectingCells.IsEmpty() && IterativeQueryCount == MaxIterativeQueries;
+	if (bDoBruteForceQuery)
+	{
+		const int32 NumCells = Grid.GetNumCells();
+		IntersectingCells.Reserve(NumCells);
+		for (int32 FlatCellIndex = 0; FlatCellIndex < NumCells; ++FlatCellIndex)
+		{
+			if (!EmptyCells[FlatCellIndex])
+			{
+				IntersectingCells.Add(FlatCellIndex);
+			}
+		}
+	}
+
+	FEmbeddingCoordinate ClosestPointCoord;
+	FReal ClosestPointDistSq = std::numeric_limits<FReal>::max();
+	for (int32 IntersectingCell : IntersectingCells)
+	{
+		const TVec3<int32> CellIndex = Grid.GetIndex(IntersectingCell);
+		const TAABB<FReal, 3> BBox = FWeightedLatticeBvEntry({ this, CellIndex }).BoundingBox();
+		if (FMath::Square(BBox.SignedDistance(DeformedPoint)) > ClosestPointDistSq)
+		{
+			continue;
+		}
+		const TVec4<TVec3<int32>>* TestCellOffsets = ((CellIndex.X + CellIndex.Y + CellIndex.Z) & 1) == 0 ? FEmbeddingCoordinate::EvenIndexTetrahedraOffsets : FEmbeddingCoordinate::OddIndexTetrahedraOffsets;
+		for (int32 TestTetIndex = 0; TestTetIndex < 5; ++TestTetIndex)
+		{
+			const TTetrahedron<FReal> Tetrahedron(
+				DeformedPoints(CellIndex + TestCellOffsets[TestTetIndex][0]),
+				DeformedPoints(CellIndex + TestCellOffsets[TestTetIndex][1]),
+				DeformedPoints(CellIndex + TestCellOffsets[TestTetIndex][2]),
+				DeformedPoints(CellIndex + TestCellOffsets[TestTetIndex][3]));
+
+			TVec4<FReal> Barycentric;
+			const FVec3 ClosestPoint = Tetrahedron.FindClosestPointAndBary(DeformedPoint, Barycentric);
+			const FReal DistSq = FVec3::DistSquared(ClosestPoint, DeformedPoint);
+			if (DistSq < ClosestPointDistSq)
+			{
+				ClosestPointDistSq = DistSq;
+				ClosestPointCoord = FEmbeddingCoordinate(CellIndex, TestTetIndex, Barycentric);
+			}
+		}
+	}
+
+	check(ClosestPointCoord.IsValid());
+
+	// If closest point is farther than last search radius, need to do an expanded search to see if we missed anything.
+	if (!bDoBruteForceQuery && ClosestPointDistSq > FMath::Square(IntersectionTestRadius))
+	{
+		TAABB<FReal, 3> QueryBounds(DeformedPoint, DeformedPoint);
+		QueryBounds.Thicken(FMath::Sqrt(ClosestPointDistSq));
+
+		const TArray<int32> FinalIntersectingCells = Spatial.FindAllIntersections(QueryBounds);
+		const TSet<int32> OrigIntersectingCellsSet(IntersectingCells);
+		for (int32 IntersectingCell : FinalIntersectingCells)
+		{
+			if (OrigIntersectingCellsSet.Contains(IntersectingCell))
+			{
+				// This cell was tested in the last pass.
+				continue;
+			}
+
+			const TVec3<int32> CellIndex = Grid.GetIndex(IntersectingCell);
+			const TAABB<FReal, 3> BBox = FWeightedLatticeBvEntry({ this, CellIndex }).BoundingBox();
+			if (FMath::Square(BBox.SignedDistance(DeformedPoint)) > ClosestPointDistSq)
+			{
+				continue;
+			}
+			const TVec4<TVec3<int32>>* TestCellOffsets = ((CellIndex.X + CellIndex.Y + CellIndex.Z) & 1) == 0 ? FEmbeddingCoordinate::EvenIndexTetrahedraOffsets : FEmbeddingCoordinate::OddIndexTetrahedraOffsets;
+			for (int32 TestTetIndex = 0; TestTetIndex < 5; ++TestTetIndex)
+			{
+				const TTetrahedron<FReal> Tetrahedron(
+					DeformedPoints(CellIndex + TestCellOffsets[TestTetIndex][0]),
+					DeformedPoints(CellIndex + TestCellOffsets[TestTetIndex][1]),
+					DeformedPoints(CellIndex + TestCellOffsets[TestTetIndex][2]),
+					DeformedPoints(CellIndex + TestCellOffsets[TestTetIndex][3]));
+
+				TVec4<FReal> Barycentric;
+				const FVec3 ClosestPoint = Tetrahedron.FindClosestPointAndBary(DeformedPoint, Barycentric);
+				const FReal DistSq = FVec3::DistSquared(ClosestPoint, DeformedPoint);
+				if (DistSq < ClosestPointDistSq)
+				{
+					ClosestPointDistSq = DistSq;
+					ClosestPointCoord = FEmbeddingCoordinate(CellIndex, TestTetIndex, Barycentric);
+				}
+			}
+		}
+	}
+	check(ClosestPointCoord.IsValid());
+	CoordinatesOut.Add(ClosestPointCoord);
+	return true;
+}
+
+void FWeightedLatticeImplicitObject::InitializeDeformedPoints()
+{
+	constexpr bool NodeValues = true;
+	DeformedPoints.SetCounts(Grid, NodeValues);
+	const TVec3<int32> NodeCounts = Grid.NodeCounts();
+	for (int32 I = 0; I < NodeCounts.X; ++I)
+	{
+		for (int32 J = 0; J < NodeCounts.Y; ++J)
+		{
+			for (int32 K = 0; K < NodeCounts.Z; ++K)
+			{
+				DeformedPoints(I, J, K) = Grid.Node(TVec3<int32>(I, J, K));
+			}
+		}
+	}
+	LocalBoundingBox = TAABB<FReal, 3>(Grid.MinCorner(), Grid.MaxCorner());
+}
+
+void FWeightedLatticeImplicitObject::SetEmptyCells()
+{
+	constexpr bool NodeValues = false;
+	EmptyCells.SetCounts(Grid, NodeValues);
+	const TVec3<int32> CellCounts = Grid.Counts();
+	for (int32 I = 0; I < CellCounts.X; ++I)
+	{
+		for (int32 J = 0; J < CellCounts.Y; ++J)
+		{
+			for (int32 K = 0; K < CellCounts.Z; ++K)
+			{
+				// Cell is empty if any of its 8 corners have zero weights
+				const bool IsEmpty = (BoneData(I, J, K).NumInfluences == 0 ||
+					BoneData(I, J, K + 1).NumInfluences == 0 ||
+					BoneData(I, J + 1, K).NumInfluences == 0 ||
+					BoneData(I, J + 1, K + 1).NumInfluences == 0 ||
+					BoneData(I + 1, J, K).NumInfluences == 0 ||
+					BoneData(I + 1, J, K + 1).NumInfluences == 0 ||
+					BoneData(I + 1, J + 1, K).NumInfluences == 0 ||
+					BoneData(I + 1, J + 1, K + 1).NumInfluences == 0);
+				EmptyCells(I, J, K) = IsEmpty;
+			}
+		}
+	}
+}
+
+uint32 FWeightedLatticeImplicitObject::GetTypeHashHelper(const uint32 InHash) const
+{
+	uint32 Result = InHash;
+
+	// TypeHash BoneData
+	TConstArrayView< FWeightedLatticeInfluenceData > BoneDataFlattened(BoneData.GetData(), BoneData.Num());
+	for (const FWeightedLatticeInfluenceData& Data : BoneDataFlattened)
+	{
+		Result = HashCombine(Result, Data.GetTypeHash());
+	}
+
+	// TypeHash DeformedPoints
+	TConstArrayView<FVec3> DeformedPointsFlattened(DeformedPoints.GetData(), DeformedPoints.Num());
+	for (const FVec3& Point : DeformedPointsFlattened)
+	{
+		Result = HashCombine(Result, HashCombine(::GetTypeHash(Point.X), HashCombine(::GetTypeHash(Point.Y), ::GetTypeHash(Point.Z))));
+	}
+	return Result;
+}
+
+template<typename TConcrete>
+TWeightedLatticeImplicitObject<TConcrete>::TWeightedLatticeImplicitObject(ObjectType&& InObject, TUniformGrid<FReal, 3>&& InGrid,
+	TArrayND<FWeightedLatticeInfluenceData, 3>&& InBoneData, TArray<FName>&& InUsedBones, TArray<FTransform>&& InReferenceRelativeTransforms)
+	: FWeightedLatticeImplicitObject(EImplicitObject::HasBoundingBox, InObject->GetType(), MoveTemp(InGrid), MoveTemp(InBoneData), MoveTemp(InUsedBones), MoveTemp(InReferenceRelativeTransforms))
+	, Object(MoveTemp(InObject))
+{
+	this->bDoCollide = Object->GetDoCollide();
+}
+
+template<typename TConcrete>
+TWeightedLatticeImplicitObject<TConcrete>::TWeightedLatticeImplicitObject(TWeightedLatticeImplicitObject&& Other)
+	: FWeightedLatticeImplicitObject(MoveTemp(Other))
+	, Object(MoveTemp(Other.Object))
+{
+	this->bDoCollide = Other.Object->GetDoCollide();
+}
+
+template<typename TConcrete>
+FImplicitObject* TWeightedLatticeImplicitObject<TConcrete>::Duplicate() const
+{
+	if (Object)
+	{
+		ObjectType DuplicatedObject((TConcrete*)Object->Duplicate());
+		TWeightedLatticeImplicitObject* NewObj = new TWeightedLatticeImplicitObject(MoveTemp(DuplicatedObject), *this);
+		return NewObj;
+	}
+	else
+	{
+		check(false);
+		return nullptr;
+	}
+}
+
+template<typename TConcrete>
+TUniquePtr<FImplicitObject> TWeightedLatticeImplicitObject<TConcrete>::Copy() const
+{
+	if (Object)
+	{
+		return TUniquePtr<FImplicitObject>(CopyHelper(this));
+	}
+	else
+	{
+		check(false);
+		return nullptr;
+	}
+}
+
+template<typename TConcrete>
+TUniquePtr<FImplicitObject> TWeightedLatticeImplicitObject<TConcrete>::DeepCopy() const
+{
+	if (Object)
+	{
+		TUniquePtr<FImplicitObject> CopiedObject = Object->DeepCopy();
+		TUniquePtr<TWeightedLatticeImplicitObject> NewObj(new TWeightedLatticeImplicitObject(reinterpret_cast<ObjectType&&>(CopiedObject), *this));
+		return NewObj;
+	}
+	else
+	{
+		check(false);
+		return nullptr;
+	}
+}
+
+template<typename TConcrete>
+FReal TWeightedLatticeImplicitObject<TConcrete>::PhiWithNormalAndSurfacePoint(const FVec3& X, FVec3& Normal, FEmbeddingCoordinate& SurfaceCoord, bool bIncludeEmptyCells) const
+{
+	TArray<FEmbeddingCoordinate> Coordinates;
+	if (GetEmbeddingCoordinates(X, Coordinates))
+	{
+		// Choose "deepest" point here.
+		FReal MinPhiUndeformed = std::numeric_limits<FReal>::max();
+		int32 MinCoordIndex = -1;
+		FVec3 MinNormalUndeformed;
+		for (int32 CoordIndex = 0; CoordIndex < Coordinates.Num(); ++CoordIndex)
+		{
+			FVec3 UndeformedNormal;
+			const FReal Phi = Object->PhiWithNormal(Coordinates[CoordIndex].UndeformedPosition(Grid), UndeformedNormal);
+			if (Phi < MinPhiUndeformed)
+			{
+				MinPhiUndeformed = Phi;
+				MinCoordIndex = CoordIndex;
+				MinNormalUndeformed = UndeformedNormal;
+			}
+		}
+		const FEmbeddingCoordinate& Coordinate = Coordinates[MinCoordIndex];
+
+		const FVec3 UndeformedPoint = Coordinate.UndeformedPosition(Grid);
+		const FVec3 UndeformedSurfacePoint = UndeformedPoint - MinPhiUndeformed * MinNormalUndeformed;
+
+		const TVec3<int32> SurfaceCell = Grid.Cell(UndeformedSurfacePoint);
+		const FVec3 SurfaceAlpha = ((UndeformedSurfacePoint - Grid.Node(SurfaceCell)) / Grid.Dx()).BoundToBox(FVec3(0.), FVec3(1.));
+		SurfaceCoord = FEmbeddingCoordinate(SurfaceCell, SurfaceAlpha);
+
+		const FVec3 DeformedSurfacePoint = SurfaceCoord.DeformedPosition(DeformedPoints);
+		Normal = X - DeformedSurfacePoint;
+		const FReal NormalLen = Normal.Length();
+		const FReal Phi = MinPhiUndeformed < 0 ? -NormalLen : NormalLen;
+
+		if (NormalLen < UE_SMALL_NUMBER)
+		{
+			const FMatrix Transform = Coordinate.DeformationTransform(DeformedPoints, Grid);
+			Normal = Transform.TransformVector(MinNormalUndeformed.GetSafeNormal());
+			Normal.SafeNormalize();
+		}
+		else
+		{
+			Normal /= Phi;
+		}
+
+		return Phi;
+	}
+	else
+	{
+		SurfaceCoord = FEmbeddingCoordinate();
+		return UE_BIG_NUMBER;
+	}
+}
+
+void FWeightedLatticeImplicitObjectBuilder::GenerateGrid(const int32 GridResolution, const TAABB<FReal, 3>& ObjectBBox)
+{
+	constexpr int32 NumGhostCells = 3;
+	const int32 DesiredGridResolutionNoPad = FMath::Max(GridResolution - 2 * NumGhostCells, 1); // We'll add ghost cells to grid. 
+	const FVec3 Extents = ObjectBBox.Extents();
+	const FReal Dx = Extents.GetMax() / DesiredGridResolutionNoPad;
+	TVec3<int32> Cells(FMath::CeilToInt32(Extents.X / Dx), FMath::CeilToInt32(Extents.Y / Dx), FMath::CeilToInt32(Extents.Z / Dx));
+	Cells = Cells.ComponentwiseMax(TVec3<int32>(1, 1, 1));
+	const FVec3 MinCornerAdjusted = ObjectBBox.Min() - Dx * FVec3(1.) * .5; // Want bbox corners to be at cell centers.
+	const FVec3 MaxCornerAdjusted = ObjectBBox.Min() + FVec3((FReal)Cells.X + .5, (FReal)Cells.Y, (FReal)Cells.Z) * Dx;
+	Grid = TUniformGrid<FReal, 3>(ObjectBBox.Min(), MaxCornerAdjusted, Cells, NumGhostCells);
+
+	BoneData.Reset();
+	BoneData.SetCounts(Grid, true);
+	FMemory::Memzero(BoneData.GetData(), sizeof(FWeightedLatticeInfluenceData) * BoneData.Num());
+
+	BuildData.Reset();
+	BuildData.SetNum(BoneData.Num());
+	BuildStep = EBuildStep::GridValid;
+}
+
+void FWeightedLatticeImplicitObjectBuilder::AddInfluence(int32 FlatIndex, uint16 BoneIndex, float Weight, bool bIsOuterWeight)
+{
+	// Keep influences sorted by weight (highest to lowest)
+	check(BuildStep == EBuildStep::GridValid);
+	FWeightedLatticeInfluenceData& Data = BoneData[FlatIndex];
+	FInfluenceBuildData& Build = BuildData[FlatIndex];
+
+	if (bIsOuterWeight && !Build.WeightsAreOuter)
+	{
+		// Inner/surface weights take precedence over outer weights.
+		return;
+	}
+	if (!bIsOuterWeight && Build.WeightsAreOuter)
+	{
+		// Clear all existing outer weights
+		Data.NumInfluences = 0;
+	}
+
+	uint8 InsertIndex = Data.NumInfluences;
+	for (uint8 Idx = 0; Idx < Data.NumInfluences; ++Idx)
+	{
+		if (Weight > Data.BoneWeights[Idx])
+		{
+			InsertIndex = Idx;
+			break;
+		}
+	}
+	if (InsertIndex == Data.NumInfluences && Data.NumInfluences == Data.MaxTotalInfluences)
+	{
+		return;
+	}
+
+	if (Data.NumInfluences != Data.MaxTotalInfluences)
+	{
+		++Data.NumInfluences;
+	}
+
+	for (uint8 Idx = Data.NumInfluences - 1; Idx > InsertIndex; --Idx)
+	{
+		Data.BoneIndices[Idx] = Data.BoneIndices[Idx - 1];
+		Data.BoneWeights[Idx] = Data.BoneWeights[Idx - 1];
+	}
+
+	Data.BoneIndices[InsertIndex] = BoneIndex;
+	Data.BoneWeights[InsertIndex] = Weight;
+	Build.WeightsAreOuter = bIsOuterWeight;
+}
+
+void FWeightedLatticeImplicitObjectBuilder::NormalizeBoneWeights()
+{
+	for (int32 Idx = 0; Idx < BoneData.Num(); ++Idx)
+	{
+		FWeightedLatticeInfluenceData& Data = BoneData[Idx];
+		float TotalWeight = 0.f;
+		for (uint8 InfIdx = 0; InfIdx < Data.NumInfluences; ++InfIdx)
+		{
+			TotalWeight += Data.BoneWeights[InfIdx];
+		}
+		if (TotalWeight < UE_SMALL_NUMBER)
+		{
+			Data.NumInfluences = 0;
+		}
+		for (uint8 InfIdx = 0; InfIdx < Data.NumInfluences; ++InfIdx)
+		{
+			Data.BoneWeights[InfIdx] /= TotalWeight;
+		}
+	}
+}
+
+// Explicit instantiations
+template class CHAOS_API TWeightedLatticeImplicitObject<FLevelSet>;
+} // namespace Chaos
\ No newline at end of file
diff --git a/Engine/Source/Runtime/Experimental/Chaos/Public/Chaos/HierarchicalSpatialHash.h b/Engine/Source/Runtime/Experimental/Chaos/Public/Chaos/HierarchicalSpatialHash.h
--- a/Engine/Source/Runtime/Experimental/Chaos/Public/Chaos/HierarchicalSpatialHash.h
+++ b/Engine/Source/Runtime/Experimental/Chaos/Public/Chaos/HierarchicalSpatialHash.h
@@ -185,6 +185,15 @@ class THierarchicalSpatialHash
 			return true;
 		}
 
+		bool Intersects(const TVectorRegisterType<T>& Point) const
+		{
+			if (VectorAnyLesserThan(Point, Min) || VectorAnyGreaterThan(Point, Max))
+			{
+				return false;
+			}
+			return true;
+		}
+
 	private:
 
 		TVectorRegisterType<T> Min;
@@ -298,6 +307,39 @@ public:
 		return Result;
 	}
 
+	TArray<PayloadType> FindAllIntersections(const TVec3<T>& Point) const
+	{
+		TVectorRegisterType<T> PointVec(MakeVectorRegister(Point.X, Point.Y, Point.Z, (T)0));
+
+		TArray<PayloadType> Result;
+		for (const typename TMap<int32, T>::ElementType& UsedLod : UsedLods)
+		{
+			const int32 Lod = UsedLod.Key;
+			const T CellSize = UsedLod.Value;
+
+			const TVec3<int32> CellIdx = FHierarchicalSpatialHashCellIdx::CellIdxForPoint(Point, CellSize);
+			const FHierarchicalSpatialHashCellIdx Key{ CellIdx, Lod };
+			for (int32 HashIndex = HashMap.First(Key); HashMap.IsValid(HashIndex); HashIndex = HashMap.Next(HashIndex))
+			{
+				if (HashMap.KeyMatches(HashIndex, Key))
+				{
+					const FPayloadAndBounds Value = HashMap.GetValue(HashIndex);
+					const FVectorAABB& PayloadBounds = Bounds[Value.BoundsLookupIdx];
+
+					if (PayloadBounds.Intersects(PointVec))
+					{
+						// Typical cloth examples do not find a lot of unique results. 
+						// In testing, the overhead to do this search is lower than using a TSet and converting the result to an array
+						Result.AddUnique(Value.Payload);
+					}
+				}
+			}
+		}
+		return Result;
+	}
+
+
+
 private:
 	TMap<int32, T> UsedLods;
 	FHashMap HashMap;
diff --git a/Engine/Source/Runtime/Experimental/Chaos/Public/Chaos/ImplicitObjectType.h b/Engine/Source/Runtime/Experimental/Chaos/Public/Chaos/ImplicitObjectType.h
--- a/Engine/Source/Runtime/Experimental/Chaos/Public/Chaos/ImplicitObjectType.h
+++ b/Engine/Source/Runtime/Experimental/Chaos/Public/Chaos/ImplicitObjectType.h
@@ -27,11 +27,15 @@ namespace Chaos
 			Triangle,
 			UnionClustered,
 			TaperedCapsule,
+			WeightedLatticeBone,
 
 			//Add entries above this line for serialization
+			ConcreteObjectCount, // Used to ensure bitflags do not overlap concrete type
+			IsWeightedLattice = 1 << 5,
 			IsInstanced = 1 << 6,
 			IsScaled = 1 << 7
 		};
+		static_assert(ConcreteObjectCount <= IsWeightedLattice, "Too many Chaos::ImplicitObjectType concrete types.");
 	}
 
 	using EImplicitObjectType = uint8;	//see ImplicitObjectType
@@ -46,9 +50,14 @@ namespace Chaos
 		return (Type & ImplicitObjectType::IsScaled) != 0;
 	}
 
+	FORCEINLINE bool IsWeightedLattice(EImplicitObjectType Type)
+	{
+		return (Type & ImplicitObjectType::IsWeightedLattice) != 0;
+	}
+
 	FORCEINLINE EImplicitObjectType GetInnerType(EImplicitObjectType Type)
 	{
-		return Type & (~(ImplicitObjectType::IsScaled | ImplicitObjectType::IsInstanced));
+		return Type & (~(ImplicitObjectType::IsWeightedLattice | ImplicitObjectType::IsScaled | ImplicitObjectType::IsInstanced));
 	}
 
 	namespace EImplicitObject
diff --git a/Engine/Source/Runtime/Experimental/Chaos/Public/Chaos/Levelset.h b/Engine/Source/Runtime/Experimental/Chaos/Public/Chaos/Levelset.h
--- a/Engine/Source/Runtime/Experimental/Chaos/Public/Chaos/Levelset.h
+++ b/Engine/Source/Runtime/Experimental/Chaos/Public/Chaos/Levelset.h
@@ -140,6 +140,7 @@ class CHAOS_API FLevelSet final : public FImplicitObject
 
 	// Used to generate a simple debug surface
 	void GetZeroIsosurfaceGridCellFaces(TArray<FVector3f>& Vertices, TArray<FIntVector>& Tris) const;
+	void GetInteriorCells(TArray<TVec3<int32>>& InteriorCells, const FReal InteriorThreshold) const;
 
 	virtual uint32 GetTypeHash() const override
 	{
diff --git a/Engine/Source/Runtime/Experimental/Chaos/Public/Chaos/UniformGrid.h b/Engine/Source/Runtime/Experimental/Chaos/Public/Chaos/UniformGrid.h
--- a/Engine/Source/Runtime/Experimental/Chaos/Public/Chaos/UniformGrid.h
+++ b/Engine/Source/Runtime/Experimental/Chaos/Public/Chaos/UniformGrid.h
@@ -263,7 +263,7 @@ class CHAOS_API TUniformGridBase
 };
 
 template<class T, int d>
-class TUniformGrid : public TUniformGridBase<T, d>
+class CHAOS_API TUniformGrid : public TUniformGridBase<T, d>
 {
 	using TUniformGridBase<T, d>::MCells;
 	using TUniformGridBase<T, d>::MMinCorner;
@@ -416,7 +416,7 @@ public:
 
 
 template<class T>
-class TUniformGrid<T, 3> : public TUniformGridBase<T, 3>
+class CHAOS_API TUniformGrid<T, 3> : public TUniformGridBase<T, 3>
 {
 	using TUniformGridBase<T, 3>::MCells;
 	using TUniformGridBase<T, 3>::MMinCorner;
@@ -426,6 +426,7 @@ class TUniformGrid<T, 3> : public TUniformGridBase<T, 3>
   public:
 	using TUniformGridBase<T, 3>::GetNumCells;
 	using TUniformGridBase<T, 3>::Location;
+	using TUniformGridBase<T, 3>::Node;
 
 	TUniformGrid() {}
 	TUniformGrid(const TVector<T, 3>& MinCorner, const TVector<T, 3>& MaxCorner, const TVector<int32, 3>& Cells, const uint32 GhostCells = 0)
@@ -449,6 +450,8 @@ class TUniformGrid<T, 3> : public TUniformGridBase<T, 3>
 	TVector<T, 3> Clamp(const TVector<T, 3>& X) const;
 	TVector<T, 3> ClampMinusHalf(const TVector<T, 3>& X) const;
 	bool IsValid(const TVector<int32, 3>& X) const;
+
+	TUniformGrid<T, 3> SubGrid(const TVector<int32, 3>& MinCell, const TVector<int32, 3>& MaxCell) const;
 };
 
 
diff --git a/Engine/Source/Runtime/Experimental/Chaos/Public/Chaos/WeightedLatticeImplicitObject.h b/Engine/Source/Runtime/Experimental/Chaos/Public/Chaos/WeightedLatticeImplicitObject.h
new file mode 100644
--- /dev/null
+++ b/Engine/Source/Runtime/Experimental/Chaos/Public/Chaos/WeightedLatticeImplicitObject.h
@@ -0,0 +1,430 @@
+// Copyright Epic Games, Inc. All Rights Reserved.
+#pragma once
+
+#include "Chaos/ImplicitObject.h"
+#include "Chaos/ArrayND.h"
+#include "Chaos/AABB.h"
+#include "Chaos/HierarchicalSpatialHash.h"
+
+struct FKSkinnedLevelSetElem;
+
+namespace Chaos
+{
+struct FWeightedLatticeImplicitObjectFactory;
+
+struct FWeightedLatticeInfluenceData
+{
+	FWeightedLatticeInfluenceData()
+	: NumInfluences(0)
+	{
+		FMemory::Memset(BoneIndices, (uint8)INDEX_NONE, sizeof(BoneIndices));
+		FMemory::Memset(BoneWeights, 0, sizeof(BoneWeights));
+	}
+
+	uint32 GetTypeHash() const
+	{
+		uint32 Result = ::GetTypeHash(NumInfluences);
+		for (uint8 Index = 0; Index < NumInfluences; ++Index)
+		{
+			Result = HashCombine(Result, ::GetTypeHash(BoneIndices[Index]));
+			Result = HashCombine(Result, ::GetTypeHash(BoneWeights[Index]));
+		}
+		return Result;
+	}
+
+	static const uint8 MaxTotalInfluences = 12;
+	uint8 NumInfluences;
+	uint16 BoneIndices[MaxTotalInfluences]; // This is index into FWeightedLatticeImplicitObject::UsedBones (which then maps into skeletalmesh bones)
+	float BoneWeights[MaxTotalInfluences];
+};
+
+inline FArchive& operator<<(FArchive& Ar, FWeightedLatticeInfluenceData& Value)
+{
+	Ar << Value.NumInfluences;
+	for (int32 Index = 0; Index < Value.NumInfluences; ++Index)
+	{
+		Ar << Value.BoneIndices[Index];
+		Ar << Value.BoneWeights[Index];
+	}
+	return Ar;
+}
+
+/**
+ * Embed the contained shape in a deformable lattice.
+ */
+class CHAOS_API FWeightedLatticeImplicitObject : public FImplicitObject
+{
+public:
+	using ObjectType = TUniquePtr<FImplicitObject>;
+	using FImplicitObject::GetTypeName;
+
+	FWeightedLatticeImplicitObject(int32 Flags, EImplicitObjectType InType, TUniformGrid<FReal, 3>&& InGrid,
+		TArrayND<FWeightedLatticeInfluenceData, 3>&& InBoneData, TArray<FName>&& InUsedBones, TArray<FTransform>&& InReferenceRelativeTransforms);
+
+	FWeightedLatticeImplicitObject(FWeightedLatticeImplicitObject&& Other);
+	
+	virtual ~FWeightedLatticeImplicitObject() override = default;
+
+	virtual const FAABB3 BoundingBox() const override { return LocalBoundingBox; }
+	const TArrayND<FWeightedLatticeInfluenceData, 3>& GetBoneData() const { return BoneData; }
+	const TArrayND<FVec3, 3>& GetDeformedPoints() const { return DeformedPoints; }
+	const TArrayND<bool, 3>& GetEmptyCells() const { return EmptyCells; }
+	const TUniformGrid<FReal, 3>& GetGrid() const { return Grid; }
+	const TArray<FName>& GetUsedBones() const { return UsedBones; }
+	const TArray<int32>& GetSolverBoneIndices() const { return SolverBoneIndices; }
+	void SetSolverBoneIndices(TArray<int32>&& InSolverBoneIndices) { SolverBoneIndices = MoveTemp(InSolverBoneIndices); }
+
+	virtual void Serialize(FChaosArchive& Ar) override;
+
+	FVec3 GetDeformedPoint(const FVec3& UndeformedPoint) const;
+	void UpdateSpatialHierarchy();
+	void FinalizeConstruction();
+	void DeformPoints(const TArray<FTransform>& RelativeTransforms);
+
+	struct CHAOS_API FEmbeddingCoordinate
+	{
+		TVec3<int32> CellIndex = TVec3<int32>(INDEX_NONE, INDEX_NONE, INDEX_NONE);
+		int32 LocalTetrahedron = INDEX_NONE;
+		FVec3 BarycentricCoordinate;
+
+		FEmbeddingCoordinate() = default;
+		FEmbeddingCoordinate(const FEmbeddingCoordinate& Other) = default;
+		FEmbeddingCoordinate& operator=(const FEmbeddingCoordinate& Other) = default;
+
+		FEmbeddingCoordinate(const TVec3<int32>& InCellIndex, int32 InLocalTetrahedron, const FVec3& InBarycentric)
+			:CellIndex(InCellIndex), LocalTetrahedron(InLocalTetrahedron), BarycentricCoordinate(InBarycentric)
+		{}
+
+		FEmbeddingCoordinate(const TVec3<int32>& InCellIndex, const FVec3& TrilinearCoordinate);
+
+		bool IsValid() const
+		{
+			checkSlow((CellIndex == TVec3<int32>(INDEX_NONE, INDEX_NONE, INDEX_NONE)) == (LocalTetrahedron == INDEX_NONE));
+			return LocalTetrahedron != INDEX_NONE;
+		}
+
+		inline static const TVec4<TVec3<int32>> EvenIndexTetrahedraOffsets[5] =
+		{
+			TVec4<TVec3<int32>>(TVec3<int32>(0,1,0), TVec3<int32>(0,0,0), TVec3<int32>(1,1,0), TVec3<int32>(0,1,1)),
+			TVec4<TVec3<int32>>(TVec3<int32>(1,1,1), TVec3<int32>(1,1,0), TVec3<int32>(1,0,1), TVec3<int32>(0,1,1)),
+			TVec4<TVec3<int32>>(TVec3<int32>(1,0,0), TVec3<int32>(1,0,1), TVec3<int32>(1,1,0), TVec3<int32>(0,0,0)),
+			TVec4<TVec3<int32>>(TVec3<int32>(0,0,1), TVec3<int32>(0,0,0), TVec3<int32>(0,1,1), TVec3<int32>(1,0,1)),
+			TVec4<TVec3<int32>>(TVec3<int32>(1,1,0), TVec3<int32>(0,1,1), TVec3<int32>(0,0,0), TVec3<int32>(1,0,1)),
+		};
+
+		inline static const TVec4<TVec3<int32>> OddIndexTetrahedraOffsets[5] =
+		{
+			TVec4<TVec3<int32>>(TVec3<int32>(0,0,0), TVec3<int32>(1,0,0), TVec3<int32>(0,1,0), TVec3<int32>(0,0,1)),
+			TVec4<TVec3<int32>>(TVec3<int32>(0,1,1), TVec3<int32>(0,1,0), TVec3<int32>(1,1,1), TVec3<int32>(0,0,1)),
+			TVec4<TVec3<int32>>(TVec3<int32>(1,0,1), TVec3<int32>(1,1,1), TVec3<int32>(1,0,0), TVec3<int32>(0,0,1)),
+			TVec4<TVec3<int32>>(TVec3<int32>(1,1,0), TVec3<int32>(0,1,0), TVec3<int32>(1,0,0), TVec3<int32>(1,1,1)),
+			TVec4<TVec3<int32>>(TVec3<int32>(1,0,0), TVec3<int32>(1,1,1), TVec3<int32>(0,1,0), TVec3<int32>(0,0,1)),
+		};
+
+		inline const TVec4<TVec3<int32>>& TetrahedronOffsets() const
+		{
+			checkSlow(LocalTetrahedron >= 0 && LocalTetrahedron < 5);
+
+			if (((CellIndex.X + CellIndex.Y + CellIndex.Z) & 1) == 0)
+			{
+				return EvenIndexTetrahedraOffsets[LocalTetrahedron];
+			}
+			else
+			{
+				return OddIndexTetrahedraOffsets[LocalTetrahedron];
+			}
+		}
+
+		FMatrix DeformationTransform(const TArrayND<FVec3, 3>& DeformedPoints, const TUniformGrid<FReal, 3>& Grid) const;
+		inline FVec3 UndeformedPosition(const TUniformGrid<FReal, 3>& InGrid) const;
+		inline FVec3 DeformedPosition(const TArrayND<FVec3, 3>& InDeformedPoints) const;
+		inline int32 GreatestInfluenceBone(const TArrayND<FWeightedLatticeInfluenceData, 3>& InBoneData) const;
+	};
+
+	bool GetEmbeddingCoordinates(const FVec3& DeformedPoint, TArray<FEmbeddingCoordinate>& CoordinatesOut, bool bFindClosest = false) const;
+
+protected:
+	FWeightedLatticeImplicitObject(int32 Flags, EImplicitObjectType InType)
+		:FImplicitObject(Flags, InType | ImplicitObjectType::IsWeightedLattice)
+	{}
+
+	FWeightedLatticeImplicitObject(const FWeightedLatticeImplicitObject& Other);
+
+	void InitializeDeformedPoints();
+	void SetEmptyCells();
+	uint32 GetTypeHashHelper(const uint32 InHash) const;
+
+
+	// Serialized data. Only non-const because of serialization
+	TUniformGrid<FReal, 3> Grid;
+	TArrayND<FWeightedLatticeInfluenceData, 3> BoneData;
+	TArray<FName> UsedBones;
+	FTransform ReferenceRootTransform;
+	TArray<FTransform> ReferenceRelativeTransforms; // ReferenceRootTransform * RefBaseMatrixInv(UsedBoneIdx)
+
+	// Non-serialized cached (const) data derived from serialized data
+	TArrayND<bool, 3> EmptyCells;
+
+	// Context-specific data (e.g., data a specific solver will need to do DeformPoints)
+	TArray<int32> SolverBoneIndices; // indices set by the solver for bone transform query
+
+	// Non-serialized deforming data
+	TAABB<FReal, 3> LocalBoundingBox;
+	TArrayND<FVec3, 3> DeformedPoints;
+	THierarchicalSpatialHash<int32, FReal> Spatial;
+	bool bSpatialDirty = true;
+};
+
+template<typename TConcrete>
+class TWeightedLatticeImplicitObject : public FWeightedLatticeImplicitObject
+{
+public:
+	using T = typename TConcrete::TType;
+	using TType = T;
+	using ObjectType = TUniquePtr<TConcrete>;
+
+	TWeightedLatticeImplicitObject(ObjectType&& InObject, TUniformGrid<FReal, 3>&& InGrid,
+		TArrayND<FWeightedLatticeInfluenceData, 3>&& InBoneData, TArray<FName>&& InUsedBones, TArray<FTransform>&& InReferenceRelativeTransforms);
+
+	TWeightedLatticeImplicitObject(TWeightedLatticeImplicitObject&& Other);
+
+	static constexpr EImplicitObjectType StaticType()
+	{
+		return TConcrete::StaticType() | ImplicitObjectType::IsWeightedLattice;
+	}
+
+	const TConcrete* GetEmbeddedObject() const
+	{
+		return Object.Get();
+	}
+
+	virtual FImplicitObject* Duplicate() const override;
+	virtual TUniquePtr<FImplicitObject> Copy() const override;
+	virtual TUniquePtr<FImplicitObject> DeepCopy() const override;
+
+	virtual uint32 GetTypeHash() const override
+	{
+		return GetTypeHashHelper(Object->GetTypeHash());
+	}
+
+	// Note that if x is in an empty cell / outside the deformed lattice and bIncludeEmptyCells is false, this will return UE_BIG_NUMBER and Normal will be unchanged.
+	// SurfaceCoord will be invalid.
+	// When bIncludeEmptyCells is true, the closest non-empty lattice cell will be used to find the closest surface point.
+	FReal PhiWithNormalAndSurfacePoint(const FVec3& X, FVec3& Normal, FEmbeddingCoordinate& SurfaceCoord, bool bIncludeEmptyCells = false) const;
+
+	virtual FReal PhiWithNormal(const FVec3& x, FVec3& Normal) const override
+	{
+		FEmbeddingCoordinate SurfaceCoordUnused;
+		constexpr bool bIncludeEmptyCells = true;
+		return PhiWithNormalAndSurfacePoint(x, Normal, SurfaceCoordUnused, bIncludeEmptyCells);
+	}
+
+	virtual void Serialize(FChaosArchive& Ar) override
+	{
+		FChaosArchiveScopedMemory ScopedMemory(Ar, GetTypeName(), false);
+		FWeightedLatticeImplicitObject::Serialize(Ar);
+		Ar << Object;
+	}
+
+private:
+	// Serialized data
+	ObjectType Object;
+
+	static TWeightedLatticeImplicitObject<TConcrete>* CopyHelper(const TWeightedLatticeImplicitObject<TConcrete>* Obj)
+	{
+		TUniquePtr<FImplicitObject> CopiedShape = Obj->Object->Copy();
+		return new TWeightedLatticeImplicitObject<TConcrete>(reinterpret_cast<TUniquePtr<TConcrete>&&>(CopiedShape), *Obj);
+	}
+
+	TWeightedLatticeImplicitObject(ObjectType&& InObject, const FWeightedLatticeImplicitObject& Other)
+		: FWeightedLatticeImplicitObject(Other)
+		, Object(MoveTemp(InObject))
+	{}
+
+	//needed for serialization
+	TWeightedLatticeImplicitObject() 
+		: FWeightedLatticeImplicitObject(EImplicitObject::HasBoundingBox, StaticType()) 
+		, Object(nullptr) 
+	{}
+
+	friend FImplicitObject;	//needed for serialization
+	friend ::FKSkinnedLevelSetElem; //needed for serialization
+};
+
+class CHAOS_API FWeightedLatticeImplicitObjectBuilder
+{
+public:
+
+	void GenerateGrid(const int32 GridResolution, const TAABB<FReal, 3>& ObjectBBox);
+	void AddInfluence(int32 FlatIndex, uint16 BoneIndex, float Weight, bool bIsOuterWeight);
+
+	// BoneIndexToBoneName(int32 BoneIndex)->FName return FName of Bone for given BoneIndex (BoneIndex used by AddInfluence)
+	// BoneIndexToReferenceTransform(int32 BoneIndex)->FTransform return transform of Bone (BoneIndex used by AddInfluence). Should return transform for RootBone (parent of this object) when BoneIndex == INDEX_NONE
+	template<typename FBoneIndexToBoneName, typename FBoneIndexToReferenceTransform>
+	inline void FinalizeInfluences(const FBoneIndexToBoneName& BoneIndexToBoneName, const FBoneIndexToReferenceTransform& BoneIndexToReferenceTransform);
+
+	template<typename TConcrete>
+	inline TUniquePtr< TWeightedLatticeImplicitObject<TConcrete> > Generate(TUniquePtr<TConcrete>&& Object);
+
+	const TUniformGrid<FReal, 3>& GetGrid() const { return Grid; }
+	const TArrayND<FWeightedLatticeInfluenceData, 3>& GetBoneData() const { return BoneData; }
+
+private:
+
+	void NormalizeBoneWeights();
+
+	template<typename FBoneIndexToBoneName, typename FBoneIndexToReferenceTransform>
+	inline void CalcUsedBonesAndReIndex(const FBoneIndexToBoneName& BoneIndexToBoneName, const FBoneIndexToReferenceTransform& BoneIndexToReferenceTransform);
+
+
+	TUniformGrid<FReal, 3> Grid;
+	TArrayND<FWeightedLatticeInfluenceData, 3> BoneData;
+	TArray<FName> UsedBones;
+	TArray<FTransform> ReferenceRelativeTransforms;
+
+	// Add on to FWeightedLatticeInfluenceData
+	struct FInfluenceBuildData
+	{
+		bool WeightsAreOuter = true;
+	};
+	TArray<FInfluenceBuildData> BuildData;
+
+	enum struct EBuildStep
+	{
+		None = 0,
+		GridValid,
+		InfluencesFinalized,
+		Finished = None, // Once Finished, all data has been moved out and the builder is no longer valid
+	};
+	EBuildStep BuildStep = EBuildStep::None;
+};
+
+// Not actually a collision object. CollisionParticles that are used to update bones used by the FWeightedLatticeImplicitObject have this as their geometry
+class FWeightedLatticeBoneProxy : public FImplicitObject
+{
+public:
+	using FImplicitObject::GetTypeName;
+
+	FWeightedLatticeBoneProxy()
+		: FImplicitObject(EImplicitObject::DisableCollisions, ImplicitObjectType::WeightedLatticeBone)
+	{}
+
+	static constexpr EImplicitObjectType StaticType() { return ImplicitObjectType::WeightedLatticeBone; }
+
+	virtual FReal PhiWithNormal(const FVec3& x, FVec3& Normal) const override
+	{
+		return UE_BIG_NUMBER;
+	}
+
+	virtual uint32 GetTypeHash() const override
+	{
+		return 0;
+	}
+};
+
+
+// inline definitions
+inline FVec3 FWeightedLatticeImplicitObject::FEmbeddingCoordinate::UndeformedPosition(const TUniformGrid<FReal, 3>& InGrid) const
+{
+	const TVec4<TVec3<int32>>& TetOffsets = TetrahedronOffsets();
+	return BarycentricCoordinate.X * InGrid.Node(CellIndex + TetOffsets[0]) + BarycentricCoordinate.Y * InGrid.Node(CellIndex + TetOffsets[1]) + BarycentricCoordinate.Z * InGrid.Node(CellIndex + TetOffsets[2])
+		+ (1. - BarycentricCoordinate.X - BarycentricCoordinate.Y - BarycentricCoordinate.Z) * InGrid.Node(CellIndex + TetOffsets[3]);
+}
+
+inline FVec3 FWeightedLatticeImplicitObject::FEmbeddingCoordinate::DeformedPosition(const TArrayND<FVec3, 3>& InDeformedPoints) const
+{
+	const TVec4<TVec3<int32>>& TetOffsets = TetrahedronOffsets();
+
+	return BarycentricCoordinate.X * InDeformedPoints(CellIndex + TetOffsets[0]) + BarycentricCoordinate.Y * InDeformedPoints(CellIndex + TetOffsets[1]) + BarycentricCoordinate.Z * InDeformedPoints(CellIndex + TetOffsets[2])
+		+ (1. - BarycentricCoordinate.X - BarycentricCoordinate.Y - BarycentricCoordinate.Z) * InDeformedPoints(CellIndex + TetOffsets[3]);
+}
+
+inline int32 FWeightedLatticeImplicitObject::FEmbeddingCoordinate::GreatestInfluenceBone(const TArrayND<FWeightedLatticeInfluenceData, 3>& InBoneData) const
+{
+	if (!IsValid())
+	{
+		return INDEX_NONE;
+	}
+
+	const TVec4<TVec3<int32>>& TetOffsets = TetrahedronOffsets();
+	int32 BoneIndex = INDEX_NONE;
+	const TVec4<float> BarycentricCoordinate4((float)BarycentricCoordinate.X, (float)BarycentricCoordinate.Y, 
+		(float)BarycentricCoordinate.Z, 1.f - (float)BarycentricCoordinate.X - (float)BarycentricCoordinate.Y - (float)BarycentricCoordinate.Z);
+	float LargestWeight = 0.f;
+	for (int32 Coord = 0; Coord < 4; ++Coord)
+	{
+		const FWeightedLatticeInfluenceData& Influences = InBoneData(CellIndex + TetOffsets[Coord]);
+		for (uint8 InfIndex = 0; InfIndex < Influences.NumInfluences; ++InfIndex)
+		{
+			const float Weight = BarycentricCoordinate4[Coord] * Influences.BoneWeights[InfIndex];
+			if (Weight > LargestWeight)
+			{
+				LargestWeight = Weight;
+				BoneIndex = Influences.BoneIndices[InfIndex];
+			}
+		}
+	}
+
+	return BoneIndex;
+}
+
+template<typename FBoneIndexToBoneName, typename FBoneIndexToReferenceTransform>
+inline void FWeightedLatticeImplicitObjectBuilder::CalcUsedBonesAndReIndex(const FBoneIndexToBoneName& BoneIndexToBoneName, const FBoneIndexToReferenceTransform& BoneIndexToReferenceTransform)
+{
+	TSet<uint16> UsedBoneSet;
+	for (int32 Idx = 0; Idx < BoneData.Num(); ++Idx)
+	{
+		const FWeightedLatticeInfluenceData& Data = BoneData[Idx];
+		for (uint8 InfIdx = 0; InfIdx < Data.NumInfluences; ++InfIdx)
+		{
+			UsedBoneSet.Add(Data.BoneIndices[InfIdx]);
+		}
+	}
+	const TArray<uint16> UsedBoneIndices = UsedBoneSet.Array();
+	UsedBones.Reset();
+	UsedBones.SetNumUninitialized(UsedBoneIndices.Num());
+	ReferenceRelativeTransforms.Reset();
+	ReferenceRelativeTransforms.SetNumUninitialized(UsedBoneIndices.Num());
+
+	const FTransform ReferenceRootTransform = (BoneIndexToReferenceTransform(INDEX_NONE)).Inverse();
+
+	TMap<uint16, uint16> UsedBoneReverseLookup;
+	for (uint16 ArrayIdx = 0; ArrayIdx < (uint16)UsedBoneIndices.Num(); ++ArrayIdx)
+	{
+		UsedBones[ArrayIdx] = BoneIndexToBoneName(UsedBoneIndices[ArrayIdx]);
+		ReferenceRelativeTransforms[ArrayIdx] = ReferenceRootTransform * BoneIndexToReferenceTransform(UsedBoneIndices[ArrayIdx]);
+		UsedBoneReverseLookup.Add(UsedBoneIndices[ArrayIdx], ArrayIdx);
+	}
+
+	for (int32 Idx = 0; Idx < BoneData.Num(); ++Idx)
+	{
+		FWeightedLatticeInfluenceData& Data = BoneData[Idx];
+		for (uint8 InfIdx = 0; InfIdx < Data.NumInfluences; ++InfIdx)
+		{
+			Data.BoneIndices[InfIdx] = UsedBoneReverseLookup[Data.BoneIndices[InfIdx]];
+		}
+	}
+}
+
+template<typename FBoneIndexToBoneName, typename FBoneIndexToReferenceTransform>
+inline void FWeightedLatticeImplicitObjectBuilder::FinalizeInfluences(const FBoneIndexToBoneName& BoneIndexToBoneName, const FBoneIndexToReferenceTransform& BoneIndexToReferenceTransform)
+{
+	check(BuildStep == EBuildStep::GridValid);
+	NormalizeBoneWeights();
+	CalcUsedBonesAndReIndex(BoneIndexToBoneName, BoneIndexToReferenceTransform);
+	BuildStep = EBuildStep::InfluencesFinalized;
+}
+
+template<typename TConcrete>
+inline TUniquePtr< TWeightedLatticeImplicitObject<TConcrete> > FWeightedLatticeImplicitObjectBuilder::Generate(TUniquePtr<TConcrete>&& Object)
+{
+	check(BuildStep == EBuildStep::InfluencesFinalized);
+
+	TUniquePtr< TWeightedLatticeImplicitObject<TConcrete> > Ret = MakeUnique< TWeightedLatticeImplicitObject<TConcrete> >(
+		MoveTemp(Object), MoveTemp(Grid), MoveTemp(BoneData), MoveTemp(UsedBones), MoveTemp(ReferenceRelativeTransforms));
+
+	BuildStep = EBuildStep::Finished;
+
+	return Ret;
+}
+}
diff --git a/Engine/Source/Runtime/Experimental/ChaosCore/Public/Chaos/ArrayND.h b/Engine/Source/Runtime/Experimental/ChaosCore/Public/Chaos/ArrayND.h
--- a/Engine/Source/Runtime/Experimental/ChaosCore/Public/Chaos/ArrayND.h
+++ b/Engine/Source/Runtime/Experimental/ChaosCore/Public/Chaos/ArrayND.h
@@ -167,6 +167,9 @@ class TArrayNDBase
 		MArray.Reset();
 	}
 
+	FORCEINLINE T* GetData() { return MArray.GetData(); }
+	FORCEINLINE const T* GetData() const { return MArray.GetData(); }
+
   protected:
 	TVector<int32, d> MCounts;
 	TArray<T> MArray;
@@ -240,15 +243,14 @@ class TArrayND<T, 3> : public TArrayNDBase<TArrayND<T, 3>, T, 3>
 #else
 	FORCEINLINE TArrayND() { MCounts = TVec3<int32>(0); }
 #endif
-	FORCEINLINE TArrayND(const TUniformGrid<FReal, 3>& grid)
+	template<typename U>
+	FORCEINLINE TArrayND(const TUniformGrid<U, 3>& Grid, bool NodeValues = false)
 	{
-		MCounts = grid.Counts();
-		MArray.SetNum(MCounts[0] * MCounts[1] * MCounts[2]);
+		SetCounts(Grid, NodeValues);
 	}
 	FORCEINLINE TArrayND(const TVec3<int32>& Counts)
 	{
-		MCounts = Counts;
-		MArray.SetNum(MCounts[0] * MCounts[1] * MCounts[2]);
+		SetCounts(Counts);
 	}
 	FORCEINLINE TArrayND(const TVec3<int32>& Counts, const TArray<T>& Array)
 	    : Base(Counts, Array) { check(Counts.Product() == Array.Num()); }
@@ -279,6 +281,20 @@ class TArrayND<T, 3> : public TArrayNDBase<TArrayND<T, 3>, T, 3>
 	{
 		return MArray[(x * MCounts[1] + y) * MCounts[2] + z];
 	}
+
+	FORCEINLINE void SetCounts(const TVector<int32, 3>& Counts)
+	{
+		MCounts = Counts;
+		MArray.SetNum(MCounts[0] * MCounts[1] * MCounts[2]);
+	}
+
+	template<typename U>
+	FORCEINLINE void SetCounts(const TUniformGrid<U, 3>& Grid, bool NodeValues = false)
+	{
+		MCounts = NodeValues ? Grid.NodeCounts() : Grid.Counts();
+		MArray.SetNum(MCounts[0] * MCounts[1] * MCounts[2]);
+	}
+	
 };
 
 #if COMPILE_WITHOUT_UNREAL_SUPPORT
