Fix comment and overriden method in ISM Pool subsystem

diff --git a/Engine/Source/Runtime/Experimental/GeometryCollectionEngine/Private/GeometryCollection/GeometryCollectionISMPoolSubSystem.cpp b/Engine/Source/Runtime/Experimental/GeometryCollectionEngine/Private/GeometryCollection/GeometryCollectionISMPoolSubSystem.cpp
--- a/Engine/Source/Runtime/Experimental/GeometryCollectionEngine/Private/GeometryCollection/GeometryCollectionISMPoolSubSystem.cpp
+++ b/Engine/Source/Runtime/Experimental/GeometryCollectionEngine/Private/GeometryCollection/GeometryCollectionISMPoolSubSystem.cpp
@@ -30,11 +30,6 @@ void UGeometryCollectionISMPoolSubSystem::Deinitialize()
 	Super::Deinitialize();
 }
 
-void UGeometryCollectionISMPoolSubSystem::OnWorldBeginPlay(UWorld& InWorld)
-{
-	Super::OnWorldBeginPlay(InWorld);
-}
-
 AGeometryCollectionISMPoolActor* UGeometryCollectionISMPoolSubSystem::FindISMPoolActor(const USceneComponent& Requester)
 {
 	// on demand creation of the actor 
diff --git a/Engine/Source/Runtime/Experimental/GeometryCollectionEngine/Public/GeometryCollection/GeometryCollectionISMPoolSubSystem.h b/Engine/Source/Runtime/Experimental/GeometryCollectionEngine/Public/GeometryCollection/GeometryCollectionISMPoolSubSystem.h
--- a/Engine/Source/Runtime/Experimental/GeometryCollectionEngine/Public/GeometryCollection/GeometryCollectionISMPoolSubSystem.h
+++ b/Engine/Source/Runtime/Experimental/GeometryCollectionEngine/Public/GeometryCollection/GeometryCollectionISMPoolSubSystem.h
@@ -8,7 +8,7 @@
 class AGeometryCollectionISMPoolActor;
 
 /**
-* A subsystem managing StateTree assets in Mass
+* A subsystem managing ISMPool actors ( used by geometry collection now but repurposed for more general use )
 */
 UCLASS()
 class GEOMETRYCOLLECTIONENGINE_API UGeometryCollectionISMPoolSubSystem : public UWorldSubsystem
@@ -24,10 +24,6 @@ public:
 	virtual void Deinitialize() override;
 	// USubsystem END
 
-	// UWorldSubsystem BEGIN
-	virtual void OnWorldBeginPlay(UWorld& InWorld) override;
-	// UWorldSubsystem END
-
 	AGeometryCollectionISMPoolActor* FindISMPoolActor(const USceneComponent& Requester);
 
 protected:
