Geometry Collection , add setter on EnableNanite UProperty

diff --git a/Engine/Source/Runtime/Experimental/GeometryCollectionEngine/Private/GeometryCollection/GeometryCollectionObject.cpp b/Engine/Source/Runtime/Experimental/GeometryCollectionEngine/Private/GeometryCollection/GeometryCollectionObject.cpp
--- a/Engine/Source/Runtime/Experimental/GeometryCollectionEngine/Private/GeometryCollection/GeometryCollectionObject.cpp
+++ b/Engine/Source/Runtime/Experimental/GeometryCollectionEngine/Private/GeometryCollection/GeometryCollectionObject.cpp
@@ -973,6 +973,23 @@ void UGeometryCollection::CreateSimulationDataImp(bool bCopyFromDDC)
 	}
 }
 
+void UGeometryCollection::SetEnableNanite(bool bValue)
+{
+	if (EnableNanite != bValue)
+	{
+		EnableNanite = bValue;
+		NaniteData = MakeUnique<FGeometryCollectionNaniteData>();
+
+		if (EnableNanite)
+		{
+			if (GeometryCollection)
+			{
+				NaniteData = UGeometryCollection::CreateNaniteData(GeometryCollection.Get());
+			}
+		}
+	}
+}
+
 void UGeometryCollection::CreateSimulationData()
 {
 	CreateSimulationDataImp(/*bCopyFromDDC=*/false);
diff --git a/Engine/Source/Runtime/Experimental/GeometryCollectionEngine/Public/GeometryCollection/GeometryCollectionObject.h b/Engine/Source/Runtime/Experimental/GeometryCollectionEngine/Public/GeometryCollection/GeometryCollectionObject.h
--- a/Engine/Source/Runtime/Experimental/GeometryCollectionEngine/Public/GeometryCollection/GeometryCollectionObject.h
+++ b/Engine/Source/Runtime/Experimental/GeometryCollectionEngine/Public/GeometryCollection/GeometryCollectionObject.h
@@ -527,13 +527,16 @@ public:
 	/**
 	 * Strip unnecessary data from the Geometry Collection to keep the memory footprint as small as possible.
 	 */
-	UPROPERTY(EditAnywhere, Category = "Nanite")
+	UPROPERTY(EditAnywhere, BlueprintReadWrite, Category = "Nanite")
 	bool bStripOnCook;
 
+	UFUNCTION(BlueprintCallable, Category = "Nanite")
+	void SetEnableNanite(bool bValue);
+
 	/**
 	 * Enable support for Nanite.
 	 */
-	UPROPERTY(EditAnywhere, BlueprintReadOnly, Category = "Nanite")
+	UPROPERTY(EditAnywhere, BlueprintReadWrite, BlueprintSetter=SetEnableNanite, Category = "Nanite")
 	bool EnableNanite;
 
 #if WITH_EDITORONLY_DATA
@@ -583,13 +586,13 @@ public:
 	/**
 	* Mass As Density, units are in kg/m^3
 	*/
-	UPROPERTY(EditAnywhere, BlueprintReadOnly, Category = "Collisions")
+	UPROPERTY(EditAnywhere, BlueprintReadWrite, Category = "Collisions")
 	bool bMassAsDensity;
 
 	/**
 	* Total Mass of Collection. If density, units are in kg/m^3
 	*/
-	UPROPERTY(EditAnywhere, BlueprintReadOnly, Category = "Collisions")
+	UPROPERTY(EditAnywhere, BlueprintReadWrite, Category = "Collisions")
 	float Mass;
 
 	/**
