Chaos : fix SetEnableNanite to not call CreateNaniteData outisde of the editor builds

diff --git a/Engine/Source/Runtime/Experimental/GeometryCollectionEngine/Private/GeometryCollection/GeometryCollectionObject.cpp b/Engine/Source/Runtime/Experimental/GeometryCollectionEngine/Private/GeometryCollection/GeometryCollectionObject.cpp
--- a/Engine/Source/Runtime/Experimental/GeometryCollectionEngine/Private/GeometryCollection/GeometryCollectionObject.cpp
+++ b/Engine/Source/Runtime/Experimental/GeometryCollectionEngine/Private/GeometryCollection/GeometryCollectionObject.cpp
@@ -941,6 +941,7 @@ void UGeometryCollection::SetEnableNanite(bool bValue)
 		EnableNanite = bValue;
 		NaniteData = MakeUnique<FGeometryCollectionNaniteData>();
 
+#if WITH_EDITOR
 		if (EnableNanite)
 		{
 			if (GeometryCollection)
@@ -948,6 +949,7 @@ void UGeometryCollection::SetEnableNanite(bool bValue)
 				NaniteData = UGeometryCollection::CreateNaniteData(GeometryCollection.Get());
 			}
 		}
+#endif
 	}
 }
 
