Fix static analysis bug in WeightedLatticeImplicitObject

diff --git a/Engine/Source/Runtime/Experimental/Chaos/Private/Chaos/WeightedLatticeImplicitObject.cpp b/Engine/Source/Runtime/Experimental/Chaos/Private/Chaos/WeightedLatticeImplicitObject.cpp
--- a/Engine/Source/Runtime/Experimental/Chaos/Private/Chaos/WeightedLatticeImplicitObject.cpp
+++ b/Engine/Source/Runtime/Experimental/Chaos/Private/Chaos/WeightedLatticeImplicitObject.cpp
@@ -751,7 +751,7 @@ void FWeightedLatticeImplicitObjectBuilder::AddInfluence(int32 FlatIndex, uint16
 		return;
 	}
 
-	if (Data.NumInfluences != Data.MaxTotalInfluences)
+	if (Data.NumInfluences < Data.MaxTotalInfluences)
 	{
 		++Data.NumInfluences;
 	}
