ChaosFlesh : Add new group for objects on different proxies.

diff --git a/Engine/Plugins/Experimental/ChaosFlesh/Source/ChaosFleshEngine/Private/ChaosFlesh/ChaosDeformableSolverComponent.cpp b/Engine/Plugins/Experimental/ChaosFlesh/Source/ChaosFleshEngine/Private/ChaosFlesh/ChaosDeformableSolverComponent.cpp
--- a/Engine/Plugins/Experimental/ChaosFlesh/Source/ChaosFleshEngine/Private/ChaosFlesh/ChaosDeformableSolverComponent.cpp
+++ b/Engine/Plugins/Experimental/ChaosFlesh/Source/ChaosFleshEngine/Private/ChaosFlesh/ChaosDeformableSolverComponent.cpp
@@ -239,9 +239,11 @@ void UDeformableSolverComponent::Reset()
 		, bUseGridBasedConstraints
 		, GridDx
 		, bDoQuasistatics
-		, EMesh
+		, YoungModulus
 		, bDoBlended
 		, BlendedZeta
+		, Damping
+		, bEnableGravity
 		}));
 
 	
diff --git a/Engine/Plugins/Experimental/ChaosFlesh/Source/ChaosFleshEngine/Private/ChaosFlesh/FleshComponent.cpp b/Engine/Plugins/Experimental/ChaosFlesh/Source/ChaosFleshEngine/Private/ChaosFlesh/FleshComponent.cpp
--- a/Engine/Plugins/Experimental/ChaosFlesh/Source/ChaosFleshEngine/Private/ChaosFlesh/FleshComponent.cpp
+++ b/Engine/Plugins/Experimental/ChaosFlesh/Source/ChaosFleshEngine/Private/ChaosFlesh/FleshComponent.cpp
@@ -165,13 +165,13 @@ UDeformablePhysicsComponent::FDataMapValue UFleshComponent::NewDeformableData()
 								}
 							}
 						}
-						return FDataMapValue(new Chaos::Softs::FFleshThreadingProxy::FFleshInputBuffer(this->GetComponentTransform(), AnimationTransforms, ComponentPose, this));
+						return FDataMapValue(new Chaos::Softs::FFleshThreadingProxy::FFleshInputBuffer(this->GetComponentTransform(), AnimationTransforms, ComponentPose, bTempEnableGravity, this));
 					}
 				}
 			}
 		}
 	}
-	return FDataMapValue(new Chaos::Softs::FFleshThreadingProxy::FFleshInputBuffer(this->GetComponentTransform(), this));
+	return FDataMapValue(new Chaos::Softs::FFleshThreadingProxy::FFleshInputBuffer(this->GetComponentTransform(), bTempEnableGravity, this));
 }
 
 
diff --git a/Engine/Plugins/Experimental/ChaosFlesh/Source/ChaosFleshEngine/Public/ChaosFlesh/ChaosDeformablePhysicsComponent.h b/Engine/Plugins/Experimental/ChaosFlesh/Source/ChaosFleshEngine/Public/ChaosFlesh/ChaosDeformablePhysicsComponent.h
--- a/Engine/Plugins/Experimental/ChaosFlesh/Source/ChaosFleshEngine/Public/ChaosFlesh/ChaosDeformablePhysicsComponent.h
+++ b/Engine/Plugins/Experimental/ChaosFlesh/Source/ChaosFleshEngine/Public/ChaosFlesh/ChaosDeformablePhysicsComponent.h
@@ -56,6 +56,9 @@ public:
 	UPROPERTY(EditAnywhere, BlueprintReadWrite, Category = Attributes)
 	TObjectPtr<ADeformableSolverActor> PrimarySolver;
 
+	UPROPERTY(EditAnywhere, BlueprintReadWrite, Category = Attributes)
+	bool bTempEnableGravity = true;
+
 	const FThreadingProxy* GetPhysicsProxy() const { return PhysicsProxy; }
 	FThreadingProxy* GetPhysicsProxy() { return PhysicsProxy; }
 
diff --git a/Engine/Plugins/Experimental/ChaosFlesh/Source/ChaosFleshEngine/Public/ChaosFlesh/ChaosDeformableSolverComponent.h b/Engine/Plugins/Experimental/ChaosFlesh/Source/ChaosFleshEngine/Public/ChaosFlesh/ChaosDeformableSolverComponent.h
--- a/Engine/Plugins/Experimental/ChaosFlesh/Source/ChaosFleshEngine/Public/ChaosFlesh/ChaosDeformableSolverComponent.h
+++ b/Engine/Plugins/Experimental/ChaosFlesh/Source/ChaosFleshEngine/Public/ChaosFlesh/ChaosDeformableSolverComponent.h
@@ -102,7 +102,7 @@ public:
 		bool bDoQuasistatics = false;
 
 	UPROPERTY(EditAnywhere, Category = Chaos)
-		float EMesh = 100000;
+		float YoungModulus = 100000;
 
 	UPROPERTY(EditAnywhere, Category = Chaos)
 		bool bDoBlended = false;
@@ -110,6 +110,11 @@ public:
 	UPROPERTY(EditAnywhere, Category = Chaos)
 		float BlendedZeta = 0;
 
+	UPROPERTY(EditAnywhere, Category = Chaos)
+		float Damping = 0;
+
+	UPROPERTY(EditAnywhere, Category = Chaos)
+		bool bEnableGravity = true;
 	//UPROPERTY(EditAnywhere, Category = Chaos)
 	//	bool bWaitForParallelDeformableTask = true;
 
diff --git a/Engine/Source/Runtime/Experimental/Chaos/Private/Chaos/Deformable/ChaosDeformableSolver.cpp b/Engine/Source/Runtime/Experimental/Chaos/Private/Chaos/Deformable/ChaosDeformableSolver.cpp
--- a/Engine/Source/Runtime/Experimental/Chaos/Private/Chaos/Deformable/ChaosDeformableSolver.cpp
+++ b/Engine/Source/Runtime/Experimental/Chaos/Private/Chaos/Deformable/ChaosDeformableSolver.cpp
@@ -82,7 +82,7 @@ namespace Chaos::Softs
 			Property.NumSolverIterations, (FSolverReal)0.,
 			/*SelfCollisionsThickness = */(FSolverReal)0.,
 			/*CoefficientOfFriction = */(FSolverReal)0.,
-			/*FSolverReal Damping = */(FSolverReal)0.04,
+			/*FSolverReal Damping = */Property.Damping,
 			/*FSolverReal LocalDamping = */(FSolverReal)0.,
 			Property.bDoQuasistatics));
 		Evolution->Particles().AddArray(&MObjects);
@@ -111,8 +111,8 @@ namespace Chaos::Softs
 		if (Property.NumSolverIterations)
 		{
 			RemoveSimulationObjects();
-			InitializeSimulationObjects();
 			UpdateProxyInputPackages();
+			InitializeSimulationObjects();
 			AdvanceDt(DeltaTime);
 		}
 	}
@@ -180,7 +180,9 @@ namespace Chaos::Softs
 		auto DoubleVert = [](FVector3f V) { return FVector3d(V.X, V.Y, V.Z); };
 
 		const FTransform& InitialTransform = Proxy.GetInitialTransform();
-		int32 ParticleStart = Evolution->AddParticleRange(NumParticles, 1, true);
+		int32 ParticleStart = Evolution->AddParticleRange(NumParticles, GroupOffset, true);
+		GroupOffset += 1;
+		
 		for (uint32 vdx = 0; vdx < NumParticles; ++vdx)
 		{
 			int32 SolverParticleIndex = ParticleStart + vdx;
@@ -191,6 +193,28 @@ namespace Chaos::Softs
 			Evolution->Particles().PAndInvM(SolverParticleIndex).InvM = Evolution->Particles().InvM(SolverParticleIndex);
 			MObjects[SolverParticleIndex] = Proxy.GetOwner();
 		}
+
+		bool ObjectEnableGravity = false;
+
+		int32 SolverParticleIndex = ParticleStart;
+		if (const UObject* Owner = this->MObjects[SolverParticleIndex]) {
+			FFleshThreadingProxy::FFleshInputBuffer* FleshInputBuffer = nullptr;
+			if (this->CurrentInputPackage->ObjectMap.Contains(Owner))
+			{
+				FleshInputBuffer = this->CurrentInputPackage->ObjectMap[Owner]->As<FFleshThreadingProxy::FFleshInputBuffer>();
+				if (FleshInputBuffer)
+				{
+					ObjectEnableGravity = FleshInputBuffer->bEnableGravity;
+				}
+			}
+		}
+
+		if (!ObjectEnableGravity || !Property.bEnableGravity)
+		{
+			FSolverVec3 ZeroGravity(0.f);
+			Evolution->SetGravity(ZeroGravity, GroupOffset - 1);
+		}
+
 		Proxy.SetSolverParticleRange(ParticleStart, NumParticles);
 	}
 
@@ -463,10 +487,11 @@ namespace Chaos::Softs
 		{
 			Chaos::FVec3 Position(0.f);
 			Chaos::FVec3 EulerRot(0.f);
-			Evolution->AddCollisionParticleRange(1, 1, true);
+			int32 CollisionParticleOffset = Evolution->AddCollisionParticleRange(1, INDEX_NONE, true);
 			Evolution->CollisionParticles().X(0) = Position;
 			Evolution->CollisionParticles().R(0) = Chaos::TRotation<Chaos::FReal, 3>::MakeFromEuler(EulerRot);
 			Evolution->CollisionParticles().SetDynamicGeometry(0, MakeUnique<Chaos::TPlane<Chaos::FReal, 3>>(Chaos::FVec3(0.f, 0.f, 0.f), Chaos::FVec3(0.f, 0.f, 1.f)));
+
 		}
 	}
 
diff --git a/Engine/Source/Runtime/Experimental/Chaos/Public/Chaos/Deformable/ChaosDeformableSolver.h b/Engine/Source/Runtime/Experimental/Chaos/Public/Chaos/Deformable/ChaosDeformableSolver.h
--- a/Engine/Source/Runtime/Experimental/Chaos/Public/Chaos/Deformable/ChaosDeformableSolver.h
+++ b/Engine/Source/Runtime/Experimental/Chaos/Public/Chaos/Deformable/ChaosDeformableSolver.h
@@ -183,6 +183,7 @@ namespace Chaos::Softs
 		int32 Frame = 0;
 		int32 Iteration = 0;
 		bool bSimulationInitialized = false;
+		int32 GroupOffset = 1;
 	};
 
 
diff --git a/Engine/Source/Runtime/Experimental/Chaos/Public/Chaos/Deformable/ChaosDeformableSolverProxy.h b/Engine/Source/Runtime/Experimental/Chaos/Public/Chaos/Deformable/ChaosDeformableSolverProxy.h
--- a/Engine/Source/Runtime/Experimental/Chaos/Public/Chaos/Deformable/ChaosDeformableSolverProxy.h
+++ b/Engine/Source/Runtime/Experimental/Chaos/Public/Chaos/Deformable/ChaosDeformableSolverProxy.h
@@ -114,24 +114,27 @@ namespace Chaos::Softs
 		public:
 			typedef FFleshThreadingProxy Source;
 
-			FFleshInputBuffer(const FTransform& InGlobalTransform, const UObject* InOwner = nullptr)
+			FFleshInputBuffer(const FTransform& InGlobalTransform, const bool InbEnableGravity, const UObject* InOwner = nullptr)
 				: Super(InOwner, FFleshThreadingProxy::TypeName())
 				, GlobalTransform(InGlobalTransform)
 				, Transforms(TArray<FTransform>())
 				, RestTransforms(TArray<FTransform>())
+				, bEnableGravity(InbEnableGravity)
 			{}
 
-			FFleshInputBuffer(const FTransform& InGlobalTransform, const TArray<FTransform> & InTransforms, const TArray<FTransform>& InRestTransforms, const UObject* InOwner = nullptr)
+			FFleshInputBuffer(const FTransform& InGlobalTransform, const TArray<FTransform>& InTransforms, const TArray<FTransform>& InRestTransforms, const bool InbEnableGravity, const UObject* InOwner = nullptr)
 				: Super(InOwner, FFleshThreadingProxy::TypeName())
 				, GlobalTransform(InGlobalTransform)
 				, Transforms(InTransforms)
 				, RestTransforms(InRestTransforms)
+				, bEnableGravity(InbEnableGravity)
 			{}
 			virtual ~FFleshInputBuffer() {}
 
 			FTransform GlobalTransform;
 			TArray<FTransform> Transforms;
 			TArray<FTransform> RestTransforms; 
+			bool bEnableGravity;
 		};
 
 		class FFleshOutputBuffer : public FThreadingProxy::FBuffer
diff --git a/Engine/Source/Runtime/Experimental/Chaos/Public/Chaos/Deformable/ChaosDeformableSolverTypes.h b/Engine/Source/Runtime/Experimental/Chaos/Public/Chaos/Deformable/ChaosDeformableSolverTypes.h
--- a/Engine/Source/Runtime/Experimental/Chaos/Public/Chaos/Deformable/ChaosDeformableSolverTypes.h
+++ b/Engine/Source/Runtime/Experimental/Chaos/Public/Chaos/Deformable/ChaosDeformableSolverTypes.h
@@ -28,7 +28,9 @@ namespace Chaos::Softs
 			bool InbDoQuasistatics = false,
 			FSolverReal InEMesh = (FSolverReal)100000.,
 			bool InbDoBlended = false,
-			FSolverReal InBlendedZeta = (FSolverReal).1)
+			FSolverReal InBlendedZeta = (FSolverReal).1,
+			FSolverReal InDamping = (FSolverReal)0,
+			bool InbEnableGravity = true)
 			: NumSolverSubSteps(InNumSolverSubSteps)
 			, NumSolverIterations(InNumSolverIterations)
 			, FixTimeStep(InFixTimeStep)
@@ -43,6 +45,8 @@ namespace Chaos::Softs
 			, EMesh(InEMesh)
 			, bDoBlended(InbDoBlended)
 			, BlendedZeta(InBlendedZeta)
+			, Damping(InDamping)
+			, bEnableGravity(InbEnableGravity)
 		{}
 
 		FDeformableSolverProperties(const FDeformableSolverProperties& InProp)
@@ -60,6 +64,8 @@ namespace Chaos::Softs
 			, EMesh(InProp.EMesh)
 			, bDoBlended(InProp.bDoBlended)
 			, BlendedZeta(InProp.BlendedZeta)
+			, Damping(InProp.Damping)
+			, bEnableGravity(InProp.bEnableGravity)
 		{}
 
 		int32 NumSolverSubSteps = 5;
@@ -76,6 +82,8 @@ namespace Chaos::Softs
 		FSolverReal EMesh = (FSolverReal)100000.;
 		bool bDoBlended = false;
 		FSolverReal BlendedZeta = (FSolverReal)0.;
+		FSolverReal Damping = (FSolverReal)0.;
+		bool bEnableGravity = true;
 	};
 
 
