Add mask to tetrahedral bindings schema.

diff --git a/Engine/Source/Runtime/Experimental/Chaos/Private/GeometryCollection/Facades/CollectionTetrahedralBindingsFacade.cpp b/Engine/Source/Runtime/Experimental/Chaos/Private/GeometryCollection/Facades/CollectionTetrahedralBindingsFacade.cpp
--- a/Engine/Source/Runtime/Experimental/Chaos/Private/GeometryCollection/Facades/CollectionTetrahedralBindingsFacade.cpp
+++ b/Engine/Source/Runtime/Experimental/Chaos/Private/GeometryCollection/Facades/CollectionTetrahedralBindingsFacade.cpp
@@ -15,59 +15,28 @@ namespace GeometryCollection::Facades
 
 	// Attributes
 	const FName FTetrahedralBindings::MeshIdAttributeName = "MeshId";
-	//const FName FTetrahedralBindings::MeshLODAttributeName = "MeshLOD";
 
 	const FName FTetrahedralBindings::ParentsAttributeName = "Parents";
 	const FName FTetrahedralBindings::WeightsAttributeName = "Weights";
 	const FName FTetrahedralBindings::OffsetsAttributeName = "Offsets";
+	const FName FTetrahedralBindings::MaskAttributeName = "Mask";
 
 	FTetrahedralBindings::FTetrahedralBindings(FManagedArrayCollection& InCollection)
 		: MeshIdAttribute(InCollection, MeshIdAttributeName, MeshBindingsGroupName)
-		, Parents(nullptr)
-		, Weights(nullptr)
-		, Offsets(nullptr)
 	{}
 
 	FTetrahedralBindings::FTetrahedralBindings(const FManagedArrayCollection& InCollection)
 		: MeshIdAttribute(InCollection, MeshIdAttributeName, MeshBindingsGroupName)
-		, Parents(nullptr)
-		, Weights(nullptr)
-		, Offsets(nullptr)
 	{}
 
 	FTetrahedralBindings::~FTetrahedralBindings()
-	{
-		delete Parents; Parents = nullptr;
-		delete Weights; Weights = nullptr;
-		delete Offsets; Offsets = nullptr;
-	}
+	{}
 
-	void FTetrahedralBindings::DefineSchema(/*const FName* MeshId, const int32 LOD*/)
+	void FTetrahedralBindings::DefineSchema()
 	{
 		check(!IsConst());
 		TManagedArray<FString>& MeshIdValues = 
 			MeshIdAttribute.IsValid() ? MeshIdAttribute.Modify() : MeshIdAttribute.Add();
-	/*	if (MeshId)
-		{
-			// Generate the group name, and if it doesn't exist, store it in the table of 
-			// available bindings groups.
-			FName GroupName = GenerateMeshGroupName(*MeshId, LOD);
-			if (!MeshIdValues.Contains(GroupName.ToString()))
-			{
-				const int32 Idx = MeshIdAttribute.AddElements(1);
-				MeshIdValues[Idx] = GroupName.ToString();
-
-				// This is a new group, so create the bindings arrays.
-				AddBindingsGroup(GroupName);
-			}
-			else
-			{
-				// This is an existing group, so find the existing bindings arrays.
-				ReadBindingsGroup(GroupName);
-			}
-
-		}
-		*/
 	}
 
 	bool FTetrahedralBindings::IsValid() const
@@ -75,7 +44,8 @@ namespace GeometryCollection::Facades
 		return MeshIdAttribute.IsValid() && 
 			(Parents && Parents->IsValid()) && 
 			(Weights && Weights->IsValid()) && 
-			(Offsets && Offsets->IsValid());
+			(Offsets && Offsets->IsValid()) &&
+			(Mask && Mask->IsValid());
 	}
 
 	FName FTetrahedralBindings::GenerateMeshGroupName(
@@ -100,6 +70,28 @@ namespace GeometryCollection::Facades
 		return MeshIdValues ? MeshIdValues->Contains(GroupName.ToString()) : false;
 	}
 
+	int32 FTetrahedralBindings::GetTetMeshIndex(const FName& MeshId, const int32 LOD) const
+	{
+		const TManagedArray<FString>* MeshIdValues =
+			MeshIdAttribute.IsValid() ? MeshIdAttribute.Find() : nullptr;
+		if (MeshIdValues)
+		{
+			FString Suffix = FString::Printf(TEXT(":%s:%d"), *MeshId.ToString(), LOD);
+			for (int32 i = 0; i < MeshIdValues->Num(); i++)
+			{
+				const FString& Entry = (*MeshIdValues)[i];
+				if (Entry.EndsWith(Suffix))
+				{
+					FString Str = Entry;
+					Str.RemoveAt(0, FString(TEXT("TetrahedralBindings:TetMeshIdx:")).Len(), false);
+					Str.RemoveAt(Str.Len() - Suffix.Len(), Suffix.Len());
+					return FCString::Atoi(*Str);
+				}
+			}
+		}
+		return INDEX_NONE;
+	}
+
 	void FTetrahedralBindings::AddBindingsGroup(const int32 TetMeshIdx, const FName& MeshId, const int32 LOD)
 	{
 		AddBindingsGroup(GenerateMeshGroupName(TetMeshIdx, MeshId, LOD));
@@ -112,21 +104,26 @@ namespace GeometryCollection::Facades
 			ReadBindingsGroup(GroupName);
 			return;
 		}
+		check(MeshIdAttribute.IsValid());
+		check(MeshIdAttribute.IsPersistent());
 
 		check(!IsConst());
 		const int32 Idx = MeshIdAttribute.AddElements(1);
 		MeshIdAttribute.Modify()[Idx] = GroupName.ToString();
 
-		delete Parents; Parents = nullptr;
-		delete Weights; Weights = nullptr;
-		delete Offsets; Offsets = nullptr;
+		Parents.Release();
+		Weights.Release();
+		Offsets.Release();
+		Mask.Release();
 		FManagedArrayCollection& Collection = *MeshIdAttribute.GetCollection();
-		Parents = new TManagedArrayAccessor<FIntVector4>(Collection, ParentsAttributeName, GroupName);
-		Weights = new TManagedArrayAccessor<FVector4f>(Collection, WeightsAttributeName, GroupName);
-		Offsets = new TManagedArrayAccessor<FVector3f>(Collection, OffsetsAttributeName, GroupName);
+		Parents.Reset(new TManagedArrayAccessor<FIntVector4>(Collection, ParentsAttributeName, GroupName));
+		Weights.Reset(new TManagedArrayAccessor<FVector4f>(Collection, WeightsAttributeName, GroupName));
+		Offsets.Reset(new TManagedArrayAccessor<FVector3f>(Collection, OffsetsAttributeName, GroupName));
+		Mask.Reset(new TManagedArrayAccessor<float>(Collection, MaskAttributeName, GroupName));
 		Parents->Add(ManageArrayAccessor::EPersistencePolicy::MakePersistent, FGeometryCollection::VerticesGroup);
-		Weights->Add(ManageArrayAccessor::EPersistencePolicy::MakePersistent);
-		Offsets->Add(ManageArrayAccessor::EPersistencePolicy::MakePersistent);
+		Weights->Add();
+		Offsets->Add();
+		Mask->Add();
 	}
 
 	bool FTetrahedralBindings::ReadBindingsGroup(const int32 TetMeshIdx, const FName& MeshId, const int32 LOD)
@@ -137,22 +134,32 @@ namespace GeometryCollection::Facades
 	bool FTetrahedralBindings::ReadBindingsGroup(const FName& GroupName)
 	{
 		check(MeshIdAttribute.IsValid());
-		delete Parents; Parents = nullptr;
-		delete Weights; Weights = nullptr;
-		delete Offsets; Offsets = nullptr;
-		if (MeshIdAttribute.Find()->Contains(GroupName.ToString()))
+		Parents.Release();
+		Weights.Release();
+		Offsets.Release();
+		Mask.Release();
+		if (!MeshIdAttribute.Find()->Contains(GroupName.ToString()))
 		{
 			return false;
 		}
 		// This is an existing group, so find the existing bindings arrays.
-		FManagedArrayCollection& Collection = *MeshIdAttribute.GetCollection();
-		Parents = new TManagedArrayAccessor<FIntVector4>(Collection, ParentsAttributeName, GroupName);
-		Weights = new TManagedArrayAccessor<FVector4f>(Collection, WeightsAttributeName, GroupName);
-		Offsets = new TManagedArrayAccessor<FVector3f>(Collection, OffsetsAttributeName, GroupName);
-		check(Parents->IsValid());
-		check(Weights->IsValid());
-		check(Offsets->IsValid());
-		return Parents->IsValid() && Weights->IsValid() && Offsets->IsValid();
+		if (!IsConst())
+		{
+			FManagedArrayCollection* Collection = MeshIdAttribute.GetCollection();
+			Parents.Reset(new TManagedArrayAccessor<FIntVector4>(*Collection, ParentsAttributeName, GroupName));
+			Weights.Reset(new TManagedArrayAccessor<FVector4f>(*Collection, WeightsAttributeName, GroupName));
+			Offsets.Reset(new TManagedArrayAccessor<FVector3f>(*Collection, OffsetsAttributeName, GroupName));
+			Mask.Reset(new TManagedArrayAccessor<float>(*Collection, MaskAttributeName, GroupName));
+		}
+		else
+		{
+			const FManagedArrayCollection& ConstCollection = MeshIdAttribute.GetConstCollection();
+			Parents.Reset(new TManagedArrayAccessor<FIntVector4>(ConstCollection, ParentsAttributeName, GroupName));
+			Weights.Reset(new TManagedArrayAccessor<FVector4f>(ConstCollection, WeightsAttributeName, GroupName));
+			Offsets.Reset(new TManagedArrayAccessor<FVector3f>(ConstCollection, OffsetsAttributeName, GroupName));
+			Mask.Reset(new TManagedArrayAccessor<float>(ConstCollection, MaskAttributeName, GroupName));
+		}
+		return Parents->IsValid() && Weights->IsValid() && Offsets->IsValid() && Mask->IsValid();
 	}
 
 	void FTetrahedralBindings::RemoveBindingsGroup(const int32 TetMeshIdx, const FName& MeshId, const int32 LOD)
@@ -176,14 +183,22 @@ namespace GeometryCollection::Facades
 		if (Parents)
 		{
 			Parents->Remove();
+			Parents.Release();
 		}
 		if (Weights)
 		{
 			Weights->Remove();
+			Weights.Release();
 		}
 		if (Offsets)
 		{
 			Offsets->Remove();
+			Offsets.Release();
+		}
+		if (Mask)
+		{
+			Mask->Remove();
+			Mask.Release();
 		}
 		// Only drop the group if it's empty at this point?
 		if (Collection.NumAttributes(GroupName) == 0)
@@ -196,11 +211,12 @@ namespace GeometryCollection::Facades
 	FTetrahedralBindings::SetBindingsData(
 		const TArray<FIntVector4>& ParentsIn,
 		const TArray<FVector4f>& WeightsIn,
-		const TArray<FVector3f>& OffsetsIn)
+		const TArray<FVector3f>& OffsetsIn,
+		const TArray<float>& MaskIn)
 	{
 		check(!IsConst());
 		check(IsValid());
-		check((ParentsIn.Num() == WeightsIn.Num()) && (ParentsIn.Num() == OffsetsIn.Num()));
+		check((ParentsIn.Num() == WeightsIn.Num()) && (ParentsIn.Num() == OffsetsIn.Num()) && (ParentsIn.Num() == MaskIn.Num()));
 
 		const int32 Num = ParentsIn.Num();
 		const int32 CurrNum = Parents->Num();//Collection.NumElements(CurrGroupName);
@@ -208,11 +224,13 @@ namespace GeometryCollection::Facades
 		TManagedArray<FIntVector4>& ParentsValues = Parents->Modify();
 		TManagedArray<FVector4f>& WeightsValues = Weights->Modify();
 		TManagedArray<FVector3f>& OffsetsValues = Offsets->Modify();
+		TManagedArray<float>& MaskValues = Mask->Modify();
 		for (int32 i = 0; i < Num; i++)
 		{
 			ParentsValues[i] = ParentsIn[i];
 			WeightsValues[i] = WeightsIn[i];
 			OffsetsValues[i] = OffsetsIn[i];
+			MaskValues[i] = MaskIn[i];
 		}
 	}
 
diff --git a/Engine/Source/Runtime/Experimental/Chaos/Public/GeometryCollection/Facades/CollectionTetrahedralBindingsFacade.h b/Engine/Source/Runtime/Experimental/Chaos/Public/GeometryCollection/Facades/CollectionTetrahedralBindingsFacade.h
--- a/Engine/Source/Runtime/Experimental/Chaos/Public/GeometryCollection/Facades/CollectionTetrahedralBindingsFacade.h
+++ b/Engine/Source/Runtime/Experimental/Chaos/Public/GeometryCollection/Facades/CollectionTetrahedralBindingsFacade.h
@@ -22,12 +22,20 @@ namespace GeometryCollection::Facades
 		// groups
 		static const FName MeshBindingsGroupName;
 
+		//
 		// Attributes
+		//
+
 		static const FName MeshIdAttributeName;
 
+		//! Tet or Tri vertex indices.
 		static const FName ParentsAttributeName;
+		//! Barycentric weight of each tet/tri vertex.
 		static const FName WeightsAttributeName;
+		//! Offset vector from barycentric tri position.
 		static const FName OffsetsAttributeName;
+		//! Per vertex amount for deformer masking.
+		static const FName MaskAttributeName;
 
 		/**
 		* FSelectionFacade Constuctor
@@ -57,6 +65,11 @@ namespace GeometryCollection::Facades
 		*/
 		static FName GenerateMeshGroupName(const int32 TetMeshIdx, const FName& MeshId, const int32 LOD);
 
+		/**
+		* For a given \p MeshId and \p LOD, return the associated tetrahedral mesh index.
+		*/
+		int32 GetTetMeshIndex(const FName& MeshId, const int32 LOD) const;
+
 		/**
 		* Returns \c true if the specified bindings group exists.
 		*/
@@ -91,32 +104,48 @@ namespace GeometryCollection::Facades
 		* \p Weights are barycentric coordinates.
 		* \p Offsets are vectors from the barycentric point to the location, in the case 
 		*    of a surface binding.
+		* \p Mask are per-vertex multipliers on the deformer, 0 for no deformation, 1.0 for 
+		*    full deformation.
 		*/
-		void SetBindingsData(const TArray<FIntVector4>& Parents, const TArray<FVector4f>& Weights, const TArray<FVector3f>& Offsets);
+		void SetBindingsData(const TArray<FIntVector4>& ParentsIn, const TArray<FVector4f>& WeightsIn, const TArray<FVector3f>& OffsetsIn, const TArray<float>& MaskIn);
+		void SetBindingsData(const TArray<FIntVector4>& ParentsIn, const TArray<FVector4f>& WeightsIn, const TArray<FVector3f>& OffsetsIn)
+		{
+			TArray<float> MaskTmp; MaskTmp.SetNum(ParentsIn.Num());
+			for (int32 i = 0; i < ParentsIn.Num(); i++)
+			{
+				MaskTmp[i] = 1.0;
+			}
+			SetBindingsData(ParentsIn, WeightsIn, OffsetsIn, MaskTmp);
+		}
 
 		/**
 		* Get Parents array.
 		*/
-		const TManagedArrayAccessor<FIntVector4>* GetParents() const { return Parents; }
-		      TManagedArrayAccessor<FIntVector4>* GetParents() { check(!IsConst()); return Parents; }
-
+		const TManagedArrayAccessor<FIntVector4>* GetParentsRO() const { return Parents.Get(); }
+		      TManagedArrayAccessor<FIntVector4>* GetParents() { check(!IsConst()); return Parents.Get(); }
 		/**
 		* Get Weights array.
 		*/
-		const TManagedArrayAccessor<FVector4f>* GetWeights() const { return Weights; }
-		      TManagedArrayAccessor<FVector4f>* GetWeights() { check(!IsConst()); return Weights; }
-
+		const TManagedArrayAccessor<FVector4f>* GetWeightsRO() const { return Weights.Get(); }
+		      TManagedArrayAccessor<FVector4f>* GetWeights() { check(!IsConst()); return Weights.Get(); }
 		/**
 		* Get Offsets array.
 		*/
-		const TManagedArrayAccessor<FVector3f>* GetOffsets() const { return Offsets; }
-		      TManagedArrayAccessor<FVector3f>* GetOffsets() { check(!IsConst());  return Offsets; }
+		const TManagedArrayAccessor<FVector3f>* GetOffsetsRO() const { return Offsets.Get(); }
+		      TManagedArrayAccessor<FVector3f>* GetOffsets() { check(!IsConst());  return Offsets.Get(); }
+		/**
+		* Get Mask array.
+		*/
+		const TManagedArrayAccessor<float>* GetMaskRO() const { return Mask.Get(); }
+		      TManagedArrayAccessor<float>* GetMask() { check(!IsConst()); return Mask.Get(); }
 
 	private:
 		TManagedArrayAccessor<FString> MeshIdAttribute;
-		TManagedArrayAccessor<FIntVector4>* Parents;
-		TManagedArrayAccessor<FVector4f>* Weights;
-		TManagedArrayAccessor<FVector3f>* Offsets;
+
+		TUniquePtr<TManagedArrayAccessor<FIntVector4>> Parents;
+		TUniquePtr<TManagedArrayAccessor<FVector4f>> Weights;
+		TUniquePtr<TManagedArrayAccessor<FVector3f>> Offsets;
+		TUniquePtr<TManagedArrayAccessor<float>> Mask;
 	};
 
 }
