Use velocity-matching to replicate geometry collection clusters.
- Disable with `p.bGeometryCollectionRepUseClusterVelocityMatch 0`
- Debug draw with `p.Chaos.DebugDraw.GeometryCollectionReplication 1`

#swarm https://p4-swarm.epicgames.net/reviews/23918942

diff --git a/Engine/Source/Runtime/Experimental/GeometryCollectionEngine/Private/GeometryCollection/GeometryCollectionComponent.cpp b/Engine/Source/Runtime/Experimental/GeometryCollectionEngine/Private/GeometryCollection/GeometryCollectionComponent.cpp
--- a/Engine/Source/Runtime/Experimental/GeometryCollectionEngine/Private/GeometryCollection/GeometryCollectionComponent.cpp
+++ b/Engine/Source/Runtime/Experimental/GeometryCollectionEngine/Private/GeometryCollection/GeometryCollectionComponent.cpp
@@ -51,6 +51,11 @@
 #include "UObject/UObjectThreadContext.h"
 #endif
 
+#if ENABLE_DRAW_DEBUG
+#include "Chaos/DebugDrawQueue.h"
+#include "Chaos/ChaosDebugDraw.h"
+#endif
+
 #include "PhysicsEngine/BodySetup.h"
 #include "PhysicsEngine/BodyInstance.h"
 #include "Chaos/ChaosGameplayEventDispatcher.h"
@@ -187,6 +192,10 @@ bool FGeometryCollectionRepData::NetSerialize(FArchive& Ar, class UPackageMap* M
 	if(Ar.IsLoading())
 	{
 		Clusters.SetNum(NumClusters);
+
+		// Resetting this received time signals that this is the first frame that this
+		// RepData will be processed.
+		RepDataReceivedTime.Reset();
 	}
 
 	for(FGeometryCollectionClusterRep& Cluster : Clusters)
@@ -1718,28 +1727,151 @@ int32 GeometryCollectionHardsnapThresholdMs = 100; // 10 Hz
 FAutoConsoleVariableRef CVarGeometryCollectionHardsnapThresholdMs(TEXT("p.GeometryCollectionHardsnapThresholdMs"), GeometryCollectionHardMissingUpdatesSnapThreshold,
 	TEXT("Determines how many ms since the last hardsnap to trigger a new one"));
 
+float GeometryCollectionRepLinearMatchStrength = 50;
+FAutoConsoleVariableRef CVarGeometryCollectionRepLinearMatchStrength(TEXT("p.GeometryCollectionRepLinearMatchStrength"), GeometryCollectionRepLinearMatchStrength, TEXT("Units can be interpreted as %/s^2 - acceleration of percent linear correction"));
+
+float GeometryCollectionRepAngularMatchTime = .5f;
+FAutoConsoleVariableRef CVarGeometryCollectionRepAngularMatchTime(TEXT("p.GeometryCollectionRepAngularMatchTime"), GeometryCollectionRepAngularMatchTime, TEXT("In seconds, how quickly should the angle match the replicated target angle"));
+
+float GeometryCollectionRepMaxExtrapolationTime = 3.f;
+FAutoConsoleVariableRef CVarGeometryCollectionRepMaxExtrapolationTime(TEXT("p.GeometryCollectionRepMaxExtrapolationTime"), GeometryCollectionRepMaxExtrapolationTime, TEXT("Number of seconds that replicated physics data will persist for a GC, extrapolating velocities"));
+
+bool bGeometryCollectionDebugDrawRep = 0;
+FAutoConsoleVariableRef CVarGeometryCollectionDebugDrawRep(TEXT("p.Chaos.DebugDraw.GeometryCollectionReplication"), bGeometryCollectionDebugDrawRep,
+	TEXT("If true debug draw deltas and corrections for geometry collection replication"));
+
+bool bGeometryCollectionRepUseClusterVelocityMatch = true;
+FAutoConsoleVariableRef CVarGeometryCollectionRepUseClusterVelocityMatch(TEXT("p.bGeometryCollectionRepUseClusterVelocityMatch"), bGeometryCollectionRepUseClusterVelocityMatch,
+	TEXT("Use physical velocity to match cluster states"));
+
 void UGeometryCollectionComponent::ProcessRepData()
+{
+	bGeometryCollectionRepUseClusterVelocityMatch = false;
+	ProcessRepData(0.f, 0.f);
+}
+
+void UGeometryCollectionComponent::ProcessRepData(const float DeltaTime, const float SimTime)
 {
 	using namespace Chaos;
-	if(!PhysicsProxy || !PhysicsProxy->IsInitializedOnPhysicsThread() || VersionProcessed == RepData.Version || PhysicsProxy->GetReplicationMode() != FGeometryCollectionPhysicsProxy::EReplicationMode::Client)
+
+	if(!PhysicsProxy || !PhysicsProxy->IsInitializedOnPhysicsThread() || PhysicsProxy->GetReplicationMode() != FGeometryCollectionPhysicsProxy::EReplicationMode::Client)
+	{
+		return;
+	}
+
+	// Track the sim time that this rep data was received on.
+	if (!RepData.RepDataReceivedTime.IsSet())
+	{
+		RepData.RepDataReceivedTime = SimTime;
+	}
+
+	// How far we must extrapolate from when we received the data
+	const float RepExtrapTime = FMath::Max(0.f, SimTime - *RepData.RepDataReceivedTime);
+
+	// If we've extrapolated past a threshold, then stop tracking
+	// the last received rep data
+	if (RepExtrapTime > GeometryCollectionRepMaxExtrapolationTime)
+	{
+		return;
+	}
+
+	// Create a little little function for applying a lambda to each
+	// corresponding pair of replicated and local clusters.
+	const auto ForEachClusterPair = [this](const auto& Lambda)
+	{
+		for (const FGeometryCollectionClusterRep& RepCluster : RepData.Clusters)
+		{
+			if (Chaos::FPBDRigidParticleHandle* Cluster = PhysicsProxy->GetParticles()[RepCluster.ClusterIdx])
+			{
+				if (RepCluster.ClusterState.IsInternalCluster())
+				{
+					// internal cluster do not have an index so we rep data send one of the children's
+					// let's find the parent
+					Cluster = Cluster->CastToClustered()->Parent();
+				}
+
+				if (Cluster && Cluster->Disabled() == false)
+				{
+					Lambda(RepCluster, *Cluster);
+				}
+			}
+		}
+	};
+
+#if ENABLE_DRAW_DEBUG
+	if (bGeometryCollectionDebugDrawRep)
+	{
+		ForEachClusterPair([RepExtrapTime](const FGeometryCollectionClusterRep& RepCluster, Chaos::FPBDRigidParticleHandle& Cluster)
+		{
+			// Don't bother debug drawing if the delta is too small
+			if ((Cluster.X() - RepCluster.Position).SizeSquared() < .1f)
+			{
+				FVector Axis;
+				float Angle;
+				(RepCluster.Rotation.Inverse() * Cluster.R()).ToAxisAndAngle(Axis, Angle);
+				if (FMath::Abs(Angle) < .1f)
+				{
+					return;
+				}
+			}
+
+			Chaos::FDebugDrawQueue& DrawQueue = Chaos::FDebugDrawQueue::GetInstance();
+			DrawQueue.DrawDebugCoordinateSystem(Cluster.X(), FRotator(Cluster.R()), 100.f, false, -1, -1, 1.f);
+			DrawQueue.DrawDebugBox(Cluster.X() + Cluster.LocalBounds().Center(), Cluster.LocalBounds().Extents(), Cluster.R(), FColor::White, false, -1, -1, 1.f);
+			DrawQueue.DrawDebugBox(RepCluster.Position + Cluster.LocalBounds().Center(), Cluster.LocalBounds().Extents(), RepCluster.Rotation, FColor::Green, false, -1, -1, 1.f);
+
+			if (bGeometryCollectionRepUseClusterVelocityMatch)
+			{
+				const FVector RepVel = RepCluster.LinearVelocity;
+				const FVector RepAngVel = RepCluster.AngularVelocity;
+				const FVector RepExtrapPos = RepCluster.Position + (RepVel * RepExtrapTime);
+				const Chaos::FRotation3 RepExtrapAng = Chaos::FRotation3::IntegrateRotationWithAngularVelocity(RepCluster.Rotation, RepAngVel, RepExtrapTime);
+				DrawQueue.DrawDebugCoordinateSystem(RepExtrapPos, FRotator(RepExtrapAng), 100.f, false, -1, -1, 1.f);
+				DrawQueue.DrawDebugDirectionalArrow(Cluster.X(), RepExtrapPos, 10.f, FColor::White, false, -1, -1, 1.f);
+				DrawQueue.DrawDebugBox(RepExtrapPos + Cluster.LocalBounds().Center(), Cluster.LocalBounds().Extents(), RepExtrapAng, FColor::Orange, false, -1, -1, 1.f);
+			}
+			else
+			{
+				DrawQueue.DrawDebugCoordinateSystem(RepCluster.Position, FRotator(RepCluster.Rotation), 100.f, false, -1, -1, 1.f);
+			}
+		});
+	}
+#endif
+
+	// If not doing velocity match, don't bother processing the same version twice.
+	// Do this one after the debug draw so that we can still easily see the diff
+	// between the position and the target position.
+	if (VersionProcessed == RepData.Version && !bGeometryCollectionRepUseClusterVelocityMatch)
 	{
 		return;
 	}
 
 	bool bHardSnap = false;
 	const int64 CurrentTimeInMs = FPlatformTime::ToMilliseconds64(FPlatformTime::Cycles64());
-	if(VersionProcessed < RepData.Version)
+
+	// Always hard snap on the very first version received
+	if (VersionProcessed == 0)
+	{
+		bHardSnap = true;
+	}
+	else if (VersionProcessed < RepData.Version)
 	{
 		//TODO: this will not really work if a fracture happens and then immediately goes to sleep without updating client enough times
 		//A time method would work better here, but is limited to async mode. Maybe we can support both
-		bHardSnap |= (RepData.Version - VersionProcessed) > GeometryCollectionHardMissingUpdatesSnapThreshold;
-		bHardSnap |= (CurrentTimeInMs - LastHardsnapTimeInMs) > GeometryCollectionHardsnapThresholdMs;
+		bHardSnap = (RepData.Version - VersionProcessed) > GeometryCollectionHardMissingUpdatesSnapThreshold;
+
+		if (!bGeometryCollectionRepUseClusterVelocityMatch)
+		{
+			// When not doing velocity match for clusters, instead we do periodic hard snapping
+			bHardSnap |= (CurrentTimeInMs - LastHardsnapTimeInMs) > GeometryCollectionHardsnapThresholdMs;
+		}
 	}
-	else
+	else if (VersionProcessed > RepData.Version)
 	{
 		//rollover so just treat as hard snap - this case is extremely rare and a one off
 		bHardSnap = true;
 	}
+
 	if (bHardSnap)
 	{
 		LastHardsnapTimeInMs = CurrentTimeInMs;
@@ -1756,7 +1888,7 @@ void UGeometryCollectionComponent::ProcessRepData()
 		check(OneOff);
 
 		// Set initial velocities if not hard snapping
-		if(!bHardSnap)
+		if (!bHardSnap)
 		{
 			// TODO: we should get an update cluster position first so that when particles break off they get the right position 
 			// TODO: should we invalidate?
@@ -1764,33 +1896,64 @@ void UGeometryCollectionComponent::ProcessRepData()
 			OneOff->SetW(ActivatedCluster.InitialAngularVelocity);
 		}
 
-		RigidClustering.ReleaseClusterParticles(TArray<FPBDRigidParticleHandle*>{ OneOff }, /* bTriggerBreakEvents */ true);
+		RigidClustering.ReleaseClusterParticles(TArray<FPBDRigidParticleHandle*>{ OneOff }, true);
 	}
 
-	if(bHardSnap)
+	ForEachClusterPair([Solver, DeltaTime, RepExtrapTime, bHardSnap](const FGeometryCollectionClusterRep& RepCluster, Chaos::FPBDRigidParticleHandle& Cluster)
 	{
-		for (const FGeometryCollectionClusterRep& RepCluster : RepData.Clusters)
+		bool bWake = false;
+
+		if (bHardSnap)
+		{
+			Cluster.SetX(RepCluster.Position);
+			Cluster.SetR(RepCluster.Rotation);
+			Cluster.SetV(RepCluster.LinearVelocity);
+			Cluster.SetW(RepCluster.AngularVelocity);
+			bWake = true;
+		}
+		else if (bGeometryCollectionRepUseClusterVelocityMatch)
 		{
-			if (FPBDRigidParticleHandle* Cluster = PhysicsProxy->GetParticles()[RepCluster.ClusterIdx])
+			//
+			// Match linear velocity
+			//
+			const FVector RepVel = RepCluster.LinearVelocity;
+			const FVector RepExtrapPos = RepCluster.Position + (RepVel * RepExtrapTime);
+			const FVec3 DeltaX = RepExtrapPos - Cluster.X();
+			const float DeltaXMagSq = DeltaX.SizeSquared();
+			if (DeltaXMagSq > SMALL_NUMBER && GeometryCollectionRepLinearMatchStrength > SMALL_NUMBER)
 			{
-				if (RepCluster.ClusterState.IsInternalCluster())
-				{
-					// internal cluster do not have an index so we rep data send one of the children's
-					// let's find the parent
-					Cluster = Cluster->CastToClustered()->Parent();
-				}
-				if(Cluster && Cluster->Disabled() == false)
-				{
-					Cluster->SetX(RepCluster.Position);
-					Cluster->SetR(RepCluster.Rotation);
-					Cluster->SetV(RepCluster.LinearVelocity);
-					Cluster->SetW(RepCluster.AngularVelocity);
-					// TODO: snap object state too once we fix interpolation
-					// Solver->GetEvolution()->SetParticleObjectState(Cluster, static_cast<Chaos::EObjectStateType>(RepCluster.ObjectState));
-				}
+				bWake = true;
+				//
+				// DeltaX * MatchStrength is an acceleration, m/s^2, which is integrated
+				// by multiplying by DeltaTime.
+				//
+				// It's formulated this way to get a larger correction for a longer time
+				// step, ie. correction velocities are framerate independent.
+				//
+				Cluster.SetV(RepVel + (DeltaX * GeometryCollectionRepLinearMatchStrength * DeltaTime));
+			}
+
+			//
+			// Match angular velocity
+			//
+			const FVector RepAngVel = RepCluster.AngularVelocity;
+			const Chaos::FRotation3 RepExtrapAng = Chaos::FRotation3::IntegrateRotationWithAngularVelocity(RepCluster.Rotation, RepAngVel, RepExtrapTime);
+			const FVector AngVel = Chaos::FRotation3::CalculateAngularVelocity(Cluster.R(), RepExtrapAng, GeometryCollectionRepAngularMatchTime);
+			if (AngVel.SizeSquared() > SMALL_NUMBER)
+			{
+				Cluster.SetW(RepAngVel + AngVel);
+				bWake = true;
 			}
 		}
-	}
+
+		//
+		// Wake up particle if it's sleeping and there's a delta to correct
+		//
+		if (bWake && Cluster.IsSleeping())
+		{
+			Solver->GetEvolution()->SetParticleObjectState(&Cluster, Chaos::EObjectStateType::Dynamic);
+		}
+	});
 
 	VersionProcessed = RepData.Version;
 }
@@ -2386,7 +2549,7 @@ void UGeometryCollectionComponent::AsyncPhysicsTickComponent(float DeltaTime, fl
 	}
 	else
 	{
-		ProcessRepData();
+		ProcessRepData(DeltaTime, SimTime);
 	}
 }
 
diff --git a/Engine/Source/Runtime/Experimental/GeometryCollectionEngine/Public/GeometryCollection/GeometryCollectionComponent.h b/Engine/Source/Runtime/Experimental/GeometryCollectionEngine/Public/GeometryCollection/GeometryCollectionComponent.h
--- a/Engine/Source/Runtime/Experimental/GeometryCollectionEngine/Public/GeometryCollection/GeometryCollectionComponent.h
+++ b/Engine/Source/Runtime/Experimental/GeometryCollectionEngine/Public/GeometryCollection/GeometryCollectionComponent.h
@@ -403,6 +403,9 @@ struct FGeometryCollectionRepData
 	// For Network Prediction Mode we require the frame number on the server when the data was gathered
 	int32 ServerFrame;
 
+	// The sim-time that this rep data was received
+	TOptional<float> RepDataReceivedTime;
+
 	// Just test version to skip having to traverse the whole pose array for replication
 	bool Identical(const FGeometryCollectionRepData* Other, uint32 PortFlags) const;
 	bool NetSerialize(FArchive& Ar, class UPackageMap* Map, bool& bOutSuccess);
@@ -1124,8 +1127,12 @@ protected:
 
 	/** Clear all rep data, this is required if the physics proxy has been recreated */
 	virtual void ResetRepData();
-
+ 
+	UE_DEPRECATED(5.3, "The argument-free version of ProcessRepData will be removed. Please use the version which takes DeltaTime and SimTime instead.")
 	virtual void ProcessRepData();
+
+	virtual void ProcessRepData(float DeltaTime, float SimTime);
+
 	int32 VersionProcessed = INDEX_NONE;
 
 private:
