Chaos - fix determinism in IslandManager

diff --git a/Engine/Source/Programs/HeadlessChaos/Private/HeadlessChaosTestGraphEvolution.cpp b/Engine/Source/Programs/HeadlessChaos/Private/HeadlessChaosTestGraphEvolution.cpp
--- a/Engine/Source/Programs/HeadlessChaos/Private/HeadlessChaosTestGraphEvolution.cpp
+++ b/Engine/Source/Programs/HeadlessChaos/Private/HeadlessChaosTestGraphEvolution.cpp
@@ -1069,4 +1069,51 @@ namespace ChaosTest
 		EXPECT_FALSE(Test.IslandManager->GetIsland(0)->IsSleeping());
 
 	}
+
+	// Test the sparse array repeatable index assignment. See TSparseArray::SortFreeList()
+	GTEST_TEST(SparseArrayTests, TestSortFreeList)
+	{
+		TSparseArray<int32> Values;
+
+		// The first time we add objects, they should be in consecutive indices starting from 0
+		Values.Add(0);
+		Values.Add(1);
+		Values.Add(2);
+		EXPECT_EQ(Values[0], 0);
+		EXPECT_EQ(Values[1], 1);
+		EXPECT_EQ(Values[2], 2);
+
+		// Remove a couple items in the same order we added them
+		Values.RemoveAt(1);
+		Values.RemoveAt(2);
+
+		// Add the items again, they will end up in reverse order
+		// We don't rely on this behaviour but I'm testing here because if this changes
+		// in the future then we may be able to remove our calls to SortFreeList in the graph.
+		Values.Add(1);
+		Values.Add(2);
+		EXPECT_EQ(Values[0], 0);
+		EXPECT_EQ(Values[1], 2);	// Swapped
+		EXPECT_EQ(Values[2], 1);	// Swapped
+
+
+		// Now do the same as above on a new array, but call SortFreeList before reusing it
+		TSparseArray<int32> Values2;
+		
+		Values2.Add(0);
+		Values2.Add(1);
+		Values2.Add(2);
+		Values2.RemoveAt(1);
+		Values2.RemoveAt(2);
+
+		// Rebuild the free list
+		Values2.SortFreeList();
+
+		// We should now get the same order as the first time we added items
+		Values2.Add(1);
+		Values2.Add(2);
+		EXPECT_EQ(Values2[0], 0);
+		EXPECT_EQ(Values2[1], 1);
+		EXPECT_EQ(Values2[2], 2);
+	}
 }
diff --git a/Engine/Source/Runtime/Engine/Classes/PhysicsEngine/PhysicsSettings.h b/Engine/Source/Runtime/Engine/Classes/PhysicsEngine/PhysicsSettings.h
--- a/Engine/Source/Runtime/Engine/Classes/PhysicsEngine/PhysicsSettings.h
+++ b/Engine/Source/Runtime/Engine/Classes/PhysicsEngine/PhysicsSettings.h
@@ -149,10 +149,6 @@ class ENGINE_API UPhysicsSettings : public UPhysicsSettingsCore
 	UPROPERTY(config, EditAnywhere, Category = Simulation)
 	bool bDisableCCD;
 
-	/** If set to true, the scene will use enhanced determinism at the cost of a bit more resources. See eENABLE_ENHANCED_DETERMINISM to learn about the specifics */
-	UPROPERTY(config, EditAnywhere, Category = Simulation)
-	bool bEnableEnhancedDeterminism;
-
 	/** Min Delta Time below which anim dynamics and rigidbody nodes will not simulate. */
 	UPROPERTY(config, EditAnywhere, meta = (ClampMin = "0.0", UIMin = "0.0", ClampMax = "1.0", UIMax = "1.0"), Category = Framerate)
 	float AnimPhysicsMinDeltaTime;
diff --git a/Engine/Source/Runtime/Engine/Private/PhysicsEngine/PhysicsSettings.cpp b/Engine/Source/Runtime/Engine/Private/PhysicsEngine/PhysicsSettings.cpp
--- a/Engine/Source/Runtime/Engine/Private/PhysicsEngine/PhysicsSettings.cpp
+++ b/Engine/Source/Runtime/Engine/Private/PhysicsEngine/PhysicsSettings.cpp
@@ -16,7 +16,6 @@ UPhysicsSettings::UPhysicsSettings(const FObjectInitializer& ObjectInitializer)
 	, LockedAxis_DEPRECATED(ESettingsLockedAxis::Invalid)
 	, bSuppressFaceRemapTable(false)
 	, bDisableActiveActors(false)
-	, bEnableEnhancedDeterminism(false)
 	, AnimPhysicsMinDeltaTime(0.f)
 	, bSimulateAnimPhysicsAfterReset(false)
 	, MinPhysicsDeltaTime(UE_SMALL_NUMBER)
diff --git a/Engine/Source/Runtime/Experimental/Chaos/Private/Chaos/Island/IslandGraph.cpp b/Engine/Source/Runtime/Experimental/Chaos/Private/Chaos/Island/IslandGraph.cpp
--- a/Engine/Source/Runtime/Experimental/Chaos/Private/Chaos/Island/IslandGraph.cpp
+++ b/Engine/Source/Runtime/Experimental/Chaos/Private/Chaos/Island/IslandGraph.cpp
@@ -963,6 +963,47 @@ void FIslandGraph<NodeType, EdgeType, IslandType, OwnerType>::UpdateGraph()
 	}
 }
 
+template<typename NodeType, typename EdgeType, typename IslandType, typename OwnerType>
+void FIslandGraph<NodeType, EdgeType, IslandType, OwnerType>::SetIsDeterministic(const bool bInIsDeterministic)
+{
+	if (bIsDeterministic != bInIsDeterministic)
+	{
+		bIsDeterministic = bInIsDeterministic;
+		if (bIsDeterministic)
+		{
+			SortFreeLists();
+		}
+	}
+}
+
+template<typename NodeType, typename EdgeType, typename IslandType, typename OwnerType>
+void FIslandGraph<NodeType, EdgeType, IslandType, OwnerType>::EndTick()
+{
+	if (bIsDeterministic)
+	{
+		SortFreeLists();
+	}
+}
+
+template<typename NodeType, typename EdgeType, typename IslandType, typename OwnerType>
+void FIslandGraph<NodeType, EdgeType, IslandType, OwnerType>::SortFreeLists()
+{
+	// For determinism: ensure that our sparse arrays will allocate indices that
+	// depend only on the current state of the graph, and not the order in which
+	// we removed elements. (We do sort the graph based on distance-to-kinematic, 
+	// but when there is an ambiguity we just take the array order so we 
+	// require that index assignment is repeatable between frames).
+	GraphNodes.SortFreeList();
+	GraphEdges.SortFreeList();
+	GraphIslands.SortFreeList();
+
+	for (FGraphNode& GraphNode : GraphNodes)
+	{
+		GraphNode.NodeEdges.SortFreeList();
+	}
+}
+
+
 template<typename NodeType, typename EdgeType, typename IslandType, typename OwnerType>
 void FIslandGraph<NodeType, EdgeType, IslandType, OwnerType>::RemoveAllAwakeEdges()
 {
diff --git a/Engine/Source/Runtime/Experimental/Chaos/Private/Chaos/Island/IslandManager.cpp b/Engine/Source/Runtime/Experimental/Chaos/Private/Chaos/Island/IslandManager.cpp
--- a/Engine/Source/Runtime/Experimental/Chaos/Private/Chaos/Island/IslandManager.cpp
+++ b/Engine/Source/Runtime/Experimental/Chaos/Private/Chaos/Island/IslandManager.cpp
@@ -187,6 +187,11 @@ FPBDIslandManager::~FPBDIslandManager()
 {
 }
 
+void FPBDIslandManager::SetIsDeterministic(const bool bInIsDeterministic)
+{
+	IslandGraph->SetIsDeterministic(bInIsDeterministic);
+}
+
 void FPBDIslandManager::Reset()
 {
 	RemoveConstraints();
@@ -276,6 +281,9 @@ void FPBDIslandManager::EndTick()
 		Island->ClearConstraints();
 	}
 
+	// Let the graph perform and end-of-frame cleanup
+	IslandGraph->EndTick();
+
 	// We can no longer use the Island's particle and constraint lists
 	bIslandsPopulated = false;
 
diff --git a/Engine/Source/Runtime/Experimental/Chaos/Private/Chaos/PBDCollisionConstraints.cpp b/Engine/Source/Runtime/Experimental/Chaos/Private/Chaos/PBDCollisionConstraints.cpp
--- a/Engine/Source/Runtime/Experimental/Chaos/Private/Chaos/PBDCollisionConstraints.cpp
+++ b/Engine/Source/Runtime/Experimental/Chaos/Private/Chaos/PBDCollisionConstraints.cpp
@@ -463,12 +463,6 @@ namespace Chaos
 		// Disable any edge collisions that are hidden by face collisions
 		// (for bodies that have the EdgePruning option enabled)
 		PruneEdgeCollisions();
-
-		// @todo(chaos): this should be required any more - we sort when adding to the graph. Test and remove...
-		if (bIsDeterministic)
-		{
-			ConstraintAllocator.SortConstraintsHandles();
-		}
 	}
 
 	void FPBDCollisionConstraints::DetectProbeCollisions(FReal Dt)
@@ -582,9 +576,14 @@ namespace Chaos
 			}
 		}
 
-		// Sort new constraints into a predictable order. This isn't strictly required, but without it we
-		// can get fairly different behaviour from run to run because the collisions detection order is 
-		// effectively random on multicore machines
+		// Sort new constraints into a predictable order. This isn't strictly required unless we have
+		// deterministic mode enabled, but we do it always because we can get fairly different behaviour 
+		// from run to run because the collisions detection order is effectively random on multicore machines.
+		//
+		// @todo(chaos): this is still not good enough for some types of determinism. Specifically if we 
+		// create two set of objects in a different order but with the same physical positions and other 
+		// state, they will behave differently which is undesirable. To fix this we need a sorting 
+		// mechanism that does not rely on properties like IDs. E.g., some kind of physical state hash?
 		TempCollisions.Sort(
 			[](const FPBDCollisionConstraintHandle& L, const FPBDCollisionConstraintHandle& R)
 			{
diff --git a/Engine/Source/Runtime/Experimental/Chaos/Private/Chaos/PBDRigidsEvolutionGBF.cpp b/Engine/Source/Runtime/Experimental/Chaos/Private/Chaos/PBDRigidsEvolutionGBF.cpp
--- a/Engine/Source/Runtime/Experimental/Chaos/Private/Chaos/PBDRigidsEvolutionGBF.cpp
+++ b/Engine/Source/Runtime/Experimental/Chaos/Private/Chaos/PBDRigidsEvolutionGBF.cpp
@@ -571,8 +571,11 @@ void FPBDRigidsEvolutionGBF::AdvanceOneTimeStepImpl(const FReal Dt, const FSubSt
 
 	ParticleUpdatePosition(Particles.GetDirtyParticlesView(), Dt);
 
-	// Clean up the transient data from the constraint graph (e.g., Islands get cleared here)
-	GetIslandManager().EndTick();
+	{
+		CSV_SCOPED_TIMING_STAT(PhysicsVerbose, StepSolver_IslandEndTick);
+		// Clean up the transient data from the constraint graph (e.g., Islands get cleared here)
+		GetIslandManager().EndTick();
+	}
 
 	if (bChaos_Solver_TestMode)
 	{
@@ -626,6 +629,9 @@ void FPBDRigidsEvolutionGBF::SetIsDeterministic(const bool bInIsDeterministic)
 {
 	// We detect collisions in parallel, so order is non-deterministic without additional processing
 	CollisionConstraints.SetIsDeterministic(bInIsDeterministic);
+
+	// IslandManager uses TSparseArray which requires free-list maintenance for determinism
+	IslandManager.SetIsDeterministic(bInIsDeterministic);
 }
 
 
diff --git a/Engine/Source/Runtime/Experimental/Chaos/Private/PBDRigidsSolver.cpp b/Engine/Source/Runtime/Experimental/Chaos/Private/PBDRigidsSolver.cpp
--- a/Engine/Source/Runtime/Experimental/Chaos/Private/PBDRigidsSolver.cpp
+++ b/Engine/Source/Runtime/Experimental/Chaos/Private/PBDRigidsSolver.cpp
@@ -241,7 +241,7 @@ namespace Chaos
 		int32 ChaosSolverProjectionIterations = -1;
 		FAutoConsoleVariableRef CVarChaosSolverProjectionIterations(TEXT("p.Chaos.Solver.Iterations.Projection"), ChaosSolverProjectionIterations, TEXT("Override number of solver projection iterations (-1 to use config)"));
 
-		int32 ChaosSolverDeterministic = -1;
+		int32 ChaosSolverDeterministic = 1;
 		FAutoConsoleVariableRef CVarChaosSolverDeterministic(TEXT("p.Chaos.Solver.Deterministic"), ChaosSolverDeterministic, TEXT("Override determinism. 0: disabled; 1: enabled; -1: use config"));
 
 		// Copied from RBAN
diff --git a/Engine/Source/Runtime/Experimental/Chaos/Public/Chaos/Island/IslandGraph.h b/Engine/Source/Runtime/Experimental/Chaos/Public/Chaos/Island/IslandGraph.h
--- a/Engine/Source/Runtime/Experimental/Chaos/Public/Chaos/Island/IslandGraph.h
+++ b/Engine/Source/Runtime/Experimental/Chaos/Public/Chaos/Island/IslandGraph.h
@@ -148,11 +148,21 @@ public:
 		Owner = InOwner;
 	}
 
+	/**
+	 * Set to true if we require deterministic behaviour.
+	 */
+	void SetIsDeterministic(const bool bInIsDeterministic);
+
 	/**
 	 * Update current graph to merge connected islands 
 	 */
 	void UpdateGraph();
 
+	/**
+	 * Perform end-of-tick cleanup if necessary
+	 */
+	 void EndTick();
+
 	/**
 	 * Reserve a number of nodes in memory before adding them later
 	 * @param NumNodes Number of nodes to be reserved in memory
@@ -375,6 +385,8 @@ public:
 		}
 	}
 
+	void SortFreeLists();
+
 	FGraphOwner* Owner = nullptr;
 
 	/** List of graph nodes */
@@ -392,6 +404,9 @@ public:
 	/** Reverse list to find given an item the matching edge */
 	TMap<EdgeType,int32> ItemEdges;
 
+	/** Whether we care about determinism */
+	bool bIsDeterministic = false;
+
 	/** Graph counter used for graph traversal to check if an edge/node/island has already been processed */
 	int32 GraphCounter = 0;
 
diff --git a/Engine/Source/Runtime/Experimental/Chaos/Public/Chaos/Island/IslandManager.h b/Engine/Source/Runtime/Experimental/Chaos/Public/Chaos/Island/IslandManager.h
--- a/Engine/Source/Runtime/Experimental/Chaos/Public/Chaos/Island/IslandManager.h
+++ b/Engine/Source/Runtime/Experimental/Chaos/Public/Chaos/Island/IslandManager.h
@@ -53,6 +53,11 @@ public:
 	*/
 	~FPBDIslandManager();
 
+	/**
+	 * Set to true if we require deterministic behaviour.
+	 */
+	void SetIsDeterministic(const bool bInIsDeterministic);
+
 	/**
 	 * Clear all nodes and edges from the graph
 	*/
diff --git a/Engine/Source/Runtime/PhysicsCore/Private/ChaosScene.cpp b/Engine/Source/Runtime/PhysicsCore/Private/ChaosScene.cpp
--- a/Engine/Source/Runtime/PhysicsCore/Private/ChaosScene.cpp
+++ b/Engine/Source/Runtime/PhysicsCore/Private/ChaosScene.cpp
@@ -78,12 +78,14 @@ FChaosScene::FChaosScene(
 	SceneSolver->PhysSceneHack = this;
 	SimCallback = SceneSolver->CreateAndRegisterSimCallbackObject_External<FChaosSceneSimCallback>();
 
+	// Apply project settings to the solver
 	if(CVar_ApplyProjectSettings.GetValueOnAnyThread() != 0)
 	{
 		UPhysicsSettingsCore* Settings = UPhysicsSettingsCore::Get();
-		SceneSolver->RegisterSimOneShotCallback([InSolver = SceneSolver, SolverConfigCopy = Settings->SolverOptions]()
+		SceneSolver->RegisterSimOneShotCallback([InSolver = SceneSolver, SolverConfigCopy = Settings->SolverOptions, bIsDeterministic = Settings->bEnableEnhancedDeterminism]()
 		{
 			InSolver->ApplyConfig(SolverConfigCopy);
+			InSolver->SetIsDeterministic(bIsDeterministic);
 		});
 	}
 
diff --git a/Engine/Source/Runtime/PhysicsCore/Private/PhysicsSettingsCore.cpp b/Engine/Source/Runtime/PhysicsCore/Private/PhysicsSettingsCore.cpp
--- a/Engine/Source/Runtime/PhysicsCore/Private/PhysicsSettingsCore.cpp
+++ b/Engine/Source/Runtime/PhysicsCore/Private/PhysicsSettingsCore.cpp
@@ -32,12 +32,13 @@ UPhysicsSettingsCore::UPhysicsSettingsCore(const FObjectInitializer& ObjectIniti
 	, SimulateScratchMemorySize(262144)
 	, RagdollAggregateThreshold(4)
 	, TriangleMeshTriangleMinAreaThreshold(5.0f)
+	, bEnableEnhancedDeterminism(false)
 	, bEnableShapeSharing(false)
 	, bEnablePCM(true)
 	, bEnableStabilization(false)
 	, bWarnMissingLocks(true)
 	, bEnable2DPhysics(false)
-	,bDefaultHasComplexCollision_DEPRECATED(true)
+	, bDefaultHasComplexCollision_DEPRECATED(true)
 	, BounceThresholdVelocity(200.f)
 	, MaxAngularVelocity(3600)	//10 revolutions per second
 	, ContactOffsetMultiplier(0.02f)
diff --git a/Engine/Source/Runtime/PhysicsCore/Public/PhysicsSettingsCore.h b/Engine/Source/Runtime/PhysicsCore/Public/PhysicsSettingsCore.h
--- a/Engine/Source/Runtime/PhysicsCore/Public/PhysicsSettingsCore.h
+++ b/Engine/Source/Runtime/PhysicsCore/Public/PhysicsSettingsCore.h
@@ -44,6 +44,10 @@ class PHYSICSCORE_API UPhysicsSettingsCore: public UDeveloperSettings
 	UPROPERTY(config,EditAnywhere,AdvancedDisplay,meta = (ClampMin = "-1.0",UIMin = "-1.0",ClampMax = "10.0",UIMax = "10.0"),Category = Constants)
 	float TriangleMeshTriangleMinAreaThreshold;
 
+	/** If set to true, the scene will use enhanced determinism at the cost of a bit more resources. See eENABLE_ENHANCED_DETERMINISM to learn about the specifics */
+	UPROPERTY(config,EditAnywhere,Category = Simulation)
+	bool bEnableEnhancedDeterminism;
+
 	/** Enables shape sharing between sync and async scene for static rigid actors */
 	UPROPERTY(config,EditAnywhere,AdvancedDisplay,Category = Simulation)
 	bool bEnableShapeSharing;
