More fixes to WeightedLatticeBuilder static analysis failures.

diff --git a/Engine/Source/Runtime/Experimental/Chaos/Private/Chaos/WeightedLatticeImplicitObject.cpp b/Engine/Source/Runtime/Experimental/Chaos/Private/Chaos/WeightedLatticeImplicitObject.cpp
--- a/Engine/Source/Runtime/Experimental/Chaos/Private/Chaos/WeightedLatticeImplicitObject.cpp
+++ b/Engine/Source/Runtime/Experimental/Chaos/Private/Chaos/WeightedLatticeImplicitObject.cpp
@@ -746,6 +746,7 @@ void FWeightedLatticeImplicitObjectBuilder::AddInfluence(int32 FlatIndex, uint16
 			break;
 		}
 	}
+
 	if (InsertIndex == Data.NumInfluences && Data.NumInfluences == Data.MaxTotalInfluences)
 	{
 		return;
@@ -756,12 +757,15 @@ void FWeightedLatticeImplicitObjectBuilder::AddInfluence(int32 FlatIndex, uint16
 		++Data.NumInfluences;
 	}
 
-	for (uint8 Idx = Data.NumInfluences - 1; Idx > InsertIndex; --Idx)
+	for (int32 Idx = Data.NumInfluences - 1; Idx > (int32)InsertIndex; --Idx)
 	{
+		check(Idx < Data.MaxTotalInfluences);
+		check(Idx > 0);
 		Data.BoneIndices[Idx] = Data.BoneIndices[Idx - 1];
 		Data.BoneWeights[Idx] = Data.BoneWeights[Idx - 1];
 	}
 
+	check(InsertIndex < Data.MaxTotalInfluences);
 	Data.BoneIndices[InsertIndex] = BoneIndex;
 	Data.BoneWeights[InsertIndex] = Weight;
 	Build.WeightsAreOuter = bIsOuterWeight;
