Fix maths error in aerodynamic drag calculation when velocity is negative (driving backwards)

diff --git a/Engine/Source/Runtime/Experimental/ChaosVehicles/ChaosVehiclesCore/Private/AerodynamicsSystem.cpp b/Engine/Source/Runtime/Experimental/ChaosVehicles/ChaosVehiclesCore/Private/AerodynamicsSystem.cpp
--- a/Engine/Source/Runtime/Experimental/ChaosVehicles/ChaosVehiclesCore/Private/AerodynamicsSystem.cpp
+++ b/Engine/Source/Runtime/Experimental/ChaosVehicles/ChaosVehiclesCore/Private/AerodynamicsSystem.cpp
@@ -21,8 +21,9 @@ namespace Chaos
 
 	FVector FSimpleAerodynamicsSim::GetCombinedForces(float VelocityIn)
 	{
-		// -ve as forces applied in opposite direction to velocity
-		float CommonSum = -DensityOfMedium * VelocityIn * VelocityIn;
+		// aerodynamic forces applied in opposite direction to velocity
+		float Sign = (VelocityIn < 0.0f) ? 1.0f : -1.0f;
+		float CommonSum = DensityOfMedium * VelocityIn * VelocityIn * Sign;
 		FVector CombinedForces(EffectiveDragConstant * CommonSum, 0.f, EffectiveLiftConstant * CommonSum);
 		return CombinedForces;
 	}
