Fixing NonUnity errors.

diff --git a/Engine/Source/Runtime/Experimental/Chaos/Public/ChaosSolversModule.h b/Engine/Source/Runtime/Experimental/Chaos/Public/ChaosSolversModule.h
--- a/Engine/Source/Runtime/Experimental/Chaos/Public/ChaosSolversModule.h
+++ b/Engine/Source/Runtime/Experimental/Chaos/Public/ChaosSolversModule.h
@@ -5,6 +5,7 @@
 #include "CoreMinimal.h"
 #include "HAL/Event.h"
 #include "HAL/IConsoleManager.h"
+#include "HAL/ThreadSafeBool.h"
 #include "Modules/ModuleInterface.h"
 #include "Modules/ModuleManager.h"
 #include "Async/AsyncWork.h"
