add ability to set collision profiles per level on the geometry collection

diff --git a/Engine/Source/Runtime/Experimental/GeometryCollectionEngine/Private/GeometryCollection/GeometryCollectionComponent.cpp b/Engine/Source/Runtime/Experimental/GeometryCollectionEngine/Private/GeometryCollection/GeometryCollectionComponent.cpp
--- a/Engine/Source/Runtime/Experimental/GeometryCollectionEngine/Private/GeometryCollection/GeometryCollectionComponent.cpp
+++ b/Engine/Source/Runtime/Experimental/GeometryCollectionEngine/Private/GeometryCollection/GeometryCollectionComponent.cpp
@@ -2911,14 +2911,10 @@ void UGeometryCollectionComponent::RegisterAndInitializePhysicsProxy()
 					}
 				});
 		}
-
-		// We also need to make sure collision is disabled on the game thread side of things if abandoned particles shouldn't have collision.
-		if (bEnableAbandonAfterLevel)
-		{
-			LoadCollisionProfileOnAbandonedParticles();
-		}
 	}
 
+	LoadCollisionProfiles();
+
 	// We need to add the geometry collection into the external acceleration structure so that it's immediately available for queries instead of waiting for the sync from the physics thread (which could take awhile).
 	// Just adding the root particle should be sufficient since that'll be the only particle we'd expect any collisions with right after initialization.
 	if (Chaos::FPhysicsObjectHandle RootObject = GetPhysicsObjectByName(NAME_None))
@@ -2934,48 +2930,94 @@ void UGeometryCollectionComponent::RegisterAndInitializePhysicsProxy()
 
 void UGeometryCollectionComponent::SetAbandonedParticleCollisionProfileName(FName CollisionProfile)
 {
-	if (!bEnableAbandonAfterLevel)
+	if (!bEnableAbandonAfterLevel || !GetIsReplicated())
 	{
 		return;
 	}
 
 	AbandonedCollisionProfileName = CollisionProfile;
-	LoadCollisionProfileOnAbandonedParticles();
+	LoadCollisionProfiles();
 }
 
-void UGeometryCollectionComponent::LoadCollisionProfileOnAbandonedParticles()
+void UGeometryCollectionComponent::SetPerLevelCollisionProfileNames(const TArray<FName>& ProfileNames)
 {
-	FCollisionResponseTemplate Template;
-	if (!PhysicsProxy || !UCollisionProfile::Get()->GetProfileTemplate(AbandonedCollisionProfileName, Template))
+	CollisionProfilePerLevel = ProfileNames;
+	LoadCollisionProfiles();
+}
+
+void UGeometryCollectionComponent::LoadCollisionProfiles()
+{
+	if (!PhysicsProxy)
 	{
 		return;
 	}
 
-	Chaos::Facades::FCollectionHierarchyFacade HierarchyFacade(*RestCollection->GetGeometryCollection());
+	// Cache the FCollisionResponseTemplate as well as the query/sim collision filter data for a given collision profile name
+	// so we don't have to recreate it every time.
+	
+	struct CollisionProfileDataCache
+	{
+		FCollisionResponseTemplate Template;
+		FCollisionFilterData QueryFilter;
+		FCollisionFilterData SimFilter;
+	};
 
-	AActor* Owner = GetOwner();
-	const uint32 ActorID = Owner ? Owner->GetUniqueID() : 0;
-	const uint32 CompID = GetUniqueID();
+	TMap<FName, CollisionProfileDataCache> CachedData;
+
+	// Returns nullptr if we can't create or get the data.
+	auto CreateOrGetCollisionProfileData = [this, &CachedData](const FName& ProfileName) -> const CollisionProfileDataCache*
+	{
+		if (const CollisionProfileDataCache* Data = CachedData.Find(ProfileName))
+		{
+			return Data;
+		}
+
+		CollisionProfileDataCache Cache;
+		if (ProfileName == NAME_None || !UCollisionProfile::Get()->GetProfileTemplate(ProfileName, Cache.Template))
+		{
+			return nullptr;
+		}
+
+		AActor* Owner = GetOwner();
+		const uint32 ActorID = Owner ? Owner->GetUniqueID() : 0;
+		const uint32 CompID = GetUniqueID();
 
-	FCollisionFilterData AbandonedQueryFilter = InitialQueryFilter;
-	FCollisionFilterData AbandonedSimFilter = InitialSimFilter;
-	CreateShapeFilterData(Template.ObjectType, BodyInstance.GetMaskFilter(), ActorID, Template.ResponseToChannels, CompID, INDEX_NONE, AbandonedQueryFilter, AbandonedSimFilter, BodyInstance.bUseCCD, bNotifyCollisions, false, false);
+		Cache.QueryFilter = InitialQueryFilter;
+		Cache.SimFilter = InitialSimFilter;
+		CreateShapeFilterData(Cache.Template.ObjectType, BodyInstance.GetMaskFilter(), ActorID, Cache.Template.ResponseToChannels, CompID, INDEX_NONE, Cache.QueryFilter, Cache.SimFilter, BodyInstance.bUseCCD, bNotifyCollisions, false, false);
+
+		// Maintain parity with the rest of the geometry collection filters.
+		Cache.QueryFilter.Word3 |= (EPDF_SimpleCollision | EPDF_ComplexCollision);
+		Cache.SimFilter.Word3 |= (EPDF_SimpleCollision | EPDF_ComplexCollision);
+		
+		return &CachedData.Add(ProfileName, Cache);
+	};
 
-	// Maintain parity with the rest of the geometry collection filters.
-	AbandonedQueryFilter.Word3 |= (EPDF_SimpleCollision | EPDF_ComplexCollision);
-	AbandonedSimFilter.Word3 |= (EPDF_SimpleCollision | EPDF_ComplexCollision);
+	const CollisionProfileDataCache* AbandonedData = (bEnableAbandonAfterLevel && GetIsReplicated()) ? CreateOrGetCollisionProfileData(AbandonedCollisionProfileName) : nullptr;
 
+	Chaos::Facades::FCollectionHierarchyFacade HierarchyFacade(*RestCollection->GetGeometryCollection());
 	TArray<Chaos::FPhysicsObjectHandle> PhysicsObjects = PhysicsProxy->GetAllPhysicsObjects();
 	FLockedWritePhysicsObjectExternalInterface Interface = FPhysicsObjectExternalInterface::LockWrite(PhysicsObjects);
 
 	for (int32 ParticleIndex = 0; ParticleIndex < PhysicsProxy->GetNumParticles(); ++ParticleIndex)
 	{
 		const int32 Level = HierarchyFacade.GetInitialLevel(ParticleIndex);
-		if (Level >= ReplicationAbandonAfterLevel + 1)
+
+		TArrayView<Chaos::FPhysicsObjectHandle> ParticleView{ &PhysicsObjects[ParticleIndex], 1 };
+		const CollisionProfileDataCache* Data = nullptr;
+		if (AbandonedData && Level >= ReplicationAbandonAfterLevel + 1)
+		{
+			Data = AbandonedData;
+		}
+		else if (!CollisionProfilePerLevel.IsEmpty())
+		{
+			Data = CreateOrGetCollisionProfileData(CollisionProfilePerLevel[FMath::Min(CollisionProfilePerLevel.Num() - 1, Level)]);
+		}
+
+		if (Data)
 		{
-			TArrayView<Chaos::FPhysicsObjectHandle> ParticleView{ &PhysicsObjects[ParticleIndex], 1 };
-			Interface->UpdateShapeCollisionFlags(ParticleView, CollisionEnabledHasPhysics(Template.CollisionEnabled), CollisionEnabledHasQuery(Template.CollisionEnabled));
-			Interface->UpdateShapeFilterData(ParticleView, AbandonedQueryFilter, AbandonedSimFilter);
+			Interface->UpdateShapeCollisionFlags(ParticleView, CollisionEnabledHasPhysics(Data->Template.CollisionEnabled), CollisionEnabledHasQuery(Data->Template.CollisionEnabled));
+			Interface->UpdateShapeFilterData(ParticleView, Data->QueryFilter, Data->SimFilter);
 		}
 	}
 }
diff --git a/Engine/Source/Runtime/Experimental/GeometryCollectionEngine/Public/GeometryCollection/GeometryCollectionComponent.h b/Engine/Source/Runtime/Experimental/GeometryCollectionEngine/Public/GeometryCollection/GeometryCollectionComponent.h
--- a/Engine/Source/Runtime/Experimental/GeometryCollectionEngine/Public/GeometryCollection/GeometryCollectionComponent.h
+++ b/Engine/Source/Runtime/Experimental/GeometryCollectionEngine/Public/GeometryCollection/GeometryCollectionComponent.h
@@ -659,6 +659,9 @@ public:
 	UFUNCTION(BlueprintCallable, Category = "ChaosPhysics")
 	void SetAbandonedParticleCollisionProfileName(FName CollisionProfile);
 
+	UFUNCTION(BlueprintCallable, Category = "ChaosPhysics")
+	void SetPerLevelCollisionProfileNames(const TArray<FName>& ProfileNames);
+
 	/** API for getting at geometry collection data */
 	int32 GetNumElements(FName Group) const;
 
@@ -1103,6 +1106,15 @@ protected:
 	UPROPERTY(EditAnywhere, BlueprintReadWrite, BlueprintSetter= SetAbandonedParticleCollisionProfileName, Category = Network)
 	FName AbandonedCollisionProfileName;
 
+	/**
+	 * A per-level collision profile name. If the name is set to NONE or an invalid collision profile, nothing will be changed.
+	 * If the there are more levels than elements in this array, then each level will use the index that best matches it.
+	 * For example, if the particle is level 2, and there is only 1 element in the array, then the particle will use the 0th
+	 * collision profile. AbandonedCollisionProfileName will override this on the client when relevant.
+	 */
+	UPROPERTY(EditAnywhere, BlueprintReadWrite, BlueprintSetter=SetPerLevelCollisionProfileNames, Category=Physics)
+	TArray<FName> CollisionProfilePerLevel;
+
 	/**
 	 * If replicating - the cluster level to stop sending corrections for geometry collection chunks.
 	 * recommended for smaller leaf levels when the size of the objects means they are no longer
@@ -1227,7 +1239,7 @@ private:
 	
 	void BuildInitialFilterData();
 
-	void LoadCollisionProfileOnAbandonedParticles();
+	void LoadCollisionProfiles();
 
 	/** The clusters we need to replicate */
 	TUniquePtr<TSet<Chaos::FPBDRigidClusteredParticleHandle*>> ClustersToRep;
