Chaos - minor debug draw improvement to FContactTriangleCollector and CCD

diff --git a/Engine/Source/Runtime/Experimental/Chaos/Private/Chaos/CCDUtilities.cpp b/Engine/Source/Runtime/Experimental/Chaos/Private/Chaos/CCDUtilities.cpp
--- a/Engine/Source/Runtime/Experimental/Chaos/Private/Chaos/CCDUtilities.cpp
+++ b/Engine/Source/Runtime/Experimental/Chaos/Private/Chaos/CCDUtilities.cpp
@@ -5,6 +5,7 @@
 #include "Chaos/PBDRigidsEvolutionGBF.h"
 #include "Chaos/CollisionResolution.h"
 #include "Chaos/ChaosDebugDraw.h"
+#include "Chaos/DebugDrawQueue.h"
 
 //PRAGMA_DISABLE_OPTIMIZATION
 
@@ -409,17 +410,29 @@ namespace Chaos
 			// Debugdraw the shape at the current TOI
 			if (CVars::ChaosSolverDrawCCDInteractions)
 			{
+				const DebugDraw::FChaosDebugDrawSettings& DebugDrawSettings = CVars::ChaosSolverDebugDebugDrawSettings;
 				if (CCDParticle0)
 				{
 					const FRigidTransform3 SweepWorldTransform0 = FRigidTransform3(FConstGenericParticleHandle(CCDParticle0->Particle)->X(), FConstGenericParticleHandle(CCDParticle0->Particle)->Q());
 					const FRigidTransform3 ShapeWorldTransformXQ0 = CCDConstraint->SweptConstraint->GetShapeRelativeTransform0() * SweepWorldTransform0;
-					DebugDraw::DrawShape(ShapeWorldTransformXQ0, CCDConstraint->SweptConstraint->GetShape0()->GetLeafGeometry(), CCDConstraint->SweptConstraint->GetShape0(), FColor::Magenta, &CVars::ChaosSolverDebugDebugDrawSettings);
+					DebugDraw::DrawShape(ShapeWorldTransformXQ0, CCDConstraint->SweptConstraint->GetShape0()->GetLeafGeometry(), CCDConstraint->SweptConstraint->GetShape0(), FColor::Magenta, &DebugDrawSettings);
 				}
 				if (CCDParticle1)
 				{
 					const FRigidTransform3 SweepWorldTransform1 = FRigidTransform3(FConstGenericParticleHandle(CCDParticle1->Particle)->X(), FConstGenericParticleHandle(CCDParticle1->Particle)->Q());
 					const FRigidTransform3 ShapeWorldTransformXR1 = CCDConstraint->SweptConstraint->GetShapeRelativeTransform1() * SweepWorldTransform1;
-					DebugDraw::DrawShape(ShapeWorldTransformXR1, CCDConstraint->SweptConstraint->GetShape1()->GetLeafGeometry(), CCDConstraint->SweptConstraint->GetShape1(), FColor::Magenta, &CVars::ChaosSolverDebugDebugDrawSettings);
+					DebugDraw::DrawShape(ShapeWorldTransformXR1, CCDConstraint->SweptConstraint->GetShape1()->GetLeafGeometry(), CCDConstraint->SweptConstraint->GetShape1(), FColor::Magenta, &DebugDrawSettings);
+				}
+				for (int32 ManifoldPointIndex = 0; ManifoldPointIndex < CCDConstraint->SweptConstraint->NumManifoldPoints(); ++ManifoldPointIndex)
+				{
+					const FManifoldPoint& ManifoldPoint = CCDConstraint->SweptConstraint->GetManifoldPoint(ManifoldPointIndex);
+					if (ManifoldPoint.Flags.bDisabled)
+					{
+						continue;
+					}
+					const FVec3 ContactPos = CCDConstraint->SweptConstraint->GetShapeWorldTransform1().TransformPositionNoScale(FVec3(ManifoldPoint.ContactPoint.ShapeContactPoints[1]));
+					const FVec3 ContactNormal = CCDConstraint->SweptConstraint->GetShapeWorldTransform1().TransformVectorNoScale(FVec3(ManifoldPoint.ContactPoint.ShapeContactNormal));
+					FDebugDrawQueue::GetInstance().DrawDebugLine(ContactPos, ContactPos + DebugDrawSettings.DrawScale * 50 * ContactNormal, FColor::Red, false, UE_KINDA_SMALL_NUMBER, DebugDrawSettings.DrawPriority, DebugDrawSettings.LineThickness);
 				}
 			}
 #endif
diff --git a/Engine/Source/Runtime/Experimental/Chaos/Private/Chaos/Collision/ContactTriangles.cpp b/Engine/Source/Runtime/Experimental/Chaos/Private/Chaos/Collision/ContactTriangles.cpp
--- a/Engine/Source/Runtime/Experimental/Chaos/Private/Chaos/Collision/ContactTriangles.cpp
+++ b/Engine/Source/Runtime/Experimental/Chaos/Private/Chaos/Collision/ContactTriangles.cpp
@@ -533,7 +533,7 @@ namespace Chaos
 		}
 	}
 
-	void FContactTriangleCollector::DebugDrawContactPoints(const FColor& Color, const FReal LineThickness)
+	void FContactTriangleCollector::DebugDrawContactPoints(const FColor& Color, const FReal LineScale)
 	{
 #if CHAOS_DEBUG_DRAW
 		if (CVars::ChaosSolverDebugDrawMeshContacts && FDebugDrawQueue::GetInstance().IsDebugDrawingEnabled())
@@ -550,10 +550,10 @@ namespace Chaos
 				const FVec3 N = ConvexTransform.TransformVectorNoScale(ContactPoint.ShapeContactNormal);
 
 				// Draw the normal from the triangle face
-				FDebugDrawQueue::GetInstance().DrawDebugLine(P1, P1 + FReal(15) * N, Color, false, FRealSingle(Duration), DrawPriority, FRealSingle(LineThickness));
+				FDebugDrawQueue::GetInstance().DrawDebugLine(P1, P1 + LineScale * FReal(50) * N, Color, false, FRealSingle(Duration), DrawPriority, 2 * FRealSingle(LineScale));
 
 				// Draw a thin black line connecting the two contact points (triangle face to convex surface)
-				FDebugDrawQueue::GetInstance().DrawDebugLine(P0, P1, FColor::Black, false, FRealSingle(Duration), DrawPriority, 0.5f * FRealSingle(LineThickness));
+				//FDebugDrawQueue::GetInstance().DrawDebugLine(P0, P1, FColor::Black, false, FRealSingle(Duration), DrawPriority, 0.5f * FRealSingle(LineThickness));
 
 				if (!bTriangleDrawn[ContactPoint.ContactTriangleIndex])
 				{
@@ -561,9 +561,9 @@ namespace Chaos
 					const FVec3 V0 = ConvexTransform.TransformPosition(Triangle.Vertices[0]);
 					const FVec3 V1 = ConvexTransform.TransformPosition(Triangle.Vertices[1]);
 					const FVec3 V2 = ConvexTransform.TransformPosition(Triangle.Vertices[2]);
-					FDebugDrawQueue::GetInstance().DrawDebugLine(V0, V1, FColor::Silver, false, FRealSingle(Duration), DrawPriority, 0.5f * FRealSingle(LineThickness));
-					FDebugDrawQueue::GetInstance().DrawDebugLine(V1, V2, FColor::Silver, false, FRealSingle(Duration), DrawPriority, 0.5f * FRealSingle(LineThickness));
-					FDebugDrawQueue::GetInstance().DrawDebugLine(V2, V0, FColor::Silver, false, FRealSingle(Duration), DrawPriority, 0.5f * FRealSingle(LineThickness));
+					FDebugDrawQueue::GetInstance().DrawDebugLine(V0, V1, FColor::Silver, false, FRealSingle(Duration), DrawPriority, FRealSingle(LineScale));
+					FDebugDrawQueue::GetInstance().DrawDebugLine(V1, V2, FColor::Silver, false, FRealSingle(Duration), DrawPriority, FRealSingle(LineScale));
+					FDebugDrawQueue::GetInstance().DrawDebugLine(V2, V0, FColor::Silver, false, FRealSingle(Duration), DrawPriority, FRealSingle(LineScale));
 					bTriangleDrawn[ContactPoint.ContactTriangleIndex] = true;
 				}
 			}
