Geometry collection : fix code to support removal of geometry, vertex and face groups when nott using nanite ( to properly support StripOnCook for ISMPool  )
- Better support for bounds computation using transform bounding box attribute ( this code path can now use ISPC )
- make sure that using ISMPool is equivalent to having visible geometry in the TickComponent method

diff --git a/Engine/Source/Runtime/Experimental/GeometryCollectionEngine/Private/GeometryCollection/GeometryCollectionComponent.cpp b/Engine/Source/Runtime/Experimental/GeometryCollectionEngine/Private/GeometryCollection/GeometryCollectionComponent.cpp
--- a/Engine/Source/Runtime/Experimental/GeometryCollectionEngine/Private/GeometryCollection/GeometryCollectionComponent.cpp
+++ b/Engine/Source/Runtime/Experimental/GeometryCollectionEngine/Private/GeometryCollection/GeometryCollectionComponent.cpp
@@ -448,24 +448,114 @@ void UGeometryCollectionComponent::GetLifetimeReplicatedProps(TArray<FLifetimePr
 
 namespace
 {
-void UpdateGlobalMatricesWithExplodedVectors(TArray<FMatrix>& GlobalMatricesIn, FGeometryCollection& GeometryCollection)
-{
-	int32 NumMatrices = GlobalMatricesIn.Num();
-	if (GlobalMatricesIn.Num() > 0)
+	void UpdateGlobalMatricesWithExplodedVectors(TArray<FMatrix>& GlobalMatricesIn, FGeometryCollection& GeometryCollection)
 	{
-		if (GeometryCollection.HasAttribute("ExplodedVector", FGeometryCollection::TransformGroup))
+		int32 NumMatrices = GlobalMatricesIn.Num();
+		if (GlobalMatricesIn.Num() > 0)
 		{
-			const TManagedArray<FVector3f>& ExplodedVectors = GeometryCollection.GetAttribute<FVector3f>("ExplodedVector", FGeometryCollection::TransformGroup);
+			if (GeometryCollection.HasAttribute("ExplodedVector", FGeometryCollection::TransformGroup))
+			{
+				const TManagedArray<FVector3f>& ExplodedVectors = GeometryCollection.GetAttribute<FVector3f>("ExplodedVector", FGeometryCollection::TransformGroup);
 
-			check(NumMatrices == ExplodedVectors.Num());
+				check(NumMatrices == ExplodedVectors.Num());
+				for (int32 TransformIndex = 0; TransformIndex < ExplodedVectors.Num(); ++TransformIndex)
+				{
+					GlobalMatricesIn[TransformIndex] = GlobalMatricesIn[TransformIndex].ConcatTranslation((FVector)ExplodedVectors[TransformIndex]);
+				}
+			}
+		}
+	}
 
-			for (int32 tt = 0, nt = NumMatrices; tt < nt; ++tt)
+	/**
+	* compute the bounding box from the bounding boxes stored in the geometry group
+	*/
+	inline FBox ComputeBoundsFromGeometryBoundingBoxes(
+		const TManagedArray<int32>& TransformToGeometryIndex,
+		const TManagedArray<int32>& TransformIndices,
+		const TManagedArray<FBox>& BoundingBoxes,
+		const TArray<FMatrix>& GlobalMatrices,
+		const FMatrix& LocalToWorldWithScale)
+	{
+		FBox BoundingBox(ForceInit);
+		if (bChaos_BoxCalcBounds_ISPC_Enabled && !bGeometryCollectionSingleThreadedBoundsCalculation)
+		{
+#if INTEL_ISPC
+			ispc::BoxCalcBoundsFromGeometryGroup(
+				(int32*)&TransformToGeometryIndex[0],
+				(int32*)&TransformIndices[0],
+				(ispc::FMatrix*)&GlobalMatrices[0],
+				(ispc::FBox*)&BoundingBoxes[0],
+				(ispc::FMatrix&)LocalToWorldWithScale,
+				(ispc::FBox&)BoundingBox,
+				BoundingBoxes.Num());
+#endif
+		}
+		else
+		{
+			for (int32 BoxIdx = 0; BoxIdx < BoundingBoxes.Num(); ++BoxIdx)
 			{
-				GlobalMatricesIn[tt] = GlobalMatricesIn[tt].ConcatTranslation((FVector)ExplodedVectors[tt]);
+				const int32 TransformIndex = TransformIndices[BoxIdx];
+
+				if (TransformToGeometryIndex[TransformIndex] != INDEX_NONE)
+				{
+					BoundingBox += BoundingBoxes[BoxIdx].TransformBy(GlobalMatrices[TransformIndex] * LocalToWorldWithScale);
+				}
 			}
 		}
+		return BoundingBox;
+	}
+
+	/**
+	* compute the bounding box from the bounding boxes stored in the transform group
+	* (used for nanite or when the geometry group data has been stripped on cook )
+	*/
+	inline FBox ComputeBoundsFromTransformBoundingBoxes(
+		const TManagedArray<int32>& TransformToGeometryIndex,
+		const TManagedArray<FBox>& BoundingBoxes,
+		const TArray<FMatrix>& GlobalMatrices,
+		const FMatrix& LocalToWorldWithScale)
+	{
+		FBox BoundingBox(ForceInit);
+		if (bChaos_BoxCalcBounds_ISPC_Enabled && !bGeometryCollectionSingleThreadedBoundsCalculation)
+		{
+#if INTEL_ISPC
+			ispc::BoxCalcBoundsFromTransformGroup(
+				(int32*)&TransformToGeometryIndex[0],
+				(ispc::FMatrix*)&GlobalMatrices[0],
+				(ispc::FBox*)&BoundingBoxes[0],
+				(ispc::FMatrix&)LocalToWorldWithScale,
+				(ispc::FBox&)BoundingBox,
+				BoundingBoxes.Num());
+#endif
+		}
+		else
+		{
+			for (int32 TransformIndex = 0; TransformIndex < BoundingBoxes.Num(); ++TransformIndex)
+			{
+				if (TransformToGeometryIndex[TransformIndex] != INDEX_NONE)
+				{
+					BoundingBox += BoundingBoxes[TransformIndex].TransformBy(GlobalMatrices[TransformIndex] * LocalToWorldWithScale);
+				}
+			}
+		}
+		return BoundingBox;
 	}
 }
+
+FBox UGeometryCollectionComponent::ComputeBoundsFromGlobalMatrices(const FMatrix& LocalToWorldWithScale, const TArray<FMatrix>& GlobalMatricesArray) const
+{
+	auto GeometryCollectionPtr = RestCollection->GetGeometryCollection();
+	const TManagedArray<FBox>* TransformBoundingBoxes = GeometryCollectionPtr->FindAttribute<FBox>("BoundingBox", FGeometryCollection::TransformGroup);
+	const TManagedArray<FBox>& GeometryBoundingBoxes = GetBoundingBoxArray();
+	const TManagedArray<int32>& TransformToGeometryIndex = GetTransformToGeometryIndexArray();
+
+	if (TransformBoundingBoxes)
+	{
+		return ComputeBoundsFromTransformBoundingBoxes(TransformToGeometryIndex, *TransformBoundingBoxes, GlobalMatricesArray, LocalToWorldWithScale);
+	}
+
+	const TManagedArray<int32>& TransformIndices = GetTransformIndexArray();
+	return ComputeBoundsFromGeometryBoundingBoxes(TransformToGeometryIndex, TransformIndices, GeometryBoundingBoxes, GlobalMatricesArray, LocalToWorldWithScale);
 }
 
 FBox UGeometryCollectionComponent::ComputeBounds(const FMatrix& LocalToWorldWithScale) const
@@ -474,29 +564,10 @@ FBox UGeometryCollectionComponent::ComputeBounds(const FMatrix& LocalToWorldWith
 	if (RestCollection)
 	{
 		//Hold on to reference so it doesn't get GC'ed
-		auto HackGeometryCollectionPtr = RestCollection->GetGeometryCollection();
+		auto GeometryCollectionPtr = RestCollection->GetGeometryCollection();
 
-		const TManagedArray<FBox>& BoundingBoxes = GetBoundingBoxArray();
-		const TManagedArray<int32>& TransformIndices = GetTransformIndexArray();
-		const TManagedArray<int32>& ParentIndices = GetParentArray();
-		const TManagedArray<int32>& TransformToGeometryIndex = GetTransformToGeometryIndexArray();
-		const TManagedArray<FTransform>& Transforms = GetTransformArray();
-
-		const int32 NumBoxes = BoundingBoxes.Num();
-	
-		int32 NumElements = HackGeometryCollectionPtr->NumElements(FGeometryCollection::TransformGroup);
-		if (RestCollection->EnableNanite && HackGeometryCollectionPtr->HasAttribute("BoundingBox", FGeometryCollection::TransformGroup) && NumElements)
-		{
-			TArray<FMatrix> TmpGlobalMatrices;
-			GeometryCollectionAlgo::GlobalMatrices(Transforms, ParentIndices, TmpGlobalMatrices);
-
-			TManagedArray<FBox>& TransformBounds = HackGeometryCollectionPtr->ModifyAttribute<FBox>("BoundingBox", "Transform");
-			for (int32 TransformIndex = 0; TransformIndex < HackGeometryCollectionPtr->NumElements(FGeometryCollection::TransformGroup); TransformIndex++)
-			{
-				BoundingBox += TransformBounds[TransformIndex].TransformBy(TmpGlobalMatrices[TransformIndex] * LocalToWorldWithScale);
-			}
-		}
-		else if (NumElements == 0 || GlobalMatrices.Num() != RestCollection->NumElements(FGeometryCollection::TransformGroup))
+		const int32 NumElements = GeometryCollectionPtr->NumElements(FGeometryCollection::TransformGroup);
+		if (NumElements == 0 || GlobalMatrices.Num() != RestCollection->NumElements(FGeometryCollection::TransformGroup))
 		{
 			// #todo(dmp): we could do the bbox transform in parallel with a bit of reformulating		
 			// #todo(dmp):  there are some cases where the calcbounds function is called before the component
@@ -504,8 +575,10 @@ FBox UGeometryCollectionComponent::ComputeBounds(const FMatrix& LocalToWorldWith
 			// to default to just calculating tmp global matrices.  This should be removed or modified somehow
 			// such that we always cache the global matrices and this method always does the correct behavior
 
+			const TManagedArray<FTransform>& Transforms = GetTransformArray();
+			const TManagedArray<int32>& ParentIndices = GetParentArray();
+
 			TArray<FMatrix> TmpGlobalMatrices;
-			
 			GeometryCollectionAlgo::GlobalMatrices(Transforms, ParentIndices, TmpGlobalMatrices);
 			if (TmpGlobalMatrices.Num() == 0)
 			{
@@ -513,59 +586,13 @@ FBox UGeometryCollectionComponent::ComputeBounds(const FMatrix& LocalToWorldWith
 			}
 			else
 			{
-				UpdateGlobalMatricesWithExplodedVectors(TmpGlobalMatrices, *(RestCollection->GetGeometryCollection()));
-				for (int32 BoxIdx = 0; BoxIdx < NumBoxes; ++BoxIdx)
-				{
-					const int32 TransformIndex = TransformIndices[BoxIdx];
-
-					if(RestCollection->GetGeometryCollection()->IsGeometry(TransformIndex))
-					{
-						BoundingBox += BoundingBoxes[BoxIdx].TransformBy(TmpGlobalMatrices[TransformIndex] * LocalToWorldWithScale);
-					}
-				}
-
-			}
-		}
-		else if (bGeometryCollectionSingleThreadedBoundsCalculation)
-		{
-			CHAOS_ENSURE(false); // this is slower and only enabled through a pvar debugging, disable bGeometryCollectionSingleThreadedBoundsCalculation in a production environment. 
-			for (int32 BoxIdx = 0; BoxIdx < NumBoxes; ++BoxIdx)
-			{
-				const int32 TransformIndex = TransformIndices[BoxIdx];
-
-				if (RestCollection->GetGeometryCollection()->IsGeometry(TransformIndex))
-				{
-					BoundingBox += BoundingBoxes[BoxIdx].TransformBy(GlobalMatrices[TransformIndex] * LocalToWorldWithScale);
-				}
+				UpdateGlobalMatricesWithExplodedVectors(TmpGlobalMatrices, *GeometryCollectionPtr);
+				BoundingBox = ComputeBoundsFromGlobalMatrices(LocalToWorldWithScale, TmpGlobalMatrices);
 			}
 		}
 		else
 		{
-			if (bChaos_BoxCalcBounds_ISPC_Enabled)
-			{
-#if INTEL_ISPC
-				ispc::BoxCalcBounds(
-					(int32*)&TransformToGeometryIndex[0],
-					(int32*)&TransformIndices[0],
-					(ispc::FMatrix*)&GlobalMatrices[0],
-					(ispc::FBox*)&BoundingBoxes[0],
-					(ispc::FMatrix&)LocalToWorldWithScale,
-					(ispc::FBox&)BoundingBox,
-					NumBoxes);
-#endif
-			}
-			else
-			{
-				for (int32 BoxIdx = 0; BoxIdx < NumBoxes; ++BoxIdx)
-				{
-					const int32 TransformIndex = TransformIndices[BoxIdx];
-
-					if (RestCollection->GetGeometryCollection()->IsGeometry(TransformIndex))
-					{
-						BoundingBox += BoundingBoxes[BoxIdx].TransformBy(GlobalMatrices[TransformIndex] * LocalToWorldWithScale);
-					}
-				}
-			}
+			BoundingBox = ComputeBoundsFromGlobalMatrices(LocalToWorldWithScale, GlobalMatrices);
 		}
 	}
 	return BoundingBox;
@@ -2487,7 +2514,8 @@ void UGeometryCollectionComponent::TickComponent(float DeltaTime, enum ELevelTic
 			// 	}
 			// }
 			
-			if (RestCollection->HasVisibleGeometry() && DynamicCollection->IsDirty())
+			const bool bHasVisibleGeometry = AssignedISMPool || RestCollection->HasVisibleGeometry();
+			if (bHasVisibleGeometry && DynamicCollection->IsDirty())
 			{
 				// #todo review: When we've made changes to ISMC, we need to move this function call to SetRenderDynamicData_Concurrent
 				RefreshEmbeddedGeometry();
diff --git a/Engine/Source/Runtime/Experimental/GeometryCollectionEngine/Private/GeometryCollection/GeometryCollectionComponent.ispc b/Engine/Source/Runtime/Experimental/GeometryCollectionEngine/Private/GeometryCollection/GeometryCollectionComponent.ispc
--- a/Engine/Source/Runtime/Experimental/GeometryCollectionEngine/Private/GeometryCollection/GeometryCollectionComponent.ispc
+++ b/Engine/Source/Runtime/Experimental/GeometryCollectionEngine/Private/GeometryCollection/GeometryCollectionComponent.ispc
@@ -91,13 +91,14 @@ inline static FBox BoxTransformByMatrix(const FBox &Box, const FMatrix &M)
 	return NewBox;
 }
 
-export void BoxCalcBounds(  const uniform int TransformToGeometryIndex[],
-							const uniform int TransformIndices[],
-							const uniform FMatrix GlobalMatrices[],
-							const uniform FBox BoundingBoxes[],
-							const uniform FMatrix &LocalToWorldWithScale,
-							uniform FBox &BoundingBox,
-							const uniform int NumBoxes)
+export void BoxCalcBoundsFromGeometryGroup(
+				const uniform int TransformToGeometryIndex[],
+				const uniform int TransformIndices[],
+				const uniform FMatrix GlobalMatrices[],
+				const uniform FBox BoundingBoxes[],
+				const uniform FMatrix &LocalToWorldWithScale,
+				uniform FBox &BoundingBox,
+				const uniform int NumBoxes)
 {
 	FBox BoundingBoxSum = BoxDefaultInit();
 
@@ -122,6 +123,47 @@ export void BoxCalcBounds(  const uniform int TransformToGeometryIndex[],
 	#pragma ignore warning(perf)
 	ReducedBox[programIndex] = BoundingBoxSum;
 
+	if (programIndex < NumBoxes)
+	{
+		foreach_active(j)
+		{
+			unmasked
+			{
+				BoundingBox = BoundingBox + ReducedBox[j];
+			}
+		}
+	}
+}
+
+export void BoxCalcBoundsFromTransformGroup(
+				const uniform int TransformToGeometryIndex[],
+				const uniform FMatrix GlobalMatrices[],
+				const uniform FBox BoundingBoxes[],
+				const uniform FMatrix &LocalToWorldWithScale,
+				uniform FBox &BoundingBox,
+				const uniform int NumBoxes)
+{
+	FBox BoundingBoxSum = BoxDefaultInit();
+
+	foreach(TransformIndex = 0 ... NumBoxes)
+	{
+		#pragma ignore warning(perf)
+		const FBox B = BoundingBoxes[TransformIndex];
+
+		#pragma ignore warning(perf)
+		if (TransformToGeometryIndex[TransformIndex] != INDEX_NONE && B.IsValid)
+		{
+			#pragma ignore warning(perf)
+			const FMatrix M = GlobalMatrices[TransformIndex];
+			BoundingBoxSum = BoundingBoxSum + BoxTransformByMatrix(B, M * LocalToWorldWithScale);
+		}
+	}
+
+	// Custom reduce for large arrays
+	uniform FBox ReducedBox[programCount];
+	#pragma ignore warning(perf)
+	ReducedBox[programIndex] = BoundingBoxSum;
+
 	if (programIndex < NumBoxes)
 	{
 		foreach_active(j)
diff --git a/Engine/Source/Runtime/Experimental/GeometryCollectionEngine/Private/GeometryCollection/GeometryCollectionObject.cpp b/Engine/Source/Runtime/Experimental/GeometryCollectionEngine/Private/GeometryCollection/GeometryCollectionObject.cpp
--- a/Engine/Source/Runtime/Experimental/GeometryCollectionEngine/Private/GeometryCollection/GeometryCollectionObject.cpp
+++ b/Engine/Source/Runtime/Experimental/GeometryCollectionEngine/Private/GeometryCollection/GeometryCollectionObject.cpp
@@ -1220,6 +1220,25 @@ TSharedPtr<FGeometryCollection, ESPMode::ThreadSafe> UGeometryCollection::CopyCo
 	const TArray<FName> GroupsToSkip{ FGeometryCollection::GeometryGroup, FGeometryCollection::VerticesGroup, FGeometryCollection::FacesGroup };
 	CollectionToCopy->CopyTo(GeometryCollectionToReturn.Get(), GroupsToSkip);
 
+	// since we are removing the bounding box attribute from the geometry group we need to move it to the transform group 
+	const TManagedArray<FBox>& GeometryBounds = CollectionToCopy->GetAttribute<FBox>("BoundingBox", "Geometry");
+	const TManagedArray<int32>& TransformToGeometryIndexArray = CollectionToCopy->TransformToGeometryIndex;
+
+	TManagedArray<FBox>& TransformBounds = GeometryCollectionToReturn->AddAttribute<FBox>("BoundingBox", "Transform");
+	
+	for (int TransformIndex = 0; TransformIndex < TransformBounds.Num(); TransformIndex++)
+	{
+		const int32 GeometryIndex = TransformToGeometryIndexArray[TransformIndex];
+		if (GeometryIndex != INDEX_NONE)
+		{
+			TransformBounds[TransformIndex] = GeometryBounds[GeometryIndex];
+		}
+		else
+		{
+			TransformBounds[TransformIndex].Init();
+		}
+	}
+
 	return GeometryCollectionToReturn;
 }
 
diff --git a/Engine/Source/Runtime/Experimental/GeometryCollectionEngine/Public/GeometryCollection/GeometryCollectionComponent.h b/Engine/Source/Runtime/Experimental/GeometryCollectionEngine/Public/GeometryCollection/GeometryCollectionComponent.h
--- a/Engine/Source/Runtime/Experimental/GeometryCollectionEngine/Public/GeometryCollection/GeometryCollectionComponent.h
+++ b/Engine/Source/Runtime/Experimental/GeometryCollectionEngine/Public/GeometryCollection/GeometryCollectionComponent.h
@@ -1080,6 +1080,7 @@ protected:
 	Chaos::FPhysicsSolver* GetSolver(const UGeometryCollectionComponent& GeometryCollectionComponent);
 	void CalculateLocalBounds();
 	void CalculateGlobalMatrices();
+	FBox ComputeBoundsFromGlobalMatrices(const FMatrix& LocalToWorldWithScale, const TArray<FMatrix>& GlobalMatricesArray) const;
 	FBox ComputeBounds(const FMatrix& LocalToWorldWithScale) const;
 
 	void RegisterForEvents();
