We desire more control over when wheels will spin in vehiceles. Introducing a max wheelspin velocity property and applying this more correctly to the wheel simulation.

diff --git a/Engine/Plugins/Experimental/ChaosVehiclesPlugin/Source/ChaosVehicles/Private/ChaosVehicleWheel.cpp b/Engine/Plugins/Experimental/ChaosVehiclesPlugin/Source/ChaosVehicles/Private/ChaosVehicleWheel.cpp
--- a/Engine/Plugins/Experimental/ChaosVehiclesPlugin/Source/ChaosVehicles/Private/ChaosVehicleWheel.cpp
+++ b/Engine/Plugins/Experimental/ChaosVehiclesPlugin/Source/ChaosVehicles/Private/ChaosVehicleWheel.cpp
@@ -30,6 +30,7 @@ UChaosVehicleWheel::UChaosVehicleWheel(const FObjectInitializer& ObjectInitializ
 	SideSlipModifier = 1.0f;
 	SlipThreshold = 20.0f;
 	SkidThreshold = 20.0f;
+	MaxWheelspinRotation = 30;
 
 	bAffectedByBrake = true;
 	bAffectedByHandbrake = false;
diff --git a/Engine/Plugins/Experimental/ChaosVehiclesPlugin/Source/ChaosVehicles/Public/ChaosVehicleWheel.h b/Engine/Plugins/Experimental/ChaosVehiclesPlugin/Source/ChaosVehicles/Public/ChaosVehicleWheel.h
--- a/Engine/Plugins/Experimental/ChaosVehiclesPlugin/Source/ChaosVehicles/Public/ChaosVehicleWheel.h
+++ b/Engine/Plugins/Experimental/ChaosVehiclesPlugin/Source/ChaosVehicles/Public/ChaosVehicleWheel.h
@@ -149,6 +149,10 @@ class UChaosWheeledVehicleMovementComponent;
 		UPROPERTY(EditAnywhere, Category = Wheel)
 		bool bTractionControlEnabled;
 
+		/** Max Wheelspin rotation rad/sec */
+		UPROPERTY(EditAnywhere, Category = Wheel)
+		float MaxWheelspinRotation;
+
 		/** Determines how the SetDriveTorque/SetBrakeTorque inputs are combined with the internal torques */
 		UPROPERTY(EditAnywhere, Category = Wheel)
 		ETorqueCombineMethod ExternalTorqueCombineMethod;
@@ -353,6 +357,7 @@ class UChaosWheeledVehicleMovementComponent;
 			PWheelConfig.TorqueRatio = 0.f;							// calculated later after all wheel info is known
 			PWheelConfig.ABSEnabled = this->bABSEnabled;
 			PWheelConfig.TractionControlEnabled = this->bTractionControlEnabled;
+			PWheelConfig.MaxSpinRotation = this->MaxWheelspinRotation;
 			PWheelConfig.AxleType = static_cast<Chaos::FSimpleWheelConfig::EAxleType>(this->AxleType);
 			PWheelConfig.FrictionMultiplier = this->FrictionForceMultiplier;
 			PWheelConfig.CorneringStiffness = this->CorneringStiffness * 10000.0f;
diff --git a/Engine/Source/Runtime/Experimental/ChaosVehicles/ChaosVehiclesCore/Private/WheelSystem.cpp b/Engine/Source/Runtime/Experimental/ChaosVehicles/ChaosVehiclesCore/Private/WheelSystem.cpp
--- a/Engine/Source/Runtime/Experimental/ChaosVehicles/ChaosVehiclesCore/Private/WheelSystem.cpp
+++ b/Engine/Source/Runtime/Experimental/ChaosVehicles/ChaosVehiclesCore/Private/WheelSystem.cpp
@@ -84,6 +84,7 @@ namespace Chaos
 		AvailableGrip = ForceIntoSurface * SurfaceFriction * FrictionMultiplier;
 
 		float FinalLongitudinalForce = 0.f;
+		float ExcessTorque = 0.0f;
 		float FinalLateralForce = 0.f;
 		float ABSGripThresholdSpeed = 5.0f;
 
@@ -110,11 +111,17 @@ namespace Chaos
 				}
 			}
 
-			// Traction control limiting drive force to match force from grip available
-			if (TractionControlEnabled && !Braking && FMath::Abs(AppliedLinearDriveForce) > AvailableGrip)
+			if (FMath::Abs(AppliedLinearDriveForce) > AvailableGrip)
 			{
-				float Sign = (AppliedLinearDriveForce > 0.0f) ? 1.0f : -1.0f;
-				AppliedLinearDriveForce = AvailableGrip * TractionControlAndABSScaling * Sign;
+				float SignTorque = AppliedLinearDriveForce < 0.0f ? -1.0f : 1.0f;
+				ExcessTorque = (FMath::Abs(AppliedLinearDriveForce) - AvailableGrip) * Re * SignTorque;
+
+				// Traction control limiting drive force to match force from grip available
+				if (TractionControlEnabled && !Braking)
+				{
+					float Sign = (AppliedLinearDriveForce > 0.0f) ? 1.0f : -1.0f;
+					AppliedLinearDriveForce = AvailableGrip * TractionControlAndABSScaling * Sign;
+				}
 			}
 
 			if (Braking)
@@ -180,11 +187,6 @@ namespace Chaos
 					{
 						WheelLocked = true;
 					}
-					else if (FMath::Abs(FinalLongitudinalForce) > AvailableGrip)
-					{
-						SlipOmega = (FinalLongitudinalForce < 0.0f) ? -Setup().MaxSpinRotation : Setup().MaxSpinRotation;
-
-					}
 
 					bClipping = true;
 					FinalLongitudinalForce *= Clip;
@@ -197,10 +199,15 @@ namespace Chaos
 		}
 		else
 		{
-			SlipOmega = DriveTorque / Inertia * DeltaTime;
-			SlipOmega = FMath::Clamp(SlipOmega, -Setup().MaxSpinRotation, Setup().MaxSpinRotation);
+			// only apply more spin torque if we haven't reached our max spin rotation
+			if (FMath::Abs(Omega) < Setup().MaxSpinRotation)
+			{
+				ExcessTorque = DriveTorque;
+			}
 		}
 
+		SlipOmega = ExcessTorque / Inertia * DeltaTime;
+		SlipOmega = FMath::Clamp(SlipOmega, -Setup().MaxSpinRotation, Setup().MaxSpinRotation);
 
 		if (WheelLocked)
 		{
