ChaosFlesh : Insights tags

diff --git a/Engine/Plugins/Experimental/ChaosFlesh/Source/ChaosFleshEngine/Private/ChaosFlesh/ChaosDeformablePhysicsComponent.cpp b/Engine/Plugins/Experimental/ChaosFlesh/Source/ChaosFleshEngine/Private/ChaosFlesh/ChaosDeformablePhysicsComponent.cpp
--- a/Engine/Plugins/Experimental/ChaosFlesh/Source/ChaosFleshEngine/Private/ChaosFlesh/ChaosDeformablePhysicsComponent.cpp
+++ b/Engine/Plugins/Experimental/ChaosFlesh/Source/ChaosFleshEngine/Private/ChaosFlesh/ChaosDeformablePhysicsComponent.cpp
@@ -94,6 +94,19 @@ void UDeformablePhysicsComponent::EnableSimulation(UDeformableSolverComponent* D
 }
 
 
+void UDeformablePhysicsComponent::EnableSimulationFromActor(ADeformableSolverActor* DeformableSolverActor)
+{
+	if (DeformableSolverActor && DeformableSolverActor->GetDeformableSolverComponent())
+	{
+		PrimarySolverComponent = DeformableSolverActor->GetDeformableSolverComponent();
+		if (!DeformableSolverActor->GetDeformableSolverComponent()->DeformableComponents.Contains(this))
+		{
+			DeformableSolverActor->GetDeformableSolverComponent()->DeformableComponents.Add(this);
+		}
+		DeformableSolverActor->GetDeformableSolverComponent()->AddDeformableProxy(this);
+	}
+}
+
 
 
 
diff --git a/Engine/Plugins/Experimental/ChaosFlesh/Source/ChaosFleshEngine/Private/ChaosFlesh/ChaosDeformableSolverComponent.cpp b/Engine/Plugins/Experimental/ChaosFlesh/Source/ChaosFleshEngine/Private/ChaosFlesh/ChaosDeformableSolverComponent.cpp
--- a/Engine/Plugins/Experimental/ChaosFlesh/Source/ChaosFleshEngine/Private/ChaosFlesh/ChaosDeformableSolverComponent.cpp
+++ b/Engine/Plugins/Experimental/ChaosFlesh/Source/ChaosFleshEngine/Private/ChaosFlesh/ChaosDeformableSolverComponent.cpp
@@ -5,9 +5,18 @@
 #include "ChaosFlesh/ChaosDeformablePhysicsComponent.h"
 #include "ChaosFlesh/ChaosDeformableSolverActor.h"
 #include "Engine/World.h"
+#include "ChaosStats.h"
 
 DEFINE_LOG_CATEGORY_STATIC(LogDeformableSolverComponentInternal, Log, All);
 
+DECLARE_CYCLE_STAT(TEXT("Chaos.Deformable.UDeformableSolverComponent.UpdateDeformableEndTickState"), STAT_ChaosDeformable_UDeformableSolverComponent_UpdateDeformableEndTickState, STATGROUP_Chaos);
+DECLARE_CYCLE_STAT(TEXT("Chaos.Deformable.UDeformableSolverComponent.TickComponent"), STAT_ChaosDeformable_UDeformableSolverComponent_TickComponent, STATGROUP_Chaos);
+DECLARE_CYCLE_STAT(TEXT("Chaos.Deformable.UDeformableSolverComponent.Reset"), STAT_ChaosDeformable_UDeformableSolverComponent_Reset, STATGROUP_Chaos);
+DECLARE_CYCLE_STAT(TEXT("Chaos.Deformable.UDeformableSolverComponent.AddDeformableProxy"), STAT_ChaosDeformable_UDeformableSolverComponent_AddDeformableProxy, STATGROUP_Chaos);
+DECLARE_CYCLE_STAT(TEXT("Chaos.Deformable.UDeformableSolverComponent.Simulate"), STAT_ChaosDeformable_UDeformableSolverComponent_Simulate, STATGROUP_Chaos);
+DECLARE_CYCLE_STAT(TEXT("Chaos.Deformable.UDeformableSolverComponent.UpdateFromGameThread"), STAT_ChaosDeformable_UDeformableSolverComponent_UpdateFromGameThread, STATGROUP_Chaos);
+DECLARE_CYCLE_STAT(TEXT("Chaos.Deformable.UDeformableSolverComponent.UpdateFromSimulation"), STAT_ChaosDeformable_UDeformableSolverComponent_UpdateFromSimulation, STATGROUP_Chaos);
+
 static TAutoConsoleVariable<int32> CVarDeformablePhysicsTickWaitForParallelDeformableTask(
 	TEXT("p.ClothPhysics.WaitForParallelDeformableTask"), 0, 
 	TEXT("If 1, always wait for deformable task completion in the Deformable Tick function. "\
@@ -94,7 +103,9 @@ bool UDeformableSolverComponent::IsSimulating(UDeformablePhysicsComponent* InCom
 
 void UDeformableSolverComponent::UpdateDeformableEndTickState(bool bRegister)
 {
-	UE_LOG(LogDeformableSolverComponentInternal, Verbose, TEXT("UDeformableSolverComponent::RegiUpdateDeformableEndTickStatesterEndPhysicsTick"));
+	SCOPE_CYCLE_COUNTER(STAT_ChaosDeformable_UDeformableSolverComponent_UpdateDeformableEndTickState);
+	TRACE_CPUPROFILER_EVENT_SCOPE(ChaosDeformable_UDeformableSolverComponent_UpdateDeformableEndTickState);
+
 	bRegister &= PrimaryComponentTick.IsTickFunctionRegistered();
 	if (bDoThreadedAdvance)
 	{
@@ -134,8 +145,8 @@ void UDeformableSolverComponent::BeginPlay()
 
 void UDeformableSolverComponent::TickComponent(float DeltaTime, enum ELevelTick TickType, FActorComponentTickFunction* ThisTickFunction)
 {
-	UE_LOG(LogDeformableSolverComponentInternal, Verbose, TEXT("UDeformableSolverComponent::TickComponent"));
-	TRACE_CPUPROFILER_EVENT_SCOPE(DeformableSolverComponent_TickComponent);
+	SCOPE_CYCLE_COUNTER(STAT_ChaosDeformable_UDeformableSolverComponent_TickComponent);
+	TRACE_CPUPROFILER_EVENT_SCOPE(ChaosDeformable_UDeformableSolverComponent_TickComponent);
 
 	if (GChaosEngineDeformableCVarParams.bEnableDeformableSolver)
 	{
@@ -169,6 +180,9 @@ void UDeformableSolverComponent::EndPlay(const EEndPlayReason::Type EndPlayReaso
 
 void UDeformableSolverComponent::Reset()
 {
+	SCOPE_CYCLE_COUNTER(STAT_ChaosDeformable_UDeformableSolverComponent_Reset);
+	TRACE_CPUPROFILER_EVENT_SCOPE(ChaosDeformable_UDeformableSolverComponent_Reset);
+
 	if (GChaosEngineDeformableCVarParams.bEnableDeformableSolver)
 	{
 		Solver.Reset(new FDeformableSolver({
@@ -205,7 +219,8 @@ void UDeformableSolverComponent::Reset()
 
 void UDeformableSolverComponent::AddDeformableProxy(UDeformablePhysicsComponent* InComponent)
 {
-	TRACE_CPUPROFILER_EVENT_SCOPE(DeformableSolverComponent_AddDeformableProxy);
+	SCOPE_CYCLE_COUNTER(STAT_ChaosDeformable_UDeformableSolverComponent_AddDeformableProxy);
+	TRACE_CPUPROFILER_EVENT_SCOPE(ChaosDeformable_UDeformableSolverComponent_AddDeformableProxy);
 
 	if (Solver && IsSimulating(InComponent))
 	{
@@ -219,7 +234,8 @@ void UDeformableSolverComponent::AddDeformableProxy(UDeformablePhysicsComponent*
 
 void UDeformableSolverComponent::Simulate(float DeltaTime)
 {
-	TRACE_CPUPROFILER_EVENT_SCOPE(DeformableSolverComponent_Simulate);
+	SCOPE_CYCLE_COUNTER(STAT_ChaosDeformable_UDeformableSolverComponent_Simulate);
+	TRACE_CPUPROFILER_EVENT_SCOPE(ChaosDeformable_UDeformableSolverComponent_Simulate);
 
 	if (Solver)
 	{
@@ -231,7 +247,8 @@ void UDeformableSolverComponent::Simulate(float DeltaTime)
 
 void UDeformableSolverComponent::UpdateFromGameThread(float DeltaTime)
 {
-	TRACE_CPUPROFILER_EVENT_SCOPE(DeformableSolverComponent_UpdateFromGameThread);
+	SCOPE_CYCLE_COUNTER(STAT_ChaosDeformable_UDeformableSolverComponent_UpdateFromGameThread);
+	TRACE_CPUPROFILER_EVENT_SCOPE(ChaosDeformable_UDeformableSolverComponent_UpdateFromGameThread);
 
 	if (Solver)
 	{
@@ -258,7 +275,8 @@ void UDeformableSolverComponent::UpdateFromGameThread(float DeltaTime)
 
 void UDeformableSolverComponent::UpdateFromSimulation(float DeltaTime)
 {
-	TRACE_CPUPROFILER_EVENT_SCOPE(DeformableSolverComponent_UpdateFromSimulation);
+	SCOPE_CYCLE_COUNTER(STAT_ChaosDeformable_UDeformableSolverComponent_UpdateFromSimulation);
+	TRACE_CPUPROFILER_EVENT_SCOPE(ChaosDeformable_UDeformableSolverComponent_UpdateFromSimulation);
 
 	if (Solver)
 	{
diff --git a/Engine/Plugins/Experimental/ChaosFlesh/Source/ChaosFleshEngine/Private/ChaosFlesh/FleshComponent.cpp b/Engine/Plugins/Experimental/ChaosFlesh/Source/ChaosFleshEngine/Private/ChaosFlesh/FleshComponent.cpp
--- a/Engine/Plugins/Experimental/ChaosFlesh/Source/ChaosFleshEngine/Private/ChaosFlesh/FleshComponent.cpp
+++ b/Engine/Plugins/Experimental/ChaosFlesh/Source/ChaosFleshEngine/Private/ChaosFlesh/FleshComponent.cpp
@@ -8,15 +8,27 @@
 #include "ChaosFlesh/FleshCollection.h"
 #include "ChaosFlesh/ChaosDeformableSolverActor.h"
 #include "ChaosFlesh/ChaosDeformableSolverComponent.h"
+#include "ChaosStats.h"
 #include "Components/SkeletalMeshComponent.h"
 #include "Dataflow/DataflowEngineUtil.h"
 #include "Engine/SkeletalMesh.h"
 #include "GeometryCollection/Facades/CollectionTransformSourceFacade.h"
 #include "GeometryCollection/TransformCollection.h"
 #include "ProceduralMeshComponent.h"
+#include "ChaosFlesh/ChaosDeformableTypes.h"
+
 
 #include UE_INLINE_GENERATED_CPP_BY_NAME(FleshComponent)
 
+FChaosEngineDeformableCVarParams CVarParams;
+FAutoConsoleVariableRef CVarDeforambleDoDrawSimulationMesh(TEXT("p.Chaos.DebugDraw.Deformable.SimulationMesh"), CVarParams.bDoDrawSimulationMesh, TEXT("Debug draw the deformable simulation resutls on the game thread. [def: true]"));
+
+DECLARE_CYCLE_STAT(TEXT("Chaos.Deformable.UFleshComponent.TickComponent"), STAT_ChaosDeformable_UFleshComponent_TickComponent, STATGROUP_Chaos);
+DECLARE_CYCLE_STAT(TEXT("Chaos.Deformable.UFleshComponent.NewProxy"), STAT_ChaosDeformable_UFleshComponent_NewProxy, STATGROUP_Chaos);
+DECLARE_CYCLE_STAT(TEXT("Chaos.Deformable.UFleshComponent.UpdateFromSimualtion"), STAT_ChaosDeformable_UFleshComponent_NewDeformableData, STATGROUP_Chaos);
+DECLARE_CYCLE_STAT(TEXT("Chaos.Deformable.UFleshComponent.UpdateFromSimualtion"), STAT_ChaosDeformable_UFleshComponent_UpdateFromSimualtion, STATGROUP_Chaos);
+DECLARE_CYCLE_STAT(TEXT("Chaos.Deformable.UFleshComponent.RenderProceduralMesh"), STAT_ChaosDeformable_UFleshComponent_RenderProceduralMesh, STATGROUP_Chaos);
+
 DEFINE_LOG_CATEGORY_STATIC(LogFleshComponentInternal, Log, All);
 
 UFleshComponent::UFleshComponent(const FObjectInitializer& ObjectInitializer)
@@ -70,10 +82,19 @@ void UFleshComponent::BeginPlay()
 
 void UFleshComponent::TickComponent(float DeltaTime, enum ELevelTick TickType, FActorComponentTickFunction *ThisTickFunction)
 {
+	SCOPE_CYCLE_COUNTER(STAT_ChaosDeformable_UFleshComponent_NewDeformableData);
+	TRACE_CPUPROFILER_EVENT_SCOPE(ChaosDeformable_UFleshComponent_NewDeformableData);
+
 	Super::TickComponent(DeltaTime, TickType, ThisTickFunction);
-	//UE_LOG(LogFleshComponentInternal, Log, TEXT("UFleshComponent::TickComponent()"));
-	
-	RenderProceduralMesh();
+
+	if (CVarParams.bDoDrawSimulationMesh)
+	{
+		RenderProceduralMesh();
+	}
+	else 
+	{
+		ResetProceduralMesh();
+	}
 	UpdateLocalBounds();
 }
 
@@ -96,7 +117,9 @@ void UFleshComponent::SetRestCollection(const UFleshAsset* InRestCollection)
 
 UDeformablePhysicsComponent::FThreadingProxy* UFleshComponent::NewProxy()
 {
-	TRACE_CPUPROFILER_EVENT_SCOPE(DeformableSolverComponent_AddFleshProxy);
+	SCOPE_CYCLE_COUNTER(STAT_ChaosDeformable_UFleshComponent_NewProxy);
+	TRACE_CPUPROFILER_EVENT_SCOPE(ChaosDeformable_UFleshComponent_NewProxy);
+
 	if (const UFleshAsset* RestAsset = GetRestCollection())
 	{
 		if (const FFleshCollection* Rest = RestAsset->GetCollection())
@@ -118,6 +141,9 @@ UDeformablePhysicsComponent::FThreadingProxy* UFleshComponent::NewProxy()
 
 UDeformablePhysicsComponent::FDataMapValue UFleshComponent::NewDeformableData()
 {
+	SCOPE_CYCLE_COUNTER(STAT_ChaosDeformable_UFleshComponent_NewDeformableData);
+	TRACE_CPUPROFILER_EVENT_SCOPE(ChaosDeformable_UFleshComponent_NewDeformableData);
+
 	using namespace GeometryCollection::Facades;
 	if (GetOwner())
 	{
@@ -174,6 +200,9 @@ UDeformablePhysicsComponent::FDataMapValue UFleshComponent::NewDeformableData()
 
 void UFleshComponent::UpdateFromSimualtion(const FDataMapValue* SimualtionBuffer)
 {
+	SCOPE_CYCLE_COUNTER(STAT_ChaosDeformable_UFleshComponent_UpdateFromSimualtion);
+	TRACE_CPUPROFILER_EVENT_SCOPE(ChaosDeformable_UFleshComponent_UpdateFromSimualtion);
+
 	FTransform CurrTransform = GetComponentTransform();
 	if (const FFleshThreadingProxy::FFleshOutputBuffer* FleshBuffer = (*SimualtionBuffer)->As<FFleshThreadingProxy::FFleshOutputBuffer>())
 	{
@@ -243,9 +272,12 @@ void UFleshComponent::ResetDynamicCollection()
 
 void UFleshComponent::ResetProceduralMesh()
 {
-	if (Mesh && RenderMesh)
+	if (Mesh )
 	{
 		Mesh->ClearAllMeshSections();
+	}
+	if (RenderMesh)
+	{
 		delete RenderMesh;
 		RenderMesh = nullptr;
 	}
@@ -253,11 +285,14 @@ void UFleshComponent::ResetProceduralMesh()
 
 void UFleshComponent::RenderProceduralMesh()
 {
+	SCOPE_CYCLE_COUNTER(STAT_ChaosDeformable_UFleshComponent_RenderProceduralMesh);
+	TRACE_CPUPROFILER_EVENT_SCOPE(ChaosDeformable_UFleshComponent_RenderProceduralMesh);
+
 	bool bCanRender = false;
 	if (const UFleshAsset* FleshAsset = GetRestCollection())
 	{
 #if WITH_EDITORONLY_DATA
-		if (FleshAsset->bRenderInEditor)
+		if (Mesh && FleshAsset->bRenderInEditor && CVarParams.bDoDrawSimulationMesh)
 #endif
 		{
 			if (const FFleshCollection* Flesh = FleshAsset->GetCollection())
diff --git a/Engine/Plugins/Experimental/ChaosFlesh/Source/ChaosFleshEngine/Public/ChaosFlesh/ChaosDeformablePhysicsComponent.h b/Engine/Plugins/Experimental/ChaosFlesh/Source/ChaosFleshEngine/Public/ChaosFlesh/ChaosDeformablePhysicsComponent.h
--- a/Engine/Plugins/Experimental/ChaosFlesh/Source/ChaosFleshEngine/Public/ChaosFlesh/ChaosDeformablePhysicsComponent.h
+++ b/Engine/Plugins/Experimental/ChaosFlesh/Source/ChaosFleshEngine/Public/ChaosFlesh/ChaosDeformablePhysicsComponent.h
@@ -30,6 +30,9 @@ public:
 	UFUNCTION(BlueprintCallable, Category = "Physics")
 	void EnableSimulation(UDeformableSolverComponent* DeformableSolverComponent);
 
+	UFUNCTION(BlueprintCallable, Category = "Physics")
+	void EnableSimulationFromActor(ADeformableSolverActor* DeformableSolverActor);
+
 	virtual FThreadingProxy* NewProxy() { return nullptr; }
 	virtual void AddProxy(Chaos::Softs::FDeformableSolver::FGameThreadAccess& GameThreadSolver);
 	virtual void RemoveProxy(Chaos::Softs::FDeformableSolver::FGameThreadAccess& GameThreadSolver);
diff --git a/Engine/Source/Runtime/Experimental/Chaos/Private/Chaos/Deformable/ChaosDeformableSolver.cpp b/Engine/Source/Runtime/Experimental/Chaos/Private/Chaos/Deformable/ChaosDeformableSolver.cpp
--- a/Engine/Source/Runtime/Experimental/Chaos/Private/Chaos/Deformable/ChaosDeformableSolver.cpp
+++ b/Engine/Source/Runtime/Experimental/Chaos/Private/Chaos/Deformable/ChaosDeformableSolver.cpp
@@ -32,8 +32,31 @@
 #include "HAL/IConsoleManager.h"
 #include "Misc/Paths.h"
 
-DECLARE_CYCLE_STAT(TEXT("DeformableSolver.Advance"), STAT_DeformableSolver_Advance, STATGROUP_Chaos);
-
+DECLARE_CYCLE_STAT(TEXT("Chaos.Deformable.Solver.Constructor"), STAT_ChaosDeformableSolver_Constructor, STATGROUP_Chaos);
+DECLARE_CYCLE_STAT(TEXT("Chaos.Deformable.Solver.Destructor"), STAT_ChaosDeformableSolver_Destructor, STATGROUP_Chaos);
+DECLARE_CYCLE_STAT(TEXT("Chaos.Deformable.Solver.UpdateProxyInputPackages"), STAT_ChaosDeformableSolver_UpdateProxyInputPackages, STATGROUP_Chaos);
+DECLARE_CYCLE_STAT(TEXT("Chaos.Deformable.Solver.Simulate"), STAT_ChaosDeformableSolver_Simulate, STATGROUP_Chaos);
+DECLARE_CYCLE_STAT(TEXT("Chaos.Deformable.Solver.AdvanceDt."), STAT_ChaosDeformableSolver_AdvanceDt, STATGROUP_Chaos);
+DECLARE_CYCLE_STAT(TEXT("Chaos.Deformable.Solver.Reset"), STAT_ChaosDeformableSolver_Reset, STATGROUP_Chaos);
+DECLARE_CYCLE_STAT(TEXT("Chaos.Deformable.Solver.Update"), STAT_ChaosDeformableSolver_Update, STATGROUP_Chaos);
+DECLARE_CYCLE_STAT(TEXT("Chaos.Deformable.Solver.UpdateOutputState"), STAT_ChaosDeformableSolver_UpdateOutputState, STATGROUP_Chaos);
+DECLARE_CYCLE_STAT(TEXT("Chaos.Deformable.Solver.PullOutputPackage"), STAT_ChaosDeformableSolver_PullOutputPackage, STATGROUP_Chaos);
+DECLARE_CYCLE_STAT(TEXT("Chaos.Deformable.Solver.PushOutputPackage"), STAT_ChaosDeformableSolver_PushOutputPackage, STATGROUP_Chaos);
+DECLARE_CYCLE_STAT(TEXT("Chaos.Deformable.Solver.PullInputPackage"), STAT_ChaosDeformableSolver_PullInputPackage, STATGROUP_Chaos);
+DECLARE_CYCLE_STAT(TEXT("Chaos.Deformable.Solver.PushInputPackage"), STAT_ChaosDeformableSolver_PushInputPackage, STATGROUP_Chaos);
+DECLARE_CYCLE_STAT(TEXT("Chaos.Deformable.Solver.InitializeSimulationObjects"), STAT_ChaosDeformableSolver_InitializeSimulationObjects, STATGROUP_Chaos);
+DECLARE_CYCLE_STAT(TEXT("Chaos.Deformable.Solver.InitializeSimulationObject"), STAT_ChaosDeformableSolver_InitializeSimulationObject, STATGROUP_Chaos);
+DECLARE_CYCLE_STAT(TEXT("Chaos.Deformable.Solver.InitializeDeformableParticles"), STAT_ChaosDeformableSolver_InitializeDeformableParticles, STATGROUP_Chaos);
+DECLARE_CYCLE_STAT(TEXT("Chaos.Deformable.Solver.InitializeKinematicParticles"), STAT_ChaosDeformableSolver_InitializeKinematicParticles, STATGROUP_Chaos);
+DECLARE_CYCLE_STAT(TEXT("Chaos.Deformable.Solver.InitializeTetrahedralConstraint"), STAT_ChaosDeformableSolver_InitializeTetrahedralConstraint, STATGROUP_Chaos);
+DECLARE_CYCLE_STAT(TEXT("Chaos.Deformable.Solver.InitializeGidBasedConstraints"), STAT_ChaosDeformableSolver_InitializeGidBasedConstraints, STATGROUP_Chaos);
+DECLARE_CYCLE_STAT(TEXT("Chaos.Deformable.Solver.InitializeKinematicConstraint"), STAT_ChaosDeformableSolver_InitializeKinematicConstraint, STATGROUP_Chaos);
+DECLARE_CYCLE_STAT(TEXT("Chaos.Deformable.Solver.InitializeCollisionBodies"), STAT_ChaosDeformableSolver_InitializeCollisionBodies, STATGROUP_Chaos);
+DECLARE_CYCLE_STAT(TEXT("Chaos.Deformable.Solver.InitializeSelfCollisionVariables"), STAT_ChaosDeformableSolver_InitializeSelfCollisionVariables, STATGROUP_Chaos);
+DECLARE_CYCLE_STAT(TEXT("Chaos.Deformable.Solver.InitializeGridBasedConstraintVariables"), STAT_ChaosDeformableSolver_InitializeGridBasedConstraintVariables, STATGROUP_Chaos);
+DECLARE_CYCLE_STAT(TEXT("Chaos.Deformable.Solver.RemoveSimulationObjects"), STAT_ChaosDeformableSolver_RemoveSimulationObjects, STATGROUP_Chaos);
+DECLARE_CYCLE_STAT(TEXT("Chaos.Deformable.Solver.RemoveProxy"), STAT_ChaosDeformableSolver_RemoveProxy, STATGROUP_Chaos);
+DECLARE_CYCLE_STAT(TEXT("Chaos.Deformable.Solver.AddProxy"), STAT_ChaosDeformableSolver_AddProxy, STATGROUP_Chaos);
 
 
 DEFINE_LOG_CATEGORY_STATIC(LogChaosDeformableSolver, Log, All);
@@ -56,11 +79,16 @@ namespace Chaos::Softs
 		, PreviousInputPackage(TUniquePtr < FDeformablePackage >(nullptr))
 		, Property(InProp)
 	{
+		SCOPE_CYCLE_COUNTER(STAT_ChaosDeformableSolver_Constructor);
+		TRACE_CPUPROFILER_EVENT_SCOPE(ChaosDeformableSolver_Constructor);
 		Reset(Property);
 	}
 
 	FDeformableSolver::~FDeformableSolver()
 	{
+		SCOPE_CYCLE_COUNTER(STAT_ChaosDeformableSolver_Destructor);
+		TRACE_CPUPROFILER_EVENT_SCOPE(ChaosDeformableSolver_Destructor);
+
 		FScopeLock Lock(&InitializationMutex);
 		for (FThreadingProxy* Proxy : UninitializedProxys_Internal)
 		{
@@ -73,7 +101,8 @@ namespace Chaos::Softs
 
 	void FDeformableSolver::Reset(const FDeformableSolverProperties& InProps)
 	{
-		TRACE_CPUPROFILER_EVENT_SCOPE(DeformableSolver_Reset);
+		SCOPE_CYCLE_COUNTER(STAT_ChaosDeformableSolver_Reset);
+		TRACE_CPUPROFILER_EVENT_SCOPE(ChaosDeformableSolver_Reset);
 
 		Property = InProps;
 		MObjects = TArrayCollectionArray<const UObject*>();
@@ -109,6 +138,9 @@ namespace Chaos::Softs
 
 	void FDeformableSolver::Simulate(FSolverReal DeltaTime)
 	{
+		SCOPE_CYCLE_COUNTER(STAT_ChaosDeformableSolver_Simulate);
+		TRACE_CPUPROFILER_EVENT_SCOPE(ChaosDeformableSolver_Simulate);
+
 		if (Property.NumSolverIterations)
 		{
 			RemoveSimulationObjects();
@@ -120,8 +152,8 @@ namespace Chaos::Softs
 
 	void FDeformableSolver::InitializeSimulationObjects()
 	{
-		TRACE_CPUPROFILER_EVENT_SCOPE(DeformableSolver_InitializeSimulationObjects);
-
+		SCOPE_CYCLE_COUNTER(STAT_ChaosDeformableSolver_InitializeSimulationObjects);
+		TRACE_CPUPROFILER_EVENT_SCOPE(ChaosDeformableSolver_InitializeSimulationObjects);
 		{
 			FScopeLock Lock(&InitializationMutex); // @todo(flesh) : change to threaded task based commands to prevent the lock. 
 			if (UninitializedProxys_Internal.Num())
@@ -155,7 +187,9 @@ namespace Chaos::Softs
 
 	void FDeformableSolver::InitializeSimulationObject(FThreadingProxy& InProxy)
 	{
-		TRACE_CPUPROFILER_EVENT_SCOPE(DeformableSolver_InitializeSimulationObject);
+		SCOPE_CYCLE_COUNTER(STAT_ChaosDeformableSolver_InitializeSimulationObject);
+		TRACE_CPUPROFILER_EVENT_SCOPE(ChaosDeformableSolver_InitializeSimulationObject);
+
 		if (FFleshThreadingProxy* Proxy = InProxy.As<FFleshThreadingProxy>())
 		{
 			InitializeDeformableParticles(*Proxy);
@@ -167,6 +201,9 @@ namespace Chaos::Softs
 
 	void FDeformableSolver::InitializeDeformableParticles(FFleshThreadingProxy& Proxy)
 	{
+		SCOPE_CYCLE_COUNTER(STAT_ChaosDeformableSolver_InitializeDeformableParticles);
+		TRACE_CPUPROFILER_EVENT_SCOPE(ChaosDeformableSolver_InitializeDeformableParticles);
+
 		const FManagedArrayCollection& Dynamic = Proxy.GetDynamicCollection();
 		const FManagedArrayCollection& Rest = Proxy.GetRestCollection();
 
@@ -221,7 +258,8 @@ namespace Chaos::Softs
 
 	void FDeformableSolver::InitializeKinematicParticles(FFleshThreadingProxy& Proxy)
 	{
-		TRACE_CPUPROFILER_EVENT_SCOPE(DeformableSolver_InitializeKinematicParticles);
+		SCOPE_CYCLE_COUNTER(STAT_ChaosDeformableSolver_InitializeKinematicParticles);
+		TRACE_CPUPROFILER_EVENT_SCOPE(ChaosDeformableSolver_InitializeKinematicParticles);
 
 		const FManagedArrayCollection& Rest = Proxy.GetRestCollection();
 		const FIntVector2& Range = Proxy.GetSolverParticleRange();
@@ -276,7 +314,8 @@ namespace Chaos::Softs
 
 	void FDeformableSolver::InitializeTetrahedralConstraint(FFleshThreadingProxy& Proxy)
 	{
-		TRACE_CPUPROFILER_EVENT_SCOPE(DeformableSolver_InitializeTetrahedralConstraint);
+		SCOPE_CYCLE_COUNTER(STAT_ChaosDeformableSolver_InitializeTetrahedralConstraint);
+		TRACE_CPUPROFILER_EVENT_SCOPE(ChaosDeformableSolver_InitializeTetrahedralConstraint);
 
 		const FManagedArrayCollection& Rest = Proxy.GetRestCollection();
 
@@ -357,7 +396,9 @@ namespace Chaos::Softs
 
 	void FDeformableSolver::InitializeGidBasedConstraints(FFleshThreadingProxy& Proxy)
 	{
-		TRACE_CPUPROFILER_EVENT_SCOPE(DeformableSolver_InitializeSimulationObject);
+		SCOPE_CYCLE_COUNTER(STAT_ChaosDeformableSolver_InitializeGidBasedConstraints);
+		TRACE_CPUPROFILER_EVENT_SCOPE(ChaosDeformableSolver_InitializeGidBasedConstraints);
+
 		if (Property.bUseGridBasedConstraints)
 		{
 			auto ChaosTet = [](FIntVector4 V, int32 dp) { return Chaos::TVec4<int32>(dp + V.X, dp + V.Y, dp + V.Z, dp + V.W); };
@@ -384,7 +425,8 @@ namespace Chaos::Softs
 	{
 		auto MKineticUpdate = [this](FSolverParticles& MParticles, const FSolverReal Dt, const FSolverReal MTime, const int32 Index)
 		{
-			TRACE_CPUPROFILER_EVENT_SCOPE(DeformableSolver_MKineticUpdate);
+			SCOPE_CYCLE_COUNTER(STAT_ChaosDeformableSolver_InitializeKinematicConstraint);
+			TRACE_CPUPROFILER_EVENT_SCOPE(ChaosDeformableSolver_InitializeKinematicConstraint);
 
 			if (0 <= Index && Index < this->MObjects.Num())
 			{
@@ -493,7 +535,9 @@ namespace Chaos::Softs
 
 	void FDeformableSolver::InitializeCollisionBodies()
 	{
-		TRACE_CPUPROFILER_EVENT_SCOPE(DeformableSolver_InitializeCollisionBodies);
+		SCOPE_CYCLE_COUNTER(STAT_ChaosDeformableSolver_InitializeCollisionBodies);
+		TRACE_CPUPROFILER_EVENT_SCOPE(ChaosDeformableSolver_InitializeCollisionBodies);
+
 		if (Property.bUseFloor && Evolution->CollisionParticles().Size() == 0)
 		{
 			Chaos::FVec3 Position(0.f);
@@ -508,7 +552,9 @@ namespace Chaos::Softs
 
 	void FDeformableSolver::InitializeSelfCollisionVariables()
 	{
-		TRACE_CPUPROFILER_EVENT_SCOPE(DeformableSolver_InitializeSelfCollisionVariables);
+		SCOPE_CYCLE_COUNTER(STAT_ChaosDeformableSolver_InitializeSelfCollisionVariables);
+		TRACE_CPUPROFILER_EVENT_SCOPE(ChaosDeformableSolver_InitializeSelfCollisionVariables);
+
 
 		int32 NumParticles = Evolution->Particles().Size();
 		SurfaceTriangleMesh->Init(*SurfaceElements);
@@ -539,7 +585,9 @@ namespace Chaos::Softs
 
 	void FDeformableSolver::InitializeGridBasedConstraintVariables()
 	{
-		TRACE_CPUPROFILER_EVENT_SCOPE(DeformableSolver_InitializeGridBasedConstraintVariables);
+		SCOPE_CYCLE_COUNTER(STAT_ChaosDeformableSolver_InitializeGridBasedConstraintVariables);
+		TRACE_CPUPROFILER_EVENT_SCOPE(ChaosDeformableSolver_InitializeGridBasedConstraintVariables);
+
 
 		GridBasedCorotatedConstraint.Reset(new Chaos::Softs::FXPBDGridBasedCorotatedConstraints<FSolverReal, FSolverParticles>(
 			Evolution->Particles(), *AllElements, Property.GridDx, /*bRecordMetric = */false, (Chaos::Softs::FSolverReal).1, (Chaos::Softs::FSolverReal).01, (Chaos::Softs::FSolverReal).4, (Chaos::Softs::FSolverReal)1000.0));
@@ -568,7 +616,9 @@ namespace Chaos::Softs
 
 	void FDeformableSolver::RemoveSimulationObjects()
 	{
-		TRACE_CPUPROFILER_EVENT_SCOPE(DeformableSolver_RemoveSimulationObjects);
+		SCOPE_CYCLE_COUNTER(STAT_ChaosDeformableSolver_RemoveSimulationObjects);
+		TRACE_CPUPROFILER_EVENT_SCOPE(ChaosDeformableSolver_RemoveSimulationObjects);
+
 		TArray< FThreadingProxy* > RemovedProxies;
 		{
 			FScopeLock Lock(&RemovalMutex); // @todo(flesh) : change to threaded task based commands to prevent the lock. 
@@ -629,8 +679,8 @@ namespace Chaos::Softs
 
 	void FDeformableSolver::AdvanceDt(FSolverReal DeltaTime)
 	{
-		SCOPE_CYCLE_COUNTER(STAT_DeformableSolver_Advance);
-		TRACE_CPUPROFILER_EVENT_SCOPE(DeformableSolver_Advance);
+		SCOPE_CYCLE_COUNTER(STAT_ChaosDeformableSolver_AdvanceDt);
+		TRACE_CPUPROFILER_EVENT_SCOPE(ChaosDeformableSolver_AdvanceDt);
 
 		EventPreSolve.Broadcast(DeltaTime);
 
@@ -699,15 +749,21 @@ namespace Chaos::Softs
 
 	void FDeformableSolver::PushInputPackage(int32 InFrame, FDeformableDataMap&& InPackage)
 	{
+		SCOPE_CYCLE_COUNTER(STAT_ChaosDeformableSolver_PushInputPackage);
+		TRACE_CPUPROFILER_EVENT_SCOPE(ChaosDeformableSolver_PushInputPackage);
+
 		FScopeLock Lock(&PackageInputMutex);
-		TRACE_CPUPROFILER_EVENT_SCOPE(DeformableSolver_PushInputPackage);
+		TRACE_CPUPROFILER_EVENT_SCOPE(ChaosDeformableSolver_PushInputPackage);
 		BufferedInputPackages.Push(TUniquePtr< FDeformablePackage >(new FDeformablePackage(InFrame, MoveTemp(InPackage))));
 	}
 
 	TUniquePtr<FDeformablePackage> FDeformableSolver::PullInputPackage()
 	{
+		SCOPE_CYCLE_COUNTER(STAT_ChaosDeformableSolver_PullInputPackage);
+		TRACE_CPUPROFILER_EVENT_SCOPE(ChaosDeformableSolver_PullInputPackage);
+
 		FScopeLock Lock(&PackageInputMutex);
-		TRACE_CPUPROFILER_EVENT_SCOPE(DeformableSolver_PullInputPackage);
+		TRACE_CPUPROFILER_EVENT_SCOPE(ChaosDeformableSolver_PullInputPackage);
 		if (BufferedInputPackages.Num())
 			return BufferedInputPackages.Pop();
 		return TUniquePtr<FDeformablePackage>(nullptr);
@@ -715,6 +771,9 @@ namespace Chaos::Softs
 
 	void FDeformableSolver::UpdateProxyInputPackages()
 	{
+		SCOPE_CYCLE_COUNTER(STAT_ChaosDeformableSolver_UpdateProxyInputPackages);
+		TRACE_CPUPROFILER_EVENT_SCOPE(ChaosDeformableSolver_UpdateProxyInputPackages);
+
 		if (CurrentInputPackage)
 		{
 			PreviousInputPackage = TUniquePtr < FDeformablePackage >(CurrentInputPackage.Release());
@@ -731,7 +790,9 @@ namespace Chaos::Softs
 
 	void FDeformableSolver::Update(FSolverReal DeltaTime)
 	{
-		TRACE_CPUPROFILER_EVENT_SCOPE(DeformableSolver_Update);
+		SCOPE_CYCLE_COUNTER(STAT_ChaosDeformableSolver_Update);
+		TRACE_CPUPROFILER_EVENT_SCOPE(ChaosDeformableSolver_Update);
+
 
 		if (!Proxies.Num()) return;
 
@@ -751,15 +812,21 @@ namespace Chaos::Softs
 
 	void FDeformableSolver::PushOutputPackage(int32 InFrame, FDeformableDataMap&& InPackage)
 	{
+		SCOPE_CYCLE_COUNTER(STAT_ChaosDeformableSolver_PushOutputPackage);
+		TRACE_CPUPROFILER_EVENT_SCOPE(ChaosDeformableSolver_PushOutputPackage);
+
 		FScopeLock Lock(&PackageOutputMutex);
-		TRACE_CPUPROFILER_EVENT_SCOPE(DeformableSolver_PushPackage);
+		TRACE_CPUPROFILER_EVENT_SCOPE(ChaosDeformableSolver_PushPackage);
 		BufferedOutputPackages.Push(TUniquePtr< FDeformablePackage >(new FDeformablePackage(InFrame, MoveTemp(InPackage))));
 	}
 
 	TUniquePtr<FDeformablePackage> FDeformableSolver::PullOutputPackage()
 	{
+		SCOPE_CYCLE_COUNTER(STAT_ChaosDeformableSolver_PullOutputPackage);
+		TRACE_CPUPROFILER_EVENT_SCOPE(ChaosDeformableSolver_PullOutputPackage);
+
 		FScopeLock Lock(&PackageOutputMutex);
-		TRACE_CPUPROFILER_EVENT_SCOPE(DeformableSolver_PullPackage);
+		TRACE_CPUPROFILER_EVENT_SCOPE(ChaosDeformableSolver_PullPackage);
 		if (BufferedOutputPackages.Num())
 			return BufferedOutputPackages.Pop();
 		return TUniquePtr<FDeformablePackage>(nullptr);
@@ -767,6 +834,9 @@ namespace Chaos::Softs
 
 	void FDeformableSolver::AddProxy(FThreadingProxy* InProxy)
 	{
+		SCOPE_CYCLE_COUNTER(STAT_ChaosDeformableSolver_AddProxy);
+		TRACE_CPUPROFILER_EVENT_SCOPE(ChaosDeformableSolver_AddProxy);
+
 		FScopeLock Lock(&InitializationMutex);
 		UninitializedProxys_Internal.Add(InProxy);
 		InitializedObjects_External.Add(InProxy->GetOwner());
@@ -774,6 +844,9 @@ namespace Chaos::Softs
 
 	void FDeformableSolver::RemoveProxy(FThreadingProxy* InProxy)
 	{
+		SCOPE_CYCLE_COUNTER(STAT_ChaosDeformableSolver_RemoveProxy);
+		TRACE_CPUPROFILER_EVENT_SCOPE(ChaosDeformableSolver_RemoveProxy);
+
 		FScopeLock LockA(&RemovalMutex);
 		FScopeLock LockB(&InitializationMutex);
 
@@ -803,7 +876,10 @@ namespace Chaos::Softs
 
 	void FDeformableSolver::UpdateOutputState(FThreadingProxy& InProxy)
 	{
-		TRACE_CPUPROFILER_EVENT_SCOPE(DeformableSolver_UpdateOutputState);
+		SCOPE_CYCLE_COUNTER(STAT_ChaosDeformableSolver_UpdateOutputState);
+		TRACE_CPUPROFILER_EVENT_SCOPE(ChaosDeformableSolver_UpdateOutputState);
+
+		TRACE_CPUPROFILER_EVENT_SCOPE(ChaosDeformableSolver_UpdateOutputState);
 		if (FFleshThreadingProxy* Proxy = InProxy.As<FFleshThreadingProxy>())
 		{
 
@@ -825,7 +901,7 @@ namespace Chaos::Softs
 
 	void FDeformableSolver::WriteFrame(FThreadingProxy& InProxy, const FSolverReal DeltaTime)
 	{
-		TRACE_CPUPROFILER_EVENT_SCOPE(DeformableSolver_WriteFrame);
+		TRACE_CPUPROFILER_EVENT_SCOPE(ChaosDeformableSolver_WriteFrame);
 		if (FFleshThreadingProxy* Proxy = InProxy.As<FFleshThreadingProxy>())
 		{
 			if (const FManagedArrayCollection* Rest = &Proxy->GetRestCollection())
