Add basic support for SkinnedLevelSet to Chaos Cloth.

This code has not been optimized. In particular, full ISPC support does not exist for the non-CCD case. This means that disabling ISPC for collisions is actually faster right now. Also, the CCD case requires calculating a reasonable Phi value when outside the deforming lattice. This requires finding the closest deformed lattice point for each query. Currently I'm using the HierarchicalSpatialHash which does not have a FindClosestPoint; it can only do intersection tests. This means I'm doing progressively larger intersection queries which is non-optimal. Future work to have a true FindClosestPoint.

diff --git a/Engine/Plugins/Experimental/ChaosCloth/Source/Chaos/Private/ChaosCloth/ChaosClothVisualization.cpp b/Engine/Plugins/Experimental/ChaosCloth/Source/Chaos/Private/ChaosCloth/ChaosClothVisualization.cpp
--- a/Engine/Plugins/Experimental/ChaosCloth/Source/Chaos/Private/ChaosCloth/ChaosClothVisualization.cpp
+++ b/Engine/Plugins/Experimental/ChaosCloth/Source/Chaos/Private/ChaosCloth/ChaosClothVisualization.cpp
@@ -24,6 +24,7 @@
 #include "Chaos/PBDTriangleMeshCollisions.h"
 #include "Chaos/XPBDBendingConstraints.h"
 #include "Chaos/XPBDSpringConstraints.h"
+#include "Chaos/WeightedLatticeImplicitObject.h"
 #include "DynamicMeshBuilder.h"
 #include "Engine/EngineTypes.h"
 #include "SceneManagement.h"
@@ -35,6 +36,9 @@
 #include "Engine/Engine.h"  //
 #endif  // #if WITH_EDITOR
 
+static int DrawSkinnedLattice = 0;
+static FAutoConsoleVariableRef CVarClothVizDrawSkinnedLattice(TEXT("p.ChaosClothVisualization.DrawSkinnedLattice"), DrawSkinnedLattice, TEXT("Draw skinned lattice, 0 = none, 1 = filled, 2 = empty, 3 = both"));
+
 namespace Chaos
 {
 	FClothVisualization::FClothVisualization(const ::Chaos::FClothingSimulationSolver* InSolver)
@@ -463,6 +467,85 @@ namespace Chaos
 		}
 	}
 
+	static void DrawSkinnedLevelSet(FPrimitiveDrawInterface* PDI, const TWeightedLatticeImplicitObject<FLevelSet>& SkinnedLevelSet, const FQuat& Rotation, const FVector& Position, const FMaterialRenderProxy* MaterialRenderProxy)
+	{
+#if WITH_EDITOR
+		if (PDI && MaterialRenderProxy)
+		{
+			TArray<FVector3f> Vertices;
+			TArray<FIntVector> Tris;
+			const FLevelSet* const LevelSet = SkinnedLevelSet.GetEmbeddedObject();
+			LevelSet->GetZeroIsosurfaceGridCellFaces(Vertices, Tris);
+
+			FDynamicMeshBuilder MeshBuilder(PDI->View->GetFeatureLevel());
+			for (const FVector3f& V : Vertices)
+			{
+				MeshBuilder.AddVertex(FDynamicMeshVertex(FVector3f(SkinnedLevelSet.GetDeformedPoint(FVec3(V)))));
+			}
+			for (const FIntVector& T : Tris)
+			{
+				MeshBuilder.AddTriangle(T[0], T[1], T[2]);
+			}
+
+			MeshBuilder.Draw(PDI, FTransform(Rotation, Position).ToMatrixWithScale(), MaterialRenderProxy, SDPG_World, false, false);
+
+			if (DrawSkinnedLattice)
+			{
+				const Chaos::TUniformGrid<double, 3>& LatticeGrid = SkinnedLevelSet.GetGrid();
+				const Chaos::TArrayND<Chaos::FVec3, 3>& DeformedPoints = SkinnedLevelSet.GetDeformedPoints();
+				const Chaos::TArrayND<bool, 3>& EmptyCells = SkinnedLevelSet.GetEmptyCells();
+				const FColor LatticeColor = FColor::Cyan;
+				const FColor EmptyLatticeColor = FColor::White;
+				const Chaos::TVec3<int32> CellCounts = LatticeGrid.Counts();
+
+				FTransform LocalToWorld(Rotation, Position);
+				for (int32 I = 0; I < CellCounts.X; ++I)
+				{
+					for (int32 J = 0; J < CellCounts.Y; ++J)
+					{
+						for (int32 K = 0; K < CellCounts.Z; ++K)
+						{
+							const bool bIsEmpty = EmptyCells(I, J, K);
+							const int32 EmptyDrawMask = bIsEmpty ? 2 : 1;
+							if (EmptyDrawMask & DrawSkinnedLattice)
+							{
+								const FVector P000 = LocalToWorld.TransformPosition(FVector(DeformedPoints(I, J, K)));
+								const FVector P001 = LocalToWorld.TransformPosition(FVector(DeformedPoints(I, J, K + 1)));
+								const FVector P010 = LocalToWorld.TransformPosition(FVector(DeformedPoints(I, J + 1, K)));
+								const FVector P011 = LocalToWorld.TransformPosition(FVector(DeformedPoints(I, J + 1, K + 1)));
+								const FVector P100 = LocalToWorld.TransformPosition(FVector(DeformedPoints(I + 1, J, K)));
+								const FVector P101 = LocalToWorld.TransformPosition(FVector(DeformedPoints(I + 1, J, K + 1)));
+								const FVector P110 = LocalToWorld.TransformPosition(FVector(DeformedPoints(I + 1, J + 1, K)));
+								const FVector P111 = LocalToWorld.TransformPosition(FVector(DeformedPoints(I + 1, J + 1, K + 1)));
+
+								PDI->AddReserveLines(SDPG_World, 12);
+								const FColor& Color = bIsEmpty ? EmptyLatticeColor : LatticeColor;
+								PDI->DrawLine(P000, P001, Color, SDPG_World);
+								PDI->DrawLine(P000, P010, Color, SDPG_World);
+								PDI->DrawLine(P000, P100, Color, SDPG_World);
+								PDI->DrawLine(P001, P011, Color, SDPG_World);
+								PDI->DrawLine(P001, P101, Color, SDPG_World);
+								PDI->DrawLine(P010, P011, Color, SDPG_World);
+								PDI->DrawLine(P010, P110, Color, SDPG_World);
+								PDI->DrawLine(P011, P111, Color, SDPG_World);
+								PDI->DrawLine(P100, P101, Color, SDPG_World);
+								PDI->DrawLine(P100, P110, Color, SDPG_World);
+								PDI->DrawLine(P101, P111, Color, SDPG_World);
+								PDI->DrawLine(P110, P111, Color, SDPG_World);
+
+							}
+						}
+					}
+				}
+			}
+		}
+		else
+#endif
+		{
+			DrawCoordinateSystem(PDI, Rotation, Position);
+		}
+	}
+
 	void FClothVisualization::DrawBounds(FPrimitiveDrawInterface* PDI) const
 	{
 		if (!Solver)
@@ -874,7 +957,19 @@ namespace Chaos
 								DrawLevelSet(PDI, CombinedTransform, MaterialRenderProxy, LevelSet);
 							}
 							break;
+						case (ImplicitObjectType::LevelSet | ImplicitObjectType::IsWeightedLattice):
+							{
+								const TWeightedLatticeImplicitObject<FLevelSet>& WeightedLevelset = Object->GetObjectChecked< TWeightedLatticeImplicitObject<FLevelSet> >();
+								const FMaterialRenderProxy* MaterialRenderProxy =
+#if WITH_EDITOR
+									CollisionMaterial->GetRenderProxy();
+#else
+									nullptr;
+#endif
+								DrawSkinnedLevelSet(PDI, WeightedLevelset, Rotation, Position, MaterialRenderProxy);
 
+							}
+							break;
 						default:
 							DrawCoordinateSystem(PDI, Rotation, Position);  // Draw everything else as a coordinate for now
 							break;
@@ -896,6 +991,7 @@ namespace Chaos
 
 		// Draw contacts
 		check(Solver->GetCollisionContacts().Num() == Solver->GetCollisionNormals().Num());
+		const bool bDrawPhis = Solver->GetCollisionContacts().Num() == Solver->GetCollisionPhis().Num();
 		constexpr FReal NormalLength = (FReal)10.;
 
 		const FVec3& LocalSpaceLocation = Solver->GetLocalSpaceLocation();
@@ -917,6 +1013,13 @@ namespace Chaos
 			static const FLinearColor Brown(0.1f, 0.05f, 0.f);
 			const FVec3 Pos1 = Pos0 + NormalLength * Normal;
 			DrawLine(PDI, Pos0, Pos1, Brown);
+
+			if (bDrawPhis)
+			{
+				const FVec3 PhiLocation = Pos0 - Solver->GetCollisionPhis()[i] * Normal;
+				DrawLine(PDI, Pos0, PhiLocation, Brown);
+				DrawPoint(PDI, PhiLocation, FLinearColor::Red, nullptr, 5);
+			}
 		}
 	}
 
diff --git a/Engine/Plugins/Experimental/ChaosCloth/Source/Chaos/Private/ChaosCloth/ChaosClothingSimulationCollider.cpp b/Engine/Plugins/Experimental/ChaosCloth/Source/Chaos/Private/ChaosCloth/ChaosClothingSimulationCollider.cpp
--- a/Engine/Plugins/Experimental/ChaosCloth/Source/Chaos/Private/ChaosCloth/ChaosClothingSimulationCollider.cpp
+++ b/Engine/Plugins/Experimental/ChaosCloth/Source/Chaos/Private/ChaosCloth/ChaosClothingSimulationCollider.cpp
@@ -12,6 +12,7 @@
 #include "Chaos/Convex.h"
 #include "Chaos/Levelset.h"
 #include "Chaos/ImplicitObjectTransformed.h"
+#include "Chaos/WeightedLatticeImplicitObject.h"
 #include "ClothingAsset.h"
 #include "Engine/SkeletalMesh.h"
 #include "PhysicsEngine/PhysicsAsset.h"
@@ -43,6 +44,7 @@ struct FClothingSimulationCollider::FLODData
 		FClothingSimulationCloth* Cloth,
 		const FClothCollisionData& InClothCollisionData,
 		const TArray<FLevelSetCollisionData>& InLevelSetCollisionData,
+		const TArray<FSkinnedLevelSetCollisionData>& InSkinnedLevelSetCollisionData,
 		const FReal InScale = 1.f,
 		const TArray<int32>& UsedBoneIndices = TArray<int32>());
 	void Remove(FClothingSimulationSolver* Solver, FClothingSimulationCloth* Cloth);
@@ -64,6 +66,7 @@ void FClothingSimulationCollider::FLODData::Add(
 	FClothingSimulationCloth* Cloth,
 	const FClothCollisionData& InClothCollisionData,
 	const TArray<FLevelSetCollisionData>& InLevelSetCollisionData,
+	const TArray<FSkinnedLevelSetCollisionData>& InSkinnedLevelSetCollisionData,
 	const FReal InScale,
 	const TArray<int32>& UsedBoneIndices)
 {
@@ -95,7 +98,14 @@ void FClothingSimulationCollider::FLODData::Add(
 	const uint32 NumConvexes = ClothCollisionData.Convexes.Num();
 	const int32 NumBoxes = ClothCollisionData.Boxes.Num();
 	const int32 NumLevelSets = InLevelSetCollisionData.Num();
-	NumGeometries = NumSpheres + NumCapsules + NumConvexes + NumBoxes + NumLevelSets;
+	const int32 NumSkinnedLevelSets = InSkinnedLevelSetCollisionData.Num();
+	int32 NumSkinnedLevelSetBones = 0;
+	for (const FSkinnedLevelSetCollisionData& SkinnedLevelSetData : InSkinnedLevelSetCollisionData)
+	{
+		NumSkinnedLevelSetBones += SkinnedLevelSetData.MappedSkinnedBones.Num();
+	}
+
+	NumGeometries = NumSpheres + NumCapsules + NumConvexes + NumBoxes + NumLevelSets + NumSkinnedLevelSets + NumSkinnedLevelSetBones;
 
 	// Retrieve cloth group Id, or use INDEX_NONE if this collider applies to all cloths (when Cloth == nullptr)
 	const uint32 GroupId = Cloth ? Cloth->GetGroupId() : INDEX_NONE;
@@ -336,6 +346,49 @@ void FClothingSimulationCollider::FLODData::Add(
 		}
 	}
 
+	const int32 SkinnedLevelSetOffset = LevelSetOffset + NumLevelSets;
+	if (NumSkinnedLevelSets != 0)
+	{
+		int32* const BoneIndices = Solver->GetCollisionBoneIndices(SkinnedLevelSetOffset);
+		Softs::FSolverRigidTransform3* const BaseTransforms = Solver->GetCollisionBaseTransforms(SkinnedLevelSetOffset);
+
+		int32 Index = 0;
+		for (const FSkinnedLevelSetCollisionData& SkinnedCollisionData : InSkinnedLevelSetCollisionData)
+		{
+			// Add bone proxies first so they get updated first?
+			TArray<int32> SolverBoneIndices;
+			SolverBoneIndices.Reserve(SkinnedCollisionData.MappedSkinnedBones.Num());
+			for (int32 MappedSubBoneIndex : SkinnedCollisionData.MappedSkinnedBones)
+			{
+				// Always initialize the collision particle transforms before setting any geometry as otherwise NaNs gets detected during the bounding box updates
+				BaseTransforms[Index] = Softs::FSolverRigidTransform3::Identity;
+
+				BoneIndices[Index] = GetMappedBoneIndex(UsedBoneIndices, MappedSubBoneIndex);
+				UE_LOG(LogChaosCloth, VeryVerbose, TEXT("Found collision skinned level set sub-bone on bone index %d."), BoneIndices[Index]);
+
+				TUniquePtr<FImplicitObject> BoneProxy = MakeUnique<FWeightedLatticeBoneProxy>();
+				Solver->SetCollisionGeometry(SkinnedLevelSetOffset, Index, MoveTemp(BoneProxy));
+
+				SolverBoneIndices.Add(Index + SkinnedLevelSetOffset);
+
+				++Index;
+			}
+
+			// Always initialize the collision particle transforms before setting any geometry as otherwise NaNs gets detected during the bounding box updates
+			BaseTransforms[Index] = Softs::FSolverRigidTransform3::Identity;
+
+			BoneIndices[Index] = GetMappedBoneIndex(UsedBoneIndices, SkinnedCollisionData.BoneIndex);
+			UE_LOG(LogChaosCloth, VeryVerbose, TEXT("Found collision skinned level set on bone index %d."), BoneIndices[Index]);
+
+			TUniquePtr<FImplicitObject> SkinnedLevelSet = SkinnedCollisionData.WeightedLevelSet->DeepCopy();
+			SkinnedLevelSet->GetObjectChecked<TWeightedLatticeImplicitObject<FLevelSet>>().SetSolverBoneIndices(MoveTemp(SolverBoneIndices));
+
+			Solver->SetCollisionGeometry(SkinnedLevelSetOffset, Index, MoveTemp(SkinnedLevelSet));
+
+			++Index;
+		}
+	}
+
 	UE_LOG(LogChaosCloth, VeryVerbose, TEXT("Added collisions: %d spheres, %d capsules, %d convexes, %d boxes, % level sets."), NumSpheres, NumCapsules, NumConvexes, NumBoxes, NumLevelSets);
 }
 
@@ -457,7 +510,7 @@ FClothCollisionData FClothingSimulationCollider::GetCollisionData(const FClothin
 	return ClothCollisionData;
 }
 
-void FClothingSimulationCollider::ExtractPhysicsAssetCollision(FClothCollisionData& ClothCollisionData, TArray<FLevelSetCollisionData>& LevelSetCollisions, TArray<int32>& UsedBoneIndices)
+void FClothingSimulationCollider::ExtractPhysicsAssetCollision(FClothCollisionData& ClothCollisionData, TArray<FLevelSetCollisionData>& LevelSetCollisions, TArray<FSkinnedLevelSetCollisionData>& SkinnedLevelSetCollisions, TArray<int32>& UsedBoneIndices)
 {
 	ClothCollisionData.Reset();
 	UsedBoneIndices.Reset();
@@ -619,6 +672,23 @@ void FClothingSimulationCollider::ExtractPhysicsAssetCollision(FClothCollisionDa
 					LevelSetCollisions.Add({ LevelSetElem.GetLevelSet(), LevelSetElem.GetTransform(), MappedBoneIndex });
 				}
 			}
+
+			for (const FKSkinnedLevelSetElem& SkinnedLevelSetElem : AggGeom.SkinnedLevelSetElems)
+			{
+				if (SkinnedLevelSetElem.GetWeightedLevelSet().IsValid())
+				{
+					TArray<int32> MappedSkinnedBones;
+					const TArray<FName>& SkinnedBones = SkinnedLevelSetElem.GetWeightedLevelSet()->GetUsedBones();
+					MappedSkinnedBones.Reserve(SkinnedBones.Num());
+					for (const FName& SkinnedBoneName : SkinnedBones)
+					{
+						const int32 SkinnedBoneIndex = ReferenceSkeleton ? ReferenceSkeleton->FindBoneIndex(SkinnedBoneName) : INDEX_NONE;
+						const int32 MappedSkinnedBoneIndex = UsedBoneIndices.Add(SkinnedBoneIndex);
+						MappedSkinnedBones.Add(MappedSkinnedBoneIndex);
+					}
+					SkinnedLevelSetCollisions.Add({ SkinnedLevelSetElem.GetWeightedLevelSet(), MappedBoneIndex, MoveTemp(MappedSkinnedBones)});
+				}
+			}
 		}  // End for PhysAsset->SkeletalBodySetups
 	}  // End if PhysicsAsset
 }
@@ -682,10 +752,11 @@ PRAGMA_ENABLE_DEPRECATION_WARNINGS
 	// Create physics asset collisions, this will affect all LODs, so store it at index 0
 	FClothCollisionData PhysicsAssetCollisionData;
 	TArray<FLevelSetCollisionData> LevelSetCollisions;
+	TArray<FSkinnedLevelSetCollisionData> SkinnedLevelSetCollisions;
 	TArray<int32> UsedBoneIndices;
-	ExtractPhysicsAssetCollision(PhysicsAssetCollisionData, LevelSetCollisions, UsedBoneIndices);
+	ExtractPhysicsAssetCollision(PhysicsAssetCollisionData, LevelSetCollisions, SkinnedLevelSetCollisions, UsedBoneIndices);
 
-	LODData[(int32)ECollisionDataType::LODless]->Add(Solver, Cloth, PhysicsAssetCollisionData, LevelSetCollisions, Scale, UsedBoneIndices);
+	LODData[(int32)ECollisionDataType::LODless]->Add(Solver, Cloth, PhysicsAssetCollisionData, LevelSetCollisions, SkinnedLevelSetCollisions, Scale, UsedBoneIndices);
 }
 
 void FClothingSimulationCollider::Remove(FClothingSimulationSolver* Solver, FClothingSimulationCloth* Cloth)
@@ -712,8 +783,9 @@ void FClothingSimulationCollider::PreUpdate(FClothingSimulationSolver* Solver, F
 
 	// TODO: Get level sets?
 	const TArray<FLevelSetCollisionData> LevelSetCollisions;
+	const TArray<FSkinnedLevelSetCollisionData> SkinnedLevelSetCollisions;
 
-	LODData[(int32)ECollisionDataType::External]->Add(Solver, Cloth, CollisionData ? *CollisionData : FClothCollisionData(), LevelSetCollisions);
+	LODData[(int32)ECollisionDataType::External]->Add(Solver, Cloth, CollisionData ? *CollisionData : FClothCollisionData(), LevelSetCollisions, SkinnedLevelSetCollisions);
 	
 	// TODO: Find a better way in case the same number but different collisions are being re-added (hash collision data? Provide user dirty function?)
 	bHasExternalCollisionChanged =
diff --git a/Engine/Plugins/Experimental/ChaosCloth/Source/Chaos/Private/ChaosCloth/ChaosClothingSimulationSolver.cpp b/Engine/Plugins/Experimental/ChaosCloth/Source/Chaos/Private/ChaosCloth/ChaosClothingSimulationSolver.cpp
--- a/Engine/Plugins/Experimental/ChaosCloth/Source/Chaos/Private/ChaosCloth/ChaosClothingSimulationSolver.cpp
+++ b/Engine/Plugins/Experimental/ChaosCloth/Source/Chaos/Private/ChaosCloth/ChaosClothingSimulationSolver.cpp
@@ -7,6 +7,8 @@
 #include "ChaosCloth/ChaosClothingSimulation.h"
 #include "ChaosCloth/ChaosClothPrivate.h"
 #include "Chaos/PBDEvolution.h"
+#include "Chaos/WeightedLatticeImplicitObject.h"
+#include "Chaos/Levelset.h"
 #if INTEL_ISPC
 #include "ChaosClothingSimulationSolver.ispc.generated.h"
 #endif
@@ -163,6 +165,22 @@ FClothingSimulationSolver::FClothingSimulationSolver()
 			const Softs::FSolverVec3 Axis = Delta.GetRotationAxis();
 			ParticlesInput.W(Index) = (Softs::FSolverVec3)Axis * Angle / Dt;
 			ParticlesInput.R(Index) = NewR;
+
+			if (TWeightedLatticeImplicitObject<FLevelSet>* SkinnedLevelSet = 
+				const_cast<FImplicitObject*>(ParticlesInput.Geometry(Index).Get())->GetObject<TWeightedLatticeImplicitObject<FLevelSet>>())
+			{
+				const TArray<int32>& SubBoneIndices = SkinnedLevelSet->GetSolverBoneIndices();
+				const FTransform RootTransformInv = TRigidTransform<FReal, 3>(ParticlesInput.X(Index), ParticlesInput.R(Index)).Inverse();
+				TArray<FTransform> SubBoneTransforms;
+				SubBoneTransforms.SetNum(SubBoneIndices.Num());
+				for (int32 SubBoneIdx = 0; SubBoneIdx < SubBoneIndices.Num(); ++SubBoneIdx)
+				{
+					checkSlow(SubBoneIndices[SubBoneIdx] < Index);
+					SubBoneTransforms[SubBoneIdx] = TRigidTransform<FReal,3>(ParticlesInput.X(SubBoneIndices[SubBoneIdx]), ParticlesInput.R(SubBoneIndices[SubBoneIdx])) * RootTransformInv;
+				}
+				SkinnedLevelSet->DeformPoints(SubBoneTransforms);
+				SkinnedLevelSet->UpdateSpatialHierarchy();
+			}
 		});
 }
 
@@ -525,6 +543,11 @@ const TArray<Softs::FSolverVec3>& FClothingSimulationSolver::GetCollisionContact
 	return Evolution->GetCollisionContacts();
 }
 
+const TArray<Softs::FSolverReal>& FClothingSimulationSolver::GetCollisionPhis() const
+{
+	return Evolution->GetCollisionPhis();
+}
+
 const TArray<Softs::FSolverVec3>& FClothingSimulationSolver::GetCollisionNormals() const
 {
 	return Evolution->GetCollisionNormals();
diff --git a/Engine/Plugins/Experimental/ChaosCloth/Source/Chaos/Public/ChaosCloth/ChaosClothingSimulationCollider.h b/Engine/Plugins/Experimental/ChaosCloth/Source/Chaos/Public/ChaosCloth/ChaosClothingSimulationCollider.h
--- a/Engine/Plugins/Experimental/ChaosCloth/Source/Chaos/Public/ChaosCloth/ChaosClothingSimulationCollider.h
+++ b/Engine/Plugins/Experimental/ChaosCloth/Source/Chaos/Public/ChaosCloth/ChaosClothingSimulationCollider.h
@@ -18,6 +18,7 @@ namespace Chaos
 	class FClothingSimulationSolver;
 	class FClothingSimulationCloth;
 	class FLevelSet;
+	template<typename T> class TWeightedLatticeImplicitObject;
 
 	// Collider simulation node
 	class CHAOSCLOTH_API FClothingSimulationCollider final
@@ -90,7 +91,14 @@ namespace Chaos
 			int32 BoneIndex;
 		};
 
-		void ExtractPhysicsAssetCollision(FClothCollisionData& ClothCollisionData, TArray<FLevelSetCollisionData>& LevelSetCollisions, TArray<int32>& UsedBoneIndices);
+		struct FSkinnedLevelSetCollisionData
+		{
+			const TSharedPtr<Chaos::TWeightedLatticeImplicitObject<Chaos::FLevelSet>, ESPMode::ThreadSafe> WeightedLevelSet;
+			int32 BoneIndex;
+			TArray<int32> MappedSkinnedBones;
+		};
+
+		void ExtractPhysicsAssetCollision(FClothCollisionData& ClothCollisionData, TArray<FLevelSetCollisionData>& LevelSetCollisions, TArray<FSkinnedLevelSetCollisionData>& SkinnedLevelSetCollisions, TArray<int32>& UsedBoneIndices);
 
 		int32 GetNumGeometries(int32 InSlotIndex) const;
 
diff --git a/Engine/Plugins/Experimental/ChaosCloth/Source/Chaos/Public/ChaosCloth/ChaosClothingSimulationSolver.h b/Engine/Plugins/Experimental/ChaosCloth/Source/Chaos/Public/ChaosCloth/ChaosClothingSimulationSolver.h
--- a/Engine/Plugins/Experimental/ChaosCloth/Source/Chaos/Public/ChaosCloth/ChaosClothingSimulationSolver.h
+++ b/Engine/Plugins/Experimental/ChaosCloth/Source/Chaos/Public/ChaosCloth/ChaosClothingSimulationSolver.h
@@ -211,6 +211,7 @@ namespace Chaos
 		const bool* GetCollisionStatus(int32 Offset) const;
 		const TArray<Softs::FSolverVec3>& GetCollisionContacts() const;
 		const TArray<Softs::FSolverVec3>& GetCollisionNormals() const;
+		const TArray<Softs::FSolverReal>& GetCollisionPhis() const;
 		// ---- End of the Collider interface ----
 
 		// ---- Field interface ----
diff --git a/Engine/Source/Runtime/Experimental/Chaos/Private/Chaos/PBDEvolution.cpp b/Engine/Source/Runtime/Experimental/Chaos/Private/Chaos/PBDEvolution.cpp
--- a/Engine/Source/Runtime/Experimental/Chaos/Private/Chaos/PBDEvolution.cpp
+++ b/Engine/Source/Runtime/Experimental/Chaos/Private/Chaos/PBDEvolution.cpp
@@ -427,7 +427,7 @@ void FPBDEvolution::AdvanceOneTimeStep(const FSolverReal Dt)
 					MCollisionTransforms[Index] = FSolverRigidTransform3(CollisionParticles.X(Index), CollisionParticles.R(Index));
 
 					// Update collision transform and velocity
-					MCollisionKinematicUpdate(CollisionParticles, Dt, MTime, Index);
+					MCollisionKinematicUpdate(CollisionParticles, Dt, MTime, Index); // This expects Sequential update.
 				});
 		}
 
@@ -452,14 +452,19 @@ void FPBDEvolution::AdvanceOneTimeStep(const FSolverReal Dt)
 	// Collision rule initializations
 	MCollisionContacts.Reset();
 	MCollisionNormals.Reset();
+	MCollisionPhis.Reset();
 
 	FPerParticlePBDCollisionConstraint CollisionRule(
 		MCollisionParticlesActiveView,
 		MCollided,
+		MCollisionContacts,
+		MCollisionNormals,
+		MCollisionPhis,
 		MParticleGroupIds,
 		MCollisionParticleGroupIds,
 		MGroupCollisionThicknesses,
-		MGroupCoefficientOfFrictions);
+		MGroupCoefficientOfFrictions,
+		bWriteCCDContacts);
 
 	FPerParticlePBDCCDCollisionConstraint CCDCollisionRule(
 		MCollisionParticlesActiveView,
diff --git a/Engine/Source/Runtime/Experimental/Chaos/Private/Chaos/PerParticlePBDCollisionConstraint.cpp b/Engine/Source/Runtime/Experimental/Chaos/Private/Chaos/PerParticlePBDCollisionConstraint.cpp
--- a/Engine/Source/Runtime/Experimental/Chaos/Private/Chaos/PerParticlePBDCollisionConstraint.cpp
+++ b/Engine/Source/Runtime/Experimental/Chaos/Private/Chaos/PerParticlePBDCollisionConstraint.cpp
@@ -25,6 +25,9 @@ static_assert(ispc::ImplicitObjectType::Capsule == Chaos::ImplicitObjectType::Ca
 static_assert(ispc::ImplicitObjectType::Union == Chaos::ImplicitObjectType::Union);
 static_assert(ispc::ImplicitObjectType::TaperedCapsule == Chaos::ImplicitObjectType::TaperedCapsule);
 static_assert(ispc::ImplicitObjectType::Convex == Chaos::ImplicitObjectType::Convex);
+static_assert(ispc::ImplicitObjectType::WeightedLatticeBone == Chaos::ImplicitObjectType::WeightedLatticeBone);
+static_assert(ispc::ImplicitObjectType::IsWeightedLattice == Chaos::ImplicitObjectType::IsWeightedLattice);
+static_assert(ispc::ImplicitObjectType::WeightedLatticeLevelSetType == (Chaos::ImplicitObjectType::IsWeightedLattice | Chaos::ImplicitObjectType::LevelSet));
 static_assert(sizeof(ispc::TArray) == sizeof(TArray<int>));
 // Sphere
 static_assert(sizeof(Chaos::FImplicitObject) == Chaos::FSphere::FISPCDataVerifier::OffsetOfCenter());
@@ -109,10 +112,8 @@ static FAutoConsoleVariableRef CVarChaosPerParticleCollisionISPCEnabled(TEXT("p.
 #endif
 
 static int32 Chaos_PerParticleCollision_ISPC_ParallelBatchSize = 128;
-static bool Chaos_PerParticleCollision_ISPC_FastFriction = true;
 #if !UE_BUILD_SHIPPING
 FAutoConsoleVariableRef CVarChaosPerParticleCollisionISPCParallelBatchSize(TEXT("p.Chaos.PerParticleCollision.ISPC.ParallelBatchSize"), Chaos_PerParticleCollision_ISPC_ParallelBatchSize, TEXT("Parallel batch size for ISPC"));
-FAutoConsoleVariableRef CVarChaosPerParticleCollisionISPCFastFriction(TEXT("p.Chaos.PerParticleCollision.ISPC.FastFriction"), Chaos_PerParticleCollision_ISPC_FastFriction, TEXT("Faster friction ISPC"));
 #endif
 
 namespace Chaos::Softs {
@@ -145,9 +146,57 @@ extern "C" void GetPhiWithNormal(const uint8* CollisionParticles, const FSolverR
 	}
 }
 
+// Helper function to call PhiWithNormal and return data to ISPC
+extern "C" void GetPhiWithNormalAndVelocityBone(const uint8 * CollisionParticles, const FSolverReal * InV, FSolverReal * Normal, FSolverReal * Phi, int32 * VelocityBone, const int32 i, const FSolverReal Thickness, const int32 ProgramCount, const int32 Mask)
+{
+	const TKinematicGeometryParticlesImp<FSolverReal, 3, EGeometryParticlesSimType::Other>& C = *(const TKinematicGeometryParticlesImp<FSolverReal, 3, EGeometryParticlesSimType::Other>*)CollisionParticles;
+
+	for (int32 Index = 0; Index < ProgramCount; ++Index)
+	{
+		if (Mask & (1 << Index))
+		{
+			FSolverVec3 V;
+
+			// aos_to_soa3
+			V.X = InV[Index];
+			V.Y = InV[Index + ProgramCount];
+			V.Z = InV[Index + 2 * ProgramCount];
+
+			VelocityBone[Index] = i;
+
+			FVec3 ImplicitNormal;
+			if (const TWeightedLatticeImplicitObject<FLevelSet>* LevelSet = C.Geometry(i)->GetObject< TWeightedLatticeImplicitObject<FLevelSet> >())
+			{
+				FWeightedLatticeImplicitObject::FEmbeddingCoordinate SurfaceCoord;
+				Phi[Index] = (FSolverReal)LevelSet->PhiWithNormalAndSurfacePoint(FVec3(V), ImplicitNormal, SurfaceCoord);
+				const FSolverReal Penetration = Thickness - Phi[Index]; // This is related to the Normal impulse
+				if (Penetration > (FSolverReal)0.)
+				{
+					const int32 StrongestBone = SurfaceCoord.GreatestInfluenceBone(LevelSet->GetBoneData());
+					if (StrongestBone != INDEX_NONE)
+					{
+						VelocityBone[Index] = LevelSet->GetSolverBoneIndices()[StrongestBone];
+					}
+				}
+			}
+			else
+			{
+				Phi[Index] = (FSolverReal)C.Geometry(i)->PhiWithNormal(FVec3(V), ImplicitNormal);
+			}
+			FSolverVec3 Norm(ImplicitNormal);
+
+			// aos_to_soa3
+			Normal[Index] = Norm.X;
+			Normal[Index + ProgramCount] = Norm.Y;
+			Normal[Index + 2 * ProgramCount] = Norm.Z;
+		}
+	}
+}
+
 void FPerParticlePBDCollisionConstraint::ApplyHelperISPC(FSolverParticles& InParticles, const FSolverReal Dt, const int32 Offset, const int32 Range) const
 {
 	check(bRealTypeCompatibleWithISPC);
+	check(bFastPositionBasedFriction);
 
 	const uint32 DynamicGroupId = MDynamicGroupIds[Offset];
 	const FSolverReal PerGroupFriction = MPerGroupFriction[DynamicGroupId];
@@ -155,115 +204,77 @@ void FPerParticlePBDCollisionConstraint::ApplyHelperISPC(FSolverParticles& InPar
 
 	const int32 NumBatches = FMath::CeilToInt((FSolverReal)(Range - Offset) / (FSolverReal)Chaos_PerParticleCollision_ISPC_ParallelBatchSize);
 
-	if (Chaos_PerParticleCollision_ISPC_FastFriction)
+	if (PerGroupFriction > UE_KINDA_SMALL_NUMBER)  // Fast friction
 	{
-		if (PerGroupFriction > UE_KINDA_SMALL_NUMBER)  // Fast friction
-		{
-			PhysicsParallelFor(NumBatches, [this, &InParticles, Dt, Offset, Range, DynamicGroupId, PerGroupFriction, PerGroupThickness](int32 BatchNumber)
+		PhysicsParallelFor(NumBatches, [this, &InParticles, Dt, Offset, Range, DynamicGroupId, PerGroupFriction, PerGroupThickness](int32 BatchNumber)
 			{
 				const int32 BatchBegin = Offset + (Chaos_PerParticleCollision_ISPC_ParallelBatchSize * BatchNumber);
-				const int32 BatchEnd = FMath::Min(Range, BatchBegin + Chaos_PerParticleCollision_ISPC_ParallelBatchSize);
+		const int32 BatchEnd = FMath::Min(Range, BatchBegin + Chaos_PerParticleCollision_ISPC_ParallelBatchSize);
 
 #if INTEL_ISPC
-				MCollisionParticlesActiveView.RangeFor(
-					[this, &InParticles, Dt, BatchBegin, BatchEnd, DynamicGroupId, PerGroupFriction, PerGroupThickness](FSolverRigidParticles& CollisionParticles, int32 CollisionOffset, int32 CollisionRange)
-					{
-						ispc::ApplyPerParticleCollisionFastFriction(
-							(ispc::FVector4f*)InParticles.GetPAndInvM().GetData(),
-							(const ispc::FVector3f*)InParticles.XArray().GetData(),
-							(const ispc::FVector3f*)CollisionParticles.AllV().GetData(),
-							(const ispc::FVector3f*)CollisionParticles.XArray().GetData(),
-							(const ispc::FVector3f*)CollisionParticles.AllW().GetData(),
-							(const ispc::FVector4f*)CollisionParticles.AllR().GetData(),
-							DynamicGroupId,
-							MKinematicGroupIds.GetData(),
-							PerGroupFriction,
-							PerGroupThickness,
-							(const uint8*)&CollisionParticles,
-							(const uint8*)CollisionParticles.GetAllGeometry().GetData(),
-							sizeof(FImplicitObject),
-							FImplicitObject::GetOffsetOfType(),
-							FImplicitObject::GetOffsetOfMargin(),
-							Dt,
-							CollisionOffset,
-							CollisionRange,
-							BatchBegin,
-							BatchEnd);
-					});
-#endif  // #if INTEL_ISPC
-				});
-		}
-		else  // No friction
-		{
-			PhysicsParallelFor(NumBatches, [this, &InParticles, Dt, Offset, Range, DynamicGroupId, PerGroupFriction, PerGroupThickness](int32 BatchNumber)
+		MCollisionParticlesActiveView.RangeFor(
+			[this, &InParticles, Dt, BatchBegin, BatchEnd, DynamicGroupId, PerGroupFriction, PerGroupThickness](FSolverRigidParticles& CollisionParticles, int32 CollisionOffset, int32 CollisionRange)
 			{
-				const int32 BatchBegin = Offset + (Chaos_PerParticleCollision_ISPC_ParallelBatchSize * BatchNumber);
-				const int32 BatchEnd = FMath::Min(Range, BatchBegin + Chaos_PerParticleCollision_ISPC_ParallelBatchSize);
-
-#if INTEL_ISPC
-				MCollisionParticlesActiveView.RangeFor(
-					[this, &InParticles, Dt, BatchBegin, BatchEnd, DynamicGroupId, PerGroupThickness](FSolverRigidParticles& CollisionParticles, int32 CollisionOffset, int32 CollisionRange)
-					{
-						ispc::ApplyPerParticleCollisionNoFriction(
-							(ispc::FVector4f*)InParticles.GetPAndInvM().GetData(),
-							(const ispc::FVector3f*)InParticles.XArray().GetData(),
-							(const ispc::FVector3f*)CollisionParticles.AllV().GetData(),
-							(const ispc::FVector3f*)CollisionParticles.XArray().GetData(),
-							(const ispc::FVector3f*)CollisionParticles.AllW().GetData(),
-							(const ispc::FVector4f*)CollisionParticles.AllR().GetData(),
-							DynamicGroupId,
-							MKinematicGroupIds.GetData(),
-							PerGroupThickness,
-							(const uint8*)&CollisionParticles,
-							(const uint8*)CollisionParticles.GetAllGeometry().GetData(),
-							sizeof(FImplicitObject),
-							FImplicitObject::GetOffsetOfType(),
-							FImplicitObject::GetOffsetOfMargin(),
-							Dt,
-							CollisionOffset,
-							CollisionRange,
-							BatchBegin,
-							BatchEnd);
-					});
+				ispc::ApplyPerParticleCollisionFastFriction(
+					(ispc::FVector4f*)InParticles.GetPAndInvM().GetData(),
+					(const ispc::FVector3f*)InParticles.XArray().GetData(),
+					(const ispc::FVector3f*)CollisionParticles.AllV().GetData(),
+					(const ispc::FVector3f*)CollisionParticles.XArray().GetData(),
+					(const ispc::FVector3f*)CollisionParticles.AllW().GetData(),
+					(const ispc::FVector4f*)CollisionParticles.AllR().GetData(),
+					DynamicGroupId,
+					MKinematicGroupIds.GetData(),
+					PerGroupFriction,
+					PerGroupThickness,
+					(const uint8*)&CollisionParticles,
+					(const uint8*)CollisionParticles.GetAllGeometry().GetData(),
+					sizeof(FImplicitObject),
+					FImplicitObject::GetOffsetOfType(),
+					FImplicitObject::GetOffsetOfMargin(),
+					Dt,
+					CollisionOffset,
+					CollisionRange,
+					BatchBegin,
+					BatchEnd);
+			});
 #endif  // #if INTEL_ISPC
-				});
-		}
+			});
 	}
-	else
+	else  // No friction
 	{
-		PhysicsParallelFor(NumBatches, [this, &InParticles, Dt, Offset, Range](int32 BatchNumber)
-		{
-			const int32 BatchBegin = Offset + (Chaos_PerParticleCollision_ISPC_ParallelBatchSize * BatchNumber);
-			const int32 BatchEnd = FMath::Min(Range, BatchBegin + Chaos_PerParticleCollision_ISPC_ParallelBatchSize);
+		PhysicsParallelFor(NumBatches, [this, &InParticles, Dt, Offset, Range, DynamicGroupId, PerGroupFriction, PerGroupThickness](int32 BatchNumber)
+			{
+				const int32 BatchBegin = Offset + (Chaos_PerParticleCollision_ISPC_ParallelBatchSize * BatchNumber);
+		const int32 BatchEnd = FMath::Min(Range, BatchBegin + Chaos_PerParticleCollision_ISPC_ParallelBatchSize);
 
 #if INTEL_ISPC
-			MCollisionParticlesActiveView.RangeFor(
-				[this, &InParticles, Dt, BatchBegin, BatchEnd](FSolverRigidParticles& CollisionParticles, int32 CollisionOffset, int32 CollisionRange)
-				{
-					ispc::ApplyPerParticleCollision(
-						(ispc::FVector4f*)InParticles.GetPAndInvM().GetData(),
-						(const ispc::FVector3f*)InParticles.XArray().GetData(),
-						(const ispc::FVector3f*)CollisionParticles.AllV().GetData(),
-						(const ispc::FVector3f*)CollisionParticles.XArray().GetData(),
-						(const ispc::FVector3f*)CollisionParticles.AllW().GetData(),
-						(const ispc::FVector4f*)CollisionParticles.AllR().GetData(),
-						MDynamicGroupIds.GetData(),
-						MKinematicGroupIds.GetData(),
-						MPerGroupFriction.GetData(),
-						MPerGroupThickness.GetData(),
-						(const uint8*)&CollisionParticles,
-						(const uint8*)CollisionParticles.GetAllGeometry().GetData(),
-						sizeof(FImplicitObject),
-						FImplicitObject::GetOffsetOfType(),
-						FImplicitObject::GetOffsetOfMargin(),
-						Dt,
-						CollisionOffset,
-						CollisionRange,
-						BatchBegin,
-						BatchEnd);
-				});
+		MCollisionParticlesActiveView.RangeFor(
+			[this, &InParticles, Dt, BatchBegin, BatchEnd, DynamicGroupId, PerGroupThickness](FSolverRigidParticles& CollisionParticles, int32 CollisionOffset, int32 CollisionRange)
+			{
+				ispc::ApplyPerParticleCollisionNoFriction(
+					(ispc::FVector4f*)InParticles.GetPAndInvM().GetData(),
+					(const ispc::FVector3f*)InParticles.XArray().GetData(),
+					(const ispc::FVector3f*)CollisionParticles.AllV().GetData(),
+					(const ispc::FVector3f*)CollisionParticles.XArray().GetData(),
+					(const ispc::FVector3f*)CollisionParticles.AllW().GetData(),
+					(const ispc::FVector4f*)CollisionParticles.AllR().GetData(),
+					DynamicGroupId,
+					MKinematicGroupIds.GetData(),
+					PerGroupThickness,
+					(const uint8*)&CollisionParticles,
+					(const uint8*)CollisionParticles.GetAllGeometry().GetData(),
+					sizeof(FImplicitObject),
+					FImplicitObject::GetOffsetOfType(),
+					FImplicitObject::GetOffsetOfMargin(),
+					Dt,
+					CollisionOffset,
+					CollisionRange,
+					BatchBegin,
+					BatchEnd);
+			});
 #endif  // #if INTEL_ISPC
-		});
+			});
+
 	}
 }
 
diff --git a/Engine/Source/Runtime/Experimental/Chaos/Private/Chaos/PerParticlePBDCollisionConstraint.ispc b/Engine/Source/Runtime/Experimental/Chaos/Private/Chaos/PerParticlePBDCollisionConstraint.ispc
--- a/Engine/Source/Runtime/Experimental/Chaos/Private/Chaos/PerParticlePBDCollisionConstraint.ispc
+++ b/Engine/Source/Runtime/Experimental/Chaos/Private/Chaos/PerParticlePBDCollisionConstraint.ispc
@@ -26,8 +26,11 @@ enum ImplicitObjectType
 	Triangle,
 	UnionClustered,
 	TaperedCapsule,
+	WeightedLatticeBone,
+	WeightedLatticeLevelSetType = 38, // == LevelSet | IsWeightedLattice (cannot use expressions like this in ISPC)
 
 	//Add entries above this line for serialization
+	IsWeightedLattice = 1 << 5,
 	IsInstanced = 1 << 6,
 	IsScaled = 1 << 7
 };
@@ -580,176 +583,7 @@ static inline void PhiWithNormalConvex(const uniform uint8 *uniform TypedObjectP
 }
 
 extern "C" unmasked void GetPhiWithNormal(const uniform uint8 *uniform CollisionParticles, const uniform float *uniform V, uniform float *uniform Normal, uniform float *uniform Phi, const uniform int i, const uniform int ProgramCount, const uniform int Mask);
-
-export void ApplyPerParticleCollision(uniform FVector4f ParticlesPandInvM[],
-										const uniform FVector3f ParticlesX[],
-										const uniform FVector3f CollisionV[],
-										const uniform FVector3f CollisionX[],
-										const uniform FVector3f CollisionW[],
-										const uniform FVector4f CollisionR[],
-										const uniform uint32 DynamicGroupIds[],
-										const uniform uint32 KinematicGroupIds[],
-										const uniform float PerGroupFriction[],
-										const uniform float PerGroupThickness[],
-										const uniform uint8 *uniform CollisionParticles,
-										const uniform uint8 *uniform Geometry,
-										const uniform int SizeofFImplicitObject,
-										const uniform int OffsetofGeometryType,
-										const uniform int OffsetOfMargin,
-										const uniform float Dt,
-										const uniform int InnerOffset,
-										const uniform int InnerRange,
-										const uniform int OuterOffset,
-										const uniform int OuterRange)
-{
-	uniform int BailAt = OuterRange;
-	uniform bool bBail = false;
-
-	// Pre-check for bail out.
-	foreach(Index = OuterOffset ... BailAt)
-	{
-		const uint32 DynamicGroupId = DynamicGroupIds[Index]; // Particle group Id
-
-		for(uniform int i = InnerOffset; i < InnerRange; i++)
-		{
-			const uniform uint32 KinematicGroupId = KinematicGroupIds[i];  // Collision group Id
-
-			if (KinematicGroupId != (uint32)INDEX_NONE && DynamicGroupId != KinematicGroupId)
-			{
-				foreach_active(j)
-				{
-					if(!bBail)
-					{
-						BailAt = extract(Index, j);
-						bBail = true;
-					}
-				}
-			}
-		}
-	}
-
-	foreach(Index = OuterOffset ... BailAt)
-	{
-		const FVector4f PandInvM = VectorLoad(&ParticlesPandInvM[extract(Index, 0)]);
-		varying FVector3f P;
-		varying float IM;
-		UnzipPandInvM(PandInvM, P, IM);
-
-		if(all(IM == 0))
-		{
-			continue;
-		}
-
-		const uint32 DynamicGroupId = DynamicGroupIds[Index]; // Particle group Id
-
-		for(uniform int i = InnerOffset; i < InnerRange; i++) // SequentialFor
-		{
-			// InverseTransformPosition. Can skip Scale calculation because it is 1 for xyz.
-			const FVector3f TranslatedVec = P - CollisionX[i];
-			const FVector4f InverseR = QuatInverse(CollisionR[i]);
-			const FVector3f PhiWithNormalInputFloat = VectorQuaternionRotateVector(InverseR, TranslatedVec);
-			const FVector PhiWithNormalInput = ConvertVector3fTo3Native(PhiWithNormalInputFloat);
-			FVector NormalReal;
-			FReal PhiReal;
-			FVector3f Normal;
-			float Phi;
-
-			const uniform SerializedPtr *uniform GeometryPtr = (const uniform SerializedPtr *uniform)Geometry;
-			const uniform uint8 *uniform Ptr = GeometryPtr[i].Ptr;
-			const uniform uint8 *uniform TypedPtr = Ptr + SizeofFImplicitObject;
-			const uniform uint8 Type = Ptr[OffsetofGeometryType] & 0x3F;
-
-			switch(Type)
-			{
-			case Sphere:
-				{
-					const uniform FReal Margin = *((const uniform FReal *uniform)(Ptr + OffsetOfMargin));
-					PhiWithNormalSphere(TypedPtr, PhiWithNormalInput, NormalReal, PhiReal, Margin);
-					Phi = (varying float)PhiReal;
-					Normal = ConvertVector3NativeTo3f(NormalReal);
-					break;
-				}
-			case Capsule:
-				{
-					const uniform FReal Margin = *((const uniform FReal *uniform)(Ptr + OffsetOfMargin));
-					PhiWithNormalCapsule(TypedPtr, PhiWithNormalInput, NormalReal, PhiReal, Margin);
-					Phi = (varying float)PhiReal;
-					Normal = ConvertVector3NativeTo3f(NormalReal);
-					break;
-				}
-			case Union:
-				{
-					PhiWithNormalUnion(TypedPtr, PhiWithNormalInput, NormalReal, PhiReal, SizeofFImplicitObject, OffsetofGeometryType, OffsetOfMargin);
-					Phi = (varying float)PhiReal;
-					Normal = ConvertVector3NativeTo3f(NormalReal);
-					break;
-				}
-			case TaperedCapsule:
-				{
-					PhiWithNormalTaperedCapsule(TypedPtr, PhiWithNormalInput, NormalReal, PhiReal);
-					Phi = (varying float)PhiReal;
-					Normal = ConvertVector3NativeTo3f(NormalReal);
-					break;
-				}
-			case Convex:
-				{
-					PhiWithNormalConvex(TypedPtr, PhiWithNormalInputFloat, Normal, Phi);
-					break;
-				}
-			default:
-				{
-					uniform int IMMask;
-					if(IM != 0)
-					{
-						IMMask = lanemask();
-					}
-
-					GetPhiWithNormal(CollisionParticles, (const uniform float *uniform)&PhiWithNormalInputFloat, (uniform float *uniform)&Normal, (uniform float *uniform)&Phi, i, programCount, IMMask);
-					break;
-				}
-			}
-
-			#pragma ignore warning(perf)
-			const float Thickness = PerGroupThickness[DynamicGroupId];
-
-			if(Phi < Thickness)
-			{
-				// TransformVector. Scale is 1 so that part can be removed.
-				const FVector3f NormalWorld = VectorQuaternionRotateVector(CollisionR[i], Normal);
-				P = P + ((-Phi + Thickness) * NormalWorld);
-
-				#pragma ignore warning(perf)
-				const float CoefficientOfFriction = PerGroupFriction[DynamicGroupId];
-
-				if(CoefficientOfFriction > 0.0f)
-				{
-					const uniform float MaximumFrictionCorrectionPerStep = 1.0f; // 1cm absolute maximum correction the friction force can provide
-					const float Penetration = (-Phi + Thickness); // This is related to the Normal impulse
-
-					FVector3f VectorToPoint = P - CollisionX[i];
-					const FVector3f X = VectorLoad(&ParticlesX[extract(Index, 0)]);
-					const FVector3f RelativeDisplacement = (P - X) - (CollisionV[i] + VectorCross(CollisionW[i], VectorToPoint)) * Dt; // This corresponds to the tangential velocity multiplied by dt (friction will drive this to zero if it is high enough)
-					const FVector3f RelativeDisplacementTangent = RelativeDisplacement - VectorDot(RelativeDisplacement, NormalWorld) * NormalWorld; // Project displacement into the tangential plane
-					const float RelativeDisplacementTangentLength = VectorSize(RelativeDisplacementTangent);
-					float PositionCorrection = clamp(Penetration * CoefficientOfFriction, 0.0f, RelativeDisplacementTangentLength);
-					if (PositionCorrection > MaximumFrictionCorrectionPerStep)
-					{
-						PositionCorrection = MaximumFrictionCorrectionPerStep;
-					}
-					if (PositionCorrection > 0)
-					{
-						P = P - ((PositionCorrection / RelativeDisplacementTangentLength) * RelativeDisplacementTangent);
-					}
-				}
-			}
-		}
-
-		if(IM != 0)
-		{
-			VectorStore(&ParticlesPandInvM[extract(Index, 0)], SetVector4(P, IM));
-		}
-	}
-}
+extern "C" unmasked void GetPhiWithNormalAndVelocityBone(const uniform uint8 *uniform CollisionParticles, const uniform float *uniform V, uniform float *uniform Normal, uniform float *uniform Phi, uniform int *uniform VelocityBone, const uniform int i, const uniform float Thickness, const uniform int ProgramCount, const uniform int Mask);
 
 export void ApplyPerParticleCollisionFastFriction(uniform FVector4f ParticlesPandInvM[],
 													const uniform FVector3f ParticlesX[],
@@ -825,6 +659,7 @@ export void ApplyPerParticleCollisionFastFriction(uniform FVector4f ParticlesPan
 			const uniform uint8 *uniform TypedPtr = Ptr + SizeofFImplicitObject;
 			const uniform uint8 Type = Ptr[OffsetofGeometryType] & 0x3F;
 
+			int32 VelocityBone = i;
 			switch(Type)
 			{
 			case Sphere:
@@ -862,6 +697,21 @@ export void ApplyPerParticleCollisionFastFriction(uniform FVector4f ParticlesPan
 					PhiWithNormalConvex(TypedPtr, PhiWithNormalInputFloat, Normal, Phi);
 					break;
 				}
+			case WeightedLatticeBone:
+				{
+					Phi = PerGroupThickness + 10;
+					break;
+				}
+			case WeightedLatticeLevelSetType:
+				{
+					uniform int IMMask;
+					if(IM != 0)
+					{
+						IMMask = lanemask();
+					}
+					GetPhiWithNormalAndVelocityBone(CollisionParticles, (const uniform float *uniform)&PhiWithNormalInputFloat, (uniform float *uniform)&Normal, (uniform float *uniform)&Phi, (uniform int *uniform)&VelocityBone, i, PerGroupThickness, programCount, IMMask);
+					break;
+				}
 			default:
 				{
 					uniform int IMMask;
@@ -881,16 +731,39 @@ export void ApplyPerParticleCollisionFastFriction(uniform FVector4f ParticlesPan
 				// TransformVector. Scale is 1 so that part can be removed.
 				const FVector3f NormalWorld = VectorQuaternionRotateVector(CollisionR[i], Normal);
 				P = P + (Penetration * NormalWorld);
+				if(Type == WeightedLatticeLevelSetType)
+				{
+					#pragma ignore warning(perf)
+					const FVector3f CollisionXVal = CollisionX[VelocityBone];
+					const FVector3f VectorToPoint = P - CollisionXVal;
+					const FVector3f X = VectorLoad(&ParticlesX[extract(Index, 0)]);
+					
+					#pragma ignore warning(perf)
+					const FVector3f CollisionVVal = CollisionV[VelocityBone];
+					#pragma ignore warning(perf)
+					const FVector3f CollisionWVal = CollisionW[VelocityBone];
+					const FVector3f RelativeDisplacement = (P - X) - (CollisionVVal + VectorCross(CollisionWVal, VectorToPoint)) * Dt; // This corresponds to the tangential velocity multiplied by dt (friction will drive this to zero if it is high enough)
+					const FVector3f RelativeDisplacementTangent = RelativeDisplacement - VectorDot(RelativeDisplacement, NormalWorld) * NormalWorld; // Project displacement into the tangential plane
+					const float RelativeDisplacementTangentLength = VectorSize(RelativeDisplacementTangent);
 
-				const FVector3f VectorToPoint = P - CollisionX[i];
-				const FVector3f X = VectorLoad(&ParticlesX[extract(Index, 0)]);
-				const FVector3f RelativeDisplacement = (P - X) - (CollisionV[i] + VectorCross(CollisionW[i], VectorToPoint)) * Dt; // This corresponds to the tangential velocity multiplied by dt (friction will drive this to zero if it is high enough)
-				const FVector3f RelativeDisplacementTangent = RelativeDisplacement - VectorDot(RelativeDisplacement, NormalWorld) * NormalWorld; // Project displacement into the tangential plane
-				const float RelativeDisplacementTangentLength = VectorSize(RelativeDisplacementTangent);
+					const float PositionCorrection = min(Penetration * PerGroupFriction, RelativeDisplacementTangentLength);
+					const float CorrectionRatio = select(RelativeDisplacementTangentLength < FLOAT_SMALL_NUMBER, 0.0f, PositionCorrection / RelativeDisplacementTangentLength);
+					P = P - (CorrectionRatio * RelativeDisplacementTangent);
+				
+				}
+				else
+				{
+					const FVector3f VectorToPoint = P - CollisionX[i];
+					const FVector3f X = VectorLoad(&ParticlesX[extract(Index, 0)]);
+
+					const FVector3f RelativeDisplacement = (P - X) - (CollisionV[i] + VectorCross(CollisionW[i], VectorToPoint)) * Dt; // This corresponds to the tangential velocity multiplied by dt (friction will drive this to zero if it is high enough)
+					const FVector3f RelativeDisplacementTangent = RelativeDisplacement - VectorDot(RelativeDisplacement, NormalWorld) * NormalWorld; // Project displacement into the tangential plane
+					const float RelativeDisplacementTangentLength = VectorSize(RelativeDisplacementTangent);
 
-				const float PositionCorrection = min(Penetration * PerGroupFriction, RelativeDisplacementTangentLength);
-				const float CorrectionRatio = select(RelativeDisplacementTangentLength < FLOAT_SMALL_NUMBER, 0.0f, PositionCorrection / RelativeDisplacementTangentLength);
-				P = P - (CorrectionRatio * RelativeDisplacementTangent);
+					const float PositionCorrection = min(Penetration * PerGroupFriction, RelativeDisplacementTangentLength);
+					const float CorrectionRatio = select(RelativeDisplacementTangentLength < FLOAT_SMALL_NUMBER, 0.0f, PositionCorrection / RelativeDisplacementTangentLength);
+					P = P - (CorrectionRatio * RelativeDisplacementTangent);
+				}
 			}
 		}
 
@@ -1011,6 +884,11 @@ export void ApplyPerParticleCollisionNoFriction(uniform FVector4f ParticlesPandI
 					PhiWithNormalConvex(TypedPtr, PhiWithNormalInputFloat, Normal, Phi);
 					break;
 				}
+			case WeightedLatticeBone:
+				{
+					Phi = PerGroupThickness + 10;
+					break;
+				}
 			default:
 				{
 					uniform int IMMask;
diff --git a/Engine/Source/Runtime/Experimental/Chaos/Public/Chaos/PBDEvolution.h b/Engine/Source/Runtime/Experimental/Chaos/Public/Chaos/PBDEvolution.h
--- a/Engine/Source/Runtime/Experimental/Chaos/Public/Chaos/PBDEvolution.h
+++ b/Engine/Source/Runtime/Experimental/Chaos/Public/Chaos/PBDEvolution.h
@@ -168,6 +168,7 @@ class CHAOS_API FPBDEvolution : public TArrayCollection
 	const TArray<bool>& GetCollisionStatus() { return MCollided; }
 	const TArray<FSolverVec3>& GetCollisionContacts() const { return MCollisionContacts; }
 	const TArray<FSolverVec3>& GetCollisionNormals() const { return MCollisionNormals; }
+	const TArray<FSolverReal>& GetCollisionPhis() const { return MCollisionPhis; }
 
 	FSolverReal GetTime() const { return MTime; }
 
@@ -192,6 +193,7 @@ private:
 	TArrayCollectionArray<uint32> MParticleGroupIds;  // Used for per group parameters for particles
 	TArray<FSolverVec3> MCollisionContacts;
 	TArray<FSolverVec3> MCollisionNormals;
+	TArray<FSolverReal> MCollisionPhis;
 
 	TArrayCollectionArray<FSolverVec3> MGroupGravityAccelerations;
 	TArrayCollectionArray<FVelocityAndPressureField> MGroupVelocityAndPressureFields;
diff --git a/Engine/Source/Runtime/Experimental/Chaos/Public/Chaos/PerParticlePBDCCDCollisionConstraint.h b/Engine/Source/Runtime/Experimental/Chaos/Public/Chaos/PerParticlePBDCCDCollisionConstraint.h
--- a/Engine/Source/Runtime/Experimental/Chaos/Public/Chaos/PerParticlePBDCCDCollisionConstraint.h
+++ b/Engine/Source/Runtime/Experimental/Chaos/Public/Chaos/PerParticlePBDCCDCollisionConstraint.h
@@ -7,6 +7,8 @@
 #include "Chaos/Transform.h"
 #include "Chaos/PBDActiveView.h"
 #include "Chaos/PBDSoftsSolverParticles.h"
+#include "Chaos/WeightedLatticeImplicitObject.h"
+#include "Chaos/Levelset.h"
 #include "Misc/ScopeLock.h"
 
 namespace Chaos::Softs
@@ -77,7 +79,7 @@ private:
 			{
 				const uint32 KinematicGroupId = KinematicGroupIds[CollisionIndex];  // Collision group Id
 
-				if (KinematicGroupId != (uint32)INDEX_NONE && DynamicGroupId != KinematicGroupId)
+				if ((KinematicGroupId != (uint32)INDEX_NONE && DynamicGroupId != KinematicGroupId) || CollisionParticles.Geometry(CollisionIndex)->GetType() == Chaos::ImplicitObjectType::WeightedLatticeBone)
 				{
 					return; // Bail out if the collision groups doesn't match the particle group id, or use INDEX_NONE (= global collision that affects all particle)
 				}
@@ -109,8 +111,35 @@ private:
 					Particles.P(Index) += Penetration * NormalWorld;
 
 					// Friction
-					const FSolverVec3 VectorToPoint = Particles.P(Index) - CollisionParticles.X(CollisionIndex);
-					const FSolverVec3 RelativeDisplacement = (Particles.P(Index) - Particles.X(Index)) - (CollisionParticles.V(CollisionIndex) + FSolverVec3::CrossProduct(CollisionParticles.W(CollisionIndex), VectorToPoint)) * Dt;  // This corresponds to the tangential velocity multiplied by dt (friction will drive this to zero if it is high enough)
+					int32 VelocityBone = CollisionIndex;
+					if (const TWeightedLatticeImplicitObject<FLevelSet>* LevelSet = CollisionParticles.Geometry(CollisionIndex)->GetObject< TWeightedLatticeImplicitObject<FLevelSet> >())
+					{
+						TArray<FWeightedLatticeImplicitObject::FEmbeddingCoordinate> Coordinates;
+						LevelSet->GetEmbeddingCoordinates(PointPair.First, Coordinates, false);
+						int32 ClosestCoordIndex = INDEX_NONE;
+						double ClosestCoordPhi = UE_BIG_NUMBER;
+						for (int32 CoordIndex = 0; CoordIndex < Coordinates.Num(); ++CoordIndex)
+						{
+							FVec3 NormalUnused;
+							const double CoordPhi = FMath::Abs(LevelSet->GetEmbeddedObject()->PhiWithNormal(Coordinates[CoordIndex].UndeformedPosition(LevelSet->GetGrid()), NormalUnused));
+							if (CoordPhi < ClosestCoordPhi)
+							{
+								ClosestCoordIndex = CoordIndex;
+								ClosestCoordPhi = CoordPhi;
+							}
+						}
+						if (ClosestCoordIndex != INDEX_NONE)
+						{
+							const int32 StrongestBone = Coordinates[ClosestCoordIndex].GreatestInfluenceBone(LevelSet->GetBoneData());
+							if (StrongestBone != INDEX_NONE)
+							{
+								VelocityBone = LevelSet->GetSolverBoneIndices()[StrongestBone];
+							}
+						}
+					}
+
+					const FSolverVec3 VectorToPoint = Particles.P(Index) - CollisionParticles.X(VelocityBone);
+					const FSolverVec3 RelativeDisplacement = (Particles.P(Index) - Particles.X(Index)) - (CollisionParticles.V(VelocityBone) + FSolverVec3::CrossProduct(CollisionParticles.W(VelocityBone), VectorToPoint)) * Dt;  // This corresponds to the tangential velocity multiplied by dt (friction will drive this to zero if it is high enough)
 					const FSolverVec3 RelativeDisplacementTangent = RelativeDisplacement - FSolverVec3::DotProduct(RelativeDisplacement, NormalWorld) * NormalWorld;  // Project displacement into the tangential plane
 					const FSolverReal RelativeDisplacementTangentLength = RelativeDisplacementTangent.Size();
 					if (RelativeDisplacementTangentLength >= UE_SMALL_NUMBER)
diff --git a/Engine/Source/Runtime/Experimental/Chaos/Public/Chaos/PerParticlePBDCollisionConstraint.h b/Engine/Source/Runtime/Experimental/Chaos/Public/Chaos/PerParticlePBDCollisionConstraint.h
--- a/Engine/Source/Runtime/Experimental/Chaos/Public/Chaos/PerParticlePBDCollisionConstraint.h
+++ b/Engine/Source/Runtime/Experimental/Chaos/Public/Chaos/PerParticlePBDCollisionConstraint.h
@@ -6,6 +6,8 @@
 #include "Chaos/PBDSoftsEvolutionFwd.h"
 #include "Chaos/PBDSoftsSolverParticles.h"
 #include "Chaos/KinematicGeometryParticles.h"
+#include "Chaos/WeightedLatticeImplicitObject.h"
+#include "Chaos/Levelset.h"
 #include "HAL/PlatformMath.h"
 
 #if !defined(CHAOS_PER_PARTICLE_COLLISION_ISPC_ENABLED_DEFAULT)
@@ -32,19 +34,47 @@ class CHAOS_API FPerParticlePBDCollisionConstraint final
 
 public:
 	FPerParticlePBDCollisionConstraint(const TPBDActiveView<FSolverRigidParticles>& InParticlesActiveView, TArray<bool>& Collided, TArray<uint32>& DynamicGroupIds, TArray<uint32>& KinematicGroupIds, const TArray<FSolverReal>& PerGroupThickness, const TArray<FSolverReal>& PerGroupFriction)
-	    : bFastPositionBasedFriction(true), MCollisionParticlesActiveView(InParticlesActiveView), MCollided(Collided), MDynamicGroupIds(DynamicGroupIds), MKinematicGroupIds(KinematicGroupIds), MPerGroupThickness(PerGroupThickness), MPerGroupFriction(PerGroupFriction) {}
+	: bFastPositionBasedFriction(true)
+	, MCollisionParticlesActiveView(InParticlesActiveView)
+	, MCollided(Collided), MDynamicGroupIds(DynamicGroupIds)
+	, MKinematicGroupIds(KinematicGroupIds), MPerGroupThickness(PerGroupThickness)
+	, MPerGroupFriction(PerGroupFriction) {}
 
-	~FPerParticlePBDCollisionConstraint() {}
+	FPerParticlePBDCollisionConstraint(const TPBDActiveView<FSolverRigidParticles>& InParticlesActiveView, TArray<bool>& Collided,
+		TArray<FSolverVec3>& InContacts,
+		TArray<FSolverVec3>& InNormals, 
+		TArray<FSolverReal>& InPhis,
+		TArray<uint32>& DynamicGroupIds, TArray<uint32>& KinematicGroupIds, const TArray<FSolverReal>& PerGroupThickness, const TArray<FSolverReal>& PerGroupFriction,
+		bool bWriteCCDContacts)
+	    : bFastPositionBasedFriction(true), MCollisionParticlesActiveView(InParticlesActiveView), MCollided(Collided)
+		, Contacts(&InContacts)
+		, Normals(&InNormals)
+		, Phis(&InPhis)
+		, MDynamicGroupIds(DynamicGroupIds), MKinematicGroupIds(KinematicGroupIds), MPerGroupThickness(PerGroupThickness), MPerGroupFriction(PerGroupFriction)
+		, Mutex(bWriteCCDContacts ? new FCriticalSection : nullptr) {}
+
+	~FPerParticlePBDCollisionConstraint() 
+	{
+		delete Mutex;
+	}
 
 	inline void ApplyRange(FSolverParticles& Particles, const FSolverReal Dt, const int32 Offset, const int32 Range) const
 	{
+		// NOTE: currently using ISPC with TWeightedLatticeImplicitObject<FLevelSet> is significantly slower than not using ISPC (largely because it has not been fully implemented)
 		if (bRealTypeCompatibleWithISPC && bChaos_PerParticleCollision_ISPC_Enabled && bFastPositionBasedFriction )
 		{
 			ApplyHelperISPC(Particles, Dt, Offset, Range);
 		}
 		else
 		{
-			ApplyHelper(Particles, Dt, Offset, Range);
+			if (Mutex)
+			{
+				ApplyHelper<true>(Particles, Dt, Offset, Range);
+			}
+			else
+			{
+				ApplyHelper<false>(Particles, Dt, Offset, Range);
+			}
 		}
 	}
 
@@ -70,6 +100,7 @@ public:
 	}
 
 private:
+	template<bool bLockAndWriteContacts>
 	inline void ApplyHelper(FSolverParticles& Particles, const FSolverReal Dt, const int32 Offset, const int32 Range) const
 	{
 		const uint32 DynamicGroupId = MDynamicGroupIds[Offset];  // Particle group Id, must be the same across the entire range
@@ -91,26 +122,59 @@ private:
 				{
 					const uint32 KinematicGroupId = MKinematicGroupIds[i];  // Collision group Id
 
-					if (KinematicGroupId != (uint32)INDEX_NONE && DynamicGroupId != KinematicGroupId)
+					if ((KinematicGroupId != (uint32)INDEX_NONE && DynamicGroupId != KinematicGroupId ) || CollisionParticles.Geometry(i)->GetType() == Chaos::ImplicitObjectType::WeightedLatticeBone)
 					{
 						return; // Bail out if the collision groups doesn't match the particle group id, or use INDEX_NONE (= global collision that affects all particle)
 					}
 					const FSolverRigidTransform3 Frame(CollisionParticles.X(i), CollisionParticles.R(i));
 					const FVec3 RigidSpacePosition(Frame.InverseTransformPosition(Particles.P(Index)));  // PhiWithNormal requires FReal based arguments
 					FVec3 ImplicitNormal;                                                                // since implicits don't use FSolverReal
-					const FSolverReal Phi = (FSolverReal)CollisionParticles.Geometry(i)->PhiWithNormal(RigidSpacePosition, ImplicitNormal);
+					FSolverReal Phi;
+					int32 VelocityBone = i;
+					FSolverReal Penetration;
+					if (const TWeightedLatticeImplicitObject<FLevelSet>* LevelSet = CollisionParticles.Geometry(i)->GetObject< TWeightedLatticeImplicitObject<FLevelSet> >())
+					{
+						FWeightedLatticeImplicitObject::FEmbeddingCoordinate SurfaceCoord;
+						Phi = (FSolverReal)LevelSet->PhiWithNormalAndSurfacePoint(RigidSpacePosition, ImplicitNormal, SurfaceCoord);
+						Penetration = PerGroupThickness - Phi; // This is related to the Normal impulse
+						if (Penetration > (FSolverReal)0.)
+						{
+							const int32 StrongestBone = SurfaceCoord.GreatestInfluenceBone(LevelSet->GetBoneData());
+							if (StrongestBone != INDEX_NONE)
+							{
+								VelocityBone = LevelSet->GetSolverBoneIndices()[StrongestBone];
+							}
+						}
+					}
+					else
+					{
+						Phi = (FSolverReal)CollisionParticles.Geometry(i)->PhiWithNormal(RigidSpacePosition, ImplicitNormal);
+						Penetration = PerGroupThickness - Phi; // This is related to the Normal impulse
+					}
 					const FSolverVec3 Normal(ImplicitNormal);
 
-					const FSolverReal Penetration = PerGroupThickness - Phi; // This is related to the Normal impulse
 					if (Penetration > (FSolverReal)0.)
 					{
 						const FSolverVec3 NormalWorld = Frame.TransformVector(Normal);
+
+						if (bLockAndWriteContacts)
+						{
+							check(Mutex);
+							checkSlow(Contacts);
+							checkSlow(Normals);
+							checkSlow(Phis);
+							FScopeLock Lock(Mutex);
+							Contacts->Emplace(Particles.P(Index));
+							Normals->Emplace(NormalWorld);
+							Phis->Emplace(Phi);
+						}
+
 						Particles.P(Index) += Penetration * NormalWorld;
 
 						if (bFastPositionBasedFriction)
 						{
-							FSolverVec3 VectorToPoint = Particles.P(Index) - CollisionParticles.X(i);
-							const FSolverVec3 RelativeDisplacement = (Particles.P(Index) - Particles.X(Index)) - (CollisionParticles.V(i) + FSolverVec3::CrossProduct(CollisionParticles.W(i), VectorToPoint)) * Dt; // This corresponds to the tangential velocity multiplied by dt (friction will drive this to zero if it is high enough)
+							FSolverVec3 VectorToPoint = Particles.P(Index) - CollisionParticles.X(VelocityBone);
+							const FSolverVec3 RelativeDisplacement = (Particles.P(Index) - Particles.X(Index)) - (CollisionParticles.V(VelocityBone) + FSolverVec3::CrossProduct(CollisionParticles.W(VelocityBone), VectorToPoint)) * Dt; // This corresponds to the tangential velocity multiplied by dt (friction will drive this to zero if it is high enough)
 							const FSolverVec3 RelativeDisplacementTangent = RelativeDisplacement - FSolverVec3::DotProduct(RelativeDisplacement, NormalWorld) * NormalWorld; // Project displacement into the tangential plane
 							const FSolverReal RelativeDisplacementTangentLength = RelativeDisplacementTangent.Size();
 							if (RelativeDisplacementTangentLength >= UE_SMALL_NUMBER)
@@ -124,8 +188,8 @@ private:
 						{
 							// Note, to fix: Only use fast position based friction for now, since adding to TMaps here is not thread safe when calling Apply on multiple threads (will cause crash)
 							FVelocityConstraint Constraint;
-							FSolverVec3 VectorToPoint = Particles.P(Index) - CollisionParticles.X(i);
-							Constraint.Velocity = CollisionParticles.V(i) + FSolverVec3::CrossProduct(CollisionParticles.W(i), VectorToPoint);
+							FSolverVec3 VectorToPoint = Particles.P(Index) - CollisionParticles.X(VelocityBone);
+							Constraint.Velocity = CollisionParticles.V(VelocityBone) + FSolverVec3::CrossProduct(CollisionParticles.W(VelocityBone), VectorToPoint);
 							Constraint.Normal = Frame.TransformVector(Normal);
 						
 							MVelocityConstraints.Add(Index, Constraint);
@@ -149,7 +213,7 @@ private:
 				{
 					const uint32 KinematicGroupId = MKinematicGroupIds[i];  // Collision group Id
 
-					if (KinematicGroupId != (uint32)INDEX_NONE && DynamicGroupId != KinematicGroupId)
+					if ((KinematicGroupId != (uint32)INDEX_NONE && DynamicGroupId != KinematicGroupId) || CollisionParticles.Geometry(i)->GetType() == Chaos::ImplicitObjectType::WeightedLatticeBone)
 					{
 						return; // Bail out if the collision groups doesn't match the particle group id, or use INDEX_NONE (= global collision that affects all particle)
 					}
@@ -163,6 +227,18 @@ private:
 					if (Penetration > (FSolverReal)0.)
 					{
 						const FSolverVec3 NormalWorld = Frame.TransformVector(Normal);
+						if (bLockAndWriteContacts)
+						{
+							check(Mutex);
+							checkSlow(Contacts);
+							checkSlow(Normals);
+							checkSlow(Phis);
+							FScopeLock Lock(Mutex);
+							Contacts->Emplace(Particles.P(Index));
+							Normals->Emplace(NormalWorld);
+							Phis->Emplace(Phi);
+						}
+
 						Particles.P(Index) += Penetration * NormalWorld;
 					}
 				});
@@ -177,11 +253,15 @@ private:
 	// TODO(mlentine): Need a bb hierarchy
 	const TPBDActiveView<FSolverRigidParticles>& MCollisionParticlesActiveView;
 	TArray<bool>& MCollided;
+	TArray<FSolverVec3>* const Contacts = nullptr;
+	TArray<FSolverVec3>* const Normals = nullptr;
+	TArray<FSolverReal>* const Phis = nullptr;
 	const TArray<uint32>& MDynamicGroupIds;
 	const TArray<uint32>& MKinematicGroupIds;
 	mutable TMap<int32, FVelocityConstraint> MVelocityConstraints;
 	const TArray<FSolverReal>& MPerGroupThickness;
 	const TArray<FSolverReal>& MPerGroupFriction;
+	FCriticalSection* const Mutex = nullptr;
 };
 
 }  // End namespace Chaos::Softs
