Geometry collection : Avoid physics state creation on ISM pool meshes

diff --git a/Engine/Source/Runtime/Experimental/GeometryCollectionEngine/Private/GeometryCollection/GeometryCollectionISMPoolComponent.cpp b/Engine/Source/Runtime/Experimental/GeometryCollectionEngine/Private/GeometryCollection/GeometryCollectionISMPoolComponent.cpp
--- a/Engine/Source/Runtime/Experimental/GeometryCollectionEngine/Private/GeometryCollection/GeometryCollectionISMPoolComponent.cpp
+++ b/Engine/Source/Runtime/Experimental/GeometryCollectionEngine/Private/GeometryCollection/GeometryCollectionISMPoolComponent.cpp
@@ -59,7 +59,7 @@ FGeometryCollectionISM::FGeometryCollectionISM(AActor* OwmingActor, const FGeome
 		}
 		ISMC->SetCullDistances(0, 0);
 		ISMC->SetCanEverAffectNavigation(false);
-		ISMC->SetCollisionProfileName(UCollisionProfile::NoCollision_ProfileName);
+		ISMC->SetCollisionEnabled(ECollisionEnabled::NoCollision);
 		ISMC->SetCastShadow(true);
 		ISMC->SetMobility(EComponentMobility::Stationary);
 		OwmingActor->AddInstanceComponent(ISMC);
