move geometry collection bone selected material out of the user-exposed materials array for new geometry collections

diff --git a/Engine/Source/Runtime/Experimental/GeometryCollectionEngine/Private/GeometryCollection/GeometryCollectionComponent.cpp b/Engine/Source/Runtime/Experimental/GeometryCollectionEngine/Private/GeometryCollection/GeometryCollectionComponent.cpp
--- a/Engine/Source/Runtime/Experimental/GeometryCollectionEngine/Private/GeometryCollection/GeometryCollectionComponent.cpp
+++ b/Engine/Source/Runtime/Experimental/GeometryCollectionEngine/Private/GeometryCollection/GeometryCollectionComponent.cpp
@@ -4171,6 +4171,26 @@ UMaterialInterface* UGeometryCollectionComponent::GetMaterial(int32 MaterialInde
 	}
 }
 
+void UGeometryCollectionComponent::GetUsedMaterials(TArray<UMaterialInterface*>& OutMaterials, bool bGetDebugMaterials) const
+{
+	Super::GetUsedMaterials(OutMaterials, bGetDebugMaterials);
+
+	if (GetRestCollection() && GetRestCollection()->GetBoneSelectedMaterial())
+	{
+		OutMaterials.Add(GetRestCollection()->GetBoneSelectedMaterial());
+	}
+}
+
+FMaterialRelevance UGeometryCollectionComponent::GetMaterialRelevance(ERHIFeatureLevel::Type InFeatureLevel) const
+{
+	FMaterialRelevance Result = Super::GetMaterialRelevance(InFeatureLevel);
+	if (RestCollection && RestCollection->GetBoneSelectedMaterial())
+	{
+		Result |= RestCollection->GetBoneSelectedMaterial()->GetRelevance_Concurrent(InFeatureLevel);
+	}
+	return Result;
+}
+
 #if WITH_EDITOR
 void UGeometryCollectionComponent::SelectEmbeddedGeometry()
 {
diff --git a/Engine/Source/Runtime/Experimental/GeometryCollectionEngine/Private/GeometryCollection/GeometryCollectionObject.cpp b/Engine/Source/Runtime/Experimental/GeometryCollectionEngine/Private/GeometryCollection/GeometryCollectionObject.cpp
--- a/Engine/Source/Runtime/Experimental/GeometryCollectionEngine/Private/GeometryCollection/GeometryCollectionObject.cpp
+++ b/Engine/Source/Runtime/Experimental/GeometryCollectionEngine/Private/GeometryCollection/GeometryCollectionObject.cpp
@@ -498,8 +498,8 @@ void UGeometryCollection::InitializeMaterials(bool bHasInternalMaterials)
 {
 	Modify();
 
-	// Last Material is the selection one
-	UMaterialInterface* BoneSelectedMaterial = LoadObject<UMaterialInterface>(nullptr, GetSelectedMaterialPath(), nullptr, LOAD_None, nullptr);
+	// Initialize the BoneSelectedMaterial separate from the materials on the collection
+	BoneSelectedMaterial = LoadObject<UMaterialInterface>(nullptr, GetSelectedMaterialPath(), nullptr, LOAD_None, nullptr);
 
 	TManagedArray<int32>& MaterialIDs = GeometryCollection->MaterialID;
 
@@ -604,8 +604,8 @@ void UGeometryCollection::InitializeMaterials(bool bHasInternalMaterials)
 	// Set new material array on the collection
 	Materials = FinalMaterials;
 
-	// Last Material is the selection one
-	BoneSelectedMaterialIndex = Materials.Add(BoneSelectedMaterial);
+	// BoneSelectedMaterial is no longer stored in the general Materials array
+	BoneSelectedMaterialIndex = INDEX_NONE;
 
 	GeometryCollection->ReindexMaterials();
 	InvalidateCollection();
@@ -731,6 +731,8 @@ void UGeometryCollection::Serialize(FArchive& Ar)
 		{
 			Materials[SelectedMaterialIndex] = Materials[0];
 		}
+		// Likewise remove the direct reference to the BoneSelectedMaterial on cook
+		BoneSelectedMaterial = nullptr;
 
 		if (bStripOnCook)
 		{
@@ -915,6 +917,24 @@ void UGeometryCollection::Serialize(FArchive& Ar)
 		DamagePropagationData.bEnabled = false;
 	}
 
+	if (Ar.IsLoading() && !bIsCookedOrCooking && BoneSelectedMaterialIndex != INDEX_NONE)
+	{
+		BoneSelectedMaterial = LoadObject<UMaterialInterface>(nullptr, GetSelectedMaterialPath(), nullptr, LOAD_None, nullptr);
+		if (Materials.IsValidIndex(BoneSelectedMaterialIndex))
+		{
+			if (!BoneSelectedMaterial)
+			{
+				BoneSelectedMaterial = Materials[BoneSelectedMaterialIndex];
+			}
+			// Remove the material assuming it's the last in the list (otherwise, leave it, as it's not clear why it would be in that state)
+			if (BoneSelectedMaterialIndex == Materials.Num() - 1)
+			{
+				Materials.RemoveAt(BoneSelectedMaterialIndex);
+			}
+		}
+		BoneSelectedMaterialIndex = INDEX_NONE;
+	}
+
 #if WITH_EDITORONLY_DATA
 	if (bCreateSimulationData)
 	{
diff --git a/Engine/Source/Runtime/Experimental/GeometryCollectionEngine/Private/GeometryCollection/GeometryCollectionSceneProxy.cpp b/Engine/Source/Runtime/Experimental/GeometryCollectionEngine/Private/GeometryCollection/GeometryCollectionSceneProxy.cpp
--- a/Engine/Source/Runtime/Experimental/GeometryCollectionEngine/Private/GeometryCollection/GeometryCollectionSceneProxy.cpp
+++ b/Engine/Source/Runtime/Experimental/GeometryCollectionEngine/Private/GeometryCollection/GeometryCollectionSceneProxy.cpp
@@ -156,7 +156,6 @@ FGeometryCollectionSceneProxy::FGeometryCollectionSceneProxy(UGeometryCollection
 	, bShowBoneColors(Component->GetShowBoneColors())
 	, bEnableBoneSelection(Component->GetEnableBoneSelection())
 	, bSuppressSelectionMaterial(Component->GetSuppressSelectionMaterial())
-	, BoneSelectionMaterialID(Component->GetBoneSelectedMaterialID())
 	, bUseFullPrecisionUVs(Component->GetRestCollection()->bUseFullPrecisionUVs)
 	, TransformVertexBuffersContainsOriginalMesh(false)
 {
@@ -171,6 +170,23 @@ FGeometryCollectionSceneProxy::FGeometryCollectionSceneProxy(UGeometryCollection
 			Materials[MaterialIndex] = UMaterial::GetDefaultMaterial(MD_Surface);
 		}
 	}
+	if (Component->GetRestCollection())
+	{
+		BoneSelectedMaterial = Component->GetRestCollection()->GetBoneSelectedMaterial();
+	}
+	if (!BoneSelectedMaterial)
+	{
+		int32 LegacyBoneMaterialID = Component->GetBoneSelectedMaterialID();
+		if (Materials.IsValidIndex(LegacyBoneMaterialID))
+		{
+			BoneSelectedMaterial = Materials[LegacyBoneMaterialID];
+		}
+	}
+	if (BoneSelectedMaterial && !BoneSelectedMaterial->CheckMaterialUsage_Concurrent(MATUSAGE_GeometryCollections))
+	{
+		// If we have an invalid BoneSelectedMaterial, switch it back to null to skip its usage in GetDynamicMeshElements below
+		BoneSelectedMaterial = nullptr;
+	}
 
 	// Make sure the vertex color material has the usage flag for rendering geometry collections
 	if (GEngine->VertexColorMaterial)
@@ -1079,9 +1095,9 @@ void FGeometryCollectionSceneProxy::GetDynamicMeshElements(const TArray<const FS
 			// bone selection is already contained in the rendered colors
 			// #note: This renders the geometry again but with the bone selection material.  Ideally we'd have one render pass and one
 			// material.
-			if ((bShowBoneColors || bEnableBoneSelection) && !bSuppressSelectionMaterial && Materials.IsValidIndex(BoneSelectionMaterialID))
+			if ((bShowBoneColors || bEnableBoneSelection) && !bSuppressSelectionMaterial && BoneSelectedMaterial)
 			{
-				FMaterialRenderProxy* MaterialRenderProxy = Materials[BoneSelectionMaterialID]->GetRenderProxy();
+				FMaterialRenderProxy* MaterialRenderProxy = BoneSelectedMaterial->GetRenderProxy();
 
 				FMeshBatch& Mesh = Collector.AllocateMesh();
 				FMeshBatchElement& BatchElement = Mesh.Elements[0];
diff --git a/Engine/Source/Runtime/Experimental/GeometryCollectionEngine/Private/GeometryCollection/GeometryCollectionSceneProxy.h b/Engine/Source/Runtime/Experimental/GeometryCollectionEngine/Private/GeometryCollection/GeometryCollectionSceneProxy.h
--- a/Engine/Source/Runtime/Experimental/GeometryCollectionEngine/Private/GeometryCollection/GeometryCollectionSceneProxy.h
+++ b/Engine/Source/Runtime/Experimental/GeometryCollectionEngine/Private/GeometryCollection/GeometryCollectionSceneProxy.h
@@ -222,6 +222,7 @@ private:
 class FGeometryCollectionSceneProxy final : public FPrimitiveSceneProxy
 {
 	TArray<UMaterialInterface*> Materials;
+	UMaterialInterface* BoneSelectedMaterial = nullptr;
 
 	FMaterialRelevance MaterialRelevance;
 
@@ -261,7 +262,6 @@ class FGeometryCollectionSceneProxy final : public FPrimitiveSceneProxy
 	bool bShowBoneColors;
 	bool bEnableBoneSelection;
 	bool bSuppressSelectionMaterial;
-	int BoneSelectionMaterialID;
 
 	bool bUseFullPrecisionUVs = false;
 
diff --git a/Engine/Source/Runtime/Experimental/GeometryCollectionEngine/Public/GeometryCollection/GeometryCollectionComponent.h b/Engine/Source/Runtime/Experimental/GeometryCollectionEngine/Public/GeometryCollection/GeometryCollectionComponent.h
--- a/Engine/Source/Runtime/Experimental/GeometryCollectionEngine/Public/GeometryCollection/GeometryCollectionComponent.h
+++ b/Engine/Source/Runtime/Experimental/GeometryCollectionEngine/Public/GeometryCollection/GeometryCollectionComponent.h
@@ -505,6 +505,8 @@ public:
 	//~ Begin UMeshComponent Interface.	
 	virtual int32 GetNumMaterials() const override;
 	virtual UMaterialInterface* GetMaterial(int32 MaterialIndex) const override;
+	virtual void GetUsedMaterials(TArray<UMaterialInterface*>& OutMaterials, bool bGetDebugMaterials = false) const override;
+	virtual FMaterialRelevance GetMaterialRelevance(ERHIFeatureLevel::Type InFeatureLevel) const override;
 	//~ End UMeshComponent Interface.
 
 	/** Chaos RBD Solver override. Will use the world's default solver actor if null. */
diff --git a/Engine/Source/Runtime/Experimental/GeometryCollectionEngine/Public/GeometryCollection/GeometryCollectionObject.h b/Engine/Source/Runtime/Experimental/GeometryCollectionEngine/Public/GeometryCollection/GeometryCollectionObject.h
--- a/Engine/Source/Runtime/Experimental/GeometryCollectionEngine/Public/GeometryCollection/GeometryCollectionObject.h
+++ b/Engine/Source/Runtime/Experimental/GeometryCollectionEngine/Public/GeometryCollection/GeometryCollectionObject.h
@@ -668,6 +668,11 @@ public:
 
 	FORCEINLINE const int32 GetBoneSelectedMaterialIndex() const { return BoneSelectedMaterialIndex; }
 
+	UMaterialInterface* GetBoneSelectedMaterial() const
+	{
+		return BoneSelectedMaterial;
+	}
+
 	/** Returns the asset path for the automatically populated selected material. */
 	static const TCHAR* GetSelectedMaterialPath();
 
@@ -737,8 +742,14 @@ private:
 #endif
 
 	// #todo(dmp): rename to be consistent BoneSelectedMaterialID?
+	// Legacy index of the bone selected material in the object's Materials array, or INDEX_NONE if it is not stored there.
+	// Note for new objects the bone selected material should not be stored in the Materials array, so this should be INDEX_NONE
+	UPROPERTY()
+	int32 BoneSelectedMaterialIndex = INDEX_NONE;
+
+	// The material to use for rendering bone selections in the editor, or nullptr
 	UPROPERTY()
-	int32 BoneSelectedMaterialIndex;
+	TObjectPtr<UMaterialInterface> BoneSelectedMaterial = nullptr;
 
 	TSharedPtr<FGeometryCollection, ESPMode::ThreadSafe> GeometryCollection;
 };
