reduce the amount of async physics ticking geometry collections do.

- Server: only update rep data when the physics sync is called (i.e. when the physics proxy is dirty)
- Client: only process rep data when we get notified the rep data is changed. plus let the tick stick around for a bit so long as the tick is active.

diff --git a/Engine/Source/Runtime/Engine/Classes/Components/ActorComponent.h b/Engine/Source/Runtime/Engine/Classes/Components/ActorComponent.h
--- a/Engine/Source/Runtime/Engine/Classes/Components/ActorComponent.h
+++ b/Engine/Source/Runtime/Engine/Classes/Components/ActorComponent.h
@@ -1154,6 +1154,10 @@ protected:
 	/** Sets bAsyncPhysicsTickEnabled which determines whether to use use the async physics tick with this component. */
 	void SetAsyncPhysicsTickEnabled(bool bEnabled);
 
+	/** Defers the call to SetAsyncPhysicsTickEnabled(false) to the next async tick. This allows you to remove the
+	 *  async physics tick callback within the async physics tick safely. */
+	void DeferRemoveAsyncPhysicsTick();
+
 private:
 #if !(UE_BUILD_SHIPPING || UE_BUILD_TEST)
 	/** This is the old name of the tick function. We just want to avoid mistakes with an attempt to override this */
diff --git a/Engine/Source/Runtime/Engine/Private/Components/ActorComponent.cpp b/Engine/Source/Runtime/Engine/Private/Components/ActorComponent.cpp
--- a/Engine/Source/Runtime/Engine/Private/Components/ActorComponent.cpp
+++ b/Engine/Source/Runtime/Engine/Private/Components/ActorComponent.cpp
@@ -1262,6 +1262,20 @@ void UActorComponent::SetAsyncPhysicsTickEnabled(bool bEnable)
 	bAsyncPhysicsTickEnabled = bEnable;
 }
 
+void UActorComponent::DeferRemoveAsyncPhysicsTick()
+{
+	if (FPhysScene_Chaos* Scene = static_cast<FPhysScene_Chaos*>(WorldPrivate->GetPhysicsScene()))
+	{
+		// Set 0 for the step so that it gets run immediately on the next async tick.
+		Scene->EnqueueAsyncPhysicsCommand(0, this,
+			[this]()
+			{
+				SetAsyncPhysicsTickEnabled(false);
+			}
+		);
+	}
+}
+
 void UActorComponent::TickComponent(float DeltaTime, enum ELevelTick TickType, FActorComponentTickFunction *ThisTickFunction)
 {
 	check(bRegistered);
diff --git a/Engine/Source/Runtime/Experimental/GeometryCollectionEngine/Private/GeometryCollection/GeometryCollectionComponent.cpp b/Engine/Source/Runtime/Experimental/GeometryCollectionEngine/Private/GeometryCollection/GeometryCollectionComponent.cpp
--- a/Engine/Source/Runtime/Experimental/GeometryCollectionEngine/Private/GeometryCollection/GeometryCollectionComponent.cpp
+++ b/Engine/Source/Runtime/Experimental/GeometryCollectionEngine/Private/GeometryCollection/GeometryCollectionComponent.cpp
@@ -1,5 +1,4 @@
 // Copyright Epic Games, Inc. All Rights Reserved.
-
 #include "GeometryCollection/GeometryCollectionComponent.h"
 
 #include "AI/Navigation/NavCollisionBase.h"
@@ -118,6 +117,9 @@ FAutoConsoleVariableRef CVarChaosGCInitConstantDataUseParallelFor(TEXT("p.Chaos.
 int32 bChaos_GC_InitConstantDataParallelForBatchSize = 5000;
 FAutoConsoleVariableRef CVarChaosGCInitConstantDataParallelForBatchSize(TEXT("p.Chaos.GC.InitConstantDataParallelForBatchSize"), bChaos_GC_InitConstantDataParallelForBatchSize, TEXT("When parallelFor is used in InitConstantData, defined the minimium size of a batch of vertex "));
 
+int32 MaxGeometryCollectionAsyncPhysicsTickIdleTimeMs = 30;
+FAutoConsoleVariableRef CVarMaxGeometryCollectionAsyncPhysicsTickIdleTimeMs(TEXT("p.Chaos.GC.MaxGeometryCollectionAsyncPhysicsTickIdleTimeMs"), MaxGeometryCollectionAsyncPhysicsTickIdleTimeMs, TEXT("Amount of time in milliseconds before the async tick turns off when it is otherwise not doing anything."));
+
 DEFINE_LOG_CATEGORY_STATIC(UGCC_LOG, Error, All);
 
 extern FGeometryCollectionDynamicDataPool GDynamicDataPool;
@@ -1777,13 +1779,13 @@ void UGeometryCollectionComponent::ProcessRepData()
 	ProcessRepData(0.f, 0.f);
 }
 
-void UGeometryCollectionComponent::ProcessRepData(const float DeltaTime, const float SimTime)
+bool UGeometryCollectionComponent::ProcessRepData(const float DeltaTime, const float SimTime)
 {
 	using namespace Chaos;
 
 	if(!PhysicsProxy || !PhysicsProxy->IsInitializedOnPhysicsThread() || PhysicsProxy->GetReplicationMode() != FGeometryCollectionPhysicsProxy::EReplicationMode::Client)
 	{
-		return;
+		return false;
 	}
 
 	// Track the sim time that this rep data was received on.
@@ -1799,7 +1801,7 @@ void UGeometryCollectionComponent::ProcessRepData(const float DeltaTime, const f
 	// the last received rep data
 	if (RepExtrapTime > GeometryCollectionRepMaxExtrapolationTime)
 	{
-		return;
+		return false;
 	}
 
 	// Create a little little function for applying a lambda to each
@@ -1870,7 +1872,7 @@ void UGeometryCollectionComponent::ProcessRepData(const float DeltaTime, const f
 	// between the position and the target position.
 	if (VersionProcessed == RepData.Version && !bGeometryCollectionRepUseClusterVelocityMatch)
 	{
-		return;
+		return false;
 	}
 
 	bool bHardSnap = false;
@@ -1914,6 +1916,14 @@ void UGeometryCollectionComponent::ProcessRepData(const float DeltaTime, const f
 		FPBDRigidParticleHandle* OneOff = PhysicsProxy->GetParticles()[ActivatedCluster.ActivatedIndex];
 		check(OneOff);
 
+		// If there's a parent cluster particle we need to release them first.
+		// This is generally an indication that something desynced between the client and server though...maybe something needs to be done
+		// to ensure internal clusters stay in sync.
+		if (FPBDRigidClusteredParticleHandle* ClusterParticle = OneOff->CastToClustered(); ClusterParticle && ClusterParticle->Parent())
+		{
+			RigidClustering.ReleaseClusterParticles(TArray<FPBDRigidParticleHandle*>{ ClusterParticle->Parent() }, true);
+		}
+
 		// Set initial velocities if not hard snapping
 		if (!bHardSnap)
 		{
@@ -1926,7 +1936,11 @@ void UGeometryCollectionComponent::ProcessRepData(const float DeltaTime, const f
 		RigidClustering.ReleaseClusterParticles(TArray<FPBDRigidParticleHandle*>{ OneOff }, true);
 	}
 
-	ForEachClusterPair([Solver, DeltaTime, RepExtrapTime, bHardSnap](const FGeometryCollectionClusterRep& RepCluster, Chaos::FPBDRigidParticleHandle& Cluster)
+	// Keep track of whether we did some "work" on this frame so we can turn off the async tick after
+	// multiple frames of not doing anything.
+	bool bProcessed = false;
+
+	ForEachClusterPair([Solver, DeltaTime, RepExtrapTime, bHardSnap, &bProcessed](const FGeometryCollectionClusterRep& RepCluster, Chaos::FPBDRigidParticleHandle& Cluster)
 	{
 		bool bWake = false;
 
@@ -1973,6 +1987,8 @@ void UGeometryCollectionComponent::ProcessRepData(const float DeltaTime, const f
 			}
 		}
 
+		bProcessed |= bWake;
+
 		//
 		// Wake up particle if it's sleeping and there's a delta to correct
 		//
@@ -1983,6 +1999,7 @@ void UGeometryCollectionComponent::ProcessRepData(const float DeltaTime, const f
 	});
 
 	VersionProcessed = RepData.Version;
+	return bProcessed;
 }
 
 void UGeometryCollectionComponent::SetDynamicState(const Chaos::EObjectStateType& NewDynamicState)
@@ -2571,13 +2588,27 @@ void UGeometryCollectionComponent::AsyncPhysicsTickComponent(float DeltaTime, fl
 
 	// using net mode for now as using local role seemed to cause other issues at initialization time
 	// we may nee dto to also use local role in the future if the authority is likely to change at runtime
-	if (GetNetMode() != ENetMode::NM_Client)
+	if (GetNetMode() == ENetMode::NM_Client)
 	{
-		UpdateRepData();
+		const int64 CurrentTimeInMs = FPlatformTime::ToMilliseconds64(FPlatformTime::Cycles64());
+		const bool bActive = ProcessRepData(DeltaTime, SimTime);
+		if (!bActive && LastAsyncPhysicsTickMs > 0 && CurrentTimeInMs - LastAsyncPhysicsTickMs > static_cast<int64>(MaxGeometryCollectionAsyncPhysicsTickIdleTimeMs))
+		{
+			DeferRemoveAsyncPhysicsTick();
+
+			// Reset LastAsyncPhysicsTickMs to 0 so that the next time when the async physics tick gets enabled, we
+			// allow the tick to run for at least MaxGeometryCollectionAsyncPhysicsTickIdleTimeMs. This handles the
+			// (hopefully unlikely) case where the async tick gets re-enabled but there's no work to do.
+			LastAsyncPhysicsTickMs = 0;
+		}
+		else
+		{
+			LastAsyncPhysicsTickMs = CurrentTimeInMs;
+		}
 	}
 	else
 	{
-		ProcessRepData(DeltaTime, SimTime);
+		UpdateRepData();
 	}
 }
 
@@ -2895,7 +2926,7 @@ void UGeometryCollectionComponent::RegisterAndInitializePhysicsProxy()
 	}
 #endif
 	PhysicsProxy = new FGeometryCollectionPhysicsProxy(this, *DynamicCollection, SimulationParameters, InitialSimFilter, InitialQueryFilter, CollectorGuid);
-	PhysicsProxy->SetPostPhysicsSyncCallback([this]() { UpdateAttachedChildrenTransform(); });
+	PhysicsProxy->SetPostPhysicsSyncCallback([this]() { OnPostPhysicsSync(); });
 
 	if (GetIsReplicated())
 	{
@@ -2953,7 +2984,39 @@ void UGeometryCollectionComponent::RegisterAndInitializePhysicsProxy()
 	}
 
 	RegisterForEvents();
-	SetAsyncPhysicsTickEnabled(GetIsReplicated());
+	//SetAsyncPhysicsTickEnabled(true);
+}
+
+void UGeometryCollectionComponent::OnPostPhysicsSync()
+{
+	UpdateAttachedChildrenTransform();
+
+	if (GetIsReplicated() && GetNetMode() != ENetMode::NM_Client)
+	{
+		// The GameThreadCollection dirty flag doesn't correspond to the "dirtiness" that should trigger replication.
+		// So as long as the physics sync happens, check for potential replication updates.
+		RequestUpdateRepData();
+	}
+}
+
+void UGeometryCollectionComponent::RequestUpdateRepData()
+{
+	if (FPhysScene* PhysScene = GetInnerChaosScene())
+	{
+		PhysScene->EnqueueAsyncPhysicsCommand(0, this,
+			[this]()
+			{
+				UpdateRepData();
+			}
+		);
+	}
+}
+
+void UGeometryCollectionComponent::OnRep_RepData()
+{
+	// We have new data that was replicated! Turn on the async tick to process instead of just requesting a one-off
+	// since we may want to keep processing for extra time afterwards.
+	SetAsyncPhysicsTickEnabled(true);
 }
 
 void UGeometryCollectionComponent::SetAbandonedParticleCollisionProfileName(FName CollisionProfile)
diff --git a/Engine/Source/Runtime/Experimental/GeometryCollectionEngine/Public/GeometryCollection/GeometryCollectionComponent.h b/Engine/Source/Runtime/Experimental/GeometryCollectionEngine/Public/GeometryCollection/GeometryCollectionComponent.h
--- a/Engine/Source/Runtime/Experimental/GeometryCollectionEngine/Public/GeometryCollection/GeometryCollectionComponent.h
+++ b/Engine/Source/Runtime/Experimental/GeometryCollectionEngine/Public/GeometryCollection/GeometryCollectionComponent.h
@@ -1134,10 +1134,14 @@ protected:
 	UPROPERTY(EditAnywhere, BlueprintReadWrite, Category = Network)
 	int32 ReplicationAbandonAfterLevel;
 
-	UPROPERTY(Replicated)
+	UPROPERTY(ReplicatedUsing=OnRep_RepData)
 	FGeometryCollectionRepData RepData;
 
 	/** Called post solve to allow authoritative components to update their replication data */
+	UFUNCTION()
+	void OnRep_RepData();
+
+	void RequestUpdateRepData();
 	virtual void UpdateRepData();
 
 	/** Clear all rep data, this is required if the physics proxy has been recreated */
@@ -1146,10 +1150,14 @@ protected:
 	UE_DEPRECATED(5.3, "The argument-free version of ProcessRepData will be removed. Please use the version which takes DeltaTime and SimTime instead.")
 	virtual void ProcessRepData();
 
-	virtual void ProcessRepData(float DeltaTime, float SimTime);
+	virtual bool ProcessRepData(float DeltaTime, float SimTime);
 
 	int32 VersionProcessed = INDEX_NONE;
 
+	// The last time (in milliseconds) the async physics component tick fired.
+	// We track this on the client to be able to turn off the tick for perf reasons
+	// if we spend a lot of ticks sequentially doing nothing.
+	int64 LastAsyncPhysicsTickMs = 0;
 private:
 
 	bool bRenderStateDirty;
@@ -1244,6 +1252,8 @@ private:
 
 	void LoadCollisionProfiles();
 
+	void OnPostPhysicsSync();
+
 	/** The clusters we need to replicate */
 	TUniquePtr<TSet<Chaos::FPBDRigidClusteredParticleHandle*>> ClustersToRep;
 
