remove commented line in geometry collection component

diff --git a/Engine/Source/Runtime/Experimental/GeometryCollectionEngine/Private/GeometryCollection/GeometryCollectionComponent.cpp b/Engine/Source/Runtime/Experimental/GeometryCollectionEngine/Private/GeometryCollection/GeometryCollectionComponent.cpp
--- a/Engine/Source/Runtime/Experimental/GeometryCollectionEngine/Private/GeometryCollection/GeometryCollectionComponent.cpp
+++ b/Engine/Source/Runtime/Experimental/GeometryCollectionEngine/Private/GeometryCollection/GeometryCollectionComponent.cpp
@@ -2984,7 +2984,6 @@ void UGeometryCollectionComponent::RegisterAndInitializePhysicsProxy()
 	}
 
 	RegisterForEvents();
-	//SetAsyncPhysicsTickEnabled(true);
 }
 
 void UGeometryCollectionComponent::OnPostPhysicsSync()
