geometry collection : fix potential crash when physics state is created while rest collection is not set yet
may happem when adding geometry collection component to a BP while having it open for edit in the editor

diff --git a/Engine/Source/Runtime/Experimental/GeometryCollectionEngine/Private/GeometryCollection/GeometryCollectionComponent.cpp b/Engine/Source/Runtime/Experimental/GeometryCollectionEngine/Private/GeometryCollection/GeometryCollectionComponent.cpp
--- a/Engine/Source/Runtime/Experimental/GeometryCollectionEngine/Private/GeometryCollection/GeometryCollectionComponent.cpp
+++ b/Engine/Source/Runtime/Experimental/GeometryCollectionEngine/Private/GeometryCollection/GeometryCollectionComponent.cpp
@@ -2732,7 +2732,7 @@ void UGeometryCollectionComponent::OnCreatePhysicsState()
 	// do the same thing, but through the geometry collection proxy and lambdas
 	// defined below.  FBodyInstance doesn't work for geometry collections 
 	// because FBodyInstance manages a single particle, where we have many.
-	if (!PhysicsProxy)
+	if (!PhysicsProxy && RestCollection)
 	{
 #if WITH_EDITOR && WITH_EDITORONLY_DATA
 		EditorActor = nullptr;
