Chaos - Simd collision solver shock propagation and incremental collision detection

diff --git a/Engine/Source/Runtime/Experimental/Chaos/Private/Chaos/Collision/PBDCollisionContainerSolverSimd.cpp b/Engine/Source/Runtime/Experimental/Chaos/Private/Chaos/Collision/PBDCollisionContainerSolverSimd.cpp
--- a/Engine/Source/Runtime/Experimental/Chaos/Private/Chaos/Collision/PBDCollisionContainerSolverSimd.cpp
+++ b/Engine/Source/Runtime/Experimental/Chaos/Private/Chaos/Collision/PBDCollisionContainerSolverSimd.cpp
@@ -40,8 +40,8 @@ namespace Chaos
 		// Transform and copy a single manifold point for use in the solver
 		template<int TNumLanes>
 		void UpdateSolverContactPointFromConstraint(
-			TArrayView<TPBDCollisionSolverManifoldPointsSimd<TNumLanes>>& ManifoldPointsBuffer,
-			TArrayView<TSolverBodyPtrPairSimd<4>>& SolverBodiesBuffer,
+			const TArrayView<TPBDCollisionSolverManifoldPointsSimd<TNumLanes>>& ManifoldPointsBuffer,
+			const TArrayView<TSolverBodyPtrPairSimd<4>>& SolverBodiesBuffer,
 			Private::FPBDCollisionSolverSimd& Solver,
 			const int32 SolverPointIndex, 
 			const FPBDCollisionConstraint* Constraint, 
@@ -141,8 +141,8 @@ namespace Chaos
 		// Transform and copy all of a constraint's manifold point data for use by the solver
 		template<int TNumLanes>
 		void UpdateSolverManifoldFromConstraint(
-			TArrayView<TPBDCollisionSolverManifoldPointsSimd<TNumLanes>>& ManifoldPointsBuffer,
-			TArrayView<TSolverBodyPtrPairSimd<4>>& SolverBodiesBuffer,
+			const TArrayView<TPBDCollisionSolverManifoldPointsSimd<TNumLanes>>& ManifoldPointsBuffer,
+			const TArrayView<TSolverBodyPtrPairSimd<4>>& SolverBodiesBuffer,
 			Private::FPBDCollisionSolverSimd& Solver,
 			const FPBDCollisionConstraint* Constraint, 
 			const FSolverReal Dt, 
@@ -179,12 +179,13 @@ namespace Chaos
 		// Transform and copy all constraint data for use by the solver
 		template<int TNumLanes>
 		void UpdateSolverFromConstraint(
-			TArrayView<TPBDCollisionSolverManifoldPointsSimd<TNumLanes>>& ManifoldPointsBuffer,
-			TArrayView<TSolverBodyPtrPairSimd<4>>& SolverBodiesBuffer,
+			const TArrayView<TPBDCollisionSolverManifoldPointsSimd<TNumLanes>>& ManifoldPointsBuffer,
+			const TArrayView<TSolverBodyPtrPairSimd<4>>& SolverBodiesBuffer,
 			FPBDCollisionSolverSimd& Solver,
 			const FPBDCollisionConstraint* Constraint,
 			const FSolverReal Dt, 
-			const FPBDCollisionSolverSettings& SolverSettings)
+			const FPBDCollisionSolverSettings& SolverSettings,
+			bool& bOutPerIterationCollision)
 		{
 			// Friction values. Static and Dynamic friction are applied in the position solve for most shapes.
 			// We can also run in a mode without static friction at all. This is faster but stacking is not possible.
@@ -225,6 +226,8 @@ namespace Chaos
 			Solver.SolverBody1().SetInvIScale(Constraint->GetInvInertiaScale1());
 			Solver.SolverBody1().SetShockPropagationScale(FReal(1));
 
+			bOutPerIterationCollision = (!Constraint->GetUseManifold() || Constraint->GetUseIncrementalCollisionDetection());
+
 			UpdateSolverManifoldFromConstraint(
 				ManifoldPointsBuffer,
 				SolverBodiesBuffer,
@@ -237,7 +240,7 @@ namespace Chaos
 
 		template<int TNumLanes>
 		void UpdateConstraintFromSolver(
-			TArrayView<TPBDCollisionSolverManifoldPointsSimd<TNumLanes>>& ManifoldPointsBuffer,
+			const TArrayView<TPBDCollisionSolverManifoldPointsSimd<TNumLanes>>& ManifoldPointsBuffer,
 			FPBDCollisionConstraint* Constraint,
 			const Private::FPBDCollisionSolverSimd& Solver, 
 			const FSolverReal Dt)
@@ -282,6 +285,7 @@ namespace Chaos
 			: FConstraintContainerSolver(InPriority)
 			, ConstraintContainer(InConstraintContainer)
 			, NumConstraints(0)
+			, bPerIterationCollisionDetection(false)
 			, DummySolverBody(FSolverBody::MakeInitialized())
 			, DummyConstraintSolverBody(DummySolverBody)
 		{
@@ -299,15 +303,16 @@ namespace Chaos
 			// @todo(chaos): allocation policy to reduce number of resizes as contacts increase
 			Constraints.Reset(MaxCollisions);
 			Solvers.Reset(MaxCollisions);
+			bCollisionConstraintPerIterationCollisionDetection.Reset(MaxCollisions);
 
 			// This is over-allocating by up to a factor of NumLanes!
-			Data.SimdManifoldPoints.Reset(MaxCollisions);
-			Data.SimdConstraintIndices.Reset(MaxCollisions);
+			SimdData.SimdManifoldPoints.Reset(MaxCollisions);
+			SimdData.SimdConstraintIndices.Reset(MaxCollisions);
 
 			for (int32 LaneIndex = 0; LaneIndex < GetNumLanes(); ++LaneIndex)
 			{
-				Data.SimdNumConstraints[LaneIndex] = 0;
-				Data.SimdNumManifoldPoints[LaneIndex] = 0;
+				SimdData.SimdNumConstraints[LaneIndex] = 0;
+				SimdData.SimdNumManifoldPoints[LaneIndex] = 0;
 			}
 			NumConstraints = 0;
 		}
@@ -323,20 +328,20 @@ namespace Chaos
 			// Decide what lane this island goes into: Find the lane with the least constraints in it
 			// @todo(chaos): should use manifold point count, not constraint count
 			int32 IslandLaneIndex = 0;
-			int32 IslandLaneNumConstraints = Data.SimdNumConstraints[0];
+			int32 IslandLaneNumConstraints = SimdData.SimdNumConstraints[0];
 			for (int32 LaneIndex = 1; LaneIndex < GetNumLanes(); ++LaneIndex)
 			{
-				if (Data.SimdNumConstraints[LaneIndex] < IslandLaneNumConstraints)
+				if (SimdData.SimdNumConstraints[LaneIndex] < IslandLaneNumConstraints)
 				{
 					IslandLaneIndex = LaneIndex;
-					IslandLaneNumConstraints = Data.SimdNumConstraints[LaneIndex];
+					IslandLaneNumConstraints = SimdData.SimdNumConstraints[LaneIndex];
 				}
 			}
 
 			// Make sure we have enough constraint rows for these constraints (space is pre-allocated)
-			if (Data.SimdConstraintIndices.Num() < IslandLaneNumConstraints + IslandConstraints.Num())
+			if (SimdData.SimdConstraintIndices.Num() < IslandLaneNumConstraints + IslandConstraints.Num())
 			{
-				Data.SimdConstraintIndices.SetNum(IslandLaneNumConstraints + IslandConstraints.Num());
+				SimdData.SimdConstraintIndices.SetNum(IslandLaneNumConstraints + IslandConstraints.Num());
 			}
 
 			// Add all the constraints in the island to the selected lane
@@ -349,10 +354,10 @@ namespace Chaos
 				const int32 ConstraintIndex = Constraints.Add(&Constraint);
 
 				// Add the constraint to its island's lane
-				const int32 RowIndex = Data.SimdNumConstraints[IslandLaneIndex];
-				Data.SimdConstraintIndices[RowIndex].ConstraintIndex[IslandLaneIndex] = ConstraintIndex;
+				const int32 RowIndex = SimdData.SimdNumConstraints[IslandLaneIndex];
+				SimdData.SimdConstraintIndices[RowIndex].ConstraintIndex[IslandLaneIndex] = ConstraintIndex;
 
-				++Data.SimdNumConstraints[IslandLaneIndex];
+				++SimdData.SimdNumConstraints[IslandLaneIndex];
 				++NumConstraints;
 			}
 		}
@@ -361,18 +366,19 @@ namespace Chaos
 		{
 			// Allocate the solvers
 			Solvers.SetNum(Constraints.Num());
+			bCollisionConstraintPerIterationCollisionDetection.SetNum(Constraints.Num());
 
 			// Reset manifold pointer per lane counter
 			for (int32 LaneIndex = 0; LaneIndex < GetNumLanes(); ++LaneIndex)
 			{
-				Data.SimdNumManifoldPoints[LaneIndex] = 0;
+				SimdData.SimdNumManifoldPoints[LaneIndex] = 0;
 			}
 
 			// Determine how many manifold point rows we need and tell each solver where its points are
 			int32 NumManifoldPointRows = 0;
-			for (int32 ConstraintRowIndex = 0; ConstraintRowIndex < Data.SimdConstraintIndices.Num(); ++ConstraintRowIndex)
+			for (int32 ConstraintRowIndex = 0; ConstraintRowIndex < SimdData.SimdConstraintIndices.Num(); ++ConstraintRowIndex)
 			{
-				TConstraintIndexSimd<4>& ConstraintIndices = Data.SimdConstraintIndices[ConstraintRowIndex];
+				TConstraintIndexSimd<4>& ConstraintIndices = SimdData.SimdConstraintIndices[ConstraintRowIndex];
 
 				for (int32 LaneIndex = 0; LaneIndex < GetNumLanes(); ++LaneIndex)
 				{
@@ -385,22 +391,22 @@ namespace Chaos
 						Solvers[ConstraintIndex].SetManifoldPointsBuffer(
 							ConstraintIndex,
 							LaneIndex,
-							Data.SimdNumManifoldPoints[LaneIndex],
+							SimdData.SimdNumManifoldPoints[LaneIndex],
 							NumManifoldPoints);
 
-						Data.SimdNumManifoldPoints[LaneIndex] += NumManifoldPoints;
-						NumManifoldPointRows = FMath::Max(NumManifoldPointRows, Data.SimdNumManifoldPoints[LaneIndex]);
+						SimdData.SimdNumManifoldPoints[LaneIndex] += NumManifoldPoints;
+						NumManifoldPointRows = FMath::Max(NumManifoldPointRows, SimdData.SimdNumManifoldPoints[LaneIndex]);
 					}
 				}
 			}
 
 			// Allocate the manifold point solver rows
-			Data.SimdManifoldPoints.SetNum(NumManifoldPointRows);
+			SimdData.SimdManifoldPoints.SetNum(NumManifoldPointRows);
 
 			// Initalize the set of body pointers with a dummy body.
 			// This allows us to avoid some branches in the body data gather for rows where not all lanes are used
-			Data.SimdSolverBodies.SetNum(NumManifoldPointRows);
-			for (TSolverBodyPtrPairSimd<4>& SolverBodies : Data.SimdSolverBodies)
+			SimdData.SimdSolverBodies.SetNum(NumManifoldPointRows);
+			for (TSolverBodyPtrPairSimd<4>& SolverBodies : SimdData.SimdSolverBodies)
 			{
 				SolverBodies.Body0.SetValues(&DummySolverBody);
 				SolverBodies.Body1.SetValues(&DummySolverBody);
@@ -450,14 +456,16 @@ namespace Chaos
 
 			const FSolverReal Dt = FSolverReal(InDt);
 
-			TArrayView<TPBDCollisionSolverManifoldPointsSimd<4>> ManifoldPointsBuffer = MakeArrayView(Data.SimdManifoldPoints);
-			TArrayView<TSolverBodyPtrPairSimd<4>> SolverBodiesBuffer = MakeArrayView(Data.SimdSolverBodies);
+			TArrayView<TPBDCollisionSolverManifoldPointsSimd<4>> ManifoldPointsBuffer = MakeArrayView(SimdData.SimdManifoldPoints);
+			TArrayView<TSolverBodyPtrPairSimd<4>> SolverBodiesBuffer = MakeArrayView(SimdData.SimdSolverBodies);
 
 			// @todo(chaos): would it be better to iterate over manifold point rows?
+			bool bAnyPerIterationCollisions = false;
 			for (int32 ConstraintIndex = 0; ConstraintIndex < Constraints.Num(); ++ConstraintIndex)
 			{
 				FPBDCollisionConstraint* Constraint = Constraints[ConstraintIndex];
 				FPBDCollisionSolverSimd& Solver = Solvers[ConstraintIndex];
+				bool& bPerIterationCollision = bCollisionConstraintPerIterationCollisionDetection[ConstraintIndex];
 
 				UpdateSolverFromConstraint(
 					ManifoldPointsBuffer,
@@ -465,7 +473,16 @@ namespace Chaos
 					Solver,
 					Constraint, 
 					Dt, 
-					ConstraintContainer.GetSolverSettings());
+					ConstraintContainer.GetSolverSettings(),
+					bPerIterationCollision);
+
+				bAnyPerIterationCollisions = bAnyPerIterationCollisions || bPerIterationCollision;
+			}
+
+			if (bAnyPerIterationCollisions)
+			{
+				// We should lock here? We only ever set to true or do nothing so I think it doesn't matter if this happens on multiple threads...
+				bPerIterationCollisionDetection = true;
 			}
 		}
 
@@ -493,7 +510,7 @@ namespace Chaos
 
 			const FSolverReal Dt = FSolverReal(InDt);
 
-			TArrayView<TPBDCollisionSolverManifoldPointsSimd<4>> ManifoldPointsBuffer = MakeArrayView(Data.SimdManifoldPoints);
+			TArrayView<TPBDCollisionSolverManifoldPointsSimd<4>> ManifoldPointsBuffer = MakeArrayView(SimdData.SimdManifoldPoints);
 
 			// @todo(chaos): would it be better to iterate over manifold point rows?
 			for (int32 ConstraintIndex = 0; ConstraintIndex < Constraints.Num(); ++ConstraintIndex)
@@ -528,11 +545,11 @@ namespace Chaos
 			// We run collision detection here under two conditions (normally it is run after Integration and before the constraint solver phase):
 			// 1) When deferring collision detection until the solver phase for better joint-collision behaviour (RBAN). In this case, we only do this on the first iteration.
 			// 2) When using no manifolds or incremental manifolds, where we may add/replace manifold points every iteration.
-			//const bool bRunDeferredCollisionDetection = (It == 0) && ConstraintContainer.GetDetectorSettings().bDeferNarrowPhase;
-			//if (bRunDeferredCollisionDetection || bPerIterationCollisionDetection)
-			//{
-			//	UpdateCollisions(InDt, BeginIndex, EndIndex);
-			//}
+			const bool bRunDeferredCollisionDetection = (It == 0) && ConstraintContainer.GetDetectorSettings().bDeferNarrowPhase;
+			if (bRunDeferredCollisionDetection || bPerIterationCollisionDetection)
+			{
+				UpdateCollisions(InDt);
+			}
 
 			// Only apply friction for the last few (tunable) iterations
 			// Adjust max pushout to attempt to make it iteration count independent
@@ -544,16 +561,16 @@ namespace Chaos
 			if (bApplyStaticFriction)
 			{
 				Private::FPBDCollisionSolverHelperSimd::SolvePositionWithFriction(
-					MakeArrayView(Data.SimdManifoldPoints),
-					MakeArrayView(Data.SimdSolverBodies),
+					MakeArrayView(SimdData.SimdManifoldPoints),
+					MakeArrayView(SimdData.SimdSolverBodies),
 					Dt,
 					MaxPushOut);
 			}
 			else
 			{
 				Private::FPBDCollisionSolverHelperSimd::SolvePositionNoFriction(
-					MakeArrayView(Data.SimdManifoldPoints),
-					MakeArrayView(Data.SimdSolverBodies),
+					MakeArrayView(SimdData.SimdManifoldPoints),
+					MakeArrayView(SimdData.SimdSolverBodies),
 					Dt,
 					MaxPushOut);
 			}
@@ -577,15 +594,15 @@ namespace Chaos
 			if (bApplyDynamicFriction)
 			{
 				Private::FPBDCollisionSolverHelperSimd::SolveVelocityWithFriction(
-					MakeArrayView(Data.SimdManifoldPoints),
-					MakeArrayView(Data.SimdSolverBodies),
+					MakeArrayView(SimdData.SimdManifoldPoints),
+					MakeArrayView(SimdData.SimdSolverBodies),
 					Dt);
 			}
 			else
 			{
 				Private::FPBDCollisionSolverHelperSimd::SolveVelocityNoFriction(
-					MakeArrayView(Data.SimdManifoldPoints),
-					MakeArrayView(Data.SimdSolverBodies),
+					MakeArrayView(SimdData.SimdManifoldPoints),
+					MakeArrayView(SimdData.SimdSolverBodies),
 					Dt);
 			}
 		}
@@ -598,84 +615,91 @@ namespace Chaos
 		void FPBDCollisionContainerSolverSimd::UpdatePositionShockPropagation(const FReal Dt, const int32 It, const int32 NumIts, const FPBDCollisionSolverSettings& SolverSettings)
 		{
 			// If this is the first shock propagation iteration, enable it on each solver
-			//const bool bEnableShockPropagation = (It == NumIts - SolverSettings.NumPositionShockPropagationIterations);
-			//if (bEnableShockPropagation)
-			//{
-			//	for (int32 SolverIndex = BeginIndex; SolverIndex < EndIndex; ++SolverIndex)
-			//	{
-			//		CollisionSolvers[SolverIndex].EnablePositionShockPropagation();
-			//	}
-			//}
+			const bool bEnableShockPropagation = (It == NumIts - SolverSettings.NumPositionShockPropagationIterations);
+			if (bEnableShockPropagation)
+			{
+				for (int32 SolverIndex = 0; SolverIndex < Solvers.Num(); ++SolverIndex)
+				{
+					Solvers[SolverIndex].EnablePositionShockPropagation(MakeArrayView(SimdData.SimdManifoldPoints));
+				}
+			}
 		}
 
 		void FPBDCollisionContainerSolverSimd::UpdateVelocityShockPropagation(const FReal Dt, const int32 It, const int32 NumIts, const FPBDCollisionSolverSettings& SolverSettings)
 		{
 			// Set/reset the shock propagation based on current iteration. The position solve may
 			// have left the bodies with a mass scale and we want to change or reset it.
-			//const bool bEnableShockPropagation = (It == NumIts - SolverSettings.NumVelocityShockPropagationIterations);
-			//if (bEnableShockPropagation)
-			//{
-			//	for (int32 SolverIndex = BeginIndex; SolverIndex < EndIndex; ++SolverIndex)
-			//	{
-			//		CollisionSolvers[SolverIndex].EnableVelocityShockPropagation();
-			//	}
-			//}
-			//else if (It == 0)
-			//{
-			//	for (int32 SolverIndex = BeginIndex; SolverIndex < EndIndex; ++SolverIndex)
-			//	{
-			//		CollisionSolvers[SolverIndex].DisableShockPropagation();
-			//	}
-			//}
+			const bool bEnableShockPropagation = (It == NumIts - SolverSettings.NumVelocityShockPropagationIterations);
+			if (bEnableShockPropagation)
+			{
+				for (int32 SolverIndex = 0; SolverIndex < Solvers.Num(); ++SolverIndex)
+				{
+					Solvers[SolverIndex].EnableVelocityShockPropagation(MakeArrayView(SimdData.SimdManifoldPoints));
+				}
+			}
+			else if (It == 0)
+			{
+				for (int32 SolverIndex = 0; SolverIndex < Solvers.Num(); ++SolverIndex)
+				{
+					Solvers[SolverIndex].DisableShockPropagation(MakeArrayView(SimdData.SimdManifoldPoints));
+				}
+			}
 		}
 
 		void FPBDCollisionContainerSolverSimd::UpdateCollisions(const FReal InDt)
 		{
-			//const FSolverReal Dt = FSolverReal(InDt);
-			//const bool bDeferredCollisionDetection = ConstraintContainer.GetDetectorSettings().bDeferNarrowPhase;
-
-			//bool bNeedsAnotherIteration = false;
-			//for (int32 SolverIndex = BeginIndex; SolverIndex < EndIndex; ++SolverIndex)
-			//{
-			//	if (bDeferredCollisionDetection || bCollisionConstraintPerIterationCollisionDetection[SolverIndex])
-			//	{
-			//		Private::FPBDCollisionSolverSimd& CollisionSolver = CollisionSolvers[SolverIndex];
-			//		FPBDCollisionConstraint* Constraint = CollisionConstraints[SolverIndex];
-
-			//		// Run collision detection at the current transforms including any correction from previous iterations
-			//		const FSolverBody& Body0 = CollisionSolver.SolverBody0().SolverBody();
-			//		const FSolverBody& Body1 = CollisionSolver.SolverBody1().SolverBody();
-			//		const FRigidTransform3 CorrectedActorWorldTransform0 = FRigidTransform3(Body0.CorrectedActorP(), Body0.CorrectedActorQ());
-			//		const FRigidTransform3 CorrectedActorWorldTransform1 = FRigidTransform3(Body1.CorrectedActorP(), Body1.CorrectedActorQ());
-			//		const FRigidTransform3 CorrectedShapeWorldTransform0 = Constraint->GetShapeRelativeTransform0() * CorrectedActorWorldTransform0;
-			//		const FRigidTransform3 CorrectedShapeWorldTransform1 = Constraint->GetShapeRelativeTransform1() * CorrectedActorWorldTransform1;
-
-			//		// @todo(chaos): this is ugly - pass these to the required functions instead and remove from the constraint class
-			//		// This is now only needed for LevelSet collision (see UpdateLevelsetLevelsetConstraint)
-			//		Constraint->SetSolverBodies(&Body0, &Body1);
-
-			//		// Reset the manifold if we are not using manifolds (we just use the first manifold point)
-			//		if (!Constraint->GetUseManifold())
-			//		{
-			//			Constraint->ResetActiveManifoldContacts();
-			//			CollisionSolver.ResetManifold();
-			//		}
-
-			//		// We need to know how many points were added to the manifold
-			//		const int32 BeginPointIndex = Constraint->NumManifoldPoints();
-
-			//		// NOTE: We deliberately have not updated the ShapwWorldTranforms on the constraint. If we did that, we would calculate 
-			//		// errors incorrectly in UpdateManifoldPoints, because the solver assumes nothing has been moved as we iterate (we accumulate 
-			//		// corrections that will be applied later.)
-			//		Constraint->ResetPhi(Constraint->GetCullDistance());
-			//		Collisions::UpdateConstraint(*Constraint, CorrectedShapeWorldTransform0, CorrectedShapeWorldTransform1, Dt);
-
-			//		// Update the manifold based on the new or updated contacts
-			//		UpdateCollisionSolverManifoldFromConstraint(CollisionSolver, Constraint, Dt, BeginPointIndex, Constraint->NumManifoldPoints());
-
-			//		Constraint->SetSolverBodies(nullptr, nullptr);
-			//	}
-			//}
+			const FSolverReal Dt = FSolverReal(InDt);
+			const bool bDeferredCollisionDetection = ConstraintContainer.GetDetectorSettings().bDeferNarrowPhase;
+
+			bool bNeedsAnotherIteration = false;
+			for (int32 SolverIndex = 0; SolverIndex < Solvers.Num(); ++SolverIndex)
+			{
+				if (bDeferredCollisionDetection || bCollisionConstraintPerIterationCollisionDetection[SolverIndex])
+				{
+					Private::FPBDCollisionSolverSimd& CollisionSolver = Solvers[SolverIndex];
+					FPBDCollisionConstraint* Constraint = Constraints[SolverIndex];
+
+					// Run collision detection at the current transforms including any correction from previous iterations
+					const FSolverBody& Body0 = CollisionSolver.SolverBody0().SolverBody();
+					const FSolverBody& Body1 = CollisionSolver.SolverBody1().SolverBody();
+					const FRigidTransform3 CorrectedActorWorldTransform0 = FRigidTransform3(Body0.CorrectedActorP(), Body0.CorrectedActorQ());
+					const FRigidTransform3 CorrectedActorWorldTransform1 = FRigidTransform3(Body1.CorrectedActorP(), Body1.CorrectedActorQ());
+					const FRigidTransform3 CorrectedShapeWorldTransform0 = Constraint->GetShapeRelativeTransform0() * CorrectedActorWorldTransform0;
+					const FRigidTransform3 CorrectedShapeWorldTransform1 = Constraint->GetShapeRelativeTransform1() * CorrectedActorWorldTransform1;
+
+					// @todo(chaos): this is ugly - pass these to the required functions instead and remove from the constraint class
+					// This is now only needed for LevelSet collision (see UpdateLevelsetLevelsetConstraint)
+					Constraint->SetSolverBodies(&Body0, &Body1);
+
+					// Reset the manifold if we are not using manifolds (we just use the first manifold point)
+					if (!Constraint->GetUseManifold())
+					{
+						Constraint->ResetActiveManifoldContacts();
+						CollisionSolver.ResetManifold();
+					}
+
+					// We need to know how many points were added to the manifold
+					const int32 BeginPointIndex = Constraint->NumManifoldPoints();
+
+					// NOTE: We deliberately have not updated the ShapwWorldTranforms on the constraint. If we did that, we would calculate 
+					// errors incorrectly in UpdateManifoldPoints, because the solver assumes nothing has been moved as we iterate (we accumulate 
+					// corrections that will be applied later.)
+					Constraint->ResetPhi(Constraint->GetCullDistance());
+					Collisions::UpdateConstraint(*Constraint, CorrectedShapeWorldTransform0, CorrectedShapeWorldTransform1, Dt);
+
+					// Update the manifold based on the new or updated contacts
+					UpdateSolverManifoldFromConstraint(
+						MakeArrayView(SimdData.SimdManifoldPoints),
+						MakeArrayView(SimdData.SimdSolverBodies),
+						CollisionSolver, 
+						Constraint, 
+						Dt, 
+						BeginPointIndex, 
+						Constraint->NumManifoldPoints());
+
+					Constraint->SetSolverBodies(nullptr, nullptr);
+				}
+			}
 		}
 
 	}	// namespace Private
diff --git a/Engine/Source/Runtime/Experimental/Chaos/Private/Chaos/Collision/PBDCollisionSolverSimd.cpp b/Engine/Source/Runtime/Experimental/Chaos/Private/Chaos/Collision/PBDCollisionSolverSimd.cpp
--- a/Engine/Source/Runtime/Experimental/Chaos/Private/Chaos/Collision/PBDCollisionSolverSimd.cpp
+++ b/Engine/Source/Runtime/Experimental/Chaos/Private/Chaos/Collision/PBDCollisionSolverSimd.cpp
@@ -18,76 +18,12 @@ namespace Chaos
 {
 	namespace CVars
 	{
-		extern float Chaos_PBDCollisionSolver_Position_MinInvMassScale;
-		extern float Chaos_PBDCollisionSolver_Velocity_MinInvMassScale;
 		extern bool bChaos_PBDCollisionSolver_Velocity_FrictionEnabled;
 		extern float Chaos_PBDCollisionSolver_Position_StaticFrictionStiffness;
 	}
-	using namespace CVars;
 
 	namespace Private
 	{
-		//////////////////////////////////////////////////////////////////////////////////////////////////
-		//////////////////////////////////////////////////////////////////////////////////////////////////
-		//////////////////////////////////////////////////////////////////////////////////////////////////
-		//////////////////////////////////////////////////////////////////////////////////////////////////
-
-
-		void FPBDCollisionSolverSimd::EnablePositionShockPropagation()
-		{
-			SetShockPropagationInvMassScale(Chaos_PBDCollisionSolver_Position_MinInvMassScale);
-		}
-
-		void FPBDCollisionSolverSimd::EnableVelocityShockPropagation()
-		{
-			SetShockPropagationInvMassScale(Chaos_PBDCollisionSolver_Velocity_MinInvMassScale);
-		}
-
-		void FPBDCollisionSolverSimd::DisableShockPropagation()
-		{
-			SetShockPropagationInvMassScale(FReal(1));
-		}
-
-		void FPBDCollisionSolverSimd::SetShockPropagationInvMassScale(const FSolverReal InvMassScale)
-		{
-			FConstraintSolverBody& Body0 = SolverBody0();
-			FConstraintSolverBody& Body1 = SolverBody1();
-
-			// Shock propagation decreases the inverse mass of bodies that are lower in the pile
-			// of objects. This significantly improves stability of heaps and stacks. Height in the pile is indictaed by the "level". 
-			// No need to set an inverse mass scale if the other body is kinematic (with inv mass of 0).
-			// Bodies at the same level do not take part in shock propagation.
-			if (Body0.IsDynamic() && Body1.IsDynamic() && (Body0.Level() != Body1.Level()))
-			{
-				// Set the inv mass scale of the "lower" body to make it heavier
-				bool bInvMassUpdated = false;
-				if (Body0.Level() < Body1.Level())
-				{
-					if (Body0.ShockPropagationScale() != InvMassScale)
-					{
-						Body0.SetShockPropagationScale(InvMassScale);
-						bInvMassUpdated = true;
-					}
-				}
-				else
-				{
-					if (Body1.ShockPropagationScale() != InvMassScale)
-					{
-						Body1.SetShockPropagationScale(InvMassScale);
-						bInvMassUpdated = true;
-					}
-				}
-
-				// If the masses changed, we need to rebuild the contact mass for each manifold point
-				//if (bInvMassUpdated)
-				//{
-				//	for (int32 PointIndex = 0; PointIndex < NumManifoldPoints(); ++PointIndex)
-				//	{
-				//		State.ManifoldPoints[PointIndex].UpdateMassNormal(Body0, Body1);
-				//	}
-				//}
-			}
-		}
 
 		//////////////////////////////////////////////////////////////////////////////////////////////////
 		//////////////////////////////////////////////////////////////////////////////////////////////////
diff --git a/Engine/Source/Runtime/Experimental/Chaos/Public/Chaos/Collision/PBDCollisionContainerSolverSimd.h b/Engine/Source/Runtime/Experimental/Chaos/Public/Chaos/Collision/PBDCollisionContainerSolverSimd.h
--- a/Engine/Source/Runtime/Experimental/Chaos/Public/Chaos/Collision/PBDCollisionContainerSolverSimd.h
+++ b/Engine/Source/Runtime/Experimental/Chaos/Public/Chaos/Collision/PBDCollisionContainerSolverSimd.h
@@ -87,7 +87,7 @@ namespace Chaos
 
 			// For testing
 			const FPBDCollisionSolverSimd& GetConstraintSolver(const int32 ConstraintIndex) const { return Solvers[ConstraintIndex]; }
-			const TPBDCollisionSolverManifoldPointsSimd<4>& GetManifoldPointSolver(const int32 RowIndex) const { return Data.SimdManifoldPoints[RowIndex]; }
+			const TPBDCollisionSolverManifoldPointsSimd<4>& GetManifoldPointSolver(const int32 RowIndex) const { return SimdData.SimdManifoldPoints[RowIndex]; }
 
 		private:
 			int32 GetNumLanes() const { return 4; }
@@ -100,8 +100,10 @@ namespace Chaos
 
 			TArray<FPBDCollisionConstraint*> Constraints;
 			TArray<FPBDCollisionSolverSimd> Solvers;
-			FDataSimd<4> Data;
+			FDataSimd<4> SimdData;
+			TArray<bool> bCollisionConstraintPerIterationCollisionDetection;
 			int32 NumConstraints;
+			bool bPerIterationCollisionDetection;
 
 			FSolverBody DummySolverBody;
 			FConstraintSolverBody DummyConstraintSolverBody;
diff --git a/Engine/Source/Runtime/Experimental/Chaos/Public/Chaos/Collision/PBDCollisionSolver.h b/Engine/Source/Runtime/Experimental/Chaos/Public/Chaos/Collision/PBDCollisionSolver.h
--- a/Engine/Source/Runtime/Experimental/Chaos/Public/Chaos/Collision/PBDCollisionSolver.h
+++ b/Engine/Source/Runtime/Experimental/Chaos/Public/Chaos/Collision/PBDCollisionSolver.h
@@ -6,12 +6,6 @@
 #include "Chaos/Defines.h"
 #include "Chaos/Evolution/SolverBody.h"
 
-// Set to 0 to use a linearized error calculation, and set to 1 to use a non-linear error calculation in collision detection. 
-// In principle nonlinear is more accurate when large rotation corrections occur, but this is not too important for collisions because 
-// when the bodies settle the corrections are small. The linearized version is significantly faster than the non-linear version because 
-// the non-linear version requires a quaternion multiply and renormalization whereas the linear version is just a cross product.
-#define CHAOS_NONLINEAR_COLLISIONS_ENABLED 0
-
 namespace Chaos
 {
 	class FManifoldPoint;
diff --git a/Engine/Source/Runtime/Experimental/Chaos/Public/Chaos/Collision/PBDCollisionSolverJacobi.h b/Engine/Source/Runtime/Experimental/Chaos/Public/Chaos/Collision/PBDCollisionSolverJacobi.h
--- a/Engine/Source/Runtime/Experimental/Chaos/Public/Chaos/Collision/PBDCollisionSolverJacobi.h
+++ b/Engine/Source/Runtime/Experimental/Chaos/Public/Chaos/Collision/PBDCollisionSolverJacobi.h
@@ -7,12 +7,6 @@
 #include "Chaos/Evolution/SolverBody.h"
 #include "Chaos/Framework/UncheckedArray.h"
 
-// Set to 0 to use a linearized error calculation, and set to 1 to use a non-linear error calculation in collision detection. 
-// In principle nonlinear is more accurate when large rotation corrections occur, but this is not too important for collisions because 
-// when the bodies settle the corrections are small. The linearized version is significantly faster than the non-linear version because 
-// the non-linear version requires a quaternion multiply and renormalization whereas the linear version is just a cross product.
-#define CHAOS_NONLINEAR_COLLISIONS_ENABLED 0
-
 namespace Chaos
 {
 	class FManifoldPoint;
diff --git a/Engine/Source/Runtime/Experimental/Chaos/Public/Chaos/Collision/PBDCollisionSolverSimd.h b/Engine/Source/Runtime/Experimental/Chaos/Public/Chaos/Collision/PBDCollisionSolverSimd.h
--- a/Engine/Source/Runtime/Experimental/Chaos/Public/Chaos/Collision/PBDCollisionSolverSimd.h
+++ b/Engine/Source/Runtime/Experimental/Chaos/Public/Chaos/Collision/PBDCollisionSolverSimd.h
@@ -7,17 +7,17 @@
 #include "Chaos/Evolution/SolverBody.h"
 #include "Chaos/Simd4.h"
 
-// Set to 0 to use a linearized error calculation, and set to 1 to use a non-linear error calculation in collision detection. 
-// In principle nonlinear is more accurate when large rotation corrections occur, but this is not too important for collisions because 
-// when the bodies settle the corrections are small. The linearized version is significantly faster than the non-linear version because 
-// the non-linear version requires a quaternion multiply and renormalization whereas the linear version is just a cross product.
-#define CHAOS_NONLINEAR_COLLISIONS_ENABLED 0
-
 namespace Chaos
 {
 	class FManifoldPoint;
 	class FPBDCollisionConstraint;
 
+	namespace CVars
+	{
+		extern float Chaos_PBDCollisionSolver_Position_MinInvMassScale;
+		extern float Chaos_PBDCollisionSolver_Velocity_MinInvMassScale;
+	}
+
 	namespace Private
 	{
 		class FPBDCollisionSolverSimd;
@@ -204,9 +204,46 @@ namespace Chaos
 				FConstraintSolverBody& Body0,
 				FConstraintSolverBody& Body1)
 			{
+				FSolverReal ContactMassInvNormal = FSolverReal(0);
+
+				SimdInvM0.SetValue(LaneIndex, 0);
+				SimdInvM1.SetValue(LaneIndex, 0);
+				SimdContactNormalAngular0.SetValue(LaneIndex, FVec3f(0));
+				SimdContactNormalAngular1.SetValue(LaneIndex, FVec3f(0));
+
+				const FVec3f ContactNormal = SimdContactNormal.GetValue(LaneIndex);
+
+				if (Body0.IsDynamic())
+				{
+					const FVec3f RelativeContactPoint0 = SimdRelativeContactPoint0.GetValue(LaneIndex);
+					const FSolverVec3 R0xN = FSolverVec3::CrossProduct(RelativeContactPoint0, ContactNormal);
+					const FSolverMatrix33 InvI0 = Body0.InvI();
+					const FSolverVec3 IR0xN = InvI0 * R0xN;
+
+					SimdInvM0.SetValue(LaneIndex, Body0.InvM());
+
+					SimdContactNormalAngular0.SetValue(LaneIndex, IR0xN);
+
+					ContactMassInvNormal += FSolverVec3::DotProduct(R0xN, IR0xN) + Body0.InvM();
+				}
+				if (Body1.IsDynamic())
+				{
+					const FVec3f RelativeContactPoint1 = SimdRelativeContactPoint1.GetValue(LaneIndex);
+					const FSolverVec3 R1xN = FSolverVec3::CrossProduct(RelativeContactPoint1, ContactNormal);
+					const FSolverMatrix33 InvI1 = Body1.InvI();
+					const FSolverVec3 IR1xN = InvI1 * R1xN;
+
+					SimdInvM1.SetValue(LaneIndex, Body1.InvM());
+
+					SimdContactNormalAngular1.SetValue(LaneIndex, IR1xN);
+
+					ContactMassInvNormal += FSolverVec3::DotProduct(R1xN, IR1xN) + Body1.InvM();
+				}
+
+				SimdContactMassNormal.SetValue(LaneIndex, (ContactMassInvNormal > FSolverReal(UE_SMALL_NUMBER)) ? FSolverReal(1) / ContactMassInvNormal : FSolverReal(0));
 			}
 
-			FORCEINLINE void GatherBodyCorrections(
+			FORCEINLINE void GatherBodyPositionCorrections(
 				const FSimdSelector& IsValid,
 				const FSimdSolverBodyPtr& Body0,
 				const FSimdSolverBodyPtr& Body1,
@@ -229,7 +266,7 @@ namespace Chaos
 				}
 			}
 
-			FORCEINLINE void ScatterBodyCorrections(
+			FORCEINLINE void ScatterBodyPositionCorrections(
 				const FSimdSelector& IsValid,
 				const FSimdVec3f& DP0,
 				const FSimdVec3f& DQ0,
@@ -252,6 +289,53 @@ namespace Chaos
 				}
 			}
 
+
+			FORCEINLINE void GatherBodyVelocities(
+				const FSimdSelector& IsValid,
+				const FSimdSolverBodyPtr& Body0,
+				const FSimdSolverBodyPtr& Body1,
+				FSimdVec3f& V0,
+				FSimdVec3f& W0,
+				FSimdVec3f& V1,
+				FSimdVec3f& W1)
+			{
+				for (int32 LaneIndex = 0; LaneIndex < TNumLanes; ++LaneIndex)
+				{
+					if (IsValid.GetValue(LaneIndex))
+					{
+						const FSolverBody* LaneBody0 = Body0.GetValue(LaneIndex);
+						const FSolverBody* LaneBody1 = Body1.GetValue(LaneIndex);
+						V0.SetValue(LaneIndex, LaneBody0->V());
+						W0.SetValue(LaneIndex, LaneBody0->W());
+						V1.SetValue(LaneIndex, LaneBody1->V());
+						W1.SetValue(LaneIndex, LaneBody1->W());
+					}
+				}
+			}
+
+			FORCEINLINE void ScatterBodyVelocities(
+				const FSimdSelector& IsValid,
+				const FSimdVec3f& V0,
+				const FSimdVec3f& W0,
+				const FSimdVec3f& V1,
+				const FSimdVec3f& W1,
+				const FSimdSolverBodyPtr& Body0,
+				const FSimdSolverBodyPtr& Body1)
+			{
+				for (int32 LaneIndex = 0; LaneIndex < TNumLanes; ++LaneIndex)
+				{
+					if (IsValid.GetValue(LaneIndex))
+					{
+						FSolverBody* LaneBody0 = Body0.GetValue(LaneIndex);
+						FSolverBody* LaneBody1 = Body1.GetValue(LaneIndex);
+						LaneBody0->SetV(V0.GetValue(LaneIndex));
+						LaneBody0->SetW(W0.GetValue(LaneIndex));
+						LaneBody1->SetV(V1.GetValue(LaneIndex));
+						LaneBody1->SetW(W1.GetValue(LaneIndex));
+					}
+				}
+			}
+
 			void SolvePositionNoFriction(
 				const FSimdSolverBodyPtr& Body0,
 				const FSimdSolverBodyPtr& Body1,
@@ -263,7 +347,7 @@ namespace Chaos
 				// Get the current corrections for each body. 
 				// NOTE: This is a gather operation
 				FSimdVec3f DP0, DQ0, DP1, DQ1;
-				GatherBodyCorrections(IsValid, Body0, Body1, DP0, DQ0, DP1, DQ1);
+				GatherBodyPositionCorrections(IsValid, Body0, Body1, DP0, DQ0, DP1, DQ1);
 
 				// Calculate the contact error
 				const FSimdVec3f DQ0xR0 = SimdCrossProduct(DQ0, SimdRelativeContactPoint0);
@@ -317,7 +401,7 @@ namespace Chaos
 
 				// Update the corrections on the bodies. 
 				// NOTE: This is a scatter operation
-				ScatterBodyCorrections(IsValid, DP0, DQ0, DP1, DQ1, Body0, Body1);
+				ScatterBodyPositionCorrections(IsValid, DP0, DQ0, DP1, DQ1, Body0, Body1);
 			}
 
 			void SolvePositionWithFriction(
@@ -332,7 +416,7 @@ namespace Chaos
 				// Get the current corrections for each body. 
 				// NOTE: This is a gather operation
 				FSimdVec3f DP0, DQ0, DP1, DQ1;
-				GatherBodyCorrections(IsValid, Body0, Body1, DP0, DQ0, DP1, DQ1);
+				GatherBodyPositionCorrections(IsValid, Body0, Body1, DP0, DQ0, DP1, DQ1);
 
 				// Calculate the contact error
 				const FSimdVec3f DQ0xR0 = SimdCrossProduct(DQ0, SimdRelativeContactPoint0);
@@ -379,18 +463,21 @@ namespace Chaos
 
 					// Should we apply tangential corrections for friction? 
 					// bUpdateFriction = ((NetPushOutNormal > 0) || (NetPushOutTangentU != 0) || (NetPushOutTangentV != 0))
-					const FSimdSelector HasNormalPushout = SimdGreater(SimdNetPushOutNormal, FSimdRealf::Zero());
+					const FSimdSelector HasFriction = SimdAnd(IsValid, SimdGreater(SimdStaticFriction, FSimdRealf::Zero()));
+					const FSimdSelector HasNormalPushout = SimdAnd(HasFriction, SimdGreater(SimdNetPushOutNormal, FSimdRealf::Zero()));
 					const FSimdSelector HasTangentUPushout = SimdNotEqual(SimdNetPushOutTangentU, FSimdRealf::Zero());
 					const FSimdSelector HasTangentVPushout = SimdNotEqual(SimdNetPushOutTangentV, FSimdRealf::Zero());
-					//const FSimdSelector UpdateFriction = SimdAnd(IsValid, SimdOr(HasNormalPushout, SimdOr(HasTangentUPushout, HasTangentVPushout)));
+					//const FSimdSelector UpdateFriction = SimdAnd(HasFriction, SimdOr(HasNormalPushout, SimdOr(HasTangentUPushout, HasTangentVPushout)));
 					//if (SimdAnyTrue(UpdateFriction))
 					{
 						// Calculate tangential errors
 						const FSimdRealf ContactErrorTangentU = SimdAdd(SimdContactDeltaTangentU, SimdDotProduct(ContactDelta, SimdContactTangentU));
 						const FSimdRealf ContactErrorTangentV = SimdAdd(SimdContactDeltaTangentV, SimdDotProduct(ContactDelta, SimdContactTangentV));
 
-						// Calculate tangential correction
+						// A stiffness multiplier on the impulses
 						const FSimdRealf FrictionStiffness = SimdMultiply(FrictionStiffnessScale, SimdStiffness);
+
+						// Calculate tangential correction
 						FSimdRealf PushOutTangentU = SimdNegate(SimdMultiply(FrictionStiffness, SimdMultiply(ContactErrorTangentU, SimdContactMassTangentU)));
 						FSimdRealf PushOutTangentV = SimdNegate(SimdMultiply(FrictionStiffness, SimdMultiply(ContactErrorTangentV, SimdContactMassTangentV)));
 
@@ -411,7 +498,7 @@ namespace Chaos
 
 						// Apply cone limits to tangential pushouts on lanes that exceed the limit
 						// NOTE: if HasNormalPushout is false in any lane, we have already zeroed the net pushout and don't need to do anything here
-						const FSimdSelector ApplyFrictionCone = SimdAnd(IsValid, SimdAnd(HasNormalPushout, SimdOr(HasTangentUPushout, HasTangentVPushout)));
+						const FSimdSelector ApplyFrictionCone = SimdAnd(HasFriction, SimdAnd(HasNormalPushout, SimdOr(HasTangentUPushout, HasTangentVPushout)));
 						//if (SimdAnyTrue(ApplyFrictionCone))
 						{
 							const FSimdRealf MaxPushOutTangent = SimdNetPushOutNormal;
@@ -457,7 +544,7 @@ namespace Chaos
 
 					// Update the corrections on the bodies. 
 					// NOTE: This is a scatter operation
-					ScatterBodyCorrections(IsValid, DP0, DQ0, DP1, DQ1, Body0, Body1);
+					ScatterBodyPositionCorrections(IsValid, DP0, DQ0, DP1, DQ1, Body0, Body1);
 				}
 			}
 
@@ -479,18 +566,7 @@ namespace Chaos
 
 				// Gather the body data we need
 				FSimdVec3f V0, W0, V1, W1;
-				for (int32 LaneIndex = 0; LaneIndex < TNumLanes; ++LaneIndex)
-				{
-					if (IsValid.GetValue(LaneIndex))
-					{
-						const FSolverBody* LaneBody0 = Body0.GetValue(LaneIndex);
-						const FSolverBody* LaneBody1 = Body1.GetValue(LaneIndex);
-						V0.SetValue(LaneIndex, LaneBody0->V());
-						W0.SetValue(LaneIndex, LaneBody0->W());
-						V1.SetValue(LaneIndex, LaneBody1->V());
-						W1.SetValue(LaneIndex, LaneBody1->W());
-					}
-				}
+				GatherBodyVelocities(IsValid, Body0, Body1, V0, W0, V1, W1);
 
 				// Calculate the velocity error we need to correct
 				const FSimdVec3f ContactVelocity0 = SimdAdd(V0, SimdCrossProduct(W0, SimdRelativeContactPoint0));
@@ -524,18 +600,97 @@ namespace Chaos
 				W0 = SimdAdd(W0, DW0);
 				V1 = SimdSubtract(V1, DV1);
 				W1 = SimdSubtract(W1, DW1);
-				for (int32 LaneIndex = 0; LaneIndex < TNumLanes; ++LaneIndex)
+
+				ScatterBodyVelocities(IsValid, V0, W0, V1, W1, Body0, Body1);
+			}
+
+			void SolveVelocityWithFrictionImpl(
+				const FSimdSolverBodyPtr& Body0,
+				const FSimdSolverBodyPtr& Body1,
+				const FSimdRealf& Dt,
+				const FSimd4Realf& FrictionStiffnessScale)
+			{
+				// Which lanes a have a valid constraint setup?
+				const FSimdSelector IsValid = SimdGreaterEqual(SimdConstraintIndex, FSimdInt32::Zero());
+
+				// Only lanes that applied a position correction or that started with an overlap require a velocity correction
+				FSimdSelector ShouldSolveVelocity = SimdOr(SimdGreater(SimdNetPushOutNormal, FSimdRealf::Zero()), SimdLess(SimdContactDeltaNormal, FSimdRealf::Zero()));
+				ShouldSolveVelocity = SimdAnd(IsValid, ShouldSolveVelocity);
+				if (!SimdAnyTrue(ShouldSolveVelocity))
 				{
-					if (IsValid.GetValue(LaneIndex))
-					{
-						FSolverBody* LaneBody0 = Body0.GetValue(LaneIndex);
-						FSolverBody* LaneBody1 = Body1.GetValue(LaneIndex);
-						LaneBody0->SetV(V0.GetValue(LaneIndex));
-						LaneBody0->SetW(W0.GetValue(LaneIndex));
-						LaneBody1->SetV(V1.GetValue(LaneIndex));
-						LaneBody1->SetW(W1.GetValue(LaneIndex));
-					}
+					return;
 				}
+
+				// Gather the body data we need
+				FSimdVec3f V0, W0, V1, W1;
+				GatherBodyVelocities(IsValid, Body0, Body1, V0, W0, V1, W1);
+
+				// Calculate the velocity error we need to correct
+				const FSimdVec3f ContactVelocity0 = SimdAdd(V0, SimdCrossProduct(W0, SimdRelativeContactPoint0));
+				const FSimdVec3f ContactVelocity1 = SimdAdd(V1, SimdCrossProduct(W1, SimdRelativeContactPoint1));
+				const FSimdVec3f ContactVelocity = SimdSubtract(ContactVelocity0, ContactVelocity1);
+				FSimdRealf ContactVelocityErrorNormal = SimdSubtract(SimdDotProduct(ContactVelocity, SimdContactNormal), SimdContactTargetVelocityNormal);
+				FSimdRealf ContactVelocityErrorTangentU = SimdDotProduct(ContactVelocity, SimdContactTangentU);
+				FSimdRealf ContactVelocityErrorTangentV = SimdDotProduct(ContactVelocity, SimdContactTangentV);
+
+				// Zero out error for lanes which should not generate friction
+				ContactVelocityErrorNormal = SimdSelect(ShouldSolveVelocity, ContactVelocityErrorNormal, FSimdRealf::Zero());
+				ContactVelocityErrorTangentU = SimdSelect(ShouldSolveVelocity, ContactVelocityErrorTangentU, FSimdRealf::Zero());
+				ContactVelocityErrorTangentV = SimdSelect(ShouldSolveVelocity, ContactVelocityErrorTangentV, FSimdRealf::Zero());
+
+				// A stiffness multiplier on the impulses (zeroed if we have no friction in a lane)
+				const FSimdRealf FrictionStiffness = SimdMultiply(FrictionStiffnessScale, SimdStiffness);
+
+				// Calculate the impulses
+				FSimdRealf ImpulseNormal = SimdNegate(SimdMultiply(SimdMultiply(SimdStiffness, SimdContactMassNormal), ContactVelocityErrorNormal));
+				FSimdRealf ImpulseTangentU = SimdNegate(SimdMultiply(SimdMultiply(FrictionStiffness, SimdContactMassTangentU), ContactVelocityErrorTangentU));
+				FSimdRealf ImpulseTangentV = SimdNegate(SimdMultiply(SimdMultiply(FrictionStiffness, SimdContactMassTangentV), ContactVelocityErrorTangentV));
+
+				// Zero out friction for lanes with no friction
+				const FSimdSelector HasFriction = SimdAnd(IsValid, SimdGreater(SimdVelocityFriction, FSimdRealf::Zero()));
+				ImpulseTangentU = SimdSelect(HasFriction, ImpulseTangentU, FSimdRealf::Zero());
+				ImpulseTangentV = SimdSelect(HasFriction, ImpulseTangentV, FSimdRealf::Zero());
+
+				// The minimum normal impulse we can apply. We are allowed to apply a negative impulse 
+				// up to an amount that would conteract the implciit velocity applied by the pushout
+				// MinImpulseNormal = FMath::Min(0, -NetPushOutNormal / Dt)
+				// if (NetImpulseNormal + ImpulseNormal < MinImpulseNormal) {...}
+				const FSimdRealf MinImpulseNormal = SimdMin(SimdNegate(SimdDivide(SimdNetPushOutNormal, Dt)), FSimdRealf::Zero());
+				const FSimdSelector ShouldClampImpulse = SimdLess(SimdAdd(SimdNetImpulseNormal, ImpulseNormal), MinImpulseNormal);
+				const FSimdRealf ClampedImpulseNormal = SimdSubtract(MinImpulseNormal, SimdNetImpulseNormal);
+				ImpulseNormal = SimdSelect(ShouldClampImpulse, ClampedImpulseNormal, ImpulseNormal);
+
+				// Calculate the impulse multipler if clamped to the dynamic friction cone
+				const FSimdRealf MaxNetImpulseAndPushOutTangent = SimdAdd(SimdNetImpulseNormal, SimdAdd(ImpulseNormal, SimdDivide(SimdNetPushOutNormal, Dt)));
+				const FSimdRealf MaxImpulseTangent = SimdMax(FSimdRealf::Zero(), SimdMultiply(SimdVelocityFriction, MaxNetImpulseAndPushOutTangent));
+				const FSimdRealf MaxImpulseTangentSq = SimdSquare(MaxImpulseTangent);
+				const FSimdRealf ImpulseTangentSq = SimdAdd(SimdSquare(ImpulseTangentU), SimdSquare(ImpulseTangentV));
+				const FSimdRealf ImpulseTangentScale = SimdMultiply(MaxImpulseTangent, SimdInvSqrt(ImpulseTangentSq));
+
+				// Apply the multiplier to lanes that exceeded the friction limit
+				const FSimdSelector ExceededFrictionCone = SimdGreater(ImpulseTangentSq, SimdAdd(MaxImpulseTangentSq, FSimdRealf::Make(UE_SMALL_NUMBER)));
+				ImpulseTangentU = SimdSelect(ExceededFrictionCone, SimdMultiply(ImpulseTangentScale, ImpulseTangentU), ImpulseTangentU);
+				ImpulseTangentV = SimdSelect(ExceededFrictionCone, SimdMultiply(ImpulseTangentScale, ImpulseTangentV), ImpulseTangentV);
+
+				SimdNetImpulseNormal = SimdAdd(SimdNetImpulseNormal, ImpulseNormal);
+				SimdNetImpulseTangentU = SimdAdd(SimdNetImpulseTangentU, ImpulseTangentU);
+				SimdNetImpulseTangentV = SimdAdd(SimdNetImpulseTangentU, ImpulseTangentV);
+
+				V0 = SimdAdd(V0, SimdMultiply(SimdMultiply(ImpulseNormal, SimdInvM0), SimdContactNormal));
+				V0 = SimdAdd(V0, SimdMultiply(SimdMultiply(ImpulseTangentU, SimdInvM0), SimdContactTangentU));
+				V0 = SimdAdd(V0, SimdMultiply(SimdMultiply(ImpulseTangentV, SimdInvM0), SimdContactTangentV));
+				W0 = SimdAdd(W0, SimdMultiply(ImpulseNormal, SimdContactNormalAngular0));
+				W0 = SimdAdd(W0, SimdMultiply(ImpulseTangentU, SimdContactTangentUAngular0));
+				W0 = SimdAdd(W0, SimdMultiply(ImpulseTangentV, SimdContactTangentVAngular0));
+
+				V1 = SimdSubtract(V1, SimdMultiply(SimdMultiply(ImpulseNormal, SimdInvM1), SimdContactNormal));
+				V1 = SimdSubtract(V1, SimdMultiply(SimdMultiply(ImpulseTangentU, SimdInvM1), SimdContactTangentU));
+				V1 = SimdSubtract(V1, SimdMultiply(SimdMultiply(ImpulseTangentV, SimdInvM1), SimdContactTangentV));
+				W1 = SimdSubtract(W1, SimdMultiply(ImpulseNormal, SimdContactNormalAngular1));
+				W1 = SimdSubtract(W1, SimdMultiply(ImpulseTangentU, SimdContactTangentUAngular1));
+				W1 = SimdSubtract(W1, SimdMultiply(ImpulseTangentV, SimdContactTangentVAngular1));
+
+				ScatterBodyVelocities(IsValid, V0, W0, V1, W1, Body0, Body1);
 			}
 
 			void SolveVelocityWithFriction(
@@ -554,8 +709,7 @@ namespace Chaos
 					return;
 				}
 
-				// @todo(chaos): velocity friction
-				SolveVelocityNoFriction(Body0, Body1, Dt);
+				SolveVelocityWithFrictionImpl(Body0, Body1, Dt, FrictionStiffnessScale);
 			}
 
 			FSolverVec3 GetNetPushOut(const int32 LaneIndex) const
@@ -737,8 +891,8 @@ namespace Chaos
 			*/
 			template<int TNumLanes>
 			void SetManifoldPoint(
-				TArrayView<TPBDCollisionSolverManifoldPointsSimd<TNumLanes>>& ManifoldPointsBuffer,
-				TArrayView<TSolverBodyPtrPairSimd<4>>& SolverBodiesBuffer,
+				const TArrayView<TPBDCollisionSolverManifoldPointsSimd<TNumLanes>>& ManifoldPointsBuffer,
+				const TArrayView<TSolverBodyPtrPairSimd<4>>& SolverBodiesBuffer,
 				const int32 ManifoldPointIndex,
 				const FSolverVec3& InRelativeContactPosition0,
 				const FSolverVec3& InRelativeContactPosition1,
@@ -774,8 +928,8 @@ namespace Chaos
 			*/
 			template<int TNumLanes>
 			void InitManifoldPoint(
-				TArrayView<TPBDCollisionSolverManifoldPointsSimd<TNumLanes>>& ManifoldPointsBuffer,
-				TArrayView<TSolverBodyPtrPairSimd<4>>& SolverBodiesBuffer,
+				const TArrayView<TPBDCollisionSolverManifoldPointsSimd<TNumLanes>>& ManifoldPointsBuffer,
+				const TArrayView<TSolverBodyPtrPairSimd<4>>& SolverBodiesBuffer,
 				const int32 ManifoldPointIndex,
 				const FSolverVec3& InRelativeContactPosition0,
 				const FSolverVec3& InRelativeContactPosition1,
@@ -821,11 +975,11 @@ namespace Chaos
 
 			/** 
 			 * Finish manifold point setup.
-			 * NOTE: Can only be called after the WorldContact has been set up (e.g., see SetManifoldPoint)
+			 * NOTE: Can only be called after the InitManifoldPoint has called
 			*/
 			template<int TNumLanes>
 			void FinalizeManifoldPoint(
-				TArrayView<TPBDCollisionSolverManifoldPointsSimd<TNumLanes>>& ManifoldPointsBuffer,
+				const TArrayView<TPBDCollisionSolverManifoldPointsSimd<TNumLanes>>& ManifoldPointsBuffer,
 				const int32 ManifoldPointIndex)
 			{
 				ManifoldPointsBuffer[GetBufferIndex(ManifoldPointIndex)].FinalizeContact(
@@ -851,21 +1005,36 @@ namespace Chaos
 			/**
 			 * @brief Set up the mass scaling for shock propagation, using the position-phase mass scale
 			*/
-			void EnablePositionShockPropagation();
+			template<int TNumLanes>
+			void EnablePositionShockPropagation(
+				const TArrayView<TPBDCollisionSolverManifoldPointsSimd<TNumLanes>>& ManifoldPointsBuffer)
+			{
+				SetShockPropagationInvMassScale(ManifoldPointsBuffer, CVars::Chaos_PBDCollisionSolver_Position_MinInvMassScale);
+			}
 
 			/**
 			 * @brief Set up the mass scaling for shock propagation, using the velocity-phase mass scale
 			*/
-			void EnableVelocityShockPropagation();
+			template<int TNumLanes>
+			void EnableVelocityShockPropagation(
+				const TArrayView<TPBDCollisionSolverManifoldPointsSimd<TNumLanes>>& ManifoldPointsBuffer)
+			{
+				SetShockPropagationInvMassScale(ManifoldPointsBuffer, CVars::Chaos_PBDCollisionSolver_Velocity_MinInvMassScale);
+			}
 
 			/**
 			 * @brief Disable mass scaling
 			*/
-			void DisableShockPropagation();
+			template<int TNumLanes>
+			void DisableShockPropagation(
+				const TArrayView<TPBDCollisionSolverManifoldPointsSimd<TNumLanes>>& ManifoldPointsBuffer)
+			{
+				SetShockPropagationInvMassScale(ManifoldPointsBuffer, FReal(1));
+			}
 
 			template<int TNumLanes>
 			FSolverVec3 GetNetPushOut(
-				TArrayView<TPBDCollisionSolverManifoldPointsSimd<TNumLanes>>& ManifoldPointsBuffer,
+				const TArrayView<TPBDCollisionSolverManifoldPointsSimd<TNumLanes>>& ManifoldPointsBuffer,
 				const int32 ManifoldPointIndex) const
 			{
 				return ManifoldPointsBuffer[GetBufferIndex(ManifoldPointIndex)].GetNetPushOut(State.LaneIndex);
@@ -873,7 +1042,7 @@ namespace Chaos
 
 			template<int TNumLanes>
 			FSolverVec3 GetNetImpulse(
-				TArrayView<TPBDCollisionSolverManifoldPointsSimd<TNumLanes>>& ManifoldPointsBuffer,
+				const TArrayView<TPBDCollisionSolverManifoldPointsSimd<TNumLanes>>& ManifoldPointsBuffer,
 				const int32 ManifoldPointIndex) const
 			{
 				return ManifoldPointsBuffer[GetBufferIndex(ManifoldPointIndex)].GetNetImpulse(State.LaneIndex);
@@ -881,7 +1050,7 @@ namespace Chaos
 
 			template<int TNumLanes>
 			FSolverReal GetStaticFrictionRatio(
-				TArrayView<TPBDCollisionSolverManifoldPointsSimd<TNumLanes>>& ManifoldPointsBuffer,
+				const TArrayView<TPBDCollisionSolverManifoldPointsSimd<TNumLanes>>& ManifoldPointsBuffer,
 				const int32 ManifoldPointIndex) const
 			{
 				return ManifoldPointsBuffer[GetBufferIndex(ManifoldPointIndex)].GetStaticFrictionRatio(State.LaneIndex);
@@ -898,7 +1067,51 @@ namespace Chaos
 			 * @brief Apply the inverse mass scale the body with the lower level
 			 * @param InvMassScale
 			*/
-			void SetShockPropagationInvMassScale(const FSolverReal InvMassScale);
+			template<int TNumLanes>
+			void SetShockPropagationInvMassScale(
+				const TArrayView<TPBDCollisionSolverManifoldPointsSimd<TNumLanes>>& ManifoldPointsBuffer,
+				const FSolverReal InvMassScale)
+			{
+				{
+					FConstraintSolverBody& Body0 = SolverBody0();
+					FConstraintSolverBody& Body1 = SolverBody1();
+
+					// Shock propagation decreases the inverse mass of bodies that are lower in the pile
+					// of objects. This significantly improves stability of heaps and stacks. Height in the pile is indictaed by the "level". 
+					// No need to set an inverse mass scale if the other body is kinematic (with inv mass of 0).
+					// Bodies at the same level do not take part in shock propagation.
+					if (Body0.IsDynamic() && Body1.IsDynamic() && (Body0.Level() != Body1.Level()))
+					{
+						// Set the inv mass scale of the "lower" body to make it heavier
+						bool bInvMassUpdated = false;
+						if (Body0.Level() < Body1.Level())
+						{
+							if (Body0.ShockPropagationScale() != InvMassScale)
+							{
+								Body0.SetShockPropagationScale(InvMassScale);
+								bInvMassUpdated = true;
+							}
+						}
+						else
+						{
+							if (Body1.ShockPropagationScale() != InvMassScale)
+							{
+								Body1.SetShockPropagationScale(InvMassScale);
+								bInvMassUpdated = true;
+							}
+						}
+
+						// If the masses changed, we need to rebuild the contact mass for each manifold point
+						if (bInvMassUpdated)
+						{
+							for (int32 PointIndex = 0; PointIndex < NumManifoldPoints(); ++PointIndex)
+							{
+								ManifoldPointsBuffer[GetBufferIndex(PointIndex)].UpdateMassNormal(State.LaneIndex, Body0, Body1);
+							}
+						}
+					}
+				}
+			}
 
 			struct FState
 			{
