Geometry collection : fix warning when cooking with strip on cook enabled ( non-nanite code path )

diff --git a/Engine/Source/Runtime/Experimental/Chaos/Public/GeometryCollection/ManagedArrayCollection.h b/Engine/Source/Runtime/Experimental/Chaos/Public/GeometryCollection/ManagedArrayCollection.h
--- a/Engine/Source/Runtime/Experimental/Chaos/Public/GeometryCollection/ManagedArrayCollection.h
+++ b/Engine/Source/Runtime/Experimental/Chaos/Public/GeometryCollection/ManagedArrayCollection.h
@@ -158,7 +158,7 @@ public:
 		return Collection;
 	}
 
-	void CopyTo(FManagedArrayCollection* Collection, const TArray<FName>& GroupsToSkip = TArray<FName>()) const
+	void CopyTo(FManagedArrayCollection* Collection, const TArray<FName>& GroupsToSkip = TArray<FName>(), TArray<TTuple<FName, FName>> AttributesToSkip = TArray<TTuple<FName, FName>>()) const
 	{
 		if (!Map.IsEmpty())
 		{
@@ -166,6 +166,14 @@ public:
 			{
 				const FName& AttributeName = Entry.Key.Get<0>();
 				const FName& GroupName = Entry.Key.Get<1>();
+				if (AttributesToSkip.Num() > 0)
+				{
+					const TTuple<FName, FName> GroupAndAttribute{ AttributeName, GroupName };
+					if (AttributesToSkip.Contains(GroupAndAttribute))
+					{
+						continue;
+					}
+				}
 				if (!GroupsToSkip.Contains(GroupName))
 				{
 					if (!Collection->HasGroup(GroupName))
@@ -185,7 +193,6 @@ public:
 		}
 	}
 
-
 	/**
 	* Add an external attribute of Type(T) to the group for size management. Lifetime is managed by the caller, must make sure the array is alive when the collection is
 	* @param Name - The name of the attribute
diff --git a/Engine/Source/Runtime/Experimental/GeometryCollectionEngine/Private/GeometryCollection/GeometryCollectionObject.cpp b/Engine/Source/Runtime/Experimental/GeometryCollectionEngine/Private/GeometryCollection/GeometryCollectionObject.cpp
--- a/Engine/Source/Runtime/Experimental/GeometryCollectionEngine/Private/GeometryCollection/GeometryCollectionObject.cpp
+++ b/Engine/Source/Runtime/Experimental/GeometryCollectionEngine/Private/GeometryCollection/GeometryCollectionObject.cpp
@@ -1238,7 +1238,20 @@ TSharedPtr<FGeometryCollection, ESPMode::ThreadSafe> UGeometryCollection::CopyCo
 	TSharedPtr<FGeometryCollection, ESPMode::ThreadSafe> GeometryCollectionToReturn(new FGeometryCollection());
 
 	const TArray<FName> GroupsToSkip{ FGeometryCollection::GeometryGroup, FGeometryCollection::VerticesGroup, FGeometryCollection::FacesGroup };
-	CollectionToCopy->CopyTo(GeometryCollectionToReturn.Get(), GroupsToSkip);
+	const TArray<TTuple<FName, FName>> AttributesToSkip{ { FGeometryDynamicCollection::SimplicialsAttribute, FTransformCollection::TransformGroup } };
+
+	CollectionToCopy->CopyTo(GeometryCollectionToReturn.Get(), GroupsToSkip, AttributesToSkip);
+
+	// recreate the simplicial attribute since we cannot copy it and we skipped it 
+	using FSimplicialUniquePtr = TUniquePtr<FCollisionStructureManager::FSimplicial>;
+	if (const TManagedArray<FSimplicialUniquePtr>* SourceSimplicials = CollectionToCopy->FindAttribute<FSimplicialUniquePtr>(FGeometryDynamicCollection::SimplicialsAttribute, FTransformCollection::TransformGroup))
+	{
+		TManagedArray<FSimplicialUniquePtr>& SimplicialsToWrite = GeometryCollectionToReturn->AddAttribute<FSimplicialUniquePtr>(FGeometryDynamicCollection::SimplicialsAttribute, FTransformCollection::TransformGroup);
+		for (int32 Index = SourceSimplicials->Num()-1; 0 <= Index; Index--)
+		{
+			SimplicialsToWrite[Index].Reset((*SourceSimplicials)[Index] ? (*SourceSimplicials)[Index]->NewCopy() : nullptr);
+		}
+	}
 
 	// since we are removing the bounding box attribute from the geometry group we need to move it to the transform group 
 	const TManagedArray<FBox>& GeometryBounds = CollectionToCopy->GetAttribute<FBox>("BoundingBox", "Geometry");
