ChaosFlesh : Add support for material with varied stiffness, damping and mass.

diff --git a/Engine/Plugins/Experimental/ChaosFlesh/Source/ChaosFleshEngine/Private/ChaosFlesh/ChaosDeformableSolverComponent.cpp b/Engine/Plugins/Experimental/ChaosFlesh/Source/ChaosFleshEngine/Private/ChaosFlesh/ChaosDeformableSolverComponent.cpp
--- a/Engine/Plugins/Experimental/ChaosFlesh/Source/ChaosFleshEngine/Private/ChaosFlesh/ChaosDeformableSolverComponent.cpp
+++ b/Engine/Plugins/Experimental/ChaosFlesh/Source/ChaosFleshEngine/Private/ChaosFlesh/ChaosDeformableSolverComponent.cpp
@@ -202,6 +202,7 @@ void UDeformableSolverComponent::Reset()
 			, BlendedZeta
 			, Damping
 			, bEnableGravity
+			, bEnableCorotatedConstraint
 		}));
 
 		for (TObjectPtr<UDeformablePhysicsComponent>& DeformableComponent : DeformableComponents)
diff --git a/Engine/Plugins/Experimental/ChaosFlesh/Source/ChaosFleshEngine/Private/ChaosFlesh/FleshComponent.cpp b/Engine/Plugins/Experimental/ChaosFlesh/Source/ChaosFleshEngine/Private/ChaosFlesh/FleshComponent.cpp
--- a/Engine/Plugins/Experimental/ChaosFlesh/Source/ChaosFleshEngine/Private/ChaosFlesh/FleshComponent.cpp
+++ b/Engine/Plugins/Experimental/ChaosFlesh/Source/ChaosFleshEngine/Private/ChaosFlesh/FleshComponent.cpp
@@ -188,13 +188,13 @@ UDeformablePhysicsComponent::FDataMapValue UFleshComponent::NewDeformableData()
 								}
 							}
 						}
-						return FDataMapValue(new Chaos::Softs::FFleshThreadingProxy::FFleshInputBuffer(this->GetComponentTransform(), AnimationTransforms, ComponentPose, bTempEnableGravity, this));
+						return FDataMapValue(new Chaos::Softs::FFleshThreadingProxy::FFleshInputBuffer(this->GetComponentTransform(), AnimationTransforms, ComponentPose, bTempEnableGravity, StiffnessMultiplier, DampingMultiplier, MassMultiplier, this));
 					}
 				}
 			}
 		}
 	}
-	return FDataMapValue(new Chaos::Softs::FFleshThreadingProxy::FFleshInputBuffer(this->GetComponentTransform(), bTempEnableGravity, this));
+	return FDataMapValue(new Chaos::Softs::FFleshThreadingProxy::FFleshInputBuffer(this->GetComponentTransform(), bTempEnableGravity, StiffnessMultiplier, DampingMultiplier, MassMultiplier, this));
 }
 
 
diff --git a/Engine/Plugins/Experimental/ChaosFlesh/Source/ChaosFleshEngine/Public/ChaosFlesh/ChaosDeformablePhysicsComponent.h b/Engine/Plugins/Experimental/ChaosFlesh/Source/ChaosFleshEngine/Public/ChaosFlesh/ChaosDeformablePhysicsComponent.h
--- a/Engine/Plugins/Experimental/ChaosFlesh/Source/ChaosFleshEngine/Public/ChaosFlesh/ChaosDeformablePhysicsComponent.h
+++ b/Engine/Plugins/Experimental/ChaosFlesh/Source/ChaosFleshEngine/Public/ChaosFlesh/ChaosDeformablePhysicsComponent.h
@@ -55,6 +55,16 @@ public:
 	UPROPERTY(EditAnywhere, BlueprintReadWrite, Category = "Chaos Deformable")
 	bool bTempEnableGravity = true;
 
+	UPROPERTY(EditAnywhere, BlueprintReadWrite, Category = Attributes, meta = (ClampMin = "0.0", ClampMax = "1.0", UIMin = "0.0", UIMax = "1.0"))
+		float DampingMultiplier = 1.f;
+
+	UPROPERTY(EditAnywhere, BlueprintReadWrite, Category = Attributes, meta = (ClampMin = "0.0", ClampMax = "1.0", UIMin = "0.0", UIMax = "1.0"))
+		float StiffnessMultiplier = 1.f;
+
+	UPROPERTY(EditAnywhere, BlueprintReadWrite, Category = Attributes, meta = (ClampMin = "0.0", ClampMax = "1.0", UIMin = "0.0", UIMax = "1.0"))
+		float MassMultiplier = 1.f;
+
+
 	const FThreadingProxy* GetPhysicsProxy() const { return PhysicsProxy; }
 	FThreadingProxy* GetPhysicsProxy() { return PhysicsProxy; }
 
diff --git a/Engine/Plugins/Experimental/ChaosFlesh/Source/ChaosFleshEngine/Public/ChaosFlesh/ChaosDeformableSolverComponent.h b/Engine/Plugins/Experimental/ChaosFlesh/Source/ChaosFleshEngine/Public/ChaosFlesh/ChaosDeformableSolverComponent.h
--- a/Engine/Plugins/Experimental/ChaosFlesh/Source/ChaosFleshEngine/Public/ChaosFlesh/ChaosDeformableSolverComponent.h
+++ b/Engine/Plugins/Experimental/ChaosFlesh/Source/ChaosFleshEngine/Public/ChaosFlesh/ChaosDeformableSolverComponent.h
@@ -109,6 +109,9 @@ public:
 
 	UPROPERTY(EditAnywhere, Category = "Chaos Deformable")
 		bool bEnableGravity = true;
+
+	UPROPERTY(EditAnywhere, Category = "Chaos Deformable")
+		bool bEnableCorotatedConstraint = true;
 	//UPROPERTY(EditAnywhere, Category = Chaos)
 	//	bool bWaitForParallelDeformableTask = true;
 
diff --git a/Engine/Plugins/Experimental/ChaosFlesh/Source/ChaosFleshNodes/Private/Dataflow/ChaosFleshEngineAssetNodes.cpp b/Engine/Plugins/Experimental/ChaosFlesh/Source/ChaosFleshNodes/Private/Dataflow/ChaosFleshEngineAssetNodes.cpp
--- a/Engine/Plugins/Experimental/ChaosFlesh/Source/ChaosFleshNodes/Private/Dataflow/ChaosFleshEngineAssetNodes.cpp
+++ b/Engine/Plugins/Experimental/ChaosFlesh/Source/ChaosFleshNodes/Private/Dataflow/ChaosFleshEngineAssetNodes.cpp
@@ -20,7 +20,7 @@ namespace Dataflow
 	{
 		DATAFLOW_NODE_REGISTER_CREATION_FACTORY(FGetFleshAssetDataflowNode);
 		DATAFLOW_NODE_REGISTER_CREATION_FACTORY(FFleshAssetTerminalDataflowNode);
-		DATAFLOW_NODE_REGISTER_CREATION_FACTORY(FComputeFleshMassNode);
+		DATAFLOW_NODE_REGISTER_CREATION_FACTORY(FSetFleshDefaultPropertiesNode);
 		DATAFLOW_NODE_REGISTER_CREATION_FACTORY(FComputeFiberFieldNode);
 	}
 }
@@ -59,14 +59,26 @@ void FFleshAssetTerminalDataflowNode::Evaluate(Dataflow::FContext& Context, cons
 
 template <class T> using MType = FManagedArrayCollection::TManagedType<T>;
 
-void FComputeFleshMassNode::Evaluate(Dataflow::FContext& Context, const FDataflowOutput* Out) const
+void FSetFleshDefaultPropertiesNode::Evaluate(Dataflow::FContext& Context, const FDataflowOutput* Out) const
 {
 	if (Out->IsA<FManagedArrayCollection>(&Collection))
 	{
 		FManagedArrayCollection InCollection = GetValue<FManagedArrayCollection>(Context, &Collection);
+		
+		//if (InCollection.HasAttributes({
+		//	MType< float >("Stiffness", FGeometryCollection::VerticesGroup),
+		//	MType< float >("Damping", FGeometryCollection::VerticesGroup) })
+		//{
+
+			TManagedArray<float>& ParticleStiffness = InCollection.AddAttribute<float>("Stiffness", FGeometryCollection::VerticesGroup);
+				TManagedArray<float>& ParticleDamping = InCollection.AddAttribute<float>("Damping", FGeometryCollection::VerticesGroup);
+		//}//TManagedArray<float>& Stiffness = InCollection.ModifyAttribute<float>("Stiffness", FGeometryCollection::VerticesGroup);
+		
 
 		if (InCollection.HasAttributes({
 			MType< float >("Mass", FGeometryCollection::VerticesGroup),
+			MType< float >("Stiffness", FGeometryCollection::VerticesGroup),
+			MType< float >("Damping", FGeometryCollection::VerticesGroup),
 			MType< FIntVector4 >(FTetrahedralCollection::TetrahedronAttribute, FTetrahedralCollection::TetrahedralGroup),
 			MType< FVector3f >("Vertex", "Vertices"),
 			MType< TArray<int32> >(FTetrahedralCollection::IncidentElementsAttribute, FGeometryCollection::VerticesGroup),
@@ -78,6 +90,8 @@ void FComputeFleshMassNode::Evaluate(Dataflow::FContext& Context, const FDataflo
 			if (VertsNum && TetsNum)
 			{
 				TManagedArray<float>& Mass = InCollection.ModifyAttribute<float>("Mass", FGeometryCollection::VerticesGroup);
+				TManagedArray<float>& Stiffness = InCollection.ModifyAttribute<float>("Stiffness", FGeometryCollection::VerticesGroup);
+				TManagedArray<float>& Damping = InCollection.ModifyAttribute<float>("Damping", FGeometryCollection::VerticesGroup);
 				const TManagedArray<FIntVector4>& Tetrahedron = InCollection.GetAttribute<FIntVector4>(FTetrahedralCollection::TetrahedronAttribute, FTetrahedralCollection::TetrahedralGroup);
 				const TManagedArray<FVector3f>& Vertex = InCollection.GetAttribute<FVector3f>("Vertex", "Vertices");
 				const TManagedArray<TArray<int32>>& ElemVertexIndex = InCollection.GetAttribute<TArray<int32>>(FTetrahedralCollection::IncidentElementsAttribute, FGeometryCollection::VerticesGroup);
@@ -133,6 +147,9 @@ void FComputeFleshMassNode::Evaluate(Dataflow::FContext& Context, const FDataflo
 				{
 					Mass.Fill(Density * Volume / Vertex.Num());
 				}
+
+				Stiffness.Fill(VertexStiffness);
+				Damping.Fill(VertexDamping);
 			}
 		}
 		SetValue<FManagedArrayCollection>(Context, InCollection, &Collection);
diff --git a/Engine/Plugins/Experimental/ChaosFlesh/Source/ChaosFleshNodes/Public/Dataflow/ChaosFleshEngineAssetNodes.h b/Engine/Plugins/Experimental/ChaosFlesh/Source/ChaosFleshNodes/Public/Dataflow/ChaosFleshEngineAssetNodes.h
--- a/Engine/Plugins/Experimental/ChaosFlesh/Source/ChaosFleshNodes/Public/Dataflow/ChaosFleshEngineAssetNodes.h
+++ b/Engine/Plugins/Experimental/ChaosFlesh/Source/ChaosFleshNodes/Public/Dataflow/ChaosFleshEngineAssetNodes.h
@@ -60,10 +60,10 @@ public:
 
 
 USTRUCT(meta = (DataflowFlesh))
-struct FComputeFleshMassNode : public FDataflowNode
+struct FSetFleshDefaultPropertiesNode : public FDataflowNode
 {
 	GENERATED_USTRUCT_BODY()
-	DATAFLOW_NODE_DEFINE_INTERNAL(FComputeFleshMassNode, "ComputeFleshMass", "Flesh", "")
+	DATAFLOW_NODE_DEFINE_INTERNAL(FSetFleshDefaultPropertiesNode, "SetFleshDefaultProperties", "Flesh", "")
 	DATAFLOW_NODE_RENDER_TYPE(FGeometryCollection::StaticType(), "Collection")
 
 public:
@@ -71,10 +71,16 @@ public:
 	UPROPERTY(EditAnywhere, Category = "Dataflow")
 	float Density = 1.f;
 
+	UPROPERTY(EditAnywhere, Category = "Dataflow")
+	float VertexStiffness = 1e6;
+
+	UPROPERTY(EditAnywhere, Category = "Dataflow", meta = (ClampMin = "0.0", ClampMax = "1.0", UIMin = "0.0", UIMax = "1.0"))
+	float VertexDamping = 0.f;
+
 	UPROPERTY(meta = (DataflowInput, DataflowOutput, DisplayName = "Collection", DataflowPassthrough  = "Collection"))
 	FManagedArrayCollection Collection;
 
-	FComputeFleshMassNode(const Dataflow::FNodeParameters& InParam, FGuid InGuid = FGuid::NewGuid())
+	FSetFleshDefaultPropertiesNode(const Dataflow::FNodeParameters& InParam, FGuid InGuid = FGuid::NewGuid())
 		: FDataflowNode(InParam, InGuid)
 	{
 		RegisterInputConnection(&Collection);
diff --git a/Engine/Source/Runtime/Experimental/Chaos/Private/Chaos/Deformable/ChaosDeformableSolver.cpp b/Engine/Source/Runtime/Experimental/Chaos/Private/Chaos/Deformable/ChaosDeformableSolver.cpp
--- a/Engine/Source/Runtime/Experimental/Chaos/Private/Chaos/Deformable/ChaosDeformableSolver.cpp
+++ b/Engine/Source/Runtime/Experimental/Chaos/Private/Chaos/Deformable/ChaosDeformableSolver.cpp
@@ -209,33 +209,69 @@ namespace Chaos::Softs
 
 		const TManagedArray<FVector3f>& DynamicVertex = Dynamic.GetAttribute<FVector3f>("Vertex", FGeometryCollection::VerticesGroup);
 		const TManagedArray<FSolverReal>* MassArray = Rest.FindAttribute<FSolverReal>("Mass", FGeometryCollection::VerticesGroup);
-		uint32 NumParticles = Rest.NumElements(FGeometryCollection::VerticesGroup);
+		const TManagedArray<FSolverReal>* DampingArray = Rest.FindAttribute<FSolverReal>("Damping", FGeometryCollection::VerticesGroup);
 		FSolverReal Mass = 100.0;// @todo(chaos) : make user attributes
 
 		auto ChaosVert = [](FVector3d V) { return Chaos::FVec3(V.X, V.Y, V.Z); };
 		auto ChaosM = [](FSolverReal M, const TManagedArray<float>* AM, int32 Index, int32 Num) { return FSolverReal((AM != nullptr) ? (*AM)[Index] : M / FSolverReal(Num)); };
 		auto ChaosInvM = [](FSolverReal M) { return FSolverReal(FMath::IsNearlyZero(M) ? 0.0 : 1 / M); };
 		auto DoubleVert = [](FVector3f V) { return FVector3d(V.X, V.Y, V.Z); };
-
-		const FTransform& InitialTransform = Proxy.GetInitialTransform();
+		uint32 NumParticles = Rest.NumElements(FGeometryCollection::VerticesGroup);
+		
 		int32 ParticleStart = Evolution->AddParticleRange(NumParticles, GroupOffset, true);
 		GroupOffset += 1;
+		for (uint32 vdx = 0; vdx < NumParticles; ++vdx)
+		{
+			MObjects[ParticleStart + vdx] = Proxy.GetOwner();
+		}
+
+		TArray<FSolverReal> MassWithMultiplier;
+		TArray<FSolverReal> DampingWithMultiplier;
+		MassWithMultiplier.Init(0.f, NumParticles);
+		DampingWithMultiplier.Init(0.f, NumParticles);
+		FSolverReal DampingMultiplier = 0.f;
+		FSolverReal MassMultiplier = 0.f;
+		if (const UObject* Owner = this->MObjects[ParticleStart]) {
+			FFleshThreadingProxy::FFleshInputBuffer* FleshInputBuffer = nullptr;
+			if (this->CurrentInputPackage->ObjectMap.Contains(Owner))
+			{
+				FleshInputBuffer = this->CurrentInputPackage->ObjectMap[Owner]->As<FFleshThreadingProxy::FFleshInputBuffer>();
+				if (FleshInputBuffer)
+				{
+					DampingMultiplier = FleshInputBuffer->DampingMultiplier;
+					MassMultiplier = FleshInputBuffer->MassMultiplier;
+				}
+			}
+		}
+
+		for (uint32 vdx = 0; vdx < NumParticles; ++vdx)
+		{
+			MassWithMultiplier[vdx] = ChaosM(Mass, MassArray, vdx, NumParticles) * MassMultiplier;
+			if (DampingArray)
+			{
+				Evolution->SetParticleDamping((*DampingArray)[vdx], ParticleStart + vdx);
+			}
+		}
+
+		Evolution->SetDamping(DampingMultiplier, GroupOffset - 1);
+		
+		const FTransform& InitialTransform = Proxy.GetInitialTransform();
+
 		
 		for (uint32 vdx = 0; vdx < NumParticles; ++vdx)
 		{
 			int32 SolverParticleIndex = ParticleStart + vdx;
 			Evolution->Particles().X(SolverParticleIndex) = ChaosVert(InitialTransform.TransformPosition(DoubleVert(DynamicVertex[vdx])));
 			Evolution->Particles().V(SolverParticleIndex) = Chaos::FVec3(0.f, 0.f, 0.f);
-			Evolution->Particles().M(SolverParticleIndex) = ChaosM(Mass, MassArray, vdx, NumParticles);
+			Evolution->Particles().M(SolverParticleIndex) = MassWithMultiplier[vdx];
 			Evolution->Particles().InvM(SolverParticleIndex) = ChaosInvM(Evolution->Particles().M(SolverParticleIndex));
 			Evolution->Particles().PAndInvM(SolverParticleIndex).InvM = Evolution->Particles().InvM(SolverParticleIndex);
-			MObjects[SolverParticleIndex] = Proxy.GetOwner();
 		}
 
 		bool ObjectEnableGravity = false;
 
-		int32 SolverParticleIndex = ParticleStart;
-		if (const UObject* Owner = this->MObjects[SolverParticleIndex]) {
+		//int32 SolverParticleIndex = ParticleStart;
+		if (const UObject* Owner = this->MObjects[ParticleStart]) {
 			FFleshThreadingProxy::FFleshInputBuffer* FleshInputBuffer = nullptr;
 			if (this->CurrentInputPackage->ObjectMap.Contains(Owner))
 			{
@@ -312,6 +348,7 @@ namespace Chaos::Softs
 #endif
 	}
 
+
 	void FDeformableSolver::InitializeTetrahedralConstraint(FFleshThreadingProxy& Proxy)
 	{
 		SCOPE_CYCLE_COUNTER(STAT_ChaosDeformableSolver_InitializeTetrahedralConstraint);
@@ -345,55 +382,95 @@ namespace Chaos::Softs
 
 			}
 
-
-			int32 InitIndex = Evolution->AddConstraintInitRange(1, true);
-			int32 ConstraintIndex = Evolution->AddConstraintRuleRange(1, true);
-
-			if (Property.bDoBlended) 
+			if (Rest.HasAttributes({ FManagedArrayCollection::TManagedType<FSolverReal>("Stiffness", FGeometryCollection::VerticesGroup) }))
 			{
-				FBlendedXPBDCorotatedConstraints<FSolverReal, FSolverParticles>* BlendedCorotatedConstraint =
-					new FBlendedXPBDCorotatedConstraints<FSolverReal, FSolverParticles>(
-						Evolution->Particles(), Elements, /*bRecordMetric = */false, Property.EMesh, (FSolverReal).3, Property.BlendedZeta);
-			
-				Evolution->ConstraintInits()[InitIndex] =
-					[BlendedCorotatedConstraint](FSolverParticles& InParticles, const FSolverReal Dt)
+				uint32 NumParticles = Rest.NumElements(FGeometryCollection::VerticesGroup);
+				TArray<FSolverReal> StiffnessWithMultiplier;
+				StiffnessWithMultiplier.Init(0.f, NumParticles);
+				FSolverReal StiffnessMultiplier = 0.f;
+
+				if (const UObject* Owner = this->MObjects[Range[0]]) {
+					FFleshThreadingProxy::FFleshInputBuffer* FleshInputBuffer = nullptr;
+					if (this->CurrentInputPackage->ObjectMap.Contains(Owner))
+					{
+						FleshInputBuffer = this->CurrentInputPackage->ObjectMap[Owner]->As<FFleshThreadingProxy::FFleshInputBuffer>();
+						if (FleshInputBuffer)
+						{
+							StiffnessMultiplier = FleshInputBuffer->StiffnessMultiplier;
+						}
+					}
+				}
+				const TManagedArray<FSolverReal>* StiffnessArray = Rest.FindAttribute<FSolverReal>("Stiffness", FGeometryCollection::VerticesGroup);
+				TArray<FSolverReal> TetStiffness;
+				TetStiffness.Init(Property.EMesh, Elements.Num());
+				if (StiffnessArray)
 				{
-					BlendedCorotatedConstraint->Init();
-				};
+					for (uint32 vdx = 0; vdx < NumParticles; ++vdx)
+					{
+						StiffnessWithMultiplier[vdx] = (*StiffnessArray)[vdx] * StiffnessMultiplier;
+					}
+					for (int32 edx = 0; edx < Elements.Num(); edx++)
+					{
+						TetStiffness[edx] = (StiffnessWithMultiplier[Tetrahedron[edx].X] + StiffnessWithMultiplier[Tetrahedron[edx].Y]
+							+ StiffnessWithMultiplier[Tetrahedron[edx].Z] + StiffnessWithMultiplier[Tetrahedron[edx].W]) / 4.f;
+					}
 
-				Evolution->ConstraintRules()[ConstraintIndex] =
-					[BlendedCorotatedConstraint](FSolverParticles& InParticles, const FSolverReal Dt)
+				}
+				
+				if (Property.bEnableCorotatedConstraints)
 				{
-					BlendedCorotatedConstraint->ApplyInParallel(InParticles, Dt);
-				};
-			
-				BlendedCorotatedConstraints.Add(TUniquePtr<FBlendedXPBDCorotatedConstraints<FSolverReal, FSolverParticles>>(BlendedCorotatedConstraint));
 
-			}
-			else
-			{
-				FXPBDCorotatedConstraints<FSolverReal, FSolverParticles>* CorotatedConstraint =
-					new FXPBDCorotatedConstraints<FSolverReal, FSolverParticles>(
-						Evolution->Particles(), Elements, /*bRecordMetric = */false, Property.EMesh);
+					int32 InitIndex = Evolution->AddConstraintInitRange(1, true);
+					int32 ConstraintIndex = Evolution->AddConstraintRuleRange(1, true);
 
-				Evolution->ConstraintInits()[InitIndex] =
-					[CorotatedConstraint](FSolverParticles& InParticles, const FSolverReal Dt)
-				{
-					CorotatedConstraint->Init();
-				};
+					if (Property.bDoBlended)
+					{
+						FBlendedXPBDCorotatedConstraints<FSolverReal, FSolverParticles>* BlendedCorotatedConstraint =
+							new FBlendedXPBDCorotatedConstraints<FSolverReal, FSolverParticles>(
+								Evolution->Particles(), Elements, TetStiffness, (FSolverReal).3,/*bRecordMetric = */false, Property.BlendedZeta);
 
-				Evolution->ConstraintRules()[ConstraintIndex] =
-					[CorotatedConstraint](FSolverParticles& InParticles, const FSolverReal Dt)
-				{
-					CorotatedConstraint->ApplyInParallel(InParticles, Dt);
-				};
+						Evolution->ConstraintInits()[InitIndex] =
+							[BlendedCorotatedConstraint](FSolverParticles& InParticles, const FSolverReal Dt)
+						{
+							BlendedCorotatedConstraint->Init();
+						};
+
+						Evolution->ConstraintRules()[ConstraintIndex] =
+							[BlendedCorotatedConstraint](FSolverParticles& InParticles, const FSolverReal Dt)
+						{
+							BlendedCorotatedConstraint->ApplyInParallel(InParticles, Dt);
+						};
+
+						BlendedCorotatedConstraints.Add(TUniquePtr<FBlendedXPBDCorotatedConstraints<FSolverReal, FSolverParticles>>(BlendedCorotatedConstraint));
+
+					}
+					else
+					{
+						FXPBDCorotatedConstraints<FSolverReal, FSolverParticles>* CorotatedConstraint =
+							new FXPBDCorotatedConstraints<FSolverReal, FSolverParticles>(
+								Evolution->Particles(), Elements, TetStiffness);
 
-				CorotatedConstraints.Add(TUniquePtr<FXPBDCorotatedConstraints<FSolverReal, FSolverParticles>>(CorotatedConstraint));
+						Evolution->ConstraintInits()[InitIndex] =
+							[CorotatedConstraint](FSolverParticles& InParticles, const FSolverReal Dt)
+						{
+							CorotatedConstraint->Init();
+						};
+
+						Evolution->ConstraintRules()[ConstraintIndex] =
+							[CorotatedConstraint](FSolverParticles& InParticles, const FSolverReal Dt)
+						{
+							CorotatedConstraint->ApplyInParallel(InParticles, Dt);
+						};
+
+						CorotatedConstraints.Add(TUniquePtr<FXPBDCorotatedConstraints<FSolverReal, FSolverParticles>>(CorotatedConstraint));
+					}
+				}
 			}
 		
 		}
 	}
 
+
 	void FDeformableSolver::InitializeGidBasedConstraints(FFleshThreadingProxy& Proxy)
 	{
 		SCOPE_CYCLE_COUNTER(STAT_ChaosDeformableSolver_InitializeGidBasedConstraints);
diff --git a/Engine/Source/Runtime/Experimental/Chaos/Private/Chaos/PBDEvolution.cpp b/Engine/Source/Runtime/Experimental/Chaos/Private/Chaos/PBDEvolution.cpp
--- a/Engine/Source/Runtime/Experimental/Chaos/Private/Chaos/PBDEvolution.cpp
+++ b/Engine/Source/Runtime/Experimental/Chaos/Private/Chaos/PBDEvolution.cpp
@@ -112,6 +112,7 @@ FPBDEvolution::FPBDEvolution(
 
 	// Add particle arrays
 	MParticles.AddArray(&MParticleGroupIds);
+	MParticles.AddArray(&MParticleDampings);
 	MCollisionParticles.AddArray(&MCollisionTransforms);
 	MCollisionParticles.AddArray(&MCollided);
 	MCollisionParticles.AddArray(&MCollisionParticleGroupIds);
@@ -139,6 +140,7 @@ int32 FPBDEvolution::AddParticleRange(int32 NumParticles, uint32 GroupId, bool b
 		for (int32 i = Offset; i < (int32)MParticles.Size(); ++i)
 		{
 			MParticleGroupIds[i] = GroupId;
+			MParticleDampings[i] = (FSolverReal)0.;
 		}
 
 		// Resize the group parameter arrays
@@ -252,7 +254,7 @@ void FPBDEvolution::PreIterationUpdate(
 		SCOPE_CYCLE_COUNTER(STAT_ChaosClothSolverIntegrate);
 
 		constexpr FSolverReal DampingFrequency = (FSolverReal)60.;  // The damping value is the percentage of velocity removed per frame when running at 60Hz
-		const FSolverReal Damping = FMath::Clamp(MGroupDampings[ParticleGroupId], (FSolverReal)0., (FSolverReal)1.);
+		FSolverReal Damping = FMath::Clamp(MGroupDampings[ParticleGroupId], (FSolverReal)0., (FSolverReal)1.);
 		FSolverReal DampingPowDt;
 		FSolverReal DampingIntegrated;
 		if (Damping > (FSolverReal)1. - (FSolverReal)UE_KINDA_SMALL_NUMBER)
@@ -293,9 +295,31 @@ void FPBDEvolution::PreIterationUpdate(
 		{
 			const int32 RangeSize = Range - Offset;
 			PhysicsParallelFor(RangeSize,
-				[this, &Offset, &ForceRule, &Gravity, &VelocityAndPressureField, &DampVelocityRule, DampingPowDt, DampingIntegrated, Dt](int32 i)
+				[this, &Offset, &ForceRule, &Gravity, &VelocityAndPressureField, &DampVelocityRule, &DampingFrequency, &ParticleGroupId,DampingPowDt, DampingIntegrated, Dt](int32 i)
 				{
 					const int32 Index = Offset + i;
+
+					const FSolverReal ParticleDamping = FMath::Clamp(MParticleDampings[Index] * MGroupDampings[ParticleGroupId], (FSolverReal)0., (FSolverReal)1.);
+					FSolverReal ParticleDampingPowDt;
+					FSolverReal ParticleDampingIntegrated;
+					if (ParticleDamping > (FSolverReal)1. - (FSolverReal)UE_KINDA_SMALL_NUMBER)
+					{
+						ParticleDampingIntegrated = ParticleDampingPowDt = (FSolverReal)0.;
+					}
+					else if (ParticleDamping > (FSolverReal)UE_SMALL_NUMBER)
+					{
+						const FSolverReal ParticleLogValueByFrequency = FMath::Loge((FSolverReal)1. - ParticleDamping) * DampingFrequency;
+
+						ParticleDampingPowDt = FMath::Exp(ParticleLogValueByFrequency * Dt);  // DampingPowDt = FMath::Pow(OneMinusDamping, Dt * DampingFrequency);
+						ParticleDampingIntegrated = (ParticleDampingPowDt - (FSolverReal)1.) / ParticleLogValueByFrequency;
+					}
+					else
+					{
+						ParticleDampingPowDt = (FSolverReal)1.;
+						ParticleDampingIntegrated = Dt;
+					}
+
+
 					if (MParticles.InvM(Index) != (FSolverReal)0.)  // Process dynamic particles
 					{
 						// Init forces with GravityForces
@@ -323,9 +347,9 @@ void FPBDEvolution::PreIterationUpdate(
 						}
 
 						// Euler Step with point damping integration
-						MParticles.P(Index) = MParticles.X(Index) + MParticles.V(Index) * DampingIntegrated;
+						MParticles.P(Index) = MParticles.X(Index) + MParticles.V(Index) * ParticleDampingIntegrated;
 
-						MParticles.V(Index) *= DampingPowDt;
+						MParticles.V(Index) *= ParticleDampingPowDt;
 					}
 					else  // Process kinematic particles
 					{
diff --git a/Engine/Source/Runtime/Experimental/Chaos/Public/Chaos/BlendedXPBDCorotatedConstraints.h b/Engine/Source/Runtime/Experimental/Chaos/Public/Chaos/BlendedXPBDCorotatedConstraints.h
--- a/Engine/Source/Runtime/Experimental/Chaos/Public/Chaos/BlendedXPBDCorotatedConstraints.h
+++ b/Engine/Source/Runtime/Experimental/Chaos/Public/Chaos/BlendedXPBDCorotatedConstraints.h
@@ -46,7 +46,30 @@ namespace Chaos::Softs
 
 			C1Contribution.Init(TVec4<TVector<T, 3>>(TVector<T, 3>((T)0.)), InMesh.Num());
 			C2Contribution.Init(TVec4<TVector<T, 3>>(TVector<T, 3>((T)0.)), InMesh.Num());
+		}
+
+		//this one only accepts tetmesh input and mesh
+		FBlendedXPBDCorotatedConstraints(
+			const ParticleType& InParticles,
+			const TArray<TVector<int32, 4>>& InMesh,
+			const TArray<T>& EMeshArray,
+			const T& NuMesh = (T).3,
+			const bool bRecordMetricIn = false,
+			const T& InZeta = (T)1.
+		)
+			: Base(InParticles, InMesh, EMeshArray, NuMesh, bRecordMetricIn), Zeta(InZeta)
+		{
+			if (Zeta < (T)0.)
+			{
+				Zeta = (T)0.;
+			}
+			if (Zeta > (T)1.)
+			{
+				Zeta = (T)1.;
+			}
 
+			C1Contribution.Init(TVec4<TVector<T, 3>>(TVector<T, 3>((T)0.)), InMesh.Num());
+			C2Contribution.Init(TVec4<TVector<T, 3>>(TVector<T, 3>((T)0.)), InMesh.Num());
 		}
 
 		virtual ~FBlendedXPBDCorotatedConstraints() {}
diff --git a/Engine/Source/Runtime/Experimental/Chaos/Public/Chaos/Deformable/ChaosDeformableSolver.h b/Engine/Source/Runtime/Experimental/Chaos/Public/Chaos/Deformable/ChaosDeformableSolver.h
--- a/Engine/Source/Runtime/Experimental/Chaos/Public/Chaos/Deformable/ChaosDeformableSolver.h
+++ b/Engine/Source/Runtime/Experimental/Chaos/Public/Chaos/Deformable/ChaosDeformableSolver.h
@@ -11,6 +11,7 @@
 #include "Chaos/PBDTriangleMeshCollisions.h"
 #include "Chaos/TriangleMesh.h"
 #include "Chaos/XPBDCorotatedConstraints.h"
+//#include "Chaos/XPBDWeakConstraints.h"
 #include "Chaos/BlendedXPBDCorotatedConstraints.h"
 #include "Chaos/XPBDGridBasedCorotatedConstraints.h"
 #include "Chaos/Deformable/ChaosDeformableSolverTypes.h"
@@ -119,6 +120,7 @@ namespace Chaos::Softs
 		void InitializeKinematicParticles(FFleshThreadingProxy&);
 		void InitializeTetrahedralConstraint(FFleshThreadingProxy&);
 		void InitializeGidBasedConstraints(FFleshThreadingProxy&);
+		//void InitializeWeakConstraints(FFleshThreadingProxy&);
 		void InitializeKinematicConstraint();
 		void InitializeCollisionBodies();
 		void InitializeSelfCollisionVariables();
diff --git a/Engine/Source/Runtime/Experimental/Chaos/Public/Chaos/Deformable/ChaosDeformableSolverProxy.h b/Engine/Source/Runtime/Experimental/Chaos/Public/Chaos/Deformable/ChaosDeformableSolverProxy.h
--- a/Engine/Source/Runtime/Experimental/Chaos/Public/Chaos/Deformable/ChaosDeformableSolverProxy.h
+++ b/Engine/Source/Runtime/Experimental/Chaos/Public/Chaos/Deformable/ChaosDeformableSolverProxy.h
@@ -114,20 +114,26 @@ namespace Chaos::Softs
 		public:
 			typedef FFleshThreadingProxy Source;
 
-			FFleshInputBuffer(const FTransform& InGlobalTransform, const bool InbEnableGravity, const UObject* InOwner = nullptr)
+			FFleshInputBuffer(const FTransform& InGlobalTransform, const bool InbEnableGravity, const float InStiffnessMultiplier, const float InDampingMultiplier, const float InMassMultiplier, const UObject* InOwner = nullptr)
 				: Super(InOwner, FFleshThreadingProxy::TypeName())
 				, GlobalTransform(InGlobalTransform)
 				, Transforms(TArray<FTransform>())
 				, RestTransforms(TArray<FTransform>())
 				, bEnableGravity(InbEnableGravity)
+				, StiffnessMultiplier(InStiffnessMultiplier)
+				, DampingMultiplier(InDampingMultiplier)
+				, MassMultiplier(InMassMultiplier)
 			{}
 
-			FFleshInputBuffer(const FTransform& InGlobalTransform, const TArray<FTransform>& InTransforms, const TArray<FTransform>& InRestTransforms, const bool InbEnableGravity, const UObject* InOwner = nullptr)
+			FFleshInputBuffer(const FTransform& InGlobalTransform, const TArray<FTransform>& InTransforms, const TArray<FTransform>& InRestTransforms, const bool InbEnableGravity, const float InStiffnessMultiplier, const float InDampingMultiplier, const float InMassMultiplier, const UObject* InOwner = nullptr)
 				: Super(InOwner, FFleshThreadingProxy::TypeName())
 				, GlobalTransform(InGlobalTransform)
 				, Transforms(InTransforms)
 				, RestTransforms(InRestTransforms)
 				, bEnableGravity(InbEnableGravity)
+				, StiffnessMultiplier(InStiffnessMultiplier)
+				, DampingMultiplier(InDampingMultiplier)
+				, MassMultiplier(InMassMultiplier)
 			{}
 			virtual ~FFleshInputBuffer() {}
 
@@ -135,6 +141,9 @@ namespace Chaos::Softs
 			TArray<FTransform> Transforms;
 			TArray<FTransform> RestTransforms; 
 			bool bEnableGravity;
+			float StiffnessMultiplier;
+			float DampingMultiplier;
+			float MassMultiplier;
 		};
 
 		class FFleshOutputBuffer : public FThreadingProxy::FBuffer
diff --git a/Engine/Source/Runtime/Experimental/Chaos/Public/Chaos/Deformable/ChaosDeformableSolverTypes.h b/Engine/Source/Runtime/Experimental/Chaos/Public/Chaos/Deformable/ChaosDeformableSolverTypes.h
--- a/Engine/Source/Runtime/Experimental/Chaos/Public/Chaos/Deformable/ChaosDeformableSolverTypes.h
+++ b/Engine/Source/Runtime/Experimental/Chaos/Public/Chaos/Deformable/ChaosDeformableSolverTypes.h
@@ -30,7 +30,8 @@ namespace Chaos::Softs
 			bool InbDoBlended = false,
 			FSolverReal InBlendedZeta = (FSolverReal).1,
 			FSolverReal InDamping = (FSolverReal)0,
-			bool InbEnableGravity = true)
+			bool InbEnableGravity = true, 
+			bool InbEnableCorotatedConstraints = true)
 			: NumSolverSubSteps(InNumSolverSubSteps)
 			, NumSolverIterations(InNumSolverIterations)
 			, FixTimeStep(InFixTimeStep)
@@ -47,6 +48,7 @@ namespace Chaos::Softs
 			, BlendedZeta(InBlendedZeta)
 			, Damping(InDamping)
 			, bEnableGravity(InbEnableGravity)
+			, bEnableCorotatedConstraints(InbEnableCorotatedConstraints)
 		{}
 
 		FDeformableSolverProperties(const FDeformableSolverProperties& InProp)
@@ -66,6 +68,7 @@ namespace Chaos::Softs
 			, BlendedZeta(InProp.BlendedZeta)
 			, Damping(InProp.Damping)
 			, bEnableGravity(InProp.bEnableGravity)
+			, bEnableCorotatedConstraints(InProp.bEnableCorotatedConstraints)
 		{}
 
 		int32 NumSolverSubSteps = 5;
@@ -84,6 +87,7 @@ namespace Chaos::Softs
 		FSolverReal BlendedZeta = (FSolverReal)0.;
 		FSolverReal Damping = (FSolverReal)0.;
 		bool bEnableGravity = true;
+		bool bEnableCorotatedConstraints = true;
 	};
 
 
diff --git a/Engine/Source/Runtime/Experimental/Chaos/Public/Chaos/PBDEvolution.h b/Engine/Source/Runtime/Experimental/Chaos/Public/Chaos/PBDEvolution.h
--- a/Engine/Source/Runtime/Experimental/Chaos/Public/Chaos/PBDEvolution.h
+++ b/Engine/Source/Runtime/Experimental/Chaos/Public/Chaos/PBDEvolution.h
@@ -162,6 +162,10 @@ class CHAOS_API FPBDEvolution : public TArrayCollection
 	bool GetUseCCD(const uint32 GroupId = 0) const { check(GroupId < TArrayCollection::Size()); return MGroupUseCCDs[GroupId]; }
 	void SetUseCCD(const bool bUseCCD, const uint32 GroupId = 0) { check(GroupId < TArrayCollection::Size()); MGroupUseCCDs[GroupId] = bUseCCD; }
 
+	FSolverReal GetParticleDamping(const uint32 ParticleIndex = 0) const { check(ParticleIndex < this->Particles().Size()); return MParticleDampings[ParticleIndex]; }
+	void SetParticleDamping(const FSolverReal Damping, const uint32 ParticleIndex = 0) { check(ParticleIndex < this->Particles().Size()); MParticleDampings[ParticleIndex] = Damping; }
+
+
 	UE_DEPRECATED(4.27, "Use GetCollisionStatus() instead")
 	const bool Collided(int32 index) { return MCollided[index]; }
 
@@ -203,6 +207,9 @@ private:
 	TArrayCollectionArray<FSolverReal> MGroupDampings;
 	TArrayCollectionArray<FSolverReal> MGroupLocalDampings;
 	TArrayCollectionArray<bool> MGroupUseCCDs;
+
+	//per particle data:
+	TArrayCollectionArray<FSolverReal> MParticleDampings;
 	
 	TArray<TFunction<void(FSolverParticles&, const FSolverReal)>> MConstraintInits;
 	TPBDActiveView<TArray<TFunction<void(FSolverParticles&, const FSolverReal)>>> MConstraintInitsActiveView;
diff --git a/Engine/Source/Runtime/Experimental/Chaos/Public/Chaos/XPBDCorotatedConstraints.h b/Engine/Source/Runtime/Experimental/Chaos/Public/Chaos/XPBDCorotatedConstraints.h
--- a/Engine/Source/Runtime/Experimental/Chaos/Public/Chaos/XPBDCorotatedConstraints.h
+++ b/Engine/Source/Runtime/Experimental/Chaos/Public/Chaos/XPBDCorotatedConstraints.h
@@ -60,6 +60,46 @@ namespace Chaos::Softs
 			InitColor(InParticles);
 		}
 
+		FXPBDCorotatedConstraints(
+			const ParticleType& InParticles,
+			const TArray<TVector<int32, 4>>& InMesh,
+			const TArray<T>& EMeshArray,
+			const T& NuMesh = (T).3,
+			const bool bRecordMetricIn = false
+		)
+			: bRecordMetric(bRecordMetricIn), MeshConstraints(InMesh)
+		{
+			ensureMsgf(EMeshArray.Num() == InMesh.Num(), TEXT("Input Young Modulus Array Size is wrong"));
+			LambdaArray.Init((T)0., 2 * MeshConstraints.Num());
+			DmInverse.Init((T)0., 9 * MeshConstraints.Num());
+			Measure.Init((T)0., MeshConstraints.Num());
+			LambdaElementArray.Init((T)0., MeshConstraints.Num());
+			MuElementArray.Init((T)0., MeshConstraints.Num());
+			
+			for (int e = 0; e < InMesh.Num(); e++)
+			{
+				LambdaElementArray[e] = EMeshArray[e] * NuMesh / (((T)1. + NuMesh) * ((T)1. - (T)2. * NuMesh));
+				MuElementArray[e] = EMeshArray[e] / ((T)2. * ((T)1. + NuMesh));
+
+				PMatrix<T, 3, 3> Dm = DsInit(e, InParticles);
+				PMatrix<T, 3, 3> DmInv = Dm.Inverse();
+				for (int r = 0; r < 3; r++) {
+					for (int c = 0; c < 3; c++) {
+						DmInverse[(3 * 3) * e + 3 * r + c] = DmInv.GetAt(r, c);
+					}
+				}
+
+				Measure[e] = Dm.Determinant() / (T)6.;
+
+				if (Measure[e] < (T)0.)
+				{
+					Measure[e] = -Measure[e];
+				}
+			}
+
+			InitColor(InParticles);
+		}
+
 		FXPBDCorotatedConstraints(
 			const ParticleType& InParticles,
 			const TArray<TVector<int32, 4>>& InMesh,
@@ -139,19 +179,25 @@ namespace Chaos::Softs
 		virtual void ApplyInSerial(ParticleType& Particles, const T Dt, const int32 ElementIndex) const
 		{
 			TRACE_CPUPROFILER_EVENT_SCOPE(TEXT("STAT_ChaosXPBDCorotatedApplySingle"));
-			TVec4<TVector<T, 3>> PolarDelta = GetPolarDelta(Particles, Dt, ElementIndex);
 
-			for (int i = 0; i < 4; i++) 
-			{
-				Particles.P(MeshConstraints[ElementIndex][i]) += PolarDelta[i];
-			}
-			
-			TVec4<TVector<T, 3>> DetDelta = GetDeterminantDelta(Particles, Dt, ElementIndex);
+			//if (MuElementArray[ElementIndex] > 300.f)
+			//{
 
-			for (int i = 0; i < 4; i++)
-			{
-				Particles.P(MeshConstraints[ElementIndex][i]) += DetDelta[i];
-			}
+
+				TVec4<TVector<T, 3>> PolarDelta = GetPolarDelta(Particles, Dt, ElementIndex);
+
+				for (int i = 0; i < 4; i++)
+				{
+					Particles.P(MeshConstraints[ElementIndex][i]) += PolarDelta[i];
+				}
+
+				TVec4<TVector<T, 3>> DetDelta = GetDeterminantDelta(Particles, Dt, ElementIndex);
+
+				for (int i = 0; i < 4; i++)
+				{
+					Particles.P(MeshConstraints[ElementIndex][i]) += DetDelta[i];
+				}
+			//}
 
 
 		}
@@ -329,14 +375,15 @@ namespace Chaos::Softs
 			PMatrix<T, 3, 3> DmInvT = ElementDmInv(ElementIndex).GetTransposed();
 			
 			T J = Fe.Determinant();
-			if (J - 1 < Tol)
+			if (J - 1 < Tol )
 			{
 				return TVec4<TVector<T, 3>>(TVector<T, 3>((T)0.));
 			}
 
 			TVec4<TVector<T, 3>> dC2 = GetDeterminantGradient(Fe, DmInvT);
 
-			T AlphaTilde = (T)2. / (Dt * Dt * Lambda * Measure[ElementIndex]);
+			//T AlphaTilde = (T)2. / (Dt * Dt * Lambda * Measure[ElementIndex]);
+			T AlphaTilde = (T)2. / (Dt * Dt * LambdaElementArray[ElementIndex] * Measure[ElementIndex]);
 
 			if (bRecordMetric)
 			{
@@ -397,19 +444,18 @@ namespace Chaos::Softs
 			}
 			C1 = FMath::Sqrt(C1);
 
-			if (C1 < Tol)
+			if (C1 < Tol )
 			{
 				return TVec4<TVector<T, 3>>(TVector<T, 3>((T)0.));
 			}
 
-			//TVector<T, 81> dRdF((T)0.);
-			//Chaos::dRdFCorotated(Fe, dRdF);
 
 			PMatrix<T, 3, 3> DmInvT = ElementDmInv(ElementIndex).GetTransposed();
 
 			TVec4<TVector<T, 3>> dC1 = GetPolarGradient(Fe, Re, DmInvT, C1);
 
-			T AlphaTilde = (T)1. / (Dt * Dt * Mu * Measure[ElementIndex]);
+			//T AlphaTilde = (T)1. / (Dt * Dt * Mu * Measure[ElementIndex]);
+			T AlphaTilde = (T)1. / (Dt * Dt * MuElementArray[ElementIndex] * Measure[ElementIndex]);
 
 			if (bRecordMetric)
 			{
@@ -456,9 +502,12 @@ namespace Chaos::Softs
 		//material constants calculated from E:
 		T Mu;
 		T Lambda;
+		TArray<T> MuElementArray;
+		TArray<T> LambdaElementArray;
 		mutable T HError;
 		mutable TArray<T> HErrorArray;
 		bool bRecordMetric;
+		bool VariableStiffness = false;
 
 		TArray<TVector<int32, 4>> MeshConstraints;
 		mutable TArray<T> Measure;
