Add cvar (p.Chaos.GC.ForceAutoAssignISMPool) to force auto assigning ISMPool to GC components - false by default

diff --git a/Engine/Source/Runtime/Experimental/GeometryCollectionEngine/Private/GeometryCollection/GeometryCollectionComponent.cpp b/Engine/Source/Runtime/Experimental/GeometryCollectionEngine/Private/GeometryCollection/GeometryCollectionComponent.cpp
--- a/Engine/Source/Runtime/Experimental/GeometryCollectionEngine/Private/GeometryCollection/GeometryCollectionComponent.cpp
+++ b/Engine/Source/Runtime/Experimental/GeometryCollectionEngine/Private/GeometryCollection/GeometryCollectionComponent.cpp
@@ -111,6 +111,9 @@ FAutoConsoleVariableRef CVarChaosGCUseISMPool(TEXT("p.Chaos.GC.UseISMPool"), bCh
 bool bChaos_GC_UseISMPoolForNonFracturedParts = true;
 FAutoConsoleVariableRef CVarChaosGCUseISMPoolForNonFracturedParts(TEXT("p.Chaos.GC.UseISMPoolForNonFracturedParts"), bChaos_GC_UseISMPoolForNonFracturedParts, TEXT("When enabled, non fractured part will use the ISM pool if specified"));
 
+bool bChaos_GC_ForceAutoAssignISMPool = false;
+FAutoConsoleVariableRef CVarChaosGCForceAutoAssignISMPool(TEXT("p.Chaos.GC.ForceAutoAssignISMPool"), bChaos_GC_ForceAutoAssignISMPool, TEXT("When enabled, force assignement of ISMPool regardgless of the settings of the components"));
+
 bool bChaos_GC_InitConstantDataUseParallelFor = true;
 FAutoConsoleVariableRef CVarChaosGCInitConstantDataUseParallelFor(TEXT("p.Chaos.GC.InitConstantDataUseParallelFor"), bChaos_GC_InitConstantDataUseParallelFor, TEXT("When enabled, InitConstant data will use parallelFor for copying some of the data"));
 
@@ -2618,7 +2621,7 @@ void UGeometryCollectionComponent::OnRegister()
 	AssignedISMPool = nullptr;
 	if (bChaos_GC_UseISMPool && GetWorld()->IsGameWorld())
 	{
-		if (bAutoAssignISMPool)
+		if (bAutoAssignISMPool || bChaos_GC_ForceAutoAssignISMPool)
 		{
 			if (UGeometryCollectionISMPoolSubSystem* ISMPoolSubSystem = UWorld::GetSubsystem<UGeometryCollectionISMPoolSubSystem>(GetWorld()))
 			{
