Chaos - add basic overlap function for use on the sim (WIP)

diff --git a/Engine/Source/Programs/HeadlessChaos/Private/HeadlessChaosTestSimSweep.cpp b/Engine/Source/Programs/HeadlessChaos/Private/HeadlessChaosTestSimSweep.cpp
new file mode 100644
--- /dev/null
+++ b/Engine/Source/Programs/HeadlessChaos/Private/HeadlessChaosTestSimSweep.cpp
@@ -0,0 +1,37 @@
+// Copyright Epic Games, Inc. All Rights Reserved.
+
+#include "HeadlessChaos.h"
+#include "HeadlessChaosTestEvolution.h"
+#include "Chaos/Collision/SimSweep.h"
+
+namespace ChaosTest
+{
+	using namespace Chaos;
+
+	GTEST_TEST(SimOverlapTests, TestBoxBoxOverlap)
+	{
+		const FVec3 FloorSize = FVec3(10000, 10000, 100);
+		const FVec3 BoxPos = FVec3(0, 0, 200);
+		const FRotation3 BoxRot = FRotation3::FromIdentity();
+		const FVec3 BoxSize = FVec3(100);
+		const FReal BoxMass = 100;
+
+		FEvolutionTest Test;
+
+		FGeometryParticleHandle* Floor = Test.AddParticleFloor(FloorSize);
+		FPBDRigidParticleHandle* Box = Test.AddParticleBox(BoxPos, BoxRot, BoxSize, BoxMass);
+
+		Test.Tick();
+
+		TArray<Private::FSimOverlapParticleShape> Overlaps;
+		const FAABB3 QueryBounds = FAABB3(FVec3(-1000), FVec3(1000));
+		Private::SimOverlapBoundsAll(Test.GetEvolution().GetSpatialAcceleration(), QueryBounds, Overlaps);
+
+		EXPECT_EQ(Overlaps.Num(), 2);
+		if (Overlaps.Num() == 2)
+		{
+			EXPECT_TRUE((Overlaps[0].HitParticle == Floor) || (Overlaps[1].HitParticle == Floor));
+			EXPECT_TRUE((Overlaps[0].HitParticle == Box) || (Overlaps[1].HitParticle == Box));
+		}
+	}
+}
\ No newline at end of file
diff --git a/Engine/Source/Runtime/Experimental/Chaos/Private/Chaos/Collision/SimSweep.cpp b/Engine/Source/Runtime/Experimental/Chaos/Private/Chaos/Collision/SimSweep.cpp
--- a/Engine/Source/Runtime/Experimental/Chaos/Private/Chaos/Collision/SimSweep.cpp
+++ b/Engine/Source/Runtime/Experimental/Chaos/Private/Chaos/Collision/SimSweep.cpp
@@ -25,5 +25,22 @@ namespace Chaos
 
 			return OutHit.IsHit();
 		}
-	}
-}
\ No newline at end of file
+
+		bool SimOverlapBoundsAll(
+			ISpatialAcceleration<FAccelerationStructureHandle, FReal, 3>* SpatialAcceleration,
+			const FAABB3& QueryBounds,
+			TArray<FSimOverlapParticleShape>& Overlaps)
+		{
+			// Accept all particles and shapes
+			const auto& ParticleFilter = [](const FGeometryParticleHandle* Particle) -> bool { return true; };
+			const auto& ShapeFilter = [](const FPerShapeData* Shape, const FImplicitObject* Implicit) { return true; };
+			const auto& OverlapCollector = [&Overlaps](const FSimOverlapParticleShape& Overlap) { Overlaps.Add(Overlap); };
+			const int32 InitialOverlapCount = Overlaps.Num();
+
+			SimOverlapBounds(SpatialAcceleration, QueryBounds, ParticleFilter, ShapeFilter, OverlapCollector);
+
+			return (Overlaps.Num() > InitialOverlapCount);
+		}
+
+	} // namespace Private
+} // namespace Chaos
\ No newline at end of file
diff --git a/Engine/Source/Runtime/Experimental/Chaos/Public/Chaos/Collision/SimSweep.h b/Engine/Source/Runtime/Experimental/Chaos/Public/Chaos/Collision/SimSweep.h
--- a/Engine/Source/Runtime/Experimental/Chaos/Public/Chaos/Collision/SimSweep.h
+++ b/Engine/Source/Runtime/Experimental/Chaos/Public/Chaos/Collision/SimSweep.h
@@ -60,6 +60,26 @@ namespace Chaos
 			int32 HitFaceIndex;
 		};
 
+		/**
+		* Results from SimOverlapBounds test
+		*/
+		class FSimOverlapParticleShape
+		{
+		public:
+			FSimOverlapParticleShape()
+			{
+			}
+
+			// The particle that we hit, or null
+			const FGeometryParticleHandle* HitParticle;
+
+			// The shape that we hit (or uninitialized)
+			const FPerShapeData* HitShape;
+
+			// The world-space transform of the geometry that was hit
+			FRigidTransform3 ShapeWorldTransform;
+		};
+
 		/**
 		* A spatial acceleration visitor that forwards callbacks to a functor
 		*/
@@ -99,6 +119,76 @@ namespace Chaos
 			const TVisitor& Visitor;
 		};
 
+		/**
+		* Produce a set of particles and their shapes that overlap the query bounds and pass the particle and shape filters.
+		* 
+		* @TParam TParticleFilter a particle that returns true to check a particle for overlaps. Signature bool(const FGeometryParticleHandle*)
+		* 
+		* @TParam TShapeFilter a shape that returns true to check a particle for overlaps. Signature bool(const FPerShapeData*, const FImplicitObject*)
+		*
+		* @TParam TOverlapCollector an overlap collector functor with the signature void(const FSimOverlapParticleShape&)
+		*/
+		template<typename TParticleFilter, typename TShapeFilter, typename TOverlapCollector>
+		void SimOverlapBounds(
+			ISpatialAcceleration<FAccelerationStructureHandle, FReal, 3>* SpatialAcceleration,
+			const FAABB3& QueryBounds,
+			TParticleFilter& ParticleFilter,
+			TShapeFilter& ShapeFilter,
+			TOverlapCollector& OverlapCollector)
+		{
+			FImplicitBox3 QueryBox(QueryBounds.Min(), QueryBounds.Max());
+
+			// Functor: Process a particle encountered by the overla through the spatial acceleration
+			const auto& ProcessParticle =
+				[&QueryBox, &ParticleFilter, &ShapeFilter, &OverlapCollector]
+				(const FGeometryParticleHandle* OtherParticle)
+				-> void
+			{
+				// Check the particle filter to see if we should process this particle pair
+				if (ParticleFilter(OtherParticle))
+				{
+					// The particle transform
+					const FRigidTransform3 OtherWorldTransform = OtherParticle->GetTransformXR();
+					for (const TUniquePtr<FPerShapeData>& OtherShape : OtherParticle->ShapesArray())
+					{
+						const FImplicitObject* OtherImplicit = OtherShape->GetLeafGeometry();
+
+						// Check the filter to see if we want to  process this shape pair
+						const bool bOverlapShape = ShapeFilter(OtherShape.Get(), OtherImplicit);
+
+						if (bOverlapShape)
+						{
+							const FRigidTransform3 OtherShapeWorldTransform = FRigidTransform3(OtherShape->GetLeafRelativeTransform()) * OtherWorldTransform;
+
+							FMTDInfo MTDInfo;
+							const bool bOverlapHit = OverlapQuery(
+								*OtherImplicit,
+								OtherShapeWorldTransform,
+								QueryBox,
+								FRigidTransform3::Identity,
+								0/*Thickness*/,
+								nullptr);
+
+							// If we have a hit, pass it to the collector
+							if (bOverlapHit)
+							{
+								FSimOverlapParticleShape Overlap;
+								Overlap.HitParticle = OtherParticle;
+								Overlap.HitShape = OtherShape.Get();
+								Overlap.ShapeWorldTransform = OtherShapeWorldTransform;
+
+								OverlapCollector(Overlap);
+							}
+						}
+					}
+				}
+			};
+
+			// Generate the set of objects that overlap the bounds
+			TSimSweepSQVisitor SpatialAcclerationVisitor(ProcessParticle);
+			SpatialAcceleration->Overlap(QueryBounds, SpatialAcclerationVisitor);
+		}
+
 		/**
 		* Sweep SweptParticle against OtherParticle. The provided filter policy can be used to reject shape pairs.
 		* The collector policy chooses which sweep results(s) to keep.
@@ -412,6 +502,15 @@ namespace Chaos
 			const FReal Length,
 			FSimSweepParticleHit& OutHit,
 			const FReal InHitDistanceEqualTolerance = UE_KINDA_SMALL_NUMBER);
+
+		/**
+		* Collect all the shapes that overlap a bounding box.
+		* @todo(chaos): Add a few more options like Complex vs Simple collection
+		*/
+		extern bool SimOverlapBoundsAll(
+			ISpatialAcceleration<FAccelerationStructureHandle, FReal, 3>* SpatialAcceleration,
+			const FAABB3& QueryBounds,
+			TArray<FSimOverlapParticleShape>& Overlaps);
 	}
 
 }
\ No newline at end of file
