Geometry collection : Improve debug info for geometry collections

diff --git a/Engine/Source/Runtime/Experimental/Chaos/Private/GeometryCollection/ManagedArrayCollection.cpp b/Engine/Source/Runtime/Experimental/Chaos/Private/GeometryCollection/ManagedArrayCollection.cpp
--- a/Engine/Source/Runtime/Experimental/Chaos/Private/GeometryCollection/ManagedArrayCollection.cpp
+++ b/Engine/Source/Runtime/Experimental/Chaos/Private/GeometryCollection/ManagedArrayCollection.cpp
@@ -547,18 +547,17 @@ bool FManagedArrayCollection::IsConnected(FName StartingNode, FName TargetNode)
 	return false;
 }
 
-#include <sstream> 
-#include <string>
 FString FManagedArrayCollection::ToString() const
 {
-	FString Buffer("Group : Attribute [Ptr] [AllocatedSize]\n");
+	FString Buffer;
 
 	const TArray<FStringFormatArg> CollectionInfos = { FString::FormatAsNumber((int32)GetAllocatedSize()) };
-	Buffer += FString::Format(TEXT("*:* [n/a] [{0}]\n"), CollectionInfos);
+	Buffer += FString::Format(TEXT("All attributes [{0} bytes]\n"), CollectionInfos);
 
 	for (FName GroupName : GroupNames())
 	{
-		Buffer += GroupName.ToString() + "\n";
+		const TArray<FStringFormatArg> GroupNameInfos = { GroupName.ToString(), FString::FormatAsNumber((int32)NumElements(GroupName)) };
+		Buffer += FString::Format(TEXT("{0} - [{1} elements]\n"), GroupNameInfos);
 		for (FName AttributeName : AttributeNames(GroupName))
 		{
 			FKeyType Key = FManagedArrayCollection::MakeMapKey(AttributeName, GroupName);
@@ -567,8 +566,8 @@ FString FManagedArrayCollection::ToString() const
 			const SIZE_T AttributeAllocatedSize = Value.Value? Value.Value->GetAllocatedSize() : 0;
 			const FString AttributeAllocatedSizeStr = FString::FormatAsNumber((int32)AttributeAllocatedSize);
 
-			const TArray<FStringFormatArg> AttributeInfos = { GroupName.ToString(), AttributeName.ToString(), (uint64)Value.Value, AttributeAllocatedSizeStr };
-			Buffer += FString::Format(TEXT("{0}:{1} [{2}] [{3}]\n"), AttributeInfos);
+			const TArray<FStringFormatArg> AttributeInfos = { AttributeName.ToString(), AttributeAllocatedSizeStr };
+			Buffer += FString::Format(TEXT(" |-- {0} [{1} bytes]\n"), AttributeInfos);
 		}
 	}
 	return Buffer;
