Copy only necessary Implicit Object

diff --git a/Engine/Source/Runtime/Experimental/Chaos/Private/PhysicsProxy/GeometryCollectionPhysicsProxy.cpp b/Engine/Source/Runtime/Experimental/Chaos/Private/PhysicsProxy/GeometryCollectionPhysicsProxy.cpp
--- a/Engine/Source/Runtime/Experimental/Chaos/Private/PhysicsProxy/GeometryCollectionPhysicsProxy.cpp
+++ b/Engine/Source/Runtime/Experimental/Chaos/Private/PhysicsProxy/GeometryCollectionPhysicsProxy.cpp
@@ -263,11 +263,22 @@ void PopulateSimulatedParticle(
 			}
 		};
 
-		TSharedPtr<Chaos::FImplicitObject, ESPMode::ThreadSafe> SharedImplicitTS(DeepCopyImplicit(Implicit).Release());
-		FCollisionStructureManager::UpdateImplicitFlags(SharedImplicitTS.Get(), SingleSupportedCollisionTypeData.CollisionType);
-		Handle->SetSharedGeometry(SharedImplicitTS);
+		Chaos::EImplicitObjectType ImplicitType = Implicit->GetType();
+		// Don't copy if it is not a level set and scale is one
+		if (SingleSupportedCollisionTypeData.CollisionType != ECollisionTypeEnum::Chaos_Surface_Volumetric && 
+			ImplicitType != Chaos::ImplicitObjectType::LevelSet && Scale.Equals(FVector::OneVector))
+		{
+			Handle->SetSharedGeometry(Implicit);
+			Handle->SetLocalBounds(Implicit->BoundingBox());
+		}
+		else
+		{
+			TSharedPtr<Chaos::FImplicitObject, ESPMode::ThreadSafe> SharedImplicitTS(DeepCopyImplicit(Implicit).Release());
+			FCollisionStructureManager::UpdateImplicitFlags(SharedImplicitTS.Get(), SingleSupportedCollisionTypeData.CollisionType);
+			Handle->SetSharedGeometry(SharedImplicitTS);
+			Handle->SetLocalBounds(SharedImplicitTS->BoundingBox());
+		}
 		Handle->SetHasBounds(true);
-		Handle->SetLocalBounds(SharedImplicitTS->BoundingBox());
 		const Chaos::FRigidTransform3 Xf(Handle->X(), Handle->R());
 		Handle->UpdateWorldSpaceState(Xf, Chaos::FVec3(0));
 	}
