Decrease Constant Data memory

diff --git a/Engine/Source/Runtime/Experimental/GeometryCollectionEngine/Private/GeometryCollection/GeometryCollectionComponent.cpp b/Engine/Source/Runtime/Experimental/GeometryCollectionEngine/Private/GeometryCollection/GeometryCollectionComponent.cpp
--- a/Engine/Source/Runtime/Experimental/GeometryCollectionEngine/Private/GeometryCollection/GeometryCollectionComponent.cpp
+++ b/Engine/Source/Runtime/Experimental/GeometryCollectionEngine/Private/GeometryCollection/GeometryCollectionComponent.cpp
@@ -701,7 +701,7 @@ FPrimitiveSceneProxy* UGeometryCollectionComponent::CreateSceneProxy()
 
 		if (RestCollection->HasVisibleGeometry())
 		{
-			FGeometryCollectionConstantData* const ConstantData = ::new FGeometryCollectionConstantData;
+			TUniquePtr<FGeometryCollectionConstantData> ConstantData = MakeUnique<FGeometryCollectionConstantData>();
 			InitConstantData(ConstantData);
 
 			FGeometryCollectionDynamicData* const DynamicData = InitDynamicData(true /* initialization */);
@@ -720,7 +720,7 @@ FPrimitiveSceneProxy* UGeometryCollectionComponent::CreateSceneProxy()
 			#endif
 
 				ENQUEUE_RENDER_COMMAND(CreateRenderState)(
-					[GeometryCollectionSceneProxy, ConstantData, DynamicData](FRHICommandListImmediate& RHICmdList)
+					[GeometryCollectionSceneProxy, ConstantData = MoveTemp(ConstantData), DynamicData](FRHICommandListImmediate& RHICmdList)
 					{
 						GeometryCollectionSceneProxy->SetConstantData_RenderThread(ConstantData);
 
@@ -753,7 +753,7 @@ FPrimitiveSceneProxy* UGeometryCollectionComponent::CreateSceneProxy()
 			#endif
 
 				ENQUEUE_RENDER_COMMAND(CreateRenderState)(
-					[GeometryCollectionSceneProxy, ConstantData, DynamicData](FRHICommandListImmediate& RHICmdList)
+					[GeometryCollectionSceneProxy, ConstantData = MoveTemp(ConstantData), DynamicData](FRHICommandListImmediate& RHICmdList)
 					{
 						GeometryCollectionSceneProxy->SetConstantData_RenderThread(ConstantData);
 						if (DynamicData)
@@ -2053,7 +2053,7 @@ void UGeometryCollectionComponent::SetInitialClusterBreaks(const TArray<int32>&
 }
 
 
-void UGeometryCollectionComponent::InitConstantData(FGeometryCollectionConstantData* ConstantData) const
+void UGeometryCollectionComponent::InitConstantData(TUniquePtr<FGeometryCollectionConstantData>& ConstantData) const
 {
 	// Constant data should all be moved to the DDC as time permits.
 	// todo : this should be computed once per asset not per component or at least the part that does not depend on the component properties
diff --git a/Engine/Source/Runtime/Experimental/GeometryCollectionEngine/Private/GeometryCollection/GeometryCollectionSceneProxy.cpp b/Engine/Source/Runtime/Experimental/GeometryCollectionEngine/Private/GeometryCollection/GeometryCollectionSceneProxy.cpp
--- a/Engine/Source/Runtime/Experimental/GeometryCollectionEngine/Private/GeometryCollection/GeometryCollectionSceneProxy.cpp
+++ b/Engine/Source/Runtime/Experimental/GeometryCollectionEngine/Private/GeometryCollection/GeometryCollectionSceneProxy.cpp
@@ -152,7 +152,6 @@ FGeometryCollectionSceneProxy::FGeometryCollectionSceneProxy(UGeometryCollection
 	, bUsesSubSections(false)
 #endif
 	, DynamicData(nullptr)
-	, ConstantData(nullptr)
 	, bShowBoneColors(Component->GetShowBoneColors())
 	, bEnableBoneSelection(Component->GetEnableBoneSelection())
 	, bSuppressSelectionMaterial(Component->GetSuppressSelectionMaterial())
@@ -299,10 +298,6 @@ void FGeometryCollectionSceneProxy::DestroyRenderThreadResources()
 		DynamicData = nullptr;
 	}
 
-	if (ConstantData != nullptr)
-	{
-		delete ConstantData;
-	}
 
 #if RHI_RAYTRACING
 	if (IsRayTracingAllowed())
@@ -355,7 +350,7 @@ void FGeometryCollectionSceneProxy::SetupVertexFactory(FGeometryCollectionVertex
 	}
 }
 
-void FGeometryCollectionSceneProxy::InitResources()
+void FGeometryCollectionSceneProxy::InitResources(const TUniquePtr<FGeometryCollectionConstantData>& ConstantData)
 {
 	check(ConstantData);
 	check(IsInRenderingThread());
@@ -475,22 +470,22 @@ void FGeometryCollectionSceneProxy::ReleaseResources()
 	VertexFactory.ReleaseResource();
 }
 
-inline static FColor GetOnlyVertexColors(const FGeometryCollectionConstantData* ConstantDataIn, int32 PointIdx, bool bShowBoneColors)
+inline static FColor GetOnlyVertexColors(const TUniquePtr<FGeometryCollectionConstantData>& ConstantDataIn, int32 PointIdx, bool bShowBoneColors)
 {
 	return ConstantDataIn->Colors[PointIdx].ToFColor(true);
 };
 
-inline static FColor GetBonesColors(const FGeometryCollectionConstantData* ConstantDataIn, int32 PointIdx, bool bShowBoneColors)
+inline static FColor GetBonesColors(const TUniquePtr<FGeometryCollectionConstantData>& ConstantDataIn, int32 PointIdx, bool bShowBoneColors)
 {
 	return bShowBoneColors
 		? ConstantDataIn->BoneColors[PointIdx].ToFColor(true)
 		: (ConstantDataIn->BoneColors[PointIdx] * ConstantDataIn->Colors[PointIdx]).ToFColor(true);
 };
 
-void FGeometryCollectionSceneProxy::BuildGeometry( const FGeometryCollectionConstantData* ConstantDataIn, TArray<FDynamicMeshVertex>& OutVertices, TArray<int32>& OutIndices, TArray<int32> &OutOriginalMeshIndices)
+void FGeometryCollectionSceneProxy::BuildGeometry(const TUniquePtr<FGeometryCollectionConstantData>& ConstantDataIn, TArray<FDynamicMeshVertex>& OutVertices, TArray<int32>& OutIndices, TArray<int32> &OutOriginalMeshIndices)
 {
 	const bool bCapturedShowBoneColors = this->bShowBoneColors;
-	const auto SetOutVertex = [&OutVertices, &ConstantDataIn, bCapturedShowBoneColors](int32 PointIdx, FColor (*GetColorFunc)(const FGeometryCollectionConstantData*, int32, bool) )
+	const auto SetOutVertex = [&OutVertices, &ConstantDataIn, bCapturedShowBoneColors](int32 PointIdx, FColor (*GetColorFunc)(const TUniquePtr<FGeometryCollectionConstantData>&, int32, bool) )
 	{
 		OutVertices[PointIdx] =
 			FDynamicMeshVertex(
@@ -515,7 +510,7 @@ void FGeometryCollectionSceneProxy::BuildGeometry( const FGeometryCollectionCons
 	const int32 MinVertexBatchSize = GeometryCollectionBuildGeometryParallelVertexProcessingMinBatch;
 	if (bShowBoneColors || bEnableBoneSelection)
 	{
-		ParallelFor(TEXT("GCSceneProxy::BuildGeometry(Vertices)"), ConstantData->Vertices.Num(), MinVertexBatchSize,
+		ParallelFor(TEXT("GCSceneProxy::BuildGeometry(Vertices)"), ConstantDataIn->Vertices.Num(), MinVertexBatchSize,
 			[&SetOutVertex](int32 PointIdx)
 			{
 				SetOutVertex(PointIdx, GetBonesColors);
@@ -523,7 +518,7 @@ void FGeometryCollectionSceneProxy::BuildGeometry( const FGeometryCollectionCons
 	}
 	else
 	{
-		ParallelFor(TEXT("GCSceneProxy::BuildGeometry(Vertices)"), ConstantData->Vertices.Num(), MinVertexBatchSize,
+		ParallelFor(TEXT("GCSceneProxy::BuildGeometry(Vertices)"), ConstantDataIn->Vertices.Num(), MinVertexBatchSize,
 			[&SetOutVertex](int32 PointIdx)
 			{
 				SetOutVertex(PointIdx, GetOnlyVertexColors);
@@ -552,22 +547,16 @@ void FGeometryCollectionSceneProxy::BuildGeometry( const FGeometryCollectionCons
 		});
 }
 
-void FGeometryCollectionSceneProxy::SetConstantData_RenderThread(FGeometryCollectionConstantData* NewConstantData, bool ForceInit)
+void FGeometryCollectionSceneProxy::SetConstantData_RenderThread(const TUniquePtr<FGeometryCollectionConstantData>& ConstantData, bool ForceInit)
 {
 	check(IsInRenderingThread());
-	check(NewConstantData);
+	check(ConstantData.IsValid());
 
-	if (ConstantData)
-	{
-		delete ConstantData;
-		ConstantData = nullptr;
-	}
-	ConstantData = NewConstantData;
 
 	if (ConstantData->Vertices.Num() != VertexBuffers.PositionVertexBuffer.GetNumVertices() || ForceInit)
 	{
 		ReleaseResources();
-		InitResources();
+		InitResources(ConstantData);
 	}
 
 	TArray<int32> Indices;
@@ -657,7 +646,7 @@ void FGeometryCollectionSceneProxy::SetConstantData_RenderThread(FGeometryCollec
 		// If we are using the GeometryCollection vertex factory, populate the vertex buffer for bone map
 		if (bSupportsManualVertexFetch)
 		{
-			void* BoneMapBufferData = RHILockBuffer(BoneMapBuffer.VertexBufferRHI, 0, Vertices.Num() * sizeof(int32), RLM_WriteOnly);								
+			void* BoneMapBufferData = RHILockBuffer(BoneMapBuffer.VertexBufferRHI, 0, Vertices.Num() * sizeof(int32), RLM_WriteOnly);
 			FMemory::Memcpy(BoneMapBufferData, &ConstantData->BoneMap[0], ConstantData->BoneMap.Num() * sizeof(int32));
 			RHIUnlockBuffer(BoneMapBuffer.VertexBufferRHI);
 
@@ -720,6 +709,13 @@ void FGeometryCollectionSceneProxy::SetConstantData_RenderThread(FGeometryCollec
 	}
 #endif
 
+	// Move don't copy it ! 
+	SubsetConstantData.BoneMap = MoveTemp(ConstantData->BoneMap);
+	SubsetConstantData.OriginalMeshIndices = MoveTemp(ConstantData->OriginalMeshIndices);
+	SubsetConstantData.OriginalMeshSections = MoveTemp(ConstantData->OriginalMeshSections);
+	SubsetConstantData.Indices = MoveTemp(ConstantData->Indices);
+	SubsetConstantData.NumTransforms = MoveTemp(ConstantData->NumTransforms);
+	SubsetConstantData.RestTransforms = MoveTemp(ConstantData->RestTransforms);
 }
 
 void FGeometryCollectionSceneProxy::SetDynamicData_RenderThread(FGeometryCollectionDynamicData* NewDynamicData)
@@ -734,7 +730,8 @@ void FGeometryCollectionSceneProxy::SetDynamicData_RenderThread(FGeometryCollect
 		}
 		DynamicData = NewDynamicData;
 		
-		check(VertexBuffers.PositionVertexBuffer.GetNumVertices() == (uint32)ConstantData->Vertices.Num());
+
+		const int32 VerticesNum = VertexBuffers.PositionVertexBuffer.GetNumVertices();
 
 		if (bSupportsManualVertexFetch)
 		{
@@ -747,8 +744,8 @@ void FGeometryCollectionSceneProxy::SetDynamicData_RenderThread(FGeometryCollect
 
 				for (int32 i = 1; i < 3; i++)
 				{
-					TransformBuffers[i].NumTransforms = ConstantData->NumTransforms;
-					PrevTransformBuffers[i].NumTransforms = ConstantData->NumTransforms;
+					TransformBuffers[i].NumTransforms = SubsetConstantData.NumTransforms;
+					PrevTransformBuffers[i].NumTransforms = SubsetConstantData.NumTransforms;
 					TransformBuffers[i].InitResource();
 					PrevTransformBuffers[i].InitResource();
 				}
@@ -776,9 +773,8 @@ void FGeometryCollectionSceneProxy::SetDynamicData_RenderThread(FGeometryCollect
 				else if (!TransformVertexBuffersContainsOriginalMesh)
 				{
 					// if we are rendering the base mesh geometry, then use rest transforms rather than the simulated one for both current and previous transforms
-					TransformBuffer.UpdateDynamicData(ConstantData->RestTransforms, LockMode);
-					PrevTransformBuffer.UpdateDynamicData(ConstantData->RestTransforms, LockMode);
-
+					TransformBuffer.UpdateDynamicData(SubsetConstantData.RestTransforms, LockMode);
+					PrevTransformBuffer.UpdateDynamicData(SubsetConstantData.RestTransforms, LockMode);
 					TransformVertexBuffersContainsOriginalMesh = true;
 				}
 
@@ -791,7 +787,7 @@ void FGeometryCollectionSceneProxy::SetDynamicData_RenderThread(FGeometryCollect
 			void* VertexBufferData = RHILockBuffer(VertexBuffer.VertexBufferRHI, 0, VertexBuffer.GetNumVertices() * VertexBuffer.GetStride(), RLM_WriteOnly);
 
 			bool bParallelGeometryCollection = true;
-			int32 TotalVertices = ConstantData->Vertices.Num();
+			int32 TotalVertices = VerticesNum;
 			int32 ParallelGeometryCollectionBatchSize = CVarParallelGeometryCollectionBatchSize.GetValueOnRenderThread();
 
 			int32 NumBatches = (TotalVertices / ParallelGeometryCollectionBatchSize);
@@ -821,7 +817,7 @@ void FGeometryCollectionSceneProxy::SetDynamicData_RenderThread(FGeometryCollect
 
 				if (ThisBatchSize > 0)
 				{
-					const FMatrix44f* RESTRICT BoneTransformsPtr = DynamicData->IsDynamic ? DynamicData->Transforms.GetData() : ConstantData->RestTransforms.GetData();
+					const FMatrix44f* RESTRICT BoneTransformsPtr = DynamicData->IsDynamic ? DynamicData->Transforms.GetData() : SubsetConstantData.RestTransforms.GetData();
 
 					if (bGeometryCollection_SetDynamicData_ISPC_Enabled)
 					{
@@ -831,16 +827,16 @@ void FGeometryCollectionSceneProxy::SetDynamicData_RenderThread(FGeometryCollect
 							(ispc::FVector3f*)VertexBufferOffset,
 							ThisBatchSize,
 							VertexBuffer.GetStride(),
-							&ConstantData->BoneMap[IndexOffset],
+							&SubsetConstantData.BoneMap[IndexOffset],
 							(ispc::FMatrix44f*)BoneTransformsPtr,
-							(ispc::FVector3f*)&ConstantData->Vertices[IndexOffset]);
+							(ispc::FVector3f*)&VertexBuffers.PositionVertexBuffer.VertexPosition(IndexOffset));
 #endif
 					}
 					else
 					{
 						for (int32 i = IndexOffset; i < IndexOffset + ThisBatchSize; i++)
 						{
-							FVector3f Transformed = BoneTransformsPtr[ConstantData->BoneMap[i]].TransformPosition(ConstantData->Vertices[i]);
+							FVector3f Transformed = BoneTransformsPtr[SubsetConstantData.BoneMap[i]].TransformPosition(VertexBuffers.PositionVertexBuffer.VertexPosition(i));
 							FMemory::Memcpy((uint8*)VertexBufferData + (i * VertexBuffer.GetStride()), &Transformed, sizeof(FVector3f));
 						}
 					}
@@ -976,10 +972,10 @@ void FGeometryCollectionSceneProxy::GetDynamicMeshElements(const TArray<const FS
 			if (!bIsDynamic)
 			{
 			#if GEOMETRYCOLLECTION_EDITOR_SELECTION
-				const TArray<FGeometryCollectionSection>& SectionArray = bUsesSubSections && SubSections.Num() ? SubSections: ConstantData->OriginalMeshSections;
+				const TArray<FGeometryCollectionSection>& SectionArray = bUsesSubSections && SubSections.Num() ? SubSections: SubsetConstantData.OriginalMeshSections;
 				UE_LOG(FGeometryCollectionSceneProxyLogging, VeryVerbose, TEXT("GetDynamicMeshElements, bUseSubSections=%d, NumSections=%d for %p."), bUsesSubSections, SectionArray.Num(), this);
 			#else
-				const TArray<FGeometryCollectionSection>& SectionArray = ConstantData->OriginalMeshSections;
+				const TArray<FGeometryCollectionSection>& SectionArray = SubsetConstantData.OriginalMeshSections;
 			#endif
 
 				// Grab the material proxies we'll be using for each section
@@ -1425,7 +1421,7 @@ void FGeometryCollectionSceneProxy::UseSubSections(bool bInUsesSubSections, bool
 void FGeometryCollectionSceneProxy::InitializeSubSections_RenderThread()
 {
 	// Exit now if there isn't any data
-	if (!ConstantData)
+	if (SubsetConstantData.OriginalMeshSections.IsEmpty())
 	{
 		SubSections.Empty();
 		SubSectionHitProxyIndexMap.Empty();
@@ -1434,9 +1430,9 @@ void FGeometryCollectionSceneProxy::InitializeSubSections_RenderThread()
 
 	// Retrieve the correct arrays depending on the dynamic state
 	const bool bIsDynamic = DynamicData && DynamicData->IsDynamic;
-	const TArray<FGeometryCollectionSection>& SectionArray = bIsDynamic ? Sections: ConstantData->OriginalMeshSections;
-	const TArray<FIntVector>& IndexArray = bIsDynamic ? ConstantData->Indices: ConstantData->OriginalMeshIndices;
-	const TArray<int32>& BoneMap = ConstantData->BoneMap;
+	const TArray<FGeometryCollectionSection>& SectionArray = bIsDynamic ? Sections: SubsetConstantData.OriginalMeshSections;
+	const TArray<FIntVector>& IndexArray = bIsDynamic ? SubsetConstantData.Indices: SubsetConstantData.OriginalMeshIndices;
+	const TArray<int32>& BoneMap = SubsetConstantData.BoneMap;
 
 	// Reserve sub sections array with a minimum of one transform per section
 	SubSections.Empty(SectionArray.Num());
@@ -1799,7 +1795,7 @@ Nanite::FResourceMeshInfo FNaniteGeometryCollectionSceneProxy::GetResourceMeshIn
 	return MoveTemp(OutInfo);
 }
 
-void FNaniteGeometryCollectionSceneProxy::SetConstantData_RenderThread(FGeometryCollectionConstantData* NewConstantData, bool ForceInit)
+void FNaniteGeometryCollectionSceneProxy::SetConstantData_RenderThread(const TUniquePtr<FGeometryCollectionConstantData>& NewConstantData, bool ForceInit)
 {
 	const TSharedPtr<FGeometryCollection, ESPMode::ThreadSafe> Collection = GeometryCollection->GetGeometryCollection();
 	const TManagedArray<int32>& TransformToGeometryIndices = Collection->TransformToGeometryIndex;
@@ -1833,8 +1829,6 @@ void FNaniteGeometryCollectionSceneProxy::SetConstantData_RenderThread(FGeometry
 		
 		InstanceHierarchyOffset.Emplace(NaniteData.HierarchyOffset);
 	}
-
-	delete NewConstantData;
 }
 
 void FNaniteGeometryCollectionSceneProxy::SetDynamicData_RenderThread(FGeometryCollectionDynamicData* NewDynamicData)
diff --git a/Engine/Source/Runtime/Experimental/GeometryCollectionEngine/Private/GeometryCollection/GeometryCollectionSceneProxy.h b/Engine/Source/Runtime/Experimental/GeometryCollectionEngine/Private/GeometryCollection/GeometryCollectionSceneProxy.h
--- a/Engine/Source/Runtime/Experimental/GeometryCollectionEngine/Private/GeometryCollection/GeometryCollectionSceneProxy.h
+++ b/Engine/Source/Runtime/Experimental/GeometryCollectionEngine/Private/GeometryCollection/GeometryCollectionSceneProxy.h
@@ -257,7 +257,16 @@ class FGeometryCollectionSceneProxy final : public FPrimitiveSceneProxy
 #endif
 
 	FGeometryCollectionDynamicData* DynamicData;
-	FGeometryCollectionConstantData* ConstantData;
+
+	struct SubsetConstantData
+	{
+		TArray<FIntVector> Indices;
+		TArray<int32> BoneMap;
+		uint32 NumTransforms;
+		TArray<FIntVector> OriginalMeshIndices;
+		TArray<FGeometryCollectionSection> OriginalMeshSections;
+		TArray<FMatrix44f> RestTransforms;
+	} SubsetConstantData;
 
 	bool bShowBoneColors;
 	bool bEnableBoneSelection;
@@ -288,13 +297,13 @@ public:
 	int32 GetRequiredIndexCount() const { return NumIndices; }
 
 	/** Called on render thread to setup static geometry for rendering */
-	void SetConstantData_RenderThread(FGeometryCollectionConstantData* NewConstantData, bool ForceInit = false);
+	void SetConstantData_RenderThread(const TUniquePtr<FGeometryCollectionConstantData>& NewConstantData, bool ForceInit = false);
 
 	/** Called on render thread to setup dynamic geometry for rendering */
 	void SetDynamicData_RenderThread(FGeometryCollectionDynamicData* NewDynamicData);
 
 	/** Called on render thread to construct the vertex definitions */
-	void BuildGeometry(const FGeometryCollectionConstantData* ConstantDataIn, TArray<FDynamicMeshVertex>& OutVertices, TArray<int32>& OutIndices, TArray<int32> &OutOriginalMeshIndices);
+	void BuildGeometry(const TUniquePtr<FGeometryCollectionConstantData>& ConstantDataIn, TArray<FDynamicMeshVertex>& OutVertices, TArray<int32>& OutIndices, TArray<int32> &OutOriginalMeshIndices);
 
 	/** Called on render thread to setup dynamic geometry for rendering */
 	virtual void GetDynamicMeshElements(const TArray<const FSceneView*>& Views, const FSceneViewFamily& ViewFamily, uint32 VisibilityMap, FMeshElementCollector& Collector) const override;
@@ -341,7 +350,7 @@ public:
 protected:
 
 	/** Create the rendering buffer resources */
-	void InitResources();
+	void InitResources(const TUniquePtr<FGeometryCollectionConstantData>& ConstantData);
 
 	/** Return the rendering buffer resources */
 	void ReleaseResources();
@@ -412,7 +421,7 @@ public:
 	virtual Nanite::FResourceMeshInfo GetResourceMeshInfo() const override;
 
 	/** Called on render thread to setup static geometry for rendering */
-	void SetConstantData_RenderThread(FGeometryCollectionConstantData* NewConstantData, bool ForceInit = false);
+	void SetConstantData_RenderThread(const TUniquePtr<FGeometryCollectionConstantData>& NewConstantData, bool ForceInit = false);
 
 	/** Called on render thread to setup dynamic geometry for rendering */
 	void SetDynamicData_RenderThread(FGeometryCollectionDynamicData* NewDynamicData);
diff --git a/Engine/Source/Runtime/Experimental/GeometryCollectionEngine/Public/GeometryCollection/GeometryCollectionComponent.h b/Engine/Source/Runtime/Experimental/GeometryCollectionEngine/Public/GeometryCollection/GeometryCollectionComponent.h
--- a/Engine/Source/Runtime/Experimental/GeometryCollectionEngine/Public/GeometryCollection/GeometryCollectionComponent.h
+++ b/Engine/Source/Runtime/Experimental/GeometryCollectionEngine/Public/GeometryCollection/GeometryCollectionComponent.h
@@ -1065,7 +1065,7 @@ protected:
 	int32 ISMPoolRootProxyMeshId = INDEX_NONE;
 
 	/** Populate the static geometry structures for the render thread. */
-	void InitConstantData(FGeometryCollectionConstantData* ConstantData) const;
+	void InitConstantData(TUniquePtr<FGeometryCollectionConstantData>& ConstantData) const;
 
 	/** Populate the dynamic particle data for the render thread. */
 	FGeometryCollectionDynamicData* InitDynamicData(bool bInitialization);
