Geometry collection : Add support for convex decomp for cluster convex collision
- Use leaf convex to create cluster convex
- Add dataflow node to generate the cluster convex
- Improved fracture editor convex vizualization ( display of convex union  when parent convex is missing )
- Add column to fracture editor outliner showing the number of convex and when a parent is using a union

diff --git a/Engine/Plugins/Experimental/ChaosEditor/Source/FractureEditor/Private/FractureToolConvex.cpp b/Engine/Plugins/Experimental/ChaosEditor/Source/FractureEditor/Private/FractureToolConvex.cpp
--- a/Engine/Plugins/Experimental/ChaosEditor/Source/FractureEditor/Private/FractureToolConvex.cpp
+++ b/Engine/Plugins/Experimental/ChaosEditor/Source/FractureEditor/Private/FractureToolConvex.cpp
@@ -225,6 +225,52 @@ TArray<UObject*> UFractureToolConvex::GetSettingsObjects() const
 	return Settings;
 }
 
+void UFractureToolConvex::AddHullVisualizationData(FGeometryCollection& Collection, const FTransform& OuterTransform, int32 TransformIdx, bool bIsUnionOfHulls)
+{
+	const TManagedArray<int32>* HasCustomConvex = FGeometryCollectionConvexUtility::GetCustomConvexFlags(&Collection, false);
+	const TManagedArray<TSet<int32>>& TransformToConvexIndices = Collection.GetAttribute<TSet<int32>>("TransformToConvexIndices", FTransformCollection::TransformGroup);
+	const TManagedArray<TUniquePtr<Chaos::FConvex>>& ConvexHull = Collection.GetAttribute<TUniquePtr<Chaos::FConvex>>("ConvexHull", "Convex");
+
+	const FTransform ComponentSpaceTransform = GeometryCollectionAlgo::GlobalMatrix(Collection.Transform, Collection.Parent, TransformIdx);
+	const FTransform WorldSpaceTransform = ComponentSpaceTransform * OuterTransform;
+
+	const TSet<int32>& ConvexIndices = TransformToConvexIndices[TransformIdx];
+	const bool bHasAnyConvex = (ConvexIndices.Num() > 0);
+	const bool bIsCustom = bHasAnyConvex && (HasCustomConvex ? bool((*HasCustomConvex)[TransformIdx]) : false);
+
+	if (bHasAnyConvex)
+	{
+		for (int32 ConvexIdx : TransformToConvexIndices[TransformIdx])
+		{
+			int32 HullPtsStart = HullPoints.Num();
+			for (Chaos::FVec3 Pt : ConvexHull[ConvexIdx]->GetVertices())
+			{
+				HullPoints.Add(WorldSpaceTransform.TransformPosition(Pt));
+			}
+			int32 NumPlanes = ConvexHull[ConvexIdx]->NumPlanes();
+			const Chaos::FConvexStructureData& HullData = ConvexHull[ConvexIdx]->GetStructureData();
+			for (int PlaneIdx = 0; PlaneIdx < NumPlanes; PlaneIdx++)
+			{
+				int32 NumPlaneVerts = HullData.NumPlaneVertices(PlaneIdx);
+				for (int32 PlaneVertexIdx = 0; PlaneVertexIdx < NumPlaneVerts; PlaneVertexIdx++)
+				{
+					int32 V0 = HullPtsStart + HullData.GetPlaneVertex(PlaneIdx, PlaneVertexIdx);
+					int32 V1 = HullPtsStart + HullData.GetPlaneVertex(PlaneIdx, (PlaneVertexIdx + 1) % NumPlaneVerts);
+					HullEdges.Add({ V0, V1, bIsCustom, bIsUnionOfHulls });
+				}
+			}
+		}
+	}
+	else
+	{
+		// no convex generated, we are using a union of the children convexes
+		for (int32 ChildTransformIndex : Collection.Children[TransformIdx])
+		{
+			AddHullVisualizationData(Collection, OuterTransform, ChildTransformIndex, /* bIsUnionOfHulls */ true);
+		}
+	}
+}
+
 void UFractureToolConvex::FractureContextChanged()
 {
 	TArray<FFractureToolContext> FractureContexts = GetFractureToolContexts();
@@ -242,53 +288,34 @@ void UFractureToolConvex::FractureContextChanged()
 		}
 
 		TManagedArray<int32>* HasCustomConvex = FGeometryCollectionConvexUtility::GetCustomConvexFlags(&Collection, false);
+		const TManagedArray<TSet<int32>>& TransformToConvexIndices = Collection.GetAttribute<TSet<int32>>("TransformToConvexIndices", FTransformCollection::TransformGroup);
+		const TManagedArray<TUniquePtr<Chaos::FConvex>>& ConvexHull = Collection.GetAttribute<TUniquePtr<Chaos::FConvex>>("ConvexHull", "Convex");
 
 		int32 CollectionIdx = VisualizedCollections.Add(FractureContext.GetGeometryCollectionComponent());
 
 		FTransform OuterTransform = FractureContext.GetTransform();
 		for (int32 TransformIdx : FractureContext.GetSelection())
 		{
-			FTransform InnerTransform = GeometryCollectionAlgo::GlobalMatrix(Collection.Transform, Collection.Parent, TransformIdx);
-			FTransform CombinedTransform = InnerTransform * OuterTransform;
-			bool bIsCustom = HasCustomConvex ? bool((*HasCustomConvex)[TransformIdx]) : false;
-
-			const TManagedArray<TSet<int32>>& TransformToConvexIndices = Collection.GetAttribute<TSet<int32>>("TransformToConvexIndices", FTransformCollection::TransformGroup);
-			const TManagedArray<TUniquePtr<Chaos::FConvex>>& ConvexHull = Collection.GetAttribute<TUniquePtr<Chaos::FConvex>>("ConvexHull", "Convex");
-
 			EdgesMappings.AddMapping(CollectionIdx, TransformIdx, HullEdges.Num());
 
-			for (int32 ConvexIdx : TransformToConvexIndices[TransformIdx])
-			{
-				int32 HullPtsStart = HullPoints.Num();
-				for (Chaos::FVec3 Pt : ConvexHull[ConvexIdx]->GetVertices())
-				{
-					HullPoints.Add(CombinedTransform.TransformPosition(Pt));
-				}
-				int32 NumPlanes = ConvexHull[ConvexIdx]->NumPlanes();
-				const Chaos::FConvexStructureData& HullData = ConvexHull[ConvexIdx]->GetStructureData();
-				for (int PlaneIdx = 0; PlaneIdx < NumPlanes; PlaneIdx++)
-				{
-					int32 NumPlaneVerts = HullData.NumPlaneVertices(PlaneIdx);
-					for (int32 PlaneVertexIdx = 0; PlaneVertexIdx < NumPlaneVerts; PlaneVertexIdx++)
-					{
-						int32 V0 = HullPtsStart + HullData.GetPlaneVertex(PlaneIdx, PlaneVertexIdx);
-						int32 V1 = HullPtsStart + HullData.GetPlaneVertex(PlaneIdx, (PlaneVertexIdx + 1) % NumPlaneVerts);
-						HullEdges.Add({ V0, V1, bIsCustom });
-					}
-				}
-			}
+			AddHullVisualizationData(Collection, OuterTransform, TransformIdx);
 		}
 	}
 }
 
 void UFractureToolConvex::Render(const FSceneView* View, FViewport* Viewport, FPrimitiveDrawInterface* PDI)
 {
+	const bool bCaptureSeeThroughLines = ConvexSettings->bSeeThroughLines;
+	const float CaptureLineThickness = ConvexSettings->LineThickness;
 	EnumerateVisualizationMapping(EdgesMappings, HullEdges.Num(), [&](int32 Idx, FVector ExplodedVector)
 	{
 		const FEdgeVisInfo& Edge = HullEdges[Idx];
 		FVector P1 = HullPoints[Edge.A] + ExplodedVector;
 		FVector P2 = HullPoints[Edge.B] + ExplodedVector;
-		PDI->DrawLine(P1, P2, Edge.bIsCustom ? FLinearColor::Red : FLinearColor::Green, SDPG_Foreground, 0.0f, 0.001f);
+		FLinearColor Color = FLinearColor::Blue;
+		if (Edge.bIsCustom) { Color = FLinearColor(1, 0, 1); /* purple */ }
+		if (Edge.bIsUnion)	{ Color = FLinearColor::Red;  }
+		PDI->DrawLine(P1, P2, Color, bCaptureSeeThroughLines? SDPG_Foreground: SDPG_World, CaptureLineThickness, 0.005f);
 	});
 }
 
diff --git a/Engine/Plugins/Experimental/ChaosEditor/Source/FractureEditor/Private/FractureToolConvex.h b/Engine/Plugins/Experimental/ChaosEditor/Source/FractureEditor/Private/FractureToolConvex.h
--- a/Engine/Plugins/Experimental/ChaosEditor/Source/FractureEditor/Private/FractureToolConvex.h
+++ b/Engine/Plugins/Experimental/ChaosEditor/Source/FractureEditor/Private/FractureToolConvex.h
@@ -42,10 +42,18 @@ public:
 	UPROPERTY(EditAnywhere, Category = AutomaticOverlapRemoval, meta = (DisplayName = "Max Removal Fraction", ClampMin = ".01", ClampMax = "1"))
 	double FractionAllowRemove = .5;
 
-		/** Delete convex hulls from selected clusters.  Does not affect hulls on leaves. */
+	/** Delete convex hulls from selected clusters.  Does not affect hulls on leaves. */
 	UFUNCTION(CallInEditor, Category = Custom, meta = (DisplayName = "Delete From Selected"))
 	void DeleteFromSelected();
 
+	/** When enabled, convex visualization lines will show through the actual geometry */
+	UPROPERTY(EditAnywhere, Category = Visualization, meta = (DisplayName = "See Through Lines"))
+	bool bSeeThroughLines = false;
+
+	/** line thickness*/
+	UPROPERTY(EditAnywhere, Category = Visualization, meta = (DisplayName = "Line Thickness", ClampMin = ".00", ClampMax = "1.0"))
+	float LineThickness = 0.5f;
+
 	// Note: this feature puts multiple convexes on a single bone, which isn't supported by sim yet
 #if 0
 	/** Promote (and save) child convex hulls on to the selected bone(s) */
@@ -140,10 +148,14 @@ protected:
 		EdgesMappings.Empty();
 	}
 
+	/* return true if the convex is a union of the children */
+	void AddHullVisualizationData(FGeometryCollection& Collection, const FTransform& OuterTransform, int32 TransformIdx, bool bIsUnionOfHulls = false);
+
 	struct FEdgeVisInfo
 	{
 		int32 A, B;
-		bool bIsCustom; // allow for different coloring of hulls that have been manually set vs auto-generated ones
+		int8 bIsCustom : 1; // allow for different coloring of hulls that have been manually set vs auto-generated ones
+		int8 bIsUnion : 1; // this edge is part of a union of hulls ( cluster does not have a convex and is using a union of the children )
 	};
 
 	TArray<FVector> HullPoints;
diff --git a/Engine/Plugins/Experimental/ChaosEditor/Source/FractureEditor/Private/SGeometryCollectionOutliner.cpp b/Engine/Plugins/Experimental/ChaosEditor/Source/FractureEditor/Private/SGeometryCollectionOutliner.cpp
--- a/Engine/Plugins/Experimental/ChaosEditor/Source/FractureEditor/Private/SGeometryCollectionOutliner.cpp
+++ b/Engine/Plugins/Experimental/ChaosEditor/Source/FractureEditor/Private/SGeometryCollectionOutliner.cpp
@@ -20,6 +20,7 @@
 #include "GeometryCollectionOutlinerDragDrop.h"
 #include "GeometryCollection/GeometryCollectionUtility.h"
 #include "GeometryCollection/Facades/CollectionAnchoringFacade.h"
+#include "GeometryCollection/Facades/CollectionHierarchyFacade.h"
 #include "PhysicsProxy/GeometryCollectionPhysicsProxy.h"
 
 #include UE_INLINE_GENERATED_CPP_BY_NAME(SGeometryCollectionOutliner)
@@ -230,6 +231,8 @@ TSharedRef<SWidget> SGeometryCollectionOutlinerRow::GenerateWidgetForColumn(cons
 		return Item->MakeRemovalTimeColumnWidget();
 	if (ColumnName == SGeometryCollectionOutlinerColumnID::ImportedCollisions)
 		return Item->MakeImportedCollisionsColumnWidget();
+	if (ColumnName == SGeometryCollectionOutlinerColumnID::ConvexCount)
+		return Item->MakeConvexCountColumnWidget();
 	return Item->MakeEmptyColumnWidget();
 }
 
@@ -341,6 +344,13 @@ void SGeometryCollectionOutliner::RegenerateHeader()
 		break;
 		
 	case EOutlinerColumnMode::Collision:
+		HeaderRowWidget->AddColumn(
+			SHeaderRow::Column(SGeometryCollectionOutlinerColumnID::ConvexCount)
+			.DefaultLabel(LOCTEXT("GCOutliner_Column_ConvexCount", "Convex Count"))
+			.DefaultTooltip(LOCTEXT("GCOutliner_Column_ConvexCount_ToolTip", "Number of convex collisions)"))
+			.HAlignHeader(EHorizontalAlignment::HAlign_Right)
+			.FillWidth(CustomFillWidth)
+		);
 		HeaderRowWidget->AddColumn(
 			SHeaderRow::Column(SGeometryCollectionOutlinerColumnID::ImportedCollisions)
 				.DefaultLabel(LOCTEXT("GCOutliner_Column_ImportedCollisions", "Imported Collisions"))
@@ -623,6 +633,7 @@ FGeometryCollectionItemDataFacade::FGeometryCollectionItemDataFacade(FManagedArr
 	, SimulationTypeAttribute(InCollection, "SimulationType", DataCollectionGroup)
 	, HasSourceCollisionAttribute(InCollection, "HasSourceCollision", DataCollectionGroup)
 	, SourceCollisionUsedAttribute(InCollection, "SourceCollisionUsed", DataCollectionGroup)
+	, ConvexCountAttribute(InCollection, "ConvexCount", DataCollectionGroup)
 {
 }
 
@@ -679,6 +690,11 @@ bool FGeometryCollectionItemDataFacade::IsSourceCollisionUsed(int32 Index) const
 	return GetAttributeValue(SourceCollisionUsedAttribute, Index, false);
 }
 
+int32 FGeometryCollectionItemDataFacade::GetConvexCount(int32 Index) const
+{
+	return GetAttributeValue(ConvexCountAttribute, Index, 0);
+}
+
 void FGeometryCollectionItemDataFacade::FillFromGeometryCollectionComponent(const UGeometryCollectionComponent& GeometryCollectionComponent, EOutlinerColumnMode ColumnMode)
 {
 	const UGeometryCollection* RestCollection = GeometryCollectionComponent.GetRestCollection();
@@ -785,6 +801,40 @@ void FGeometryCollectionItemDataFacade::FillFromGeometryCollectionComponent(cons
 					SourceCollisionUsed[Index] = RestCollection->bImportCollisionFromSource;
 				}
 			}
+			
+			TManagedArrayAccessor<TSet<int32>> GCTransformToConvexIndicesAttribute(GeometryCollection, "TransformToConvexIndices", FGeometryCollection::TransformGroup);
+			TManagedArray<int32>& ConvexCountArray = ConvexCountAttribute.Add();
+			ConvexCountAttribute.Fill(0);
+
+			if (GCTransformToConvexIndicesAttribute.IsValid())
+			{
+				const TManagedArray<TSet<int32>>& GCTransformToConvexIndices = GCTransformToConvexIndicesAttribute.Get();
+
+				Chaos::Facades::FCollectionHierarchyFacade HierarchyFacade(GeometryCollection);
+				const TArray<int32> TransformIndices = HierarchyFacade.GetTransformArrayInDepthFirstOrder();
+
+				for (int32 TransformIndex: TransformIndices)
+				{
+					const int32 ConvexCount = GCTransformToConvexIndices[TransformIndex].Num();
+					if (ConvexCount > 0)
+					{
+						ConvexCountArray[TransformIndex] = ConvexCount;
+					}
+					// if count == 0 then we have already accumulated the children in the parents, no need to do anything
+					// so now just pass it to the direct parent ( we parse the index in a depth first manner ) 
+					int32 ParentTransformIndex = GeometryCollection.Parent[TransformIndex];
+					if (ParentTransformIndex != INDEX_NONE)
+					{		
+						// if parent has no convex then it will be a union of the aggregated children
+						const int32 ParentConvexCount = GCTransformToConvexIndices[ParentTransformIndex].Num();
+						if (ParentConvexCount == 0)
+						{
+							// negative values means union 
+							ConvexCountArray[ParentTransformIndex] -= FMath::Abs(ConvexCountArray[TransformIndex]);
+						}
+					}
+				}
+			}			
 		}
 		else if (ColumnMode == EOutlinerColumnMode::Removal)
 		{
@@ -1265,6 +1315,25 @@ TSharedRef<SWidget> FGeometryCollectionTreeItemBone::MakeImportedCollisionsColum
 		];
 }
 
+TSharedRef<SWidget> FGeometryCollectionTreeItemBone::MakeConvexCountColumnWidget() const
+{
+	const FGeometryCollectionItemDataFacade& DataCollectionFacade = GetDataCollectionFacade();
+	const int32 ConvexCount = DataCollectionFacade.GetConvexCount(BoneIndex);
+	const bool bIsUnionOfConvex = (ConvexCount < 0);
+
+	const FText ItemText = bIsUnionOfConvex? FText::Format(LOCTEXT("GCOutliner_ConvexCount_Format", "Union of {0}"), FMath::Abs(ConvexCount)): FText::AsNumber(ConvexCount);
+
+	return SNew(SHorizontalBox)
+		+ SHorizontalBox::Slot()
+		.Padding(12.f, 0.f)
+		.HAlign(HAlign_Right)
+		[
+			SNew(STextBlock)
+			.Text(ItemText)
+		.ColorAndOpacity(ItemColor)
+		];
+}
+
 TSharedRef<SWidget> FGeometryCollectionTreeItemBone::MakeInitialStateColumnWidget() const
 {
 	const FGeometryCollectionItemDataFacade& DataCollectionFacade = GetDataCollectionFacade();
diff --git a/Engine/Plugins/Experimental/ChaosEditor/Source/FractureEditor/Private/SGeometryCollectionOutliner.h b/Engine/Plugins/Experimental/ChaosEditor/Source/FractureEditor/Private/SGeometryCollectionOutliner.h
--- a/Engine/Plugins/Experimental/ChaosEditor/Source/FractureEditor/Private/SGeometryCollectionOutliner.h
+++ b/Engine/Plugins/Experimental/ChaosEditor/Source/FractureEditor/Private/SGeometryCollectionOutliner.h
@@ -117,6 +117,7 @@ public:
 	FRemoveOnBreakData GetRemoveOnBreakData(int32 Index) const;
 	bool HasSourceCollision(int32 Index) const;
 	bool IsSourceCollisionUsed(int32 Index) const;
+	int32 GetConvexCount(int32 Index) const;
 
 private:
 	FManagedArrayCollection&			DataCollection;
@@ -135,6 +136,7 @@ private:
 	TManagedArrayAccessor<int32>		SimulationTypeAttribute;
 	TManagedArrayAccessor<bool>			HasSourceCollisionAttribute;
 	TManagedArrayAccessor<bool>			SourceCollisionUsedAttribute;
+	TManagedArrayAccessor<int32>		ConvexCountAttribute;
 };
 
 class FGeometryCollectionTreeItemComponent : public FGeometryCollectionTreeItem
@@ -213,6 +215,7 @@ public:
 	TSharedRef<SWidget> MakePostBreakTimeColumnWidget() const;
 	TSharedRef<SWidget> MakeRemovalTimeColumnWidget() const;
 	TSharedRef<SWidget> MakeImportedCollisionsColumnWidget() const;
+	TSharedRef<SWidget> MakeConvexCountColumnWidget() const;
 	TSharedRef<SWidget> MakeEmptyColumnWidget() const;
 	virtual void GetChildren(FGeometryCollectionTreeItemList& OutChildren) override;
 	virtual int32 GetBoneIndex() const override { return BoneIndex; }
@@ -254,6 +257,7 @@ namespace SGeometryCollectionOutlinerColumnID
 	const FName PostBreakTime("PostBreakTime");
 	const FName RemovalTime("RemovalTime");
 	// Collision Column Mode
+	const FName ConvexCount("Convex Count");
 	const FName ImportedCollisions("ImportedCollisions");
 }
 
diff --git a/Engine/Plugins/Experimental/GeometryCollectionPlugin/Source/GeometryCollectionNodes/Private/Dataflow/GeometryCollectionUtilityNodes.cpp b/Engine/Plugins/Experimental/GeometryCollectionPlugin/Source/GeometryCollectionNodes/Private/Dataflow/GeometryCollectionUtilityNodes.cpp
--- a/Engine/Plugins/Experimental/GeometryCollectionPlugin/Source/GeometryCollectionNodes/Private/Dataflow/GeometryCollectionUtilityNodes.cpp
+++ b/Engine/Plugins/Experimental/GeometryCollectionPlugin/Source/GeometryCollectionNodes/Private/Dataflow/GeometryCollectionUtilityNodes.cpp
@@ -17,10 +17,20 @@ namespace Dataflow
 		static const FLinearColor CDefaultNodeBodyTintColor = FLinearColor(0.f, 0.f, 0.f, 0.5f);
 
 		DATAFLOW_NODE_REGISTER_CREATION_FACTORY(FCreateNonOverlappingConvexHullsDataflowNode);
-
+		DATAFLOW_NODE_REGISTER_CREATION_FACTORY(FGenerateClusterConvexHullsFromLeafHullsDataflowNode);
 	}
 }
 
+FCreateNonOverlappingConvexHullsDataflowNode::FCreateNonOverlappingConvexHullsDataflowNode(const Dataflow::FNodeParameters& InParam, FGuid InGuid)
+	: FDataflowNode(InParam, InGuid)
+{
+	RegisterInputConnection(&Collection);
+	RegisterInputConnection(&CanRemoveFraction);
+	RegisterInputConnection(&SimplificationDistanceThreshold);
+	RegisterInputConnection(&CanExceedFraction);
+	RegisterInputConnection(&OverlapRemovalShrinkPercent);
+	RegisterOutputConnection(&Collection);
+}
 
 void FCreateNonOverlappingConvexHullsDataflowNode::Evaluate(Dataflow::FContext& Context, const FDataflowOutput* Out) const
 {
@@ -46,3 +56,35 @@ void FCreateNonOverlappingConvexHullsDataflowNode::Evaluate(Dataflow::FContext&
 		}
 	}
 }
+
+FGenerateClusterConvexHullsFromLeafHullsDataflowNode::FGenerateClusterConvexHullsFromLeafHullsDataflowNode(const Dataflow::FNodeParameters& InParam, FGuid InGuid)
+	: FDataflowNode(InParam, InGuid)
+{
+	RegisterInputConnection(&Collection);
+	RegisterInputConnection(&ConvexCount);
+	RegisterInputConnection(&ErrorTolerance);
+
+	RegisterOutputConnection(&Collection);
+}
+
+void FGenerateClusterConvexHullsFromLeafHullsDataflowNode::Evaluate(Dataflow::FContext& Context, const FDataflowOutput* Out) const
+{
+	if (Out->IsA<FManagedArrayCollection>(&Collection))
+	{
+		const FManagedArrayCollection& InCollection = GetValue<FManagedArrayCollection>(Context, &Collection);
+
+		if (TUniquePtr<FGeometryCollection> GeomCollection = TUniquePtr<FGeometryCollection>(InCollection.NewCopy<FGeometryCollection>()))
+		{
+			const int32 InConvexCount = GetValue(Context, &ConvexCount);
+			const double InErrorToleranceInCm = GetValue(Context, &ErrorTolerance);
+
+			FGeometryCollectionConvexUtility::GenerateClusterConvexHullsFromLeafHulls(
+				*GeomCollection,
+				InConvexCount,
+				InErrorToleranceInCm
+			);
+
+			SetValue<FManagedArrayCollection>(Context, static_cast<const FManagedArrayCollection>(*GeomCollection), &Collection);
+		}
+	}
+}
\ No newline at end of file
diff --git a/Engine/Plugins/Experimental/GeometryCollectionPlugin/Source/GeometryCollectionNodes/Public/Dataflow/GeometryCollectionUtilityNodes.h b/Engine/Plugins/Experimental/GeometryCollectionPlugin/Source/GeometryCollectionNodes/Public/Dataflow/GeometryCollectionUtilityNodes.h
--- a/Engine/Plugins/Experimental/GeometryCollectionPlugin/Source/GeometryCollectionNodes/Public/Dataflow/GeometryCollectionUtilityNodes.h
+++ b/Engine/Plugins/Experimental/GeometryCollectionPlugin/Source/GeometryCollectionNodes/Public/Dataflow/GeometryCollectionUtilityNodes.h
@@ -51,16 +51,40 @@ public:
 	UPROPERTY(EditAnywhere, Category = "Convex", meta = (DataflowInput, UIMin = 0.f, UIMax = 100.f))
 	float OverlapRemovalShrinkPercent = 0.f;
 
-	FCreateNonOverlappingConvexHullsDataflowNode(const Dataflow::FNodeParameters& InParam, FGuid InGuid = FGuid::NewGuid())
-		: FDataflowNode(InParam, InGuid)
-	{
-		RegisterInputConnection(&Collection);
-		RegisterInputConnection(&CanRemoveFraction);
-		RegisterInputConnection(&SimplificationDistanceThreshold);
-		RegisterInputConnection(&CanExceedFraction);
-		RegisterInputConnection(&OverlapRemovalShrinkPercent);
-		RegisterOutputConnection(&Collection);
-	}
+	FCreateNonOverlappingConvexHullsDataflowNode(const Dataflow::FNodeParameters& InParam, FGuid InGuid = FGuid::NewGuid());
+
+	virtual void Evaluate(Dataflow::FContext& Context, const FDataflowOutput* Out) const override;
+
+};
+
+/**
+ *
+ * Generates cluster convex hulls for leafs hulls
+ *
+ */
+USTRUCT(meta = (DataflowGeometryCollection))
+struct FGenerateClusterConvexHullsFromLeafHullsDataflowNode : public FDataflowNode
+{
+	GENERATED_USTRUCT_BODY()
+	DATAFLOW_NODE_DEFINE_INTERNAL(FGenerateClusterConvexHullsFromLeafHullsDataflowNode, "GenerateClusterConvexHullsFromLeafHulls", "GeometryCollection|Utilities", "")
+
+public:
+	UPROPERTY(meta = (DataflowInput, DataflowOutput))
+	FManagedArrayCollection Collection;
+
+	/** Maximum number of convex to generate for a specific cluster. Will be ignored if error tolerance is used instead */
+	UPROPERTY(EditAnywhere, Category = "Convex", meta = (DataflowInput, EditCondition = "ErrorTolerance == 0"))
+	int32 ConvexCount = 2;
+	
+	/** 
+	* Error tolerance to use to decide to merge leaf convex together. 
+	* This is in centimeters and represents the side of a cube, the volume of which will be used as threshold
+	* to know if the volume of the generated convex is too large compared to the sum of the volume of the leaf convex
+	*/
+	UPROPERTY(EditAnywhere, Category = "Convex", meta = (DataflowInput, UIMin = "0", UIMax = "100.", Units = cm))
+	double ErrorTolerance = 0.0;
+
+	FGenerateClusterConvexHullsFromLeafHullsDataflowNode(const Dataflow::FNodeParameters& InParam, FGuid InGuid = FGuid::NewGuid());
 
 	virtual void Evaluate(Dataflow::FContext& Context, const FDataflowOutput* Out) const override;
 
diff --git a/Engine/Source/Runtime/Experimental/Chaos/Private/GeometryCollection/Facades/CollectionHierarchyFacade.cpp b/Engine/Source/Runtime/Experimental/Chaos/Private/GeometryCollection/Facades/CollectionHierarchyFacade.cpp
--- a/Engine/Source/Runtime/Experimental/Chaos/Private/GeometryCollection/Facades/CollectionHierarchyFacade.cpp
+++ b/Engine/Source/Runtime/Experimental/Chaos/Private/GeometryCollection/Facades/CollectionHierarchyFacade.cpp
@@ -55,6 +55,16 @@ namespace Chaos::Facades
 		return GetRootIndices(ParentAttribute);
 	}
 
+	TArray<int32> FCollectionHierarchyFacade::GetChildrenAsArray(int32 TransformIndex) const
+	{
+		TArray<int32> ChildrenArray;
+		if (ChildrenAttribute.IsValid())
+		{
+			return ChildrenAttribute.Get()[TransformIndex].Array();
+		}
+		return ChildrenArray;
+	}
+
 	int32 FCollectionHierarchyFacade::GetInitialLevel(int32 TransformIndex) const
 	{
 		if (HasLevelAttribute())
@@ -134,4 +144,69 @@ namespace Chaos::Facades
 
 		return Roots;
 	}
+
+	TArray<int32> FCollectionHierarchyFacade::GetTransformArrayInDepthFirstOrder() const
+	{
+		TArray<int32> OrderedTransforms;
+		if (ParentAttribute.IsValid() && ChildrenAttribute.IsValid())
+		{
+			const TManagedArray<int32>& Parent = ParentAttribute.Get();
+			const TManagedArray<TSet<int32>>& Children = ChildrenAttribute.Get();
+
+			//traverse cluster hierarchy in depth first and record order
+			struct FClusterProcessing
+			{
+				int32 TransformGroupIndex;
+				enum
+				{
+					None,
+					VisitingChildren
+				} State;
+
+				FClusterProcessing(int32 InIndex) : TransformGroupIndex(InIndex), State(None) {};
+			};
+
+			TArray<FClusterProcessing> ClustersToProcess;
+			//enqueue all roots
+			for (int32 TransformGroupIndex = 0; TransformGroupIndex < Parent.Num(); TransformGroupIndex++)
+			{
+				if (Parent[TransformGroupIndex] == INDEX_NONE)
+				{
+					ClustersToProcess.Emplace(TransformGroupIndex);
+				}
+			}
+
+			OrderedTransforms.Reserve(Parent.Num());
+
+			while (ClustersToProcess.Num())
+			{
+				FClusterProcessing CurCluster = ClustersToProcess.Pop();
+				const int32 ClusterTransformIdx = CurCluster.TransformGroupIndex;
+				if (CurCluster.State == FClusterProcessing::VisitingChildren)
+				{
+					//children already visited
+					OrderedTransforms.Add(ClusterTransformIdx);
+				}
+				else
+				{
+					if (Children[ClusterTransformIdx].Num())
+					{
+						CurCluster.State = FClusterProcessing::VisitingChildren;
+						ClustersToProcess.Add(CurCluster);
+
+						//order of children doesn't matter as long as all children appear before parent
+						for (int32 ChildIdx : Children[ClusterTransformIdx])
+						{
+							ClustersToProcess.Emplace(ChildIdx);
+						}
+					}
+					else
+					{
+						OrderedTransforms.Add(ClusterTransformIdx);
+					}
+				}
+			}
+		}
+		return OrderedTransforms;
+	}
 }
\ No newline at end of file
diff --git a/Engine/Source/Runtime/Experimental/Chaos/Private/GeometryCollection/GeometryCollectionClusteringUtility.cpp b/Engine/Source/Runtime/Experimental/Chaos/Private/GeometryCollection/GeometryCollectionClusteringUtility.cpp
--- a/Engine/Source/Runtime/Experimental/Chaos/Private/GeometryCollection/GeometryCollectionClusteringUtility.cpp
+++ b/Engine/Source/Runtime/Experimental/Chaos/Private/GeometryCollection/GeometryCollectionClusteringUtility.cpp
@@ -806,32 +806,37 @@ void FGeometryCollectionClusteringUtility::ValidateResults(FGeometryCollection*
 	}
 }
 
-void FGeometryCollectionClusteringUtility::GetLeafBones(const FGeometryCollection* GeometryCollection, int BoneIndex, bool bOnlyRigids, TArray<int32>& LeafBonesOut)
+void FGeometryCollectionClusteringUtility::GetLeafBones(const FManagedArrayCollection* Collection, int BoneIndex, bool bOnlyRigids, TArray<int32>& LeafBonesOut)
 {
-	if (!ensure(BoneIndex >= 0))
+	if (!ensure(BoneIndex >= 0 && Collection != nullptr))
 	{
 		return;
 	}
 
-	const TManagedArray<TSet<int32>>& Children = GeometryCollection->Children;
-	const TManagedArray<int32>& SimulationType = GeometryCollection->SimulationType;
+	const TManagedArrayAccessor<TSet<int32>> ChildrenAttribute(*Collection, FGeometryCollection::ChildrenAttribute, FGeometryCollection::TransformGroup);
+	const TManagedArrayAccessor<int32> SimulationTypeAttribute(*Collection , FGeometryCollection::SimulationTypeAttribute, FGeometryCollection::TransformGroup);
 
-	if (!bOnlyRigids && Children[BoneIndex].Num() == 0)
-	{
-		LeafBonesOut.Push(BoneIndex);
-	}
-	else if (bOnlyRigids && GeometryCollection->IsRigid(BoneIndex))
+	if (ChildrenAttribute.IsValid() && SimulationTypeAttribute.IsValid())
 	{
-		LeafBonesOut.Push(BoneIndex);
-	}
-	else if (Children[BoneIndex].Num() > 0)
-	{
-		for (int32 ChildElement : Children[BoneIndex])
+		const TManagedArray<TSet<int32>>& Children = ChildrenAttribute.Get();
+		const TManagedArray<int32>& SimulationType = SimulationTypeAttribute.Get();
+
+		if (!bOnlyRigids && Children[BoneIndex].Num() == 0)
 		{
-			GetLeafBones(GeometryCollection, ChildElement, bOnlyRigids, LeafBonesOut);
+			LeafBonesOut.Push(BoneIndex);
+		}
+		else if (bOnlyRigids && SimulationType[BoneIndex] == FGeometryCollection::ESimulationTypes::FST_Rigid)
+		{
+			LeafBonesOut.Push(BoneIndex);
+		}
+		else if (Children[BoneIndex].Num() > 0)
+		{
+			for (int32 ChildElement : Children[BoneIndex])
+			{
+				GetLeafBones(Collection, ChildElement, bOnlyRigids, LeafBonesOut);
+			}
 		}
 	}
-
 }
 
 void FGeometryCollectionClusteringUtility::MoveUpOneHierarchyLevel(FGeometryCollection* GeometryCollection, const TArray<int32>& SelectedBones)
diff --git a/Engine/Source/Runtime/Experimental/Chaos/Private/GeometryCollection/GeometryCollectionConvexUtility.cpp b/Engine/Source/Runtime/Experimental/Chaos/Private/GeometryCollection/GeometryCollectionConvexUtility.cpp
--- a/Engine/Source/Runtime/Experimental/Chaos/Private/GeometryCollection/GeometryCollectionConvexUtility.cpp
+++ b/Engine/Source/Runtime/Experimental/Chaos/Private/GeometryCollection/GeometryCollectionConvexUtility.cpp
@@ -6,9 +6,13 @@
 #include "GeometryCollection/GeometryCollection.h"
 #include "GeometryCollection/GeometryCollectionAlgo.h"
 #include "GeometryCollection/GeometryCollectionProximityUtility.h"
+#include "GeometryCollection/Facades/CollectionTransformFacade.h"
+#include "GeometryCollection/Facades/CollectionHierarchyFacade.h"
 #include "CompGeom/ConvexHull3.h"
 #include "Templates/Sorting.h"
 #include "Spatial/PointHashGrid3.h"
+#include "CompGeom/ConvexDecomposition3.h"
+#include "GeometryCollection/GeometryCollectionClusteringUtility.h"
 
 bool UseVolumeToComputeRelativeSize = false;
 FAutoConsoleVariableRef CVarUseVolumeToComputeRelativeSize(TEXT("p.gc.UseVolumeToComputeRelativeSize"), UseVolumeToComputeRelativeSize, TEXT("Use Volume To Compute RelativeSize instead of the side of the cubic volume (def: false)"));
@@ -1353,31 +1357,155 @@ FGeometryCollectionConvexUtility::FGeometryCollectionConvexData FGeometryCollect
 	ConvexHull = MoveTemp(UseLeafHulls->Hulls);
 
 	// clear all null and empty hulls
-	TArray<int32> EmptyConvex;
-	for (int32 ConvexIdx = 0; ConvexIdx < ConvexHull.Num(); ConvexIdx++)
+	RemoveEmptyConvexHulls(*GeometryCollection);
+
+	checkSlow(FGeometryCollectionConvexUtility::ValidateConvexData(GeometryCollection));
+
+	return { TransformToConvexIndices, ConvexHull };
+}
+
+
+void FGeometryCollectionConvexUtility::GenerateClusterConvexHullsFromLeafHulls(FGeometryCollection& Collection,	int32 ConvexCount, double ErrorToleranceInCm)
+{
+	static FName ConvexGroupName("Convex");
+	static FName ConvexHullAttributeName("ConvexHull");
+	static FName TransformToConvexIndicesName("TransformToConvexIndices");
+
+	TManagedArrayAccessor<TUniquePtr<Chaos::FConvex>> ConvexHullAttribute(Collection, ConvexHullAttributeName, ConvexGroupName);
+	TManagedArrayAccessor<TSet<int32>> TransformToConvexIndicesAttribute(Collection, TransformToConvexIndicesName, FGeometryCollection::TransformGroup);
+
+	// if any of the leaves hulls has been computed those two attributes should be valid
+	if (!ConvexHullAttribute.IsValid() || !TransformToConvexIndicesAttribute.IsValid())
+	{
+		return;
+	}
+
+	GeometryCollection::Facades::FCollectionTransformFacade TransformFacade(Collection);
+	Chaos::Facades::FCollectionHierarchyFacade HierarchyFacade(Collection);
+
+	FGeometryCollectionProximityUtility ProximityUtility(&Collection);
+	ProximityUtility.RequireProximity();
+	const TManagedArray<TSet<int32>>& Proximity = Collection.GetAttribute<TSet<int32>>("Proximity", FGeometryCollection::GeometryGroup);
+
+	const TArray<FTransform> GlobalTransforms = TransformFacade.ComputeCollectionSpaceTransforms();
+	const TArray<int32> DepthFirstTransformIndices = HierarchyFacade.GetTransformArrayInDepthFirstOrder();
+
+	const TManagedArrayAccessor<int32> SimulationTypeAttribute(Collection, FGeometryCollection::SimulationTypeAttribute, FTransformCollection::TransformGroup);
+
+	for (int32 TransformIndex: DepthFirstTransformIndices)
 	{
-		if (ConvexHull[ConvexIdx].IsValid())
+		// only do this for clusters
+		const bool bIsCluster = (SimulationTypeAttribute.IsValid() && SimulationTypeAttribute.Get()[TransformIndex] == FGeometryCollection::ESimulationTypes::FST_Clustered);
+		if (bIsCluster)
 		{
-			if (ConvexHull[ConvexIdx]->NumVertices() == 0)
+			// compute using the leaf nodes 
+			// we cannot use the direct cluster children because they do not have proximity data that is stored at the geo level )
+			TArray<int32> SourceTransformIndices;
+			FGeometryCollectionClusteringUtility::GetLeafBones(&Collection, TransformIndex, true, SourceTransformIndices);
+
+			UE::Geometry::FDynamicMesh3 CombinedGeometry;
+			if (SourceTransformIndices.Num() > 0)
 			{
-				ConvexHull[ConvexIdx].Reset();
-				EmptyConvex.Add(ConvexIdx);
+				const FTransform& ParentTransform = GlobalTransforms[TransformIndex];
+
+				TManagedArray<TSet<int32>>& TransformToConvexIndices = TransformToConvexIndicesAttribute.Modify();
+				TManagedArray<TUniquePtr<Chaos::FConvex>>& ConvexHull = ConvexHullAttribute.Modify();
+
+				struct FHullInfo
+				{
+					Chaos::FConvex* Convex = nullptr;
+					FTransform Transform;
+				};
+
+				// build information required for merging hulls 
+				TArray<FHullInfo> Hulls;
+				TMap<int32, TArray<int32>> TransformToHullIdx;
+				TArray<TPair<int32, int32>> HullProximity;
+
+				for (int32 SourceTransformIndex : SourceTransformIndices)
+				{
+					const FTransform InnerTransform = GlobalTransforms[SourceTransformIndex];
+					FTransform ChildToParentTransform = InnerTransform.GetRelativeTransform(ParentTransform);
+					int32 HullIdxStart = Hulls.Num();
+					for (int32 SourceConvexIdx : TransformToConvexIndices[SourceTransformIndex])
+					{
+						FHullInfo HullInfo;
+						HullInfo.Convex = ConvexHull[SourceConvexIdx].Get();
+						HullInfo.Transform = ChildToParentTransform;
+						const int32 HullIdx = Hulls.Emplace(HullInfo);
+						TransformToHullIdx.FindOrAdd(SourceTransformIndex).Add(HullIdx);
+					}
+				}
+				// get hull proximity out of geometry proximity.  (Note: Only works if merging leaf convexes ... need a more general source of proximity!)
+				for (int32 SourceTransformIndex : SourceTransformIndices)
+				{
+					const TArray<int32, TInlineAllocator<8>> SourceHullIndices{ TransformToHullIdx.FindOrAdd(SourceTransformIndex) };
+					const int32 GeoIdx = Collection.TransformToGeometryIndex[SourceTransformIndex];
+					if (GeoIdx > INDEX_NONE)
+					{
+						for (int32 NeighborGeoIndex : Proximity[GeoIdx])
+						{
+							const int32 NeightborTransformIdx = Collection.TransformIndex[NeighborGeoIndex];
+							if (const TArray<int32>* NeighborHullIndices = TransformToHullIdx.Find(NeightborTransformIdx))
+							{
+								for (int32 SourceHullIdx : SourceHullIndices)
+								{
+									for (int32 NeighborHullIdx : (*NeighborHullIndices))
+									{
+										HullProximity.Emplace(SourceHullIdx, NeighborHullIdx);
+									}
+								}
+							}
+						}
+					}
+				}
+
+				// try to merge the leaf convex into a simpler set of convex
+				const auto GetHullVolume = [&Hulls](int32 Idx) { return (double)Hulls[Idx].Convex->GetVolume(); };
+				const auto GetHullNumVertices = [&Hulls](int32 Idx) { return Hulls[Idx].Convex->NumVertices(); };
+				const auto GetHullVertex = [&Hulls](int32 Hull, int32 V) {
+					const FVector3d LocalVertex(Hulls[Hull].Convex->GetVertex(V));
+					return Hulls[Hull].Transform.TransformPosition(LocalVertex);
+				};
+				UE::Geometry::FConvexDecomposition3 ConvexDecomposition;
+				ConvexDecomposition.InitializeFromHulls(Hulls.Num(), GetHullVolume, GetHullNumVertices, GetHullVertex, HullProximity);
+				ConvexDecomposition.MergeBest(ConvexCount, ErrorToleranceInCm, 0, true);
+				
+				// reset existing hulls for this transform index
+				for (int32 ConvexIndex : TransformToConvexIndices[TransformIndex])
+				{
+					ConvexHull[ConvexIndex] = nullptr;
+				}
+				TransformToConvexIndices[TransformIndex].Reset();
+
+				// add new computed ones
+				for (int32 DecompIndex = 0; DecompIndex < ConvexDecomposition.Decomposition.Num(); DecompIndex++)
+				{
+					// Make Implicit Convex of a decomposed element
+					TArray<Chaos::FConvex::FVec3Type> Particles;
+					const TArray<FVector> Points = ConvexDecomposition.GetVertices<double>(DecompIndex);
+					Particles.SetNum(Points.Num());
+					for (int32 PointIndex = 0; PointIndex < Points.Num(); PointIndex++)
+					{
+						Particles[PointIndex] = Points[PointIndex];
+					}
+					TUniquePtr<Chaos::FConvex> ImplicitConvex = MakeUnique<Chaos::FConvex>(MoveTemp(Particles), 0.0f);
+
+					// Add the the element to the union
+					const int32 NewConvexIndex = Collection.AddElements(1, ConvexGroupName);
+					ConvexHull[NewConvexIndex] = MoveTemp(ImplicitConvex);
+					TransformToConvexIndices[TransformIndex].Add(NewConvexIndex);
+				}
 			}
 		}
-		else // (!ConvexHull[ConvexIdx].IsValid())
-		{
-			EmptyConvex.Add(ConvexIdx);
-		}
 	}
-	FManagedArrayCollection::FProcessingParameters Params;
-	Params.bDoValidation = false; // for perf reasons
-	GeometryCollection->RemoveElements("Convex", EmptyConvex, Params);
 
-	checkSlow(FGeometryCollectionConvexUtility::ValidateConvexData(GeometryCollection));
+	// remove empty or null convex hulls
+	RemoveEmptyConvexHulls(Collection);
 
-	return { TransformToConvexIndices, ConvexHull };
-}
+	checkSlow(FGeometryCollectionConvexUtility::ValidateConvexData(&Collection));
 
+}
 
 bool FGeometryCollectionConvexUtility::ValidateConvexData(const FGeometryCollection* GeometryCollection)
 {
@@ -1453,6 +1581,40 @@ void FGeometryCollectionConvexUtility::RemoveConvexHulls(FGeometryCollection* Ge
 	}
 }
 
+/** Delete the convex hulls that are null */
+void FGeometryCollectionConvexUtility::RemoveEmptyConvexHulls(FManagedArrayCollection& Collection)
+{
+	const FName ConvexGroupName("Convex");
+	const FName ConvexAttributeName("ConvexHull");
+
+	TManagedArrayAccessor<TUniquePtr<Chaos::FConvex>> ConvexHullAttribute(Collection, ConvexAttributeName, ConvexGroupName);
+	if (ConvexHullAttribute.IsValid())
+	{
+		TManagedArray<TUniquePtr<Chaos::FConvex>>& ConvexHull = ConvexHullAttribute.Modify();
+
+		// clear all null and empty hulls
+		TArray<int32> EmptyConvex;
+		for (int32 ConvexIdx = 0; ConvexIdx < ConvexHull.Num(); ConvexIdx++)
+		{
+			if (ConvexHull[ConvexIdx].IsValid())
+			{
+				if (ConvexHull[ConvexIdx]->NumVertices() == 0)
+				{
+					ConvexHull[ConvexIdx].Reset();
+					EmptyConvex.Add(ConvexIdx);
+				}
+			}
+			else // (!ConvexHull[ConvexIdx].IsValid())
+			{
+				EmptyConvex.Add(ConvexIdx);
+			}
+		}
+
+		FManagedArrayCollection::FProcessingParameters Params;
+		Params.bDoValidation = false; // for perf reasons
+		Collection.RemoveElements(ConvexGroupName, EmptyConvex, Params);
+	}
+}
 
 void FGeometryCollectionConvexUtility::SetDefaults(FGeometryCollection* GeometryCollection, FName Group, uint32 StartSize, uint32 NumElements)
 {
diff --git a/Engine/Source/Runtime/Experimental/Chaos/Public/GeometryCollection/Facades/CollectionHierarchyFacade.h b/Engine/Source/Runtime/Experimental/Chaos/Public/GeometryCollection/Facades/CollectionHierarchyFacade.h
--- a/Engine/Source/Runtime/Experimental/Chaos/Public/GeometryCollection/Facades/CollectionHierarchyFacade.h
+++ b/Engine/Source/Runtime/Experimental/Chaos/Public/GeometryCollection/Facades/CollectionHierarchyFacade.h
@@ -44,6 +44,9 @@ namespace Chaos::Facades
 		/** Get the root indicies */
 		TArray<int32> GetRootIndices() const;
 
+		/** Get direct children from a specific transform index */
+		TArray<int32> GetChildrenAsArray(int32 TransformIndex) const;
+
 		/** 
 		* Get initial level of a specific transform index 
 		* If the attribute is missing return INDEX_NONE
@@ -56,6 +59,11 @@ namespace Chaos::Facades
 		*/ 
 		void GenerateLevelAttribute();
 
+		/**
+		* Get transform indices in a depth first order 
+		*/
+		TArray<int32> GetTransformArrayInDepthFirstOrder() const;
+
 	public:
 		/** Get the root indicies */
 		static TArray<int32> GetRootIndices(const TManagedArrayAccessor<int32>& ParentAttribute);
diff --git a/Engine/Source/Runtime/Experimental/Chaos/Public/GeometryCollection/GeometryCollectionClusteringUtility.h b/Engine/Source/Runtime/Experimental/Chaos/Public/GeometryCollection/GeometryCollectionClusteringUtility.h
--- a/Engine/Source/Runtime/Experimental/Chaos/Public/GeometryCollection/GeometryCollectionClusteringUtility.h
+++ b/Engine/Source/Runtime/Experimental/Chaos/Public/GeometryCollection/GeometryCollectionClusteringUtility.h
@@ -114,7 +114,7 @@ public:
 	static void RenameBone(FGeometryCollection* GeometryCollection, int32 BoneIndex, const FString& NewName, bool UpdateChildren = true);
 
 	/** return an array of all child leaf nodes below the specified node. If bOnlyRigids is true, the first Rigid node dound is considered a leaf, regardless of an children it might have. */
-	static void GetLeafBones(const FGeometryCollection* GeometryCollection, int BoneIndex, bool bOnlyRigids, TArray<int32>& LeafBonesOut);
+	static void GetLeafBones(const FManagedArrayCollection* GeometryCollection, int BoneIndex, bool bOnlyRigids, TArray<int32>& LeafBonesOut);
 
 	/** move the selected node up a level in direction of root */
 	static void MoveUpOneHierarchyLevel(FGeometryCollection* GeometryCollection, const TArray<int32>& SelectedBones);
diff --git a/Engine/Source/Runtime/Experimental/Chaos/Public/GeometryCollection/GeometryCollectionConvexUtility.h b/Engine/Source/Runtime/Experimental/Chaos/Public/GeometryCollection/GeometryCollectionConvexUtility.h
--- a/Engine/Source/Runtime/Experimental/Chaos/Public/GeometryCollection/GeometryCollectionConvexUtility.h
+++ b/Engine/Source/Runtime/Experimental/Chaos/Public/GeometryCollection/GeometryCollectionConvexUtility.h
@@ -79,12 +79,17 @@ public:
 	static FGeometryCollectionConvexData CreateNonOverlappingConvexHullData(FGeometryCollection* GeometryCollection, double FractionAllowRemove = .3, double SimplificationDistanceThreshold = 0.0, double CanExceedFraction = .5, 
 		EConvexOverlapRemoval OverlapRemovalMethod = EConvexOverlapRemoval::All, double OverlapRemovalShrinkPercent = 0.0, UE::GeometryCollectionConvexUtility::FConvexHulls* ComputedLeafHullsToModify = nullptr);
 
+	static void GenerateClusterConvexHullsFromLeafHulls(FGeometryCollection& Collection, int32 ConvexCount, double ErrorToleranceInCm);
+
 	/** Returns the convex hull of the vertices contained in the specified geometry. */
 	static TUniquePtr<Chaos::FConvex> FindConvexHull(const FGeometryCollection* GeometryCollection, int32 GeometryIndex);
 
 	/** Delete the convex hulls pointed at by the transform indices provided. */
 	static void RemoveConvexHulls(FGeometryCollection* GeometryCollection, const TArray<int32>& SortedTransformDeletes);
 
+	/** Delete the convex hulls that are null */
+	static void RemoveEmptyConvexHulls(FManagedArrayCollection& GeometryCollection);
+
 	/** Set default values for convex hull related managed arrays. */
 	static void SetDefaults(FGeometryCollection* GeometryCollection, FName Group, uint32 StartSize, uint32 NumElements);
 
