Chaos - fix average normal impulse calculation in non-SIMD collision solver

diff --git a/Engine/Source/Runtime/Experimental/Chaos/Public/Chaos/Collision/PBDCollisionSolver.h b/Engine/Source/Runtime/Experimental/Chaos/Public/Chaos/Collision/PBDCollisionSolver.h
--- a/Engine/Source/Runtime/Experimental/Chaos/Public/Chaos/Collision/PBDCollisionSolver.h
+++ b/Engine/Source/Runtime/Experimental/Chaos/Public/Chaos/Collision/PBDCollisionSolver.h
@@ -391,38 +391,42 @@ namespace Chaos
 			FSolverReal& InOutNetPushOutTangentV,
 			FSolverReal& OutStaticFrictionRatio)
 		{
-			// Assume we stay in the friction cone...
-			OutStaticFrictionRatio = FSolverReal(1);
+			FSolverReal ClampedPushOutTangentU = InOutPushOutTangentU;
+			FSolverReal ClampedPushOutTangentV = InOutPushOutTangentV;
+			FSolverReal ClampedNetPushOutTangentU = InOutNetPushOutTangentU + InOutPushOutTangentU;
+			FSolverReal ClampedNetPushOutTangentV = InOutNetPushOutTangentV + InOutPushOutTangentV;
+			FSolverReal ClampedStaticFrictionRatio = FSolverReal(1);
 
 			if (MaxFrictionPushOut < FSolverReal(UE_KINDA_SMALL_NUMBER))
 			{
-				// Note: we have already added the current iteration's PushOut to the NetPushOut but it has not been applied to the body
-				// so we must subtract it again to calculate the actual pushout we want to undo (i.e., the net pushout that has been applied 
-				// to the body so far from previous iterations)
-				InOutPushOutTangentU = -(InOutNetPushOutTangentU - InOutPushOutTangentU);
-				InOutPushOutTangentV = -(InOutNetPushOutTangentV - InOutPushOutTangentV);
-				InOutNetPushOutTangentU = FSolverReal(0);
-				InOutNetPushOutTangentV = FSolverReal(0);
-				OutStaticFrictionRatio = FSolverReal(0);
+				ClampedPushOutTangentU = -InOutNetPushOutTangentU;
+				ClampedPushOutTangentV = -InOutNetPushOutTangentV;
+				ClampedNetPushOutTangentU = FSolverReal(0);
+				ClampedNetPushOutTangentV = FSolverReal(0);
+				ClampedStaticFrictionRatio = FSolverReal(0);
 			}
 			else
 			{
 				// If we exceed the static friction cone, clip to the dynamic friction cone
 				const FSolverReal MaxStaticPushOutTangentSq = FMath::Square(StaticFriction * MaxFrictionPushOut);
-				const FSolverReal NetPushOutTangentSq = FMath::Square(InOutNetPushOutTangentU) + FMath::Square(InOutNetPushOutTangentV);
+				const FSolverReal NetPushOutTangentSq = FMath::Square(ClampedNetPushOutTangentU) + FMath::Square(ClampedNetPushOutTangentV);
 				if (NetPushOutTangentSq > MaxStaticPushOutTangentSq)
 				{
 					const FSolverReal MaxDynamicPushOutTangent = DynamicFriction * MaxFrictionPushOut;
 					const FSolverReal FrictionMultiplier = MaxDynamicPushOutTangent * FMath::InvSqrt(NetPushOutTangentSq);
-					const FSolverReal NetPushOutTangentU = FrictionMultiplier * InOutNetPushOutTangentU;
-					const FSolverReal NetPushOutTangentV = FrictionMultiplier * InOutNetPushOutTangentV;
-					InOutPushOutTangentU = NetPushOutTangentU - (InOutNetPushOutTangentU - InOutPushOutTangentU);
-					InOutPushOutTangentV = NetPushOutTangentV - (InOutNetPushOutTangentV - InOutPushOutTangentV);
-					InOutNetPushOutTangentU = NetPushOutTangentU;
-					InOutNetPushOutTangentV = NetPushOutTangentV;
-					OutStaticFrictionRatio = FrictionMultiplier;
+					ClampedNetPushOutTangentU = FrictionMultiplier * ClampedNetPushOutTangentU;
+					ClampedNetPushOutTangentV = FrictionMultiplier * ClampedNetPushOutTangentV;
+					ClampedPushOutTangentU = ClampedNetPushOutTangentU - InOutNetPushOutTangentU;
+					ClampedPushOutTangentV = ClampedNetPushOutTangentV - InOutNetPushOutTangentV;
+					ClampedStaticFrictionRatio = FrictionMultiplier;
 				}
 			}
+
+			InOutPushOutTangentU = ClampedPushOutTangentU;
+			InOutPushOutTangentV = ClampedPushOutTangentV;
+			InOutNetPushOutTangentU = ClampedNetPushOutTangentU;
+			InOutNetPushOutTangentV = ClampedNetPushOutTangentV;
+			OutStaticFrictionRatio = ClampedStaticFrictionRatio;
 		}
 
 
@@ -454,9 +458,7 @@ namespace Chaos
 				ManifoldPoint.NetPushOutTangentV,
 				PushOutTangentV);					// Out
 
-			ManifoldPoint.NetPushOutTangentU += PushOutTangentU;
-			ManifoldPoint.NetPushOutTangentV += PushOutTangentV;
-
+			// NOTE: This function modifies NetPushOutTangentU and V
 			ApplyFrictionCone(
 				StaticFriction,
 				DynamicFriction,
@@ -815,8 +817,9 @@ namespace Chaos
 			FConstraintSolverBody& Body1 = SolverBody1();
 
 			// Accumulate net pushout for friction limits below
-			bool bApplyFriction[MaxPointsPerConstraint] = { false, };
-			int32 NumFrictionContacts = 0;
+			bool bApplyFriction = false;
+			bool bApplyPointFriction[MaxPointsPerConstraint] = { false, };
+			FSolverReal NumFrictionContacts = FSolverReal(0);	// NOTE: deliberately not an int
 			FSolverReal TotalPushOutNormal = FSolverReal(0);
 
 			// Apply the position correction along the normal and determine if we want to run friction on each point
@@ -837,35 +840,37 @@ namespace Chaos
 						SolverManifoldPoint,
 						Body0,
 						Body1);
-
-					TotalPushOutNormal += SolverManifoldPoint.NetPushOutNormal;
 				}
 
 				// Friction gets updated for any point with a net normal correction or where we have previously had a normal correction and 
 				// already applied friction (in which case we may need to zero it)
-				if ((SolverManifoldPoint.NetPushOutNormal != 0) || (SolverManifoldPoint.NetPushOutTangentU != 0) || (SolverManifoldPoint.NetPushOutTangentV != 0))
+				if (SolverManifoldPoint.NetPushOutNormal > 0)
 				{
-					bApplyFriction[PointIndex] = true;
-					++NumFrictionContacts;
+					TotalPushOutNormal += SolverManifoldPoint.NetPushOutNormal;
+					NumFrictionContacts += FSolverReal(1);
+				}
+				if ((SolverManifoldPoint.NetPushOutNormal > 0) || (SolverManifoldPoint.NetPushOutTangentU != 0) || (SolverManifoldPoint.NetPushOutTangentV != 0))
+				{
+					bApplyFriction = true;
+					bApplyPointFriction[PointIndex] = ((SolverManifoldPoint.NetPushOutNormal > 0) || (SolverManifoldPoint.NetPushOutTangentU != 0) || (SolverManifoldPoint.NetPushOutTangentV != 0));
 				}
 			}
 
 			// Apply the tangential position correction if required
-			if (NumFrictionContacts > 0)
+			if (bApplyFriction)
 			{
 				// We clip the tangential correction at each contact to the friction cone, but we use to average impulse
 				// among all contacts as the clipping limit. This is not really correct but it is much more stable to 
 				// differences in contacts from tick to tick
 				// @todo(chaos): try a decaying maximum per contact point rather than an average (again - we had that once!)
-				const FSolverReal FrictionMaxPushOut = TotalPushOutNormal / FSolverReal(NumFrictionContacts);
+				const FSolverReal FrictionMaxPushOut = (NumFrictionContacts > FSolverReal(0)) ? TotalPushOutNormal / NumFrictionContacts : FSolverReal(0);
 				const FSolverReal FrictionStiffness = State.Stiffness * CVars::Chaos_PBDCollisionSolver_Position_StaticFrictionStiffness;
 
 				for (int32 PointIndex = 0; PointIndex < NumManifoldPoints(); ++PointIndex)
 				{
-					if (bApplyFriction[PointIndex])
+					FPBDCollisionSolverManifoldPoint& SolverManifoldPoint = State.ManifoldPoints[PointIndex];
+					if (bApplyPointFriction[PointIndex])
 					{
-						FPBDCollisionSolverManifoldPoint& SolverManifoldPoint = State.ManifoldPoints[PointIndex];
-
 						FSolverReal ContactDeltaTangentU, ContactDeltaTangentV;
 						SolverManifoldPoint.CalculateContactPositionErrorTangential(Body0.SolverBody(), Body1.SolverBody(), ContactDeltaTangentU, ContactDeltaTangentV);
 
