Chaos - 16-byte align corrections in SolverBody for cheaper SIMD fetch

diff --git a/Engine/Source/Runtime/Experimental/Chaos/Private/Chaos/Collision/PBDCollisionSolver.ispc b/Engine/Source/Runtime/Experimental/Chaos/Private/Chaos/Collision/PBDCollisionSolver.ispc
--- a/Engine/Source/Runtime/Experimental/Chaos/Private/Chaos/Collision/PBDCollisionSolver.ispc
+++ b/Engine/Source/Runtime/Experimental/Chaos/Private/Chaos/Collision/PBDCollisionSolver.ispc
@@ -13,6 +13,22 @@
 
 export struct FSolverBody
 {
+	// Net position delta applied by all constraints (constantly changing as we iterate over constraints)
+	FVector3f DP;
+	float DPPad;
+
+	// Net rotation delta applied by all constraints (constantly changing as we iterate over constraints)
+	FVector3f DQ;
+	float DQPad;
+
+	// World-space center of mass velocity
+	FVector3f V;
+	float VPad;
+
+	// World-space center of mass angular velocity
+	FVector3f W;
+	float WPad;
+
 	// World-space inverse inertia
 	// @todo(chaos): do we need this, or should we force all systems to use the FConstraintSolverBody decorator?
 	FMatrix33f InvI;
@@ -42,12 +58,6 @@ export struct FSolverBody
 	// Inverse mass
 	float InvM;
 
-	// Net position delta applied by all constraints (constantly changing as we iterate over constraints)
-	FVector3f DP;
-
-	// Net rotation delta applied by all constraints (constantly changing as we iterate over constraints)
-	FVector3f DQ;
-
 	// Net position correction delta applied by all constraints (constantly changing as we iterate over constraints)
 	// Will translate the body without introducing linear velocity
 	FVector3f CP;
@@ -56,12 +66,6 @@ export struct FSolverBody
 	// Will rotate the body without introducing angular velocity
 	FVector3f CQ;
 
-	// World-space center of mass velocity
-	FVector3f V;
-
-	// World-space center of mass angular velocity
-	FVector3f W;
-
 	// Distance to a kinematic body (through the contact graph). Used by collision shock propagation
 	int32 Level;
 
diff --git a/Engine/Source/Runtime/Experimental/Chaos/Public/Chaos/Evolution/SolverBody.h b/Engine/Source/Runtime/Experimental/Chaos/Public/Chaos/Evolution/SolverBody.h
--- a/Engine/Source/Runtime/Experimental/Chaos/Public/Chaos/Evolution/SolverBody.h
+++ b/Engine/Source/Runtime/Experimental/Chaos/Public/Chaos/Evolution/SolverBody.h
@@ -429,13 +429,13 @@ namespace Chaos
 		void PrefetchPositionSolverData() const
 		{
 			// The position solver only uses DP,DQ
-			FPlatformMisc::PrefetchBlock(&State.DP, 2 * sizeof(FSolverVec3));
+			FPlatformMisc::PrefetchBlock(&State.DP, 8 * sizeof(float));
 		}
 
 		void PrefetchVelocitySolverData() const
 		{
 			// The velocity solver only uses V,W
-			FPlatformMisc::PrefetchBlock(&State.V, 2 * sizeof(FSolverVec3));
+			FPlatformMisc::PrefetchBlock(&State.V, 8 * sizeof(float));
 		}
 
 	private:
@@ -449,6 +449,10 @@ namespace Chaos
 
 			void Init()
 			{
+				DP = FSolverVec3(0);
+				DQ = FSolverVec3(0);
+				V = FSolverVec3(0);
+				W = FSolverVec3(0);
 				InvI = FSolverMatrix33(0);
 				RoM = FRotation3::FromIdentity();
 				R = FRotation3::FromIdentity();
@@ -458,15 +462,23 @@ namespace Chaos
 				P = FVec3(0);
 				InvILocal = FSolverVec3(0);
 				InvM = 0;
-				DP = FSolverVec3(0);
-				DQ = FSolverVec3(0);
-				V = FSolverVec3(0);
-				W = FSolverVec3(0);
 				CP = FSolverVec3(0);
 				CQ = FSolverVec3(0);
 				Level = 0;
 			}
 
+			// Net position delta applied by all constraints (constantly changing as we iterate over constraints)
+			alignas(16) FSolverVec3 DP;
+
+			// Net rotation delta applied by all constraints (constantly changing as we iterate over constraints)
+			alignas(16) FSolverVec3 DQ;
+
+			// World-space center of mass velocity
+			alignas(16) FSolverVec3 V;
+
+			// World-space center of mass angular velocity
+			alignas(16) FSolverVec3 W;
+
 			// World-space inverse inertia
 			// NOTE: Matrix and Rotation are alignas(16) so beware of padding
 			// @todo(chaos): do we need this, or should we force all systems to use the FConstraintSolverBody decorator?
@@ -498,18 +510,6 @@ namespace Chaos
 			// Inverse mass
 			FSolverReal InvM;
 
-			// Net position delta applied by all constraints (constantly changing as we iterate over constraints)
-			FSolverVec3 DP;
-
-			// Net rotation delta applied by all constraints (constantly changing as we iterate over constraints)
-			FSolverVec3 DQ;
-
-			// World-space center of mass velocity
-			FSolverVec3 V;
-
-			// World-space center of mass angular velocity
-			FSolverVec3 W;
-
 			// Net position correction delta applied by all constraints (constantly changing as we iterate over constraints)
 			// Will translate the body without introducing linear velocity
 			FSolverVec3 CP;
