Chaos - simd collision solver WIP
- go back to per-constraint solving rather than per-manifold-point to dedupe data
- add velocity-phase dynamic friction support
- add shock propagation support
- AB Testing additions

diff --git a/Engine/Source/Runtime/Experimental/Chaos/Private/Chaos/Collision/PBDCollisionContainerSolverSimd.cpp b/Engine/Source/Runtime/Experimental/Chaos/Private/Chaos/Collision/PBDCollisionContainerSolverSimd.cpp
--- a/Engine/Source/Runtime/Experimental/Chaos/Private/Chaos/Collision/PBDCollisionContainerSolverSimd.cpp
+++ b/Engine/Source/Runtime/Experimental/Chaos/Private/Chaos/Collision/PBDCollisionContainerSolverSimd.cpp
@@ -24,9 +24,6 @@ namespace Chaos
 	{
 		extern bool bChaos_PBDCollisionSolver_Position_SolveEnabled;
 		extern bool bChaos_PBDCollisionSolver_Velocity_SolveEnabled;
-
-		extern FRealSingle Chaos_PBDCollisionSolver_AutoStiffness_MassRatio1;
-		extern FRealSingle Chaos_PBDCollisionSolver_AutoStiffness_MassRatio2;
 	}
 
 	namespace Private
@@ -37,244 +34,300 @@ namespace Chaos
 		//////////////////////////////////////////////////////////////////////////////////////////////////
 		//////////////////////////////////////////////////////////////////////////////////////////////////
 
+		bool CalculateBodyShockPropagation(
+			const FSolverBody& Body0,
+			const FSolverBody& Body1,
+			const FSolverReal ShockPropagation,
+			FSolverReal& OutShockPropagation0,
+			FSolverReal& OutShockPropagation1)
+		{
+			OutShockPropagation0 = FSolverReal(1);
+			OutShockPropagation1 = FSolverReal(1);
+
+			// Shock propagation decreases the inverse mass of bodies that are lower in the pile
+			// of objects. This significantly improves stability of heaps and stacks. Height in the pile is indictaed by the "level". 
+			// No need to set an inverse mass scale if the other body is kinematic (with inv mass of 0).
+			// Bodies at the same level do not take part in shock propagation.
+			if (Body0.IsDynamic() && Body1.IsDynamic() && (Body0.Level() != Body1.Level()))
+			{
+				// Set the inv mass scale of the "lower" body to make it heavier
+				if (Body0.Level() < Body1.Level())
+				{
+					OutShockPropagation0 = ShockPropagation;
+				}
+				else
+				{
+					OutShockPropagation1 = ShockPropagation;
+				}
+				return true;
+			}
+
+			return false;
+		}
+
 		// Transform and copy a single manifold point for use in the solver
 		template<int TNumLanes>
-		void UpdateSolverContactPointFromConstraint(
+		void UpdateSolverManifoldPointFromConstraint(
+			const int32 ManifoldPointIndex,
+			TPBDCollisionSolverSimd<4>& Solvers,
+			TConstraintPtrSimd<4>& Constraints,
+			TSolverBodyPtrPairSimd<4>& Bodies,
 			const TArrayView<TPBDCollisionSolverManifoldPointsSimd<TNumLanes>>& ManifoldPointsBuffer,
-			const TArrayView<TSolverBodyPtrPairSimd<4>>& SolverBodiesBuffer,
-			Private::FPBDCollisionSolverSimd& Solver,
-			const int32 SolverPointIndex, 
-			const FPBDCollisionConstraint* Constraint, 
-			const int32 ConstraintPointIndex, 
-			const FRealSingle Dt, 
-			const FSolverBody& Body0, 
-			const FSolverBody& Body1)
+			const FRealSingle Dt)
 		{
-			const FManifoldPoint& ManifoldPoint = Constraint->GetManifoldPoint(ConstraintPointIndex);
+			// @todo(chaos): this can be partly simd
+			for (int32 LaneIndex = 0; LaneIndex < TNumLanes; ++LaneIndex)
+			{
+				const FPBDCollisionConstraint* Constraint = Constraints.GetValue(LaneIndex);
+				if (Constraint == nullptr)
+				{
+					continue;
+				}
 
-			const FRealSingle Restitution = FRealSingle(Constraint->GetRestitution());
-			const FRealSingle RestitutionVelocityThreshold = FRealSingle(Constraint->GetRestitutionThreshold()) * Dt;
+				if (!Constraint->IsManifoldPointActive(ManifoldPointIndex))
+				{
+					continue;
+				}
 
-			// World-space shape transforms. Manifold data is currently relative to these spaces
-			const FRigidTransform3& ShapeWorldTransform0 = Constraint->GetShapeWorldTransform0();
-			const FRigidTransform3& ShapeWorldTransform1 = Constraint->GetShapeWorldTransform1();
+				const FManifoldPoint& ManifoldPoint = Constraint->GetManifoldPoint(ManifoldPointIndex);
+				const FSolverBody& Body0 = *Bodies.Body0.GetValue(LaneIndex);
+				const FSolverBody& Body1 = *Bodies.Body1.GetValue(LaneIndex);
 
-			// World-space contact points on each shape
-			const FVec3 WorldContact0 = ShapeWorldTransform0.TransformPositionNoScale(FVec3(ManifoldPoint.ContactPoint.ShapeContactPoints[0]));
-			const FVec3 WorldContact1 = ShapeWorldTransform1.TransformPositionNoScale(FVec3(ManifoldPoint.ContactPoint.ShapeContactPoints[1]));
-			const FVec3 WorldContact = FReal(0.5) * (WorldContact0 + WorldContact1);
-			const FVec3f WorldRelativeContact0 = FVec3f(WorldContact - Body0.P());
-			const FVec3f WorldRelativeContact1 = FVec3f(WorldContact - Body1.P());
+				const FRealSingle Restitution = FRealSingle(Constraint->GetRestitution());
+				const FRealSingle RestitutionVelocityThreshold = FRealSingle(Constraint->GetRestitutionThreshold()) * Dt;
 
-			// World-space normal
-			const FVec3f WorldContactNormal = FVec3f(ShapeWorldTransform1.TransformVectorNoScale(FVec3(ManifoldPoint.ContactPoint.ShapeContactNormal)));
+				// World-space shape transforms. Manifold data is currently relative to these spaces
+				const FRigidTransform3& ShapeWorldTransform0 = Constraint->GetShapeWorldTransform0();
+				const FRigidTransform3& ShapeWorldTransform1 = Constraint->GetShapeWorldTransform1();
 
-			// World-space tangents
-			FVec3f WorldContactTangentU = FVec3f::CrossProduct(FVec3f(0, 1, 0), WorldContactNormal);
-			if (!WorldContactTangentU.Normalize(FRealSingle(UE_KINDA_SMALL_NUMBER)))
-			{
-				WorldContactTangentU = FVec3f::CrossProduct(FVec3f(1, 0, 0), WorldContactNormal);
-				WorldContactTangentU = WorldContactTangentU.GetUnsafeNormal();
-			}
-			const FVec3f WorldContactTangentV = FVec3f::CrossProduct(WorldContactNormal, WorldContactTangentU);
+				// World-space contact points on each shape
+				const FVec3 WorldContact0 = ShapeWorldTransform0.TransformPositionNoScale(FVec3(ManifoldPoint.ContactPoint.ShapeContactPoints[0]));
+				const FVec3 WorldContact1 = ShapeWorldTransform1.TransformPositionNoScale(FVec3(ManifoldPoint.ContactPoint.ShapeContactPoints[1]));
+				const FVec3 WorldContact = FReal(0.5) * (WorldContact0 + WorldContact1);
+				const FVec3f WorldRelativeContact0 = FVec3f(WorldContact - Body0.P());
+				const FVec3f WorldRelativeContact1 = FVec3f(WorldContact - Body1.P());
 
-			// Calculate contact velocity if we will need it below (restitution and/or frist-contact friction)
-			const bool bNeedContactVelocity = (!ManifoldPoint.Flags.bHasStaticFrictionAnchor) || (Restitution > FRealSingle(0));
-			FVec3f ContactVel = FVec3(0);
-			if (bNeedContactVelocity)
-			{
-				const FVec3f ContactVel0 = Body0.V() + FVec3f::CrossProduct(Body0.W(), WorldRelativeContact0);
-				const FVec3f ContactVel1 = Body1.V() + FVec3f::CrossProduct(Body1.W(), WorldRelativeContact1);
-				ContactVel = ContactVel0 - ContactVel1;
-			}
+				// World-space normal
+				const FVec3f WorldContactNormal = FVec3f(ShapeWorldTransform1.TransformVectorNoScale(FVec3(ManifoldPoint.ContactPoint.ShapeContactNormal)));
 
-			// If we have contact data from a previous tick, use it to calculate the lateral position delta we need
-			// to apply to move the contacts back to their original relative locations (i.e., to enforce static friction).
-			// Otherwise, estimate the friction correction from the contact velocity.
-			// NOTE: quadratic shapes use the velocity-based path most of the time, unless the relative motion is very small.
-			FVec3f WorldFrictionDelta = FVec3f(0);
-			if (ManifoldPoint.Flags.bHasStaticFrictionAnchor)
-			{
-				const FVec3f FrictionDelta0 = ShapeWorldTransform0.TransformPositionNoScale(FVec3(ManifoldPoint.ShapeAnchorPoints[0])) - WorldContact0;
-				const FVec3f FrictionDelta1 = ShapeWorldTransform1.TransformPositionNoScale(FVec3(ManifoldPoint.ShapeAnchorPoints[1])) - WorldContact1;
-				WorldFrictionDelta = FrictionDelta0 - FrictionDelta1;
-			}
-			else
-			{
-				// @todo(chaos): consider adding a multiplier to the initial contact friction
-				WorldFrictionDelta = ContactVel * Dt;
-			}
+				// World-space tangents
+				FVec3f WorldContactTangentU = FVec3f::CrossProduct(FVec3f(0, 1, 0), WorldContactNormal);
+				if (!WorldContactTangentU.Normalize(FRealSingle(UE_KINDA_SMALL_NUMBER)))
+				{
+					WorldContactTangentU = FVec3f::CrossProduct(FVec3f(1, 0, 0), WorldContactNormal);
+					WorldContactTangentU = WorldContactTangentU.GetUnsafeNormal();
+				}
+				const FVec3f WorldContactTangentV = FVec3f::CrossProduct(WorldContactNormal, WorldContactTangentU);
 
-			// The contact point error we are trying to correct in this solver
-			const FRealSingle TargetPhi = ManifoldPoint.TargetPhi;
-			const FVec3f WorldContactDelta = FVec3f(WorldContact0 - WorldContact1);
-			const FRealSingle WorldContactDeltaNormal = FVec3f::DotProduct(WorldContactDelta, WorldContactNormal) - TargetPhi;
-			const FRealSingle WorldContactDeltaTangentU = FVec3f::DotProduct(WorldContactDelta + WorldFrictionDelta, WorldContactTangentU);
-			const FRealSingle WorldContactDeltaTangentV = FVec3f::DotProduct(WorldContactDelta + WorldFrictionDelta, WorldContactTangentV);
+				// Calculate contact velocity if we will need it below (restitution and/or frist-contact friction)
+				const bool bNeedContactVelocity = (!ManifoldPoint.Flags.bHasStaticFrictionAnchor) || (Restitution > FRealSingle(0));
+				FVec3f ContactVel = FVec3(0);
+				if (bNeedContactVelocity)
+				{
+					const FVec3f ContactVel0 = Body0.V() + FVec3f::CrossProduct(Body0.W(), WorldRelativeContact0);
+					const FVec3f ContactVel1 = Body1.V() + FVec3f::CrossProduct(Body1.W(), WorldRelativeContact1);
+					ContactVel = ContactVel0 - ContactVel1;
+				}
 
-			// The target contact velocity, taking restitution into account
-			FRealSingle WorldContactTargetVelocityNormal = FRealSingle(0);
-			if (Restitution > FRealSingle(0))
-			{
-				const FRealSingle ContactVelocityNormal = FVec3f::DotProduct(ContactVel, WorldContactNormal);
-				if (ContactVelocityNormal < -RestitutionVelocityThreshold)
+				// If we have contact data from a previous tick, use it to calculate the lateral position delta we need
+				// to apply to move the contacts back to their original relative locations (i.e., to enforce static friction).
+				// Otherwise, estimate the friction correction from the contact velocity.
+				// NOTE: quadratic shapes use the velocity-based path most of the time, unless the relative motion is very small.
+				FVec3f WorldFrictionDelta = FVec3f(0);
+				if (ManifoldPoint.Flags.bHasStaticFrictionAnchor)
 				{
-					WorldContactTargetVelocityNormal = -Restitution * ContactVelocityNormal;
+					const FVec3f FrictionDelta0 = ShapeWorldTransform0.TransformPositionNoScale(FVec3(ManifoldPoint.ShapeAnchorPoints[0])) - WorldContact0;
+					const FVec3f FrictionDelta1 = ShapeWorldTransform1.TransformPositionNoScale(FVec3(ManifoldPoint.ShapeAnchorPoints[1])) - WorldContact1;
+					WorldFrictionDelta = FrictionDelta0 - FrictionDelta1;
+				}
+				else
+				{
+					// @todo(chaos): consider adding a multiplier to the initial contact friction
+					WorldFrictionDelta = ContactVel * Dt;
 				}
-			}
 
-			Solver.SetManifoldPoint(
-				ManifoldPointsBuffer,
-				SolverBodiesBuffer,
-				SolverPointIndex,
-				WorldRelativeContact0,
-				WorldRelativeContact1,
-				WorldContactNormal,
-				WorldContactTangentU,
-				WorldContactTangentV,
-				WorldContactDeltaNormal,
-				WorldContactDeltaTangentU,
-				WorldContactDeltaTangentV,
-				WorldContactTargetVelocityNormal);
+				// The contact point error we are trying to correct in this solver
+				const FRealSingle TargetPhi = ManifoldPoint.TargetPhi;
+				const FVec3f WorldContactDelta = FVec3f(WorldContact0 - WorldContact1);
+				const FRealSingle WorldContactDeltaNormal = FVec3f::DotProduct(WorldContactDelta, WorldContactNormal) - TargetPhi;
+				const FRealSingle WorldContactDeltaTangentU = FVec3f::DotProduct(WorldContactDelta + WorldFrictionDelta, WorldContactTangentU);
+				const FRealSingle WorldContactDeltaTangentV = FVec3f::DotProduct(WorldContactDelta + WorldFrictionDelta, WorldContactTangentV);
+
+				// The target contact velocity, taking restitution into account
+				FRealSingle WorldContactTargetVelocityNormal = FRealSingle(0);
+				if (Restitution > FRealSingle(0))
+				{
+					const FRealSingle ContactVelocityNormal = FVec3f::DotProduct(ContactVel, WorldContactNormal);
+					if (ContactVelocityNormal < -RestitutionVelocityThreshold)
+					{
+						WorldContactTargetVelocityNormal = -Restitution * ContactVelocityNormal;
+					}
+				}
+
+				const FSolverReal InvMScale0 = FSolverReal(Constraint->GetInvMassScale0());
+				const FSolverReal InvMScale1 = FSolverReal(Constraint->GetInvMassScale1());
+				const FSolverReal InvIScale0 = FSolverReal(Constraint->GetInvInertiaScale0());
+				const FSolverReal InvIScale1 = FSolverReal(Constraint->GetInvInertiaScale1());
+
+				Solvers.SetManifoldPoint(
+					ManifoldPointIndex,
+					LaneIndex,
+					ManifoldPointsBuffer,
+					WorldRelativeContact0,
+					WorldRelativeContact1,
+					WorldContactNormal,
+					WorldContactTangentU,
+					WorldContactTangentV,
+					WorldContactDeltaNormal,
+					WorldContactDeltaTangentU,
+					WorldContactDeltaTangentV,
+					WorldContactTargetVelocityNormal,
+					Body0,
+					Body1,
+					InvMScale0,
+					InvIScale0,
+					InvMScale1,
+					InvIScale1);
+			}
 		}
 
 		// Transform and copy all of a constraint's manifold point data for use by the solver
 		template<int TNumLanes>
-		void UpdateSolverManifoldFromConstraint(
+		void UpdateSolverManifoldPointsFromConstraint(
+			TPBDCollisionSolverSimd<4>& Solvers,
+			TConstraintPtrSimd<4>& Constraints,
+			TSolverBodyPtrPairSimd<4>& Bodies,
 			const TArrayView<TPBDCollisionSolverManifoldPointsSimd<TNumLanes>>& ManifoldPointsBuffer,
-			const TArrayView<TSolverBodyPtrPairSimd<4>>& SolverBodiesBuffer,
-			Private::FPBDCollisionSolverSimd& Solver,
-			const FPBDCollisionConstraint* Constraint, 
-			const FSolverReal Dt, 
-			const int32 ConstraintPointBeginIndex, 
-			const int32 ConstraintPointEndIndex)
+			const FSolverReal Dt,
+			const int32 ManifoldPointBeginIndex, 
+			const int32 ManifoldPointEndIndex)
 		{
-			const FSolverBody& Body0 = Solver.SolverBody0().SolverBody();
-			const FSolverBody& Body1 = Solver.SolverBody1().SolverBody();
+			Solvers.InitManifoldPoints(ManifoldPointsBuffer);
 
 			// Only calculate state for newly added contacts. Normally this is all of them, but maybe not if incremental collision is used by RBAN.
 			// Also we only add active points to the solver's manifold points list
-			for (int32 ConstraintManifoldPointIndex = ConstraintPointBeginIndex; ConstraintManifoldPointIndex < ConstraintPointEndIndex; ++ConstraintManifoldPointIndex)
-			{
-				if (!Constraint->GetManifoldPoint(ConstraintManifoldPointIndex).Flags.bDisabled)
-				{
-					const int32 SolverManifoldPointIndex = Solver.AddManifoldPoint();
-
-					// Transform the constraint contact data into world space for use by the solver
-					// We build this data directly into the solver's world-space contact data which looks a bit odd with "Init" called after but there you go
-					UpdateSolverContactPointFromConstraint(
-						ManifoldPointsBuffer,
-						SolverBodiesBuffer,
-						Solver, 
-						SolverManifoldPointIndex, 
-						Constraint, 
-						ConstraintManifoldPointIndex, 
-						Dt, 
-						Body0, 
-						Body1);
-				}
+			for (int32 ManifoldPointIndex = ManifoldPointBeginIndex; ManifoldPointIndex < ManifoldPointEndIndex; ++ManifoldPointIndex)
+			{
+				// Transform the constraint contact data into world space for use by the solver
+				// We build this data directly into the solver's world-space contact data which looks a bit odd with "Init" called after but there you go
+				UpdateSolverManifoldPointFromConstraint(
+					ManifoldPointIndex,
+					Solvers,
+					Constraints,
+					Bodies,
+					ManifoldPointsBuffer,
+					Dt);
 			}
 		}
 
-		// Transform and copy all constraint data for use by the solver
 		template<int TNumLanes>
 		void UpdateSolverFromConstraint(
+			TPBDCollisionSolverSimd<4>& Solvers,
+			TConstraintPtrSimd<4>& Constraints,
+			TSolverBodyPtrPairSimd<4>& Bodies,
 			const TArrayView<TPBDCollisionSolverManifoldPointsSimd<TNumLanes>>& ManifoldPointsBuffer,
-			const TArrayView<TSolverBodyPtrPairSimd<4>>& SolverBodiesBuffer,
-			FPBDCollisionSolverSimd& Solver,
-			const FPBDCollisionConstraint* Constraint,
-			const FSolverReal Dt, 
+			const FSolverReal Dt,
 			const FPBDCollisionSolverSettings& SolverSettings,
 			bool& bOutPerIterationCollision)
 		{
 			// Friction values. Static and Dynamic friction are applied in the position solve for most shapes.
 			// We can also run in a mode without static friction at all. This is faster but stacking is not possible.
-			const FSolverReal StaticFriction = FSolverReal(Constraint->GetStaticFriction());
-			const FSolverReal DynamicFriction = FSolverReal(Constraint->GetDynamicFriction());
-			FSolverReal PositionStaticFriction = FSolverReal(0);
-			FSolverReal PositionDynamicFriction = FSolverReal(0);
-			FSolverReal VelocityDynamicFriction = FSolverReal(0);
-			if (SolverSettings.NumPositionFrictionIterations > 0)
-			{
-				PositionStaticFriction = StaticFriction;
-				if (!Constraint->HasQuadraticShape())
-				{
-					PositionDynamicFriction = DynamicFriction;
-				}
-				else
+			TSimdRealf<TNumLanes> PositionStaticFriction = TSimdRealf<TNumLanes>::Zero();
+			TSimdRealf<TNumLanes> PositionDynamicFriction = TSimdRealf<TNumLanes>::Zero();
+			TSimdRealf<TNumLanes> VelocityDynamicFriction = TSimdRealf<TNumLanes>::Zero();
+			TSimdRealf<TNumLanes> Stiffness = TSimdRealf<TNumLanes>::Zero();
+			TSimdInt32<TNumLanes> NumManifoldPoints = TSimdInt32<TNumLanes>::Zero();
+			for (int32 LaneIndex = 0; LaneIndex < TNumLanes; ++LaneIndex)
+			{
+				const FPBDCollisionConstraint* Constraint = Constraints.GetValue(LaneIndex);
+				if (Constraint != nullptr)
 				{
-					// Quadratic shapes don't use PBD dynamic friction - it has issues at slow speeds where the WxR is
-					// less than the position tolerance for friction point matching
-					// @todo(chaos): fix PBD dynamic friction on quadratic shapes
-					VelocityDynamicFriction = DynamicFriction;
+					const FSolverReal StaticFriction = FSolverReal(Constraint->GetStaticFriction());
+					const FSolverReal DynamicFriction = FSolverReal(Constraint->GetDynamicFriction());
+						
+					PositionStaticFriction.SetValue(LaneIndex, StaticFriction);
+					if (!Constraint->HasQuadraticShape())
+					{
+						PositionDynamicFriction.SetValue(LaneIndex, DynamicFriction);
+					}
+					else
+					{
+						// Quadratic shapes don't use PBD dynamic friction - it has issues at slow speeds where the WxR is
+						// less than the position tolerance for friction point matching
+						// @todo(chaos): fix PBD dynamic friction on quadratic shapes
+						VelocityDynamicFriction.SetValue(LaneIndex, DynamicFriction);
+					}
+
+					Stiffness.SetValue(LaneIndex, FSolverReal(Constraint->GetStiffness()));
+
+					NumManifoldPoints.SetValue(LaneIndex, Constraint->NumManifoldPoints());
 				}
 			}
-			else
-			{
-				VelocityDynamicFriction = DynamicFriction;
-			}
-
-			Solver.SetFriction(PositionStaticFriction, PositionDynamicFriction, VelocityDynamicFriction);
 
-			const FReal SolverStiffness = Constraint->GetStiffness();
-			Solver.SetStiffness(FSolverReal(SolverStiffness));
+			// NOTE: NumManifoldPoints includes any disabled points so that the point indices match between
+			// the solver and the constraint. 
+			Solvers.SetNumManifoldPoints(NumManifoldPoints);
 
-			Solver.SolverBody0().SetInvMScale(Constraint->GetInvMassScale0());
-			Solver.SolverBody0().SetInvIScale(Constraint->GetInvInertiaScale0());
-			Solver.SolverBody0().SetShockPropagationScale(FReal(1));
-			Solver.SolverBody1().SetInvMScale(Constraint->GetInvMassScale1());
-			Solver.SolverBody1().SetInvIScale(Constraint->GetInvInertiaScale1());
-			Solver.SolverBody1().SetShockPropagationScale(FReal(1));
+			Solvers.SetFriction(PositionStaticFriction, PositionDynamicFriction, VelocityDynamicFriction);
+			Solvers.SetStiffness(Stiffness);
 
-			bOutPerIterationCollision = (!Constraint->GetUseManifold() || Constraint->GetUseIncrementalCollisionDetection());
+			//bOutPerIterationCollision = (!Constraint->GetUseManifold() || Constraint->GetUseIncrementalCollisionDetection());
 
-			UpdateSolverManifoldFromConstraint(
+			UpdateSolverManifoldPointsFromConstraint(
+				Solvers,
+				Constraints,
+				Bodies,
 				ManifoldPointsBuffer,
-				SolverBodiesBuffer,
-				Solver, 
-				Constraint, 
-				Dt, 
-				0, 
-				Constraint->NumManifoldPoints());
+				Dt,
+				0,
+				Solvers.GetMaxManifoldPoints());
 		}
 
 		template<int TNumLanes>
 		void UpdateConstraintFromSolver(
+			TConstraintPtrSimd<4>& Constraints,
+			TPBDCollisionSolverSimd<4>& Solvers,
+			TSolverBodyPtrPairSimd<4>& Bodies,
 			const TArrayView<TPBDCollisionSolverManifoldPointsSimd<TNumLanes>>& ManifoldPointsBuffer,
-			FPBDCollisionConstraint* Constraint,
-			const Private::FPBDCollisionSolverSimd& Solver, 
 			const FSolverReal Dt)
 		{
-			Constraint->ResetSolverResults();
-
-			// NOTE: We only put the non-pruned manifold points into the solver so the ManifoldPointIndex and
-			// SolverManifoldPointIndex do not necessarily match. See GatherManifoldPoints
-			int32 SolverManifoldPointIndex = 0;
-			for (int32 ManifoldPointIndex = 0; ManifoldPointIndex < Constraint->NumManifoldPoints(); ++ManifoldPointIndex)
+			// @todo(chaos): this can be partly simd
+			for (int32 LaneIndex = 0; LaneIndex < TNumLanes; ++LaneIndex)
 			{
-				FSolverVec3 NetPushOut = FSolverVec3(0);
-				FSolverVec3 NetImpulse = FSolverVec3(0);
-				FSolverReal StaticFrictionRatio = FSolverReal(0);
+				FPBDCollisionConstraint* Constraint = Constraints.GetValue(LaneIndex);
+				if (Constraint == nullptr)
+				{
+					continue;
+				}
+
+				Constraint->ResetSolverResults();
 
-				if (!Constraint->GetManifoldPoint(ManifoldPointIndex).Flags.bDisabled)
+				for (int32 ManifoldPointIndex = 0; ManifoldPointIndex < Constraint->NumManifoldPoints(); ++ManifoldPointIndex)
 				{
-					NetPushOut = Solver.GetNetPushOut(ManifoldPointsBuffer, SolverManifoldPointIndex);
-					NetImpulse = Solver.GetNetImpulse(ManifoldPointsBuffer, SolverManifoldPointIndex);
-					StaticFrictionRatio = Solver.GetStaticFrictionRatio(ManifoldPointsBuffer, SolverManifoldPointIndex);
-					++SolverManifoldPointIndex;
+					FSolverVec3 NetPushOut = FSolverVec3(0);
+					FSolverVec3 NetImpulse = FSolverVec3(0);
+					FSolverReal StaticFrictionRatio = FSolverReal(0);
+
+					if (Constraint->IsManifoldPointActive(ManifoldPointIndex))
+					{
+						NetPushOut = Solvers.GetNetPushOut(ManifoldPointIndex, LaneIndex, ManifoldPointsBuffer);
+						NetImpulse = Solvers.GetNetImpulse(ManifoldPointIndex, LaneIndex, ManifoldPointsBuffer);
+						StaticFrictionRatio = Solvers.GetStaticFrictionRatio(ManifoldPointIndex, LaneIndex, ManifoldPointsBuffer);
+					}
+
+					// NOTE: We call this even for points we did not run the solver for (but with zero results)
+					Constraint->SetSolverResults(ManifoldPointIndex,
+						NetPushOut,
+						NetImpulse,
+						StaticFrictionRatio,
+						Dt);
 				}
 
-				// NOTE: We call this even for points we did not run the solver for (but with zero results)
-				Constraint->SetSolverResults(ManifoldPointIndex,
-					NetPushOut,
-					NetImpulse,
-					StaticFrictionRatio,
-					Dt);
+				Constraint->EndTick();
 			}
-
-			Constraint->EndTick();
 		}
 
+
 		//////////////////////////////////////////////////////////////////////////////////////////////////
 		//////////////////////////////////////////////////////////////////////////////////////////////////
 		//////////////////////////////////////////////////////////////////////////////////////////////////
@@ -285,13 +338,18 @@ namespace Chaos
 			: FConstraintContainerSolver(InPriority)
 			, ConstraintContainer(InConstraintContainer)
 			, NumConstraints(0)
+			, AppliedShockPropagation(1)
 			, bPerIterationCollisionDetection(false)
-			, DummySolverBody(FSolverBody::MakeInitialized())
-			, DummyConstraintSolverBody(DummySolverBody)
 		{
 	#if !UE_BUILD_SHIPPING && !UE_BUILD_TEST && INTEL_ISPC
 			Private::FPBDCollisionSolverHelperSimd::CheckISPC();
 	#endif
+
+			for (int32 LaneIndex = 0; LaneIndex < GetNumLanes(); ++LaneIndex)
+			{
+				SimdData.DummySolverBody0[LaneIndex] = FSolverBody::MakeInitialized();
+				SimdData.DummySolverBody1[LaneIndex] = FSolverBody::MakeInitialized();
+			}
 		}
 
 		FPBDCollisionContainerSolverSimd::~FPBDCollisionContainerSolverSimd()
@@ -300,21 +358,18 @@ namespace Chaos
 
 		void FPBDCollisionContainerSolverSimd::Reset(const int32 MaxCollisions)
 		{
-			// @todo(chaos): allocation policy to reduce number of resizes as contacts increase
-			Constraints.Reset(MaxCollisions);
-			Solvers.Reset(MaxCollisions);
-			bCollisionConstraintPerIterationCollisionDetection.Reset(MaxCollisions);
+			NumConstraints = 0;
+			AppliedShockPropagation = FSolverReal(1);
+
+			ConstraintSolverIds.Reset(MaxCollisions);
 
-			// This is over-allocating by up to a factor of NumLanes!
+			// @todo(chaos): This is over-allocating by up to a factor of NumLanes!
+			SimdData.SimdSolverBodies.Reset(MaxCollisions);
+			SimdData.SimdSolvers.Reset(MaxCollisions);
 			SimdData.SimdManifoldPoints.Reset(MaxCollisions);
-			SimdData.SimdConstraintIndices.Reset(MaxCollisions);
+			SimdData.SimdConstraints.Reset(MaxCollisions);
 
-			for (int32 LaneIndex = 0; LaneIndex < GetNumLanes(); ++LaneIndex)
-			{
-				SimdData.SimdNumConstraints[LaneIndex] = 0;
-				SimdData.SimdNumManifoldPoints[LaneIndex] = 0;
-			}
-			NumConstraints = 0;
+			SimdData.SimdNumConstraints.SetValues(0);
 		}
 
 		void FPBDCollisionContainerSolverSimd::AddConstraints()
@@ -326,22 +381,21 @@ namespace Chaos
 		void FPBDCollisionContainerSolverSimd::AddConstraints(const TArrayView<Private::FPBDIslandConstraint>& IslandConstraints)
 		{
 			// Decide what lane this island goes into: Find the lane with the least constraints in it
-			// @todo(chaos): should use manifold point count, not constraint count
 			int32 IslandLaneIndex = 0;
-			int32 IslandLaneNumConstraints = SimdData.SimdNumConstraints[0];
+			int32 IslandLaneNumConstraints = SimdData.SimdNumConstraints.GetValue(0);
 			for (int32 LaneIndex = 1; LaneIndex < GetNumLanes(); ++LaneIndex)
 			{
-				if (SimdData.SimdNumConstraints[LaneIndex] < IslandLaneNumConstraints)
+				if (SimdData.SimdNumConstraints.GetValue(LaneIndex) < IslandLaneNumConstraints)
 				{
 					IslandLaneIndex = LaneIndex;
-					IslandLaneNumConstraints = SimdData.SimdNumConstraints[LaneIndex];
+					IslandLaneNumConstraints = SimdData.SimdNumConstraints.GetValue(LaneIndex);
 				}
 			}
 
 			// Make sure we have enough constraint rows for these constraints (space is pre-allocated)
-			if (SimdData.SimdConstraintIndices.Num() < IslandLaneNumConstraints + IslandConstraints.Num())
+			if (SimdData.SimdConstraints.Num() < IslandLaneNumConstraints + IslandConstraints.Num())
 			{
-				SimdData.SimdConstraintIndices.SetNum(IslandLaneNumConstraints + IslandConstraints.Num());
+				SimdData.SimdConstraints.SetNumZeroed(IslandLaneNumConstraints + IslandConstraints.Num(), false);
 			}
 
 			// Add all the constraints in the island to the selected lane
@@ -350,67 +404,68 @@ namespace Chaos
 				// NOTE: We will only ever be given constraints from our container (asserts in non-shipping)
 				FPBDCollisionConstraint& Constraint = IslandConstraint.GetConstraint()->AsUnsafe<FPBDCollisionConstraintHandle>()->GetContact();
 
-				// Add the constraint to our list
-				const int32 ConstraintIndex = Constraints.Add(&Constraint);
+				const int32 SolverIndex = SimdData.SimdNumConstraints.GetValue(IslandLaneIndex);
+				SimdData.SimdNumConstraints.SetValue(IslandLaneIndex, SolverIndex + 1);
 
 				// Add the constraint to its island's lane
-				const int32 RowIndex = SimdData.SimdNumConstraints[IslandLaneIndex];
-				SimdData.SimdConstraintIndices[RowIndex].ConstraintIndex[IslandLaneIndex] = ConstraintIndex;
+				SimdData.SimdConstraints[SolverIndex].SetValue(IslandLaneIndex, &Constraint);
 
-				++SimdData.SimdNumConstraints[IslandLaneIndex];
-				++NumConstraints;
+				// Store the mapping from constraint to solver/lane
+				ConstraintSolverIds.Add({ SolverIndex, IslandLaneIndex });
 			}
+
+			NumConstraints += IslandConstraints.Num();
 		}
 
 		void FPBDCollisionContainerSolverSimd::CreateSolvers()
 		{
-			// Allocate the solvers
-			Solvers.SetNum(Constraints.Num());
-			bCollisionConstraintPerIterationCollisionDetection.SetNum(Constraints.Num());
+			const int32 NumSolvers = SimdData.SimdNumConstraints.GetMaxValue();
 
-			// Reset manifold pointer per lane counter
-			for (int32 LaneIndex = 0; LaneIndex < GetNumLanes(); ++LaneIndex)
-			{
-				SimdData.SimdNumManifoldPoints[LaneIndex] = 0;
-			}
+			// Allocate the solvers (NOTE: unitialized data)
+			SimdData.SimdSolvers.SetNum(NumSolvers, false);
 
 			// Determine how many manifold point rows we need and tell each solver where its points are
-			int32 NumManifoldPointRows = 0;
-			for (int32 ConstraintRowIndex = 0; ConstraintRowIndex < SimdData.SimdConstraintIndices.Num(); ++ConstraintRowIndex)
+			int32 NumManifoldPoints = 0;
+			for (int32 SolverIndex = 0; SolverIndex < NumSolvers; ++SolverIndex)
 			{
-				TConstraintIndexSimd<4>& ConstraintIndices = SimdData.SimdConstraintIndices[ConstraintRowIndex];
+				TConstraintPtrSimd<4>& Constraints = SimdData.SimdConstraints[SolverIndex];
 
-				for (int32 LaneIndex = 0; LaneIndex < GetNumLanes(); ++LaneIndex)
+				// How may manifold points do we need for this row of constraints?
+				int32 MaxSolverManifoldPoints = 0;
+				for (int32 LaneIndex = 0; LaneIndex < 4; ++LaneIndex)
 				{
-					const int32 ConstraintIndex = ConstraintIndices.ConstraintIndex[LaneIndex];
-					if (ConstraintIndex != INDEX_NONE)
+					const FPBDCollisionConstraint* Constraint = Constraints.GetValue(LaneIndex);
+					if (Constraint != nullptr)
 					{
-						const int32 NumManifoldPoints = Constraints[ConstraintIndex]->NumManifoldPoints();
-
-						// Tell the solver where its manifold points are in the manifold points buffer
-						Solvers[ConstraintIndex].SetManifoldPointsBuffer(
-							ConstraintIndex,
-							LaneIndex,
-							SimdData.SimdNumManifoldPoints[LaneIndex],
-							NumManifoldPoints);
-
-						SimdData.SimdNumManifoldPoints[LaneIndex] += NumManifoldPoints;
-						NumManifoldPointRows = FMath::Max(NumManifoldPointRows, SimdData.SimdNumManifoldPoints[LaneIndex]);
+						MaxSolverManifoldPoints = FMath::Max(MaxSolverManifoldPoints, Constraint->NumManifoldPoints());
 					}
 				}
+
+				// Assign manifold point buffer indices
+				SimdData.SimdSolvers[SolverIndex].SetManifoldPointsBuffer(NumManifoldPoints, MaxSolverManifoldPoints);
+				NumManifoldPoints += MaxSolverManifoldPoints;
 			}
 
-			// Allocate the manifold point solver rows
-			SimdData.SimdManifoldPoints.SetNum(NumManifoldPointRows);
+			// Allocate the manifold point solver rows (NOTE: all data cleared to zero)
+			// @todo(chaos): we could be more selective about what gets zeroed
+			SimdData.SimdManifoldPoints.SetNumZeroed(NumManifoldPoints, false);
+
 
-			// Initalize the set of body pointers with a dummy body.
-			// This allows us to avoid some branches in the body data gather for rows where not all lanes are used
-			SimdData.SimdSolverBodies.SetNum(NumManifoldPointRows);
-			for (TSolverBodyPtrPairSimd<4>& SolverBodies : SimdData.SimdSolverBodies)
+#if !CHAOS_SIMDCOLLISIONSOLVER_USEDUMMYSOLVERBODY
+			// Allocate the body pointers (NOTE: set to null so that we don't have to visit lanes with no constraint again)
+			SimdData.SimdSolverBodies.SetNumZeroed(NumSolvers, false);
+#else
+			// Allocate the body pointers (NOTE: set to point to a dummy body so that even unused lanes have a valid body)
+			SimdData.SimdSolverBodies.SetNumUninitialized(NumSolvers, false);
+			for (TSolverBodyPtrPairSimd<4>& BodyPtrPair : SimdData.SimdSolverBodies)
 			{
-				SolverBodies.Body0.SetValues(&DummySolverBody);
-				SolverBodies.Body1.SetValues(&DummySolverBody);
+				for (int32 LaneIndex = 0; LaneIndex < GetNumLanes(); ++LaneIndex)
+				{
+					BodyPtrPair.Body0.SetValue(LaneIndex, &SimdData.DummySolverBody0[LaneIndex]);
+					BodyPtrPair.Body1.SetValue(LaneIndex, &SimdData.DummySolverBody1[LaneIndex]);
+				}
 			}
+#endif
 		}
 
 		void FPBDCollisionContainerSolverSimd::AddBodies(FSolverBodyContainer& SolverBodyContainer)
@@ -418,17 +473,20 @@ namespace Chaos
 			// All constraints and bodies are now known, so we can initialize the array of solvers.
 			CreateSolvers();
 
-			for (int32 ConstraintIndex = 0; ConstraintIndex < Constraints.Num(); ++ConstraintIndex)
+			for (int32 SolverIndex = 0; SolverIndex < SimdData.SimdSolvers.Num(); ++SolverIndex)
 			{
-				FPBDCollisionConstraint* Constraint = Constraints[ConstraintIndex];
-				FPBDCollisionSolverSimd& Solver = Solvers[ConstraintIndex];
-
-				FSolverBody* Body0 = SolverBodyContainer.FindOrAdd(Constraint->GetParticle0());
-				FSolverBody* Body1 = SolverBodyContainer.FindOrAdd(Constraint->GetParticle1());
-				check(Body0 != nullptr);
-				check(Body1 != nullptr);
-
-				Solver.SetSolverBodies(*Body0, *Body1);
+				TConstraintPtrSimd<4>& Constraints = SimdData.SimdConstraints[SolverIndex];
+				for (int32 LaneIndex = 0; LaneIndex < 4; ++LaneIndex)
+				{
+					const FPBDCollisionConstraint* Constraint = Constraints.GetValue(LaneIndex);
+					if (Constraint != nullptr)
+					{
+						FSolverBody* Body0 = SolverBodyContainer.FindOrAdd(Constraint->GetParticle0());
+						FSolverBody* Body1 = SolverBodyContainer.FindOrAdd(Constraint->GetParticle1());
+						SimdData.SimdSolverBodies[SolverIndex].Body0.SetValue(LaneIndex, Body0);
+						SimdData.SimdSolverBodies[SolverIndex].Body1.SetValue(LaneIndex, Body1);
+					}
+				}
 			}
 		}
 
@@ -448,31 +506,25 @@ namespace Chaos
 				return;
 			}
 
-			int32 RowBeginIndex = 0;
-			int32 RowEndIndex = Constraints.Num();
-
-			check(RowBeginIndex >= 0);
-			check(RowEndIndex <= Constraints.Num());
-
 			const FSolverReal Dt = FSolverReal(InDt);
 
 			TArrayView<TPBDCollisionSolverManifoldPointsSimd<4>> ManifoldPointsBuffer = MakeArrayView(SimdData.SimdManifoldPoints);
 			TArrayView<TSolverBodyPtrPairSimd<4>> SolverBodiesBuffer = MakeArrayView(SimdData.SimdSolverBodies);
 
-			// @todo(chaos): would it be better to iterate over manifold point rows?
 			bool bAnyPerIterationCollisions = false;
-			for (int32 ConstraintIndex = 0; ConstraintIndex < Constraints.Num(); ++ConstraintIndex)
+			for (int32 SolverIndex = 0; SolverIndex < SimdData.SimdSolvers.Num(); ++SolverIndex)
 			{
-				FPBDCollisionConstraint* Constraint = Constraints[ConstraintIndex];
-				FPBDCollisionSolverSimd& Solver = Solvers[ConstraintIndex];
-				bool& bPerIterationCollision = bCollisionConstraintPerIterationCollisionDetection[ConstraintIndex];
+				TConstraintPtrSimd<4>& SolverConstraints = SimdData.SimdConstraints[SolverIndex];
+				TPBDCollisionSolverSimd<4>& Solvers = SimdData.SimdSolvers[SolverIndex];
+				//bool& bPerIterationCollision = bCollisionConstraintPerIterationCollisionDetection[ConstraintIndex];
+				bool bPerIterationCollision = false;
 
 				UpdateSolverFromConstraint(
+					SimdData.SimdSolvers[SolverIndex],
+					SimdData.SimdConstraints[SolverIndex],
+					SimdData.SimdSolverBodies[SolverIndex],
 					ManifoldPointsBuffer,
-					SolverBodiesBuffer,
-					Solver,
-					Constraint, 
-					Dt, 
+					Dt,
 					ConstraintContainer.GetSolverSettings(),
 					bPerIterationCollision);
 
@@ -502,31 +554,22 @@ namespace Chaos
 				return;
 			}
 
-			int32 RowBeginIndex = 0;
-			int32 RowEndIndex = Constraints.Num();
-
-			check(RowBeginIndex >= 0);
-			check(RowEndIndex <= Constraints.Num());
-
 			const FSolverReal Dt = FSolverReal(InDt);
 
 			TArrayView<TPBDCollisionSolverManifoldPointsSimd<4>> ManifoldPointsBuffer = MakeArrayView(SimdData.SimdManifoldPoints);
 
 			// @todo(chaos): would it be better to iterate over manifold point rows?
-			for (int32 ConstraintIndex = 0; ConstraintIndex < Constraints.Num(); ++ConstraintIndex)
+			for (int32 SolverIndex = 0; SolverIndex < SimdData.SimdSolvers.Num(); ++SolverIndex)
 			{
-				FPBDCollisionConstraint* Constraint = Constraints[ConstraintIndex];
-				FPBDCollisionSolverSimd& Solver = Solvers[ConstraintIndex];
+				TConstraintPtrSimd<4>& SolverConstraints = SimdData.SimdConstraints[SolverIndex];
+				TPBDCollisionSolverSimd<4>& Solvers = SimdData.SimdSolvers[SolverIndex];
 
 				UpdateConstraintFromSolver(
+					SimdData.SimdConstraints[SolverIndex],
+					SimdData.SimdSolvers[SolverIndex],
+					SimdData.SimdSolverBodies[SolverIndex],
 					ManifoldPointsBuffer,
-					Constraint, 
-					Solver, 
 					Dt);
-
-				// Reset the collision solver here as the pointers will be invalid on the next tick
-				// @todo(chaos): Pointers are always reinitalized before use next tick so this isn't strictly necessary
-				Solver.Reset();
 			}
 		}
 
@@ -539,8 +582,9 @@ namespace Chaos
 			}
 
 			const FPBDCollisionSolverSettings& SolverSettings = ConstraintContainer.GetSolverSettings();
+			const FSolverReal Dt = FSolverReal(InDt);
 
-			UpdatePositionShockPropagation(InDt, It, NumIts, SolverSettings);
+			UpdatePositionShockPropagation(It, NumIts, SolverSettings);
 
 			// We run collision detection here under two conditions (normally it is run after Integration and before the constraint solver phase):
 			// 1) When deferring collision detection until the solver phase for better joint-collision behaviour (RBAN). In this case, we only do this on the first iteration.
@@ -548,12 +592,11 @@ namespace Chaos
 			const bool bRunDeferredCollisionDetection = (It == 0) && ConstraintContainer.GetDetectorSettings().bDeferNarrowPhase;
 			if (bRunDeferredCollisionDetection || bPerIterationCollisionDetection)
 			{
-				UpdateCollisions(InDt);
+				UpdateCollisions(Dt);
 			}
 
 			// Only apply friction for the last few (tunable) iterations
 			// Adjust max pushout to attempt to make it iteration count independent
-			const FSolverReal Dt = FSolverReal(InDt);
 			const bool bApplyStaticFriction = (It >= (NumIts - SolverSettings.NumPositionFrictionIterations));
 			const FSolverReal MaxPushOut = (SolverSettings.MaxPushOutVelocity > 0) ? (FSolverReal(SolverSettings.MaxPushOutVelocity) * Dt) / FSolverReal(NumIts) : FSolverReal(0);
 
@@ -561,6 +604,7 @@ namespace Chaos
 			if (bApplyStaticFriction)
 			{
 				Private::FPBDCollisionSolverHelperSimd::SolvePositionWithFriction(
+					MakeArrayView(SimdData.SimdSolvers),
 					MakeArrayView(SimdData.SimdManifoldPoints),
 					MakeArrayView(SimdData.SimdSolverBodies),
 					Dt,
@@ -569,6 +613,7 @@ namespace Chaos
 			else
 			{
 				Private::FPBDCollisionSolverHelperSimd::SolvePositionNoFriction(
+					MakeArrayView(SimdData.SimdSolvers),
 					MakeArrayView(SimdData.SimdManifoldPoints),
 					MakeArrayView(SimdData.SimdSolverBodies),
 					Dt,
@@ -586,7 +631,7 @@ namespace Chaos
 
 			const FPBDCollisionSolverSettings& SolverSettings = ConstraintContainer.GetSolverSettings();
 
-			UpdateVelocityShockPropagation(InDt, It, NumIts, SolverSettings);
+			UpdateVelocityShockPropagation(It, NumIts, SolverSettings);
 
 			const FSolverReal Dt = FSolverReal(InDt);
 			const bool bApplyDynamicFriction = (It >= NumIts - SolverSettings.NumVelocityFrictionIterations);
@@ -594,6 +639,7 @@ namespace Chaos
 			if (bApplyDynamicFriction)
 			{
 				Private::FPBDCollisionSolverHelperSimd::SolveVelocityWithFriction(
+					MakeArrayView(SimdData.SimdSolvers),
 					MakeArrayView(SimdData.SimdManifoldPoints),
 					MakeArrayView(SimdData.SimdSolverBodies),
 					Dt);
@@ -601,6 +647,7 @@ namespace Chaos
 			else
 			{
 				Private::FPBDCollisionSolverHelperSimd::SolveVelocityNoFriction(
+					MakeArrayView(SimdData.SimdSolvers),
 					MakeArrayView(SimdData.SimdManifoldPoints),
 					MakeArrayView(SimdData.SimdSolverBodies),
 					Dt);
@@ -612,94 +659,110 @@ namespace Chaos
 			// Not supported for collisions
 		}
 
-		void FPBDCollisionContainerSolverSimd::UpdatePositionShockPropagation(const FReal Dt, const int32 It, const int32 NumIts, const FPBDCollisionSolverSettings& SolverSettings)
+		void FPBDCollisionContainerSolverSimd::ApplyShockPropagation(const FSolverReal ShockPropagation)
 		{
-			// If this is the first shock propagation iteration, enable it on each solver
-			const bool bEnableShockPropagation = (It == NumIts - SolverSettings.NumPositionShockPropagationIterations);
-			if (bEnableShockPropagation)
+			// @todo(chaos): cache the mass scales so we don't have to look in the constraint again
+			if (ShockPropagation != AppliedShockPropagation)
 			{
-				for (int32 SolverIndex = 0; SolverIndex < Solvers.Num(); ++SolverIndex)
+				for (int32 SolverIndex = 0; SolverIndex < SimdData.SimdSolvers.Num(); ++SolverIndex)
 				{
-					Solvers[SolverIndex].EnablePositionShockPropagation(MakeArrayView(SimdData.SimdManifoldPoints));
+					for (int32 LaneIndex = 0; LaneIndex < 4; ++LaneIndex)
+					{
+						const FPBDCollisionConstraint* Constraint = SimdData.SimdConstraints[SolverIndex].GetValue(LaneIndex);
+						if (Constraint != nullptr)
+						{
+							const FSolverBody* Body0 = SimdData.SimdSolverBodies[SolverIndex].Body0.GetValue(LaneIndex);
+							const FSolverBody* Body1 = SimdData.SimdSolverBodies[SolverIndex].Body1.GetValue(LaneIndex);
+
+							FSolverReal ShockPropagation0, ShockPropagation1;
+							if (CalculateBodyShockPropagation(*Body0, *Body1, ShockPropagation, ShockPropagation0, ShockPropagation1))
+							{
+								SimdData.SimdSolvers[SolverIndex].UpdateMassNormal(
+									LaneIndex,
+									MakeArrayView(SimdData.SimdManifoldPoints),
+									*Body0,
+									*Body1,
+									ShockPropagation0 * FSolverReal(Constraint->GetInvMassScale0()),
+									ShockPropagation0 * FSolverReal(Constraint->GetInvInertiaScale0()),
+									ShockPropagation1 * FSolverReal(Constraint->GetInvMassScale1()),
+									ShockPropagation1 * FSolverReal(Constraint->GetInvInertiaScale1()));
+							}
+						}
+					}
 				}
+
+				AppliedShockPropagation = ShockPropagation;
 			}
 		}
 
-		void FPBDCollisionContainerSolverSimd::UpdateVelocityShockPropagation(const FReal Dt, const int32 It, const int32 NumIts, const FPBDCollisionSolverSettings& SolverSettings)
+		void FPBDCollisionContainerSolverSimd::UpdatePositionShockPropagation(const int32 It, const int32 NumIts, const FPBDCollisionSolverSettings& SolverSettings)
 		{
-			// Set/reset the shock propagation based on current iteration. The position solve may
-			// have left the bodies with a mass scale and we want to change or reset it.
-			const bool bEnableShockPropagation = (It == NumIts - SolverSettings.NumVelocityShockPropagationIterations);
-			if (bEnableShockPropagation)
-			{
-				for (int32 SolverIndex = 0; SolverIndex < Solvers.Num(); ++SolverIndex)
-				{
-					Solvers[SolverIndex].EnableVelocityShockPropagation(MakeArrayView(SimdData.SimdManifoldPoints));
-				}
-			}
-			else if (It == 0)
-			{
-				for (int32 SolverIndex = 0; SolverIndex < Solvers.Num(); ++SolverIndex)
-				{
-					Solvers[SolverIndex].DisableShockPropagation(MakeArrayView(SimdData.SimdManifoldPoints));
-				}
-			}
+			const bool bEnableShockPropagation = (It >= NumIts - SolverSettings.NumPositionShockPropagationIterations);
+			const FSolverReal ShockPropagation = (bEnableShockPropagation) ? CVars::Chaos_PBDCollisionSolver_Position_MinInvMassScale : FSolverReal(1);
+			ApplyShockPropagation(ShockPropagation);
 		}
 
-		void FPBDCollisionContainerSolverSimd::UpdateCollisions(const FReal InDt)
+		void FPBDCollisionContainerSolverSimd::UpdateVelocityShockPropagation(const int32 It, const int32 NumIts, const FPBDCollisionSolverSettings& SolverSettings)
 		{
-			const FSolverReal Dt = FSolverReal(InDt);
-			const bool bDeferredCollisionDetection = ConstraintContainer.GetDetectorSettings().bDeferNarrowPhase;
-
-			bool bNeedsAnotherIteration = false;
-			for (int32 SolverIndex = 0; SolverIndex < Solvers.Num(); ++SolverIndex)
-			{
-				if (bDeferredCollisionDetection || bCollisionConstraintPerIterationCollisionDetection[SolverIndex])
-				{
-					Private::FPBDCollisionSolverSimd& CollisionSolver = Solvers[SolverIndex];
-					FPBDCollisionConstraint* Constraint = Constraints[SolverIndex];
-
-					// Run collision detection at the current transforms including any correction from previous iterations
-					const FSolverBody& Body0 = CollisionSolver.SolverBody0().SolverBody();
-					const FSolverBody& Body1 = CollisionSolver.SolverBody1().SolverBody();
-					const FRigidTransform3 CorrectedActorWorldTransform0 = FRigidTransform3(Body0.CorrectedActorP(), Body0.CorrectedActorQ());
-					const FRigidTransform3 CorrectedActorWorldTransform1 = FRigidTransform3(Body1.CorrectedActorP(), Body1.CorrectedActorQ());
-					const FRigidTransform3 CorrectedShapeWorldTransform0 = Constraint->GetShapeRelativeTransform0() * CorrectedActorWorldTransform0;
-					const FRigidTransform3 CorrectedShapeWorldTransform1 = Constraint->GetShapeRelativeTransform1() * CorrectedActorWorldTransform1;
-
-					// @todo(chaos): this is ugly - pass these to the required functions instead and remove from the constraint class
-					// This is now only needed for LevelSet collision (see UpdateLevelsetLevelsetConstraint)
-					Constraint->SetSolverBodies(&Body0, &Body1);
-
-					// Reset the manifold if we are not using manifolds (we just use the first manifold point)
-					if (!Constraint->GetUseManifold())
-					{
-						Constraint->ResetActiveManifoldContacts();
-						CollisionSolver.ResetManifold();
-					}
+			const bool bEnableShockPropagation = (It >= NumIts - SolverSettings.NumVelocityShockPropagationIterations);
+			const FSolverReal ShockPropagation = (bEnableShockPropagation) ? CVars::Chaos_PBDCollisionSolver_Velocity_MinInvMassScale : FSolverReal(1);
+			ApplyShockPropagation(ShockPropagation);
+		}
 
-					// We need to know how many points were added to the manifold
-					const int32 BeginPointIndex = Constraint->NumManifoldPoints();
-
-					// NOTE: We deliberately have not updated the ShapwWorldTranforms on the constraint. If we did that, we would calculate 
-					// errors incorrectly in UpdateManifoldPoints, because the solver assumes nothing has been moved as we iterate (we accumulate 
-					// corrections that will be applied later.)
-					Constraint->ResetPhi(Constraint->GetCullDistance());
-					Collisions::UpdateConstraint(*Constraint, CorrectedShapeWorldTransform0, CorrectedShapeWorldTransform1, Dt);
-
-					// Update the manifold based on the new or updated contacts
-					UpdateSolverManifoldFromConstraint(
-						MakeArrayView(SimdData.SimdManifoldPoints),
-						MakeArrayView(SimdData.SimdSolverBodies),
-						CollisionSolver, 
-						Constraint, 
-						Dt, 
-						BeginPointIndex, 
-						Constraint->NumManifoldPoints());
-
-					Constraint->SetSolverBodies(nullptr, nullptr);
-				}
-			}
+		void FPBDCollisionContainerSolverSimd::UpdateCollisions(const FSolverReal InDt)
+		{
+			//const FSolverReal Dt = FSolverReal(InDt);
+			//const bool bDeferredCollisionDetection = ConstraintContainer.GetDetectorSettings().bDeferNarrowPhase;
+
+			//bool bNeedsAnotherIteration = false;
+			//for (int32 SolverIndex = 0; SolverIndex < Solvers.Num(); ++SolverIndex)
+			//{
+			//	if (bDeferredCollisionDetection || bCollisionConstraintPerIterationCollisionDetection[SolverIndex])
+			//	{
+			//		Private::FPBDCollisionSolverSimd& CollisionSolver = Solvers[SolverIndex];
+			//		FPBDCollisionConstraint* Constraint = Constraints[SolverIndex];
+
+			//		// Run collision detection at the current transforms including any correction from previous iterations
+			//		const FSolverBody& Body0 = CollisionSolver.SolverBody0().SolverBody();
+			//		const FSolverBody& Body1 = CollisionSolver.SolverBody1().SolverBody();
+			//		const FRigidTransform3 CorrectedActorWorldTransform0 = FRigidTransform3(Body0.CorrectedActorP(), Body0.CorrectedActorQ());
+			//		const FRigidTransform3 CorrectedActorWorldTransform1 = FRigidTransform3(Body1.CorrectedActorP(), Body1.CorrectedActorQ());
+			//		const FRigidTransform3 CorrectedShapeWorldTransform0 = Constraint->GetShapeRelativeTransform0() * CorrectedActorWorldTransform0;
+			//		const FRigidTransform3 CorrectedShapeWorldTransform1 = Constraint->GetShapeRelativeTransform1() * CorrectedActorWorldTransform1;
+
+			//		// @todo(chaos): this is ugly - pass these to the required functions instead and remove from the constraint class
+			//		// This is now only needed for LevelSet collision (see UpdateLevelsetLevelsetConstraint)
+			//		Constraint->SetSolverBodies(&Body0, &Body1);
+
+			//		// Reset the manifold if we are not using manifolds (we just use the first manifold point)
+			//		if (!Constraint->GetUseManifold())
+			//		{
+			//			Constraint->ResetActiveManifoldContacts();
+			//			CollisionSolver.ResetManifold();
+			//		}
+
+			//		// We need to know how many points were added to the manifold
+			//		const int32 BeginPointIndex = Constraint->NumManifoldPoints();
+
+			//		// NOTE: We deliberately have not updated the ShapwWorldTranforms on the constraint. If we did that, we would calculate 
+			//		// errors incorrectly in UpdateManifoldPoints, because the solver assumes nothing has been moved as we iterate (we accumulate 
+			//		// corrections that will be applied later.)
+			//		Constraint->ResetPhi(Constraint->GetCullDistance());
+			//		Collisions::UpdateConstraint(*Constraint, CorrectedShapeWorldTransform0, CorrectedShapeWorldTransform1, Dt);
+
+			//		// Update the manifold based on the new or updated contacts
+			//		UpdateSolverManifoldFromConstraint(
+			//			MakeArrayView(SimdData.SimdManifoldPoints),
+			//			MakeArrayView(SimdData.SimdSolverBodies),
+			//			CollisionSolver, 
+			//			Constraint, 
+			//			Dt, 
+			//			BeginPointIndex, 
+			//			Constraint->NumManifoldPoints());
+
+			//		Constraint->SetSolverBodies(nullptr, nullptr);
+			//	}
+			//}
 		}
 
 	}	// namespace Private
diff --git a/Engine/Source/Runtime/Experimental/Chaos/Private/Chaos/Collision/PBDCollisionSolverSimd.cpp b/Engine/Source/Runtime/Experimental/Chaos/Private/Chaos/Collision/PBDCollisionSolverSimd.cpp
--- a/Engine/Source/Runtime/Experimental/Chaos/Private/Chaos/Collision/PBDCollisionSolverSimd.cpp
+++ b/Engine/Source/Runtime/Experimental/Chaos/Private/Chaos/Collision/PBDCollisionSolverSimd.cpp
@@ -49,46 +49,55 @@ namespace Chaos
 			}
 		}
 
-		FORCEINLINE void PrefetchManifoldPoint(const TArrayView<TPBDCollisionSolverManifoldPointsSimd<4>>& ManifoldPoints, const int32 Index)
-		{
-			if (Index < ManifoldPoints.Num())
-			{
-				//PrefetchObject(&ManifoldPoints[Index]);
-				FPlatformMisc::PrefetchBlock(&ManifoldPoints[Index], sizeof(TPBDCollisionSolverManifoldPointsSimd<4>));
-			}
-		}
+		// Is there any point in having a value greater than 1? Tune this
+		static const int32 PrefetchCount = 4;
 
-		FORCEINLINE void PrefetchPositionSolverBodies(const TArrayView<TSolverBodyPtrPairSimd<4>>& SolverBodies, const int32 Index)
+		FORCEINLINE void PrefetchSolvePosition(
+			const int32 Index,
+			const TArrayView<TPBDCollisionSolverSimd<4>>& Solvers,
+			const TArrayView<TPBDCollisionSolverManifoldPointsSimd<4>>& ManifoldPoints,
+			const TArrayView<TSolverBodyPtrPairSimd<4>>& Bodies)
 		{
-			if (Index < SolverBodies.Num())
+			if (Index < Solvers.Num())
 			{
+				FPlatformMisc::PrefetchBlock(&Solvers[Index], sizeof(TPBDCollisionSolverSimd<4>));
+
+				FPlatformMisc::PrefetchBlock(&Solvers[Index].GetManifoldPoint(0, ManifoldPoints), 4 * sizeof(TPBDCollisionSolverManifoldPointsSimd<4>));
+
 				for (int32 LaneIndex = 0; LaneIndex < 4; ++LaneIndex)
 				{
-					SolverBodies[Index].Body0.GetValue(LaneIndex)->PrefetchPositionSolverData();
-					SolverBodies[Index].Body1.GetValue(LaneIndex)->PrefetchPositionSolverData();
+					Bodies[Index].Body0.GetValue(LaneIndex)->PrefetchPositionSolverData();
+					Bodies[Index].Body1.GetValue(LaneIndex)->PrefetchPositionSolverData();
 				}
 			}
 		}
 
 
-		FORCEINLINE void PrefetchVelocitySolverBodies(const TArrayView<TSolverBodyPtrPairSimd<4>>& SolverBodies, const int32 Index)
+		FORCEINLINE void PrefetchSolveVelocity(
+			const int32 Index,
+			const TArrayView<TPBDCollisionSolverSimd<4>>& Solvers,
+			const TArrayView<TPBDCollisionSolverManifoldPointsSimd<4>>& ManifoldPoints,
+			const TArrayView<TSolverBodyPtrPairSimd<4>>& Bodies)
 		{
-			if (Index < SolverBodies.Num())
+			if (Index < Solvers.Num())
 			{
+				FPlatformMisc::PrefetchBlock(&Solvers[Index], sizeof(TPBDCollisionSolverSimd<4>));
+
+				FPlatformMisc::PrefetchBlock(&Solvers[Index].GetManifoldPoint(0, ManifoldPoints), 4 * sizeof(TPBDCollisionSolverManifoldPointsSimd<4>));
+
 				for (int32 LaneIndex = 0; LaneIndex < 4; ++LaneIndex)
 				{
-					SolverBodies[Index].Body0.GetValue(LaneIndex)->PrefetchVelocitySolverData();
-					SolverBodies[Index].Body1.GetValue(LaneIndex)->PrefetchVelocitySolverData();
+					Bodies[Index].Body0.GetValue(LaneIndex)->PrefetchVelocitySolverData();
+					Bodies[Index].Body1.GetValue(LaneIndex)->PrefetchVelocitySolverData();
 				}
 			}
 		}
 
-		static const int32 PrefetchCount = 4;
-
 		template<>
 		void FPBDCollisionSolverHelperSimd::SolvePositionNoFriction<4>(
+			const TArrayView<TPBDCollisionSolverSimd<4>>& Solvers,
 			const TArrayView<TPBDCollisionSolverManifoldPointsSimd<4>>& ManifoldPoints,
-			const TArrayView<TSolverBodyPtrPairSimd<4>>& SolverBodies,
+			const TArrayView<TSolverBodyPtrPairSimd<4>>& Bodies,
 			const FSolverReal InDt,
 			const FSolverReal InMaxPushOut)
 		{
@@ -96,23 +105,22 @@ namespace Chaos
 
 			for (int32 Index = 0; Index < PrefetchCount; ++Index)
 			{
-				PrefetchManifoldPoint(ManifoldPoints, Index);
-				PrefetchPositionSolverBodies(SolverBodies, Index);
+				PrefetchSolvePosition(Index, Solvers, ManifoldPoints, Bodies);
 			}
 
-			for (int32 Index = 0; Index < ManifoldPoints.Num(); ++Index)
+			for (int32 Index = 0; Index < Solvers.Num(); ++Index)
 			{
-				PrefetchManifoldPoint(ManifoldPoints, Index + PrefetchCount);
-				PrefetchPositionSolverBodies(SolverBodies, Index + PrefetchCount);
+				PrefetchSolvePosition(Index + PrefetchCount, Solvers, ManifoldPoints, Bodies);
 
-				ManifoldPoints[Index].SolvePositionNoFriction(SolverBodies[Index].Body0, SolverBodies[Index].Body1, MaxPushOut);
+				Solvers[Index].SolvePositionNoFriction(ManifoldPoints, Bodies[Index].Body0, Bodies[Index].Body1, MaxPushOut);
 			}
 		}
 
 		template<>
 		void FPBDCollisionSolverHelperSimd::SolvePositionWithFriction<4>(
+			const TArrayView<TPBDCollisionSolverSimd<4>>& Solvers,
 			const TArrayView<TPBDCollisionSolverManifoldPointsSimd<4>>& ManifoldPoints,
-			const TArrayView<TSolverBodyPtrPairSimd<4>>& SolverBodies,
+			const TArrayView<TSolverBodyPtrPairSimd<4>>& Bodies,
 			const FSolverReal InDt,
 			const FSolverReal InMaxPushOut)
 		{
@@ -121,69 +129,65 @@ namespace Chaos
 
 			for (int32 Index = 0; Index < PrefetchCount; ++Index)
 			{
-				PrefetchManifoldPoint(ManifoldPoints, Index);
-				PrefetchPositionSolverBodies(SolverBodies, Index);
+				PrefetchSolvePosition(Index, Solvers, ManifoldPoints, Bodies);
 			}
 
-			for (int32 Index = 0; Index < ManifoldPoints.Num(); ++Index)
+			for (int32 Index = 0; Index < Solvers.Num(); ++Index)
 			{
-				PrefetchManifoldPoint(ManifoldPoints, Index + PrefetchCount);
-				PrefetchPositionSolverBodies(SolverBodies, Index + PrefetchCount);
+				PrefetchSolvePosition(Index + PrefetchCount, Solvers, ManifoldPoints, Bodies);
 
-				ManifoldPoints[Index].SolvePositionWithFriction(SolverBodies[Index].Body0, SolverBodies[Index].Body1, MaxPushOut, FrictionStiffnessScale);
+				Solvers[Index].SolvePositionWithFriction(ManifoldPoints, Bodies[Index].Body0, Bodies[Index].Body1, MaxPushOut, FrictionStiffnessScale);
 			}
 		}
 
 		template<>
 		void FPBDCollisionSolverHelperSimd::SolveVelocityNoFriction<4>(
+			const TArrayView<TPBDCollisionSolverSimd<4>>& Solvers,
 			const TArrayView<TPBDCollisionSolverManifoldPointsSimd<4>>& ManifoldPoints,
-			const TArrayView<TSolverBodyPtrPairSimd<4>>& SolverBodies,
+			const TArrayView<TSolverBodyPtrPairSimd<4>>& Bodies,
 			const FSolverReal InDt)
 		{
 			const FSimd4Realf Dt = FSimd4Realf::Make(InDt);
 
 			for (int32 Index = 0; Index < PrefetchCount; ++Index)
 			{
-				PrefetchManifoldPoint(ManifoldPoints, Index);
-				PrefetchVelocitySolverBodies(SolverBodies, Index);
+				PrefetchSolveVelocity(Index, Solvers, ManifoldPoints, Bodies);
 			}
 
-			for (int32 Index = 0; Index < ManifoldPoints.Num(); ++Index)
+			for (int32 Index = 0; Index < Solvers.Num(); ++Index)
 			{
-				PrefetchManifoldPoint(ManifoldPoints, Index + PrefetchCount);
-				PrefetchVelocitySolverBodies(SolverBodies, Index + PrefetchCount);
+				PrefetchSolveVelocity(Index + PrefetchCount, Solvers, ManifoldPoints, Bodies);
 
-				ManifoldPoints[Index].SolveVelocityNoFriction(SolverBodies[Index].Body0, SolverBodies[Index].Body1, Dt);
+				Solvers[Index].SolveVelocityNoFriction(ManifoldPoints, Bodies[Index].Body0, Bodies[Index].Body1, Dt);
 			}
 		}
 
 		template<>
 		void FPBDCollisionSolverHelperSimd::SolveVelocityWithFriction<4>(
+			const TArrayView<TPBDCollisionSolverSimd<4>>& Solvers,
 			const TArrayView<TPBDCollisionSolverManifoldPointsSimd<4>>& ManifoldPoints,
-			const TArrayView<TSolverBodyPtrPairSimd<4>>& SolverBodies,
+			const TArrayView<TSolverBodyPtrPairSimd<4>>& Bodies,
 			const FSolverReal InDt)
 		{
 			if (!CVars::bChaos_PBDCollisionSolver_Velocity_FrictionEnabled)
 			{
-				SolveVelocityNoFriction(ManifoldPoints, SolverBodies, InDt);
+				SolveVelocityNoFriction(Solvers, ManifoldPoints, Bodies, InDt);
 				return;
 			}
 
 			for (int32 Index = 0; Index < PrefetchCount; ++Index)
 			{
-				PrefetchManifoldPoint(ManifoldPoints, Index);
-				PrefetchVelocitySolverBodies(SolverBodies, Index);
+				PrefetchSolveVelocity(Index, Solvers, ManifoldPoints, Bodies);
 			}
 
 			const FSimd4Realf Dt = FSimd4Realf::Make(InDt);
 			const FSimd4Realf FrictionStiffnessScale = FSimd4Realf::Make(CVars::Chaos_PBDCollisionSolver_Position_StaticFrictionStiffness);
 
-			for (int32 Index = 0; Index < ManifoldPoints.Num(); ++Index)
+			for (int32 Index = 0; Index < Solvers.Num(); ++Index)
 			{
-				PrefetchManifoldPoint(ManifoldPoints, Index + PrefetchCount);
-				PrefetchVelocitySolverBodies(SolverBodies, Index + PrefetchCount);
+				PrefetchSolveVelocity(Index + PrefetchCount, Solvers, ManifoldPoints, Bodies);
 
-				ManifoldPoints[Index].SolveVelocityWithFriction(SolverBodies[Index].Body0, SolverBodies[Index].Body1, Dt, FrictionStiffnessScale);
+				Solvers[Index].SolveVelocityWithFriction(ManifoldPoints, Bodies[Index].Body0, Bodies[Index].Body1, Dt, FrictionStiffnessScale);
 			}
 		}
 
diff --git a/Engine/Source/Runtime/Experimental/Chaos/Private/Chaos/PBDCollisionConstraints.cpp b/Engine/Source/Runtime/Experimental/Chaos/Private/Chaos/PBDCollisionConstraints.cpp
--- a/Engine/Source/Runtime/Experimental/Chaos/Private/Chaos/PBDCollisionConstraints.cpp
+++ b/Engine/Source/Runtime/Experimental/Chaos/Private/Chaos/PBDCollisionConstraints.cpp
@@ -104,7 +104,7 @@ namespace Chaos
 	using FABTestingCollisionContainerSolver = Private::TABTestingConstraintContainerSolver<FPBDCollisionContainerSolver, Private::FPBDCollisionContainerSolverSimd>;
 
 	// Simd AB testing callback
-	void ABTestCollisionContainerSolver(
+	void ABTestSimdCollisionSolver(
 		FABTestingCollisionContainerSolver::ESolverPhase Phase,
 		const FPBDCollisionContainerSolver& SolverA,
 		const Private::FPBDCollisionContainerSolverSimd& SolverB,
@@ -114,50 +114,108 @@ namespace Chaos
 		for (int32 ConstraintIndex = 0; ConstraintIndex < SolverA.GetNumConstraints(); ++ConstraintIndex)
 		{
 			const Private::FPBDCollisionSolver& ConstraintSolverA = SolverA.GetConstraintSolver(ConstraintIndex);
-			const Private::FPBDCollisionSolverSimd& ConstraintSolverB = SolverB.GetConstraintSolver(ConstraintIndex);
+			const FSolverBody& Body0 = ConstraintSolverA.SolverBody0().SolverBody();
+			const FSolverBody& Body1 = ConstraintSolverA.SolverBody1().SolverBody();
 
-			if (ConstraintSolverA.NumManifoldPoints() != ConstraintSolverB.NumManifoldPoints())
+			const Private::FPBDCollisionContainerSolverSimd::FConstraintSolverId ConstraintSolverBId = SolverB.GetConstraintSolverId(ConstraintIndex);
+			const Private::TPBDCollisionSolverSimd<4>& ConstraintSolverB = SolverB.GetConstraintSolver(ConstraintSolverBId.SolverIndex);
+			const int32 LaneIndexB = ConstraintSolverBId.LaneIndex;
+
+			if (ConstraintSolverA.NumManifoldPoints() != ConstraintSolverB.NumManifoldPoints().GetValue(LaneIndexB))
 			{
-				UE_LOG(LogChaos, Warning, TEXT("ABTestSolverBodies: ManifoldPoint count mismatch"));
+				UE_LOG(LogChaos, Warning, TEXT("ManifoldPoint count mismatch"));
 				return;
 			}
 
 			for (int32 ManifoldPointIndex = 0; ManifoldPointIndex < ConstraintSolverA.NumManifoldPoints(); ++ManifoldPointIndex)
 			{
 				const Private::FPBDCollisionSolverManifoldPoint& ManifoldPointSolverA = ConstraintSolverA.GetManifoldPoint(ManifoldPointIndex);
-				const Private::TPBDCollisionSolverManifoldPointsSimd<4>& ManifoldPointSolverB = SolverB.GetManifoldPointSolver(ConstraintSolverB.GetManifoldPointBeginIndex() + ManifoldPointIndex);
-				const int32 LaneIndexB = ConstraintSolverB.GetLaneIndex();
+				const Private::TPBDCollisionSolverManifoldPointsSimd<4>& ManifoldPointSolverB = ConstraintSolverB.GetManifoldPoint(ManifoldPointIndex, SolverB.GetManifoldPointBuffer());
+				bool bIsError = false;
 
 				if (ManifoldPointSolverA.WorldContact.ContactNormal != ManifoldPointSolverB.SimdContactNormal.GetValue(LaneIndexB))
 				{
-					UE_LOG(LogChaos, Warning, TEXT("ABTestSolverBodies: ContactNormal mismatch"));
+					UE_LOG(LogChaos, Warning, TEXT("ContactNormal mismatch"));
+					bIsError = true;
+				}
+				if (ManifoldPointSolverA.WorldContact.ContactTangentU != ManifoldPointSolverB.SimdContactTangentU.GetValue(LaneIndexB))
+				{
+					UE_LOG(LogChaos, Warning, TEXT("ContactTangentU mismatch"));
+					bIsError = true;
+				}
+				if (ManifoldPointSolverA.WorldContact.ContactTangentV != ManifoldPointSolverB.SimdContactTangentV.GetValue(LaneIndexB))
+				{
+					UE_LOG(LogChaos, Warning, TEXT("ContactTangentV mismatch"));
+					bIsError = true;
 				}
 				if (ManifoldPointSolverA.WorldContact.RelativeContactPoints[0] != ManifoldPointSolverB.SimdRelativeContactPoint0.GetValue(LaneIndexB))
 				{
-					UE_LOG(LogChaos, Warning, TEXT("ABTestSolverBodies: RelativeContactPoints mismatch"));
+					UE_LOG(LogChaos, Warning, TEXT("RelativeContactPoints mismatch"));
+					bIsError = true;
 				}
 				if (ManifoldPointSolverA.WorldContact.RelativeContactPoints[1] != ManifoldPointSolverB.SimdRelativeContactPoint1.GetValue(LaneIndexB))
 				{
-					UE_LOG(LogChaos, Warning, TEXT("ABTestSolverBodies: RelativeContactPoints mismatch"));
+					UE_LOG(LogChaos, Warning, TEXT("RelativeContactPoints mismatch"));
+					bIsError = true;
 				}
 				if (ManifoldPointSolverA.ContactMassNormal != ManifoldPointSolverB.SimdContactMassNormal.GetValue(LaneIndexB))
 				{
-					UE_LOG(LogChaos, Warning, TEXT("ABTestSolverBodies: ContactMassNormal mismatch"));
+					UE_LOG(LogChaos, Warning, TEXT("ContactMassNormal mismatch"));
+					bIsError = true;
+				}
+				if (ManifoldPointSolverA.ContactMassTangentU != ManifoldPointSolverB.SimdContactMassTangentU.GetValue(LaneIndexB))
+				{
+					UE_LOG(LogChaos, Warning, TEXT("ContactMassTangentU mismatch"));
+					bIsError = true;
+				}
+				if (ManifoldPointSolverA.ContactMassTangentV != ManifoldPointSolverB.SimdContactMassTangentV.GetValue(LaneIndexB))
+				{
+					UE_LOG(LogChaos, Warning, TEXT("ContactMassTangentV mismatch"));
+					bIsError = true;
+				}
+				if (Body0.IsDynamic())
+				{
+					if (ManifoldPointSolverA.WorldContactTangentUAngular0 != ManifoldPointSolverB.SimdContactTangentUAngular0.GetValue(LaneIndexB))
+					{
+						UE_LOG(LogChaos, Warning, TEXT("ContactTangentUAngular0 mismatch"));
+						bIsError = true;
+					}
+					if (ManifoldPointSolverA.WorldContactTangentVAngular0 != ManifoldPointSolverB.SimdContactTangentVAngular0.GetValue(LaneIndexB))
+					{
+						UE_LOG(LogChaos, Warning, TEXT("ContactTangentVAngular0 mismatch"));
+						bIsError = true;
+					}
+				}
+				if (Body1.IsDynamic())
+				{
+					if (ManifoldPointSolverA.WorldContactTangentUAngular1 != ManifoldPointSolverB.SimdContactTangentUAngular1.GetValue(LaneIndexB))
+					{
+						UE_LOG(LogChaos, Warning, TEXT("ContactTangentUAngular1 mismatch"));
+						bIsError = true;
+					}
+					if (ManifoldPointSolverA.WorldContactTangentVAngular1 != ManifoldPointSolverB.SimdContactTangentVAngular1.GetValue(LaneIndexB))
+					{
+						UE_LOG(LogChaos, Warning, TEXT("ContactTangentVAngular1 mismatch"));
+						bIsError = true;
+					}
 				}
 
 				if (Phase == FABTestingCollisionContainerSolver::ESolverPhase::PostApplyPositionConstraints)
 				{
 					if (ManifoldPointSolverA.NetPushOutNormal != ManifoldPointSolverB.SimdNetPushOutNormal.GetValue(LaneIndexB))
 					{
-						UE_LOG(LogChaos, Warning, TEXT("ABTestSolverBodies: NetPushOutNormal mismatch"));
+						UE_LOG(LogChaos, Warning, TEXT("NetPushOutNormal mismatch"));
+						bIsError = true;
 					}
 					if (ManifoldPointSolverA.NetPushOutTangentU != ManifoldPointSolverB.SimdNetPushOutTangentU.GetValue(LaneIndexB))
 					{
-						UE_LOG(LogChaos, Warning, TEXT("ABTestSolverBodies: NetPushOutTangentU mismatch"));
+						UE_LOG(LogChaos, Warning, TEXT("NetPushOutTangentU mismatch"));
+						bIsError = true;
 					}
 					if (ManifoldPointSolverA.NetPushOutTangentV != ManifoldPointSolverB.SimdNetPushOutTangentV.GetValue(LaneIndexB))
 					{
-						UE_LOG(LogChaos, Warning, TEXT("ABTestSolverBodies: NetPushOutTangentV mismatch"));
+						UE_LOG(LogChaos, Warning, TEXT("NetPushOutTangentV mismatch"));
+						bIsError = true;
 					}
 				}
 
@@ -165,23 +223,31 @@ namespace Chaos
 				{
 					if (ManifoldPointSolverA.NetImpulseNormal != ManifoldPointSolverB.SimdNetImpulseNormal.GetValue(LaneIndexB))
 					{
-						UE_LOG(LogChaos, Warning, TEXT("ABTestSolverBodies: NetImpulseNormal mismatch"));
+						UE_LOG(LogChaos, Warning, TEXT("NetImpulseNormal mismatch"));
+						bIsError = true;
 					}
 					if (ManifoldPointSolverA.NetImpulseTangentU != ManifoldPointSolverB.SimdNetImpulseTangentU.GetValue(LaneIndexB))
 					{
-						UE_LOG(LogChaos, Warning, TEXT("ABTestSolverBodies: NetImpulseTangentU mismatch"));
+						UE_LOG(LogChaos, Warning, TEXT("NetImpulseTangentU mismatch"));
+						bIsError = true;
 					}
 					if (ManifoldPointSolverA.NetImpulseTangentV != ManifoldPointSolverB.SimdNetImpulseTangentV.GetValue(LaneIndexB))
 					{
-						UE_LOG(LogChaos, Warning, TEXT("ABTestSolverBodies: NetImpulseTangentV mismatch"));
+						UE_LOG(LogChaos, Warning, TEXT("NetImpulseTangentV mismatch"));
+						bIsError = true;
 					}
 				}
+
+				if (bIsError)
+				{
+					UE_LOG(LogChaos, Warning, TEXT("SimdCollisionSolver mismatch"));
+				}
 			}
 		}
 
 		if (SolverBodyContainerA.Num() != SolverBodyContainerB.Num())
 		{
-			UE_LOG(LogChaos, Warning, TEXT("ABTestSolverBodies: Body count mismatch"));
+			UE_LOG(LogChaos, Warning, TEXT("Body count mismatch"));
 			return;
 		}
 
@@ -189,33 +255,59 @@ namespace Chaos
 		{
 			const FSolverBody& BodyA = SolverBodyContainerA.GetSolverBody(BodyIndex);
 			const FSolverBody& BodyB = SolverBodyContainerB.GetSolverBody(BodyIndex);
+			bool bIsError = false;
 
 			if (BodyA.P() != BodyB.P())
 			{
-				UE_LOG(LogChaos, Warning, TEXT("Position mismatch: (%f %f %f) != (%f %f %f)"), 
-					BodyA.P().X, BodyA.P().Y, BodyA.P().Z,
-					BodyB.P().X, BodyB.P().Y, BodyB.P().Z);
+				UE_LOG(LogChaos, Warning, TEXT("Position mismatch"));
+				bIsError = true;
+			}
+
+			if (BodyA.Q() != BodyB.Q())
+			{
+				UE_LOG(LogChaos, Warning, TEXT("Rotation mismatch"));
+				bIsError = true;
 			}
 
 			if (BodyA.DP() != BodyB.DP())
 			{
-				UE_LOG(LogChaos, Warning, TEXT("PositionDelta mismatch: (%f %f %f) != (%f %f %f)"),
-					BodyA.DP().X, BodyA.DP().Y, BodyA.DP().Z,
-					BodyB.DP().X, BodyB.DP().Y, BodyB.DP().Z);
+				UE_LOG(LogChaos, Warning, TEXT("PositionDelta mismatch"));
+				bIsError = true;
 			}
 
 			if (BodyA.DQ() != BodyB.DQ())
 			{
-				UE_LOG(LogChaos, Warning, TEXT("RotationDelta mismatch: (%f %f %f) != (%f %f %f)"),
-					BodyA.DQ().X, BodyA.DQ().Y, BodyA.DQ().Z,
-					BodyB.DQ().X, BodyB.DQ().Y, BodyB.DQ().Z);
+				UE_LOG(LogChaos, Warning, TEXT("RotationDelta mismatch"));
+				bIsError = true;
 			}
 
 			if (BodyA.V() != BodyB.V())
 			{
-				UE_LOG(LogChaos, Warning, TEXT("Velocity mismatch: (%f %f %f) != (%f %f %f)"),
-					BodyA.V().X, BodyA.V().Y, BodyA.V().Z,
-					BodyB.V().X, BodyB.V().Y, BodyB.V().Z);
+				UE_LOG(LogChaos, Warning, TEXT("Velocity mismatch"));
+				bIsError = true;
+			}
+
+			if (BodyA.W() != BodyB.W())
+			{
+				UE_LOG(LogChaos, Warning, TEXT("AngVel mismatch"));
+				bIsError = true;
+			}
+
+			if (BodyA.InvM() != BodyB.InvM())
+			{
+				UE_LOG(LogChaos, Warning, TEXT("InvM mismatch"));
+				bIsError = true;
+			}
+
+			if (BodyA.InvI() != BodyB.InvI())
+			{
+				UE_LOG(LogChaos, Warning, TEXT("InvI mismatch"));
+				bIsError = true;
+			}
+
+			if (bIsError)
+			{
+				UE_LOG(LogChaos, Warning, TEXT("SimdCollisionSolver mismatch"));
 			}
 		}
 	}
@@ -284,7 +376,7 @@ namespace Chaos
 					MakeUnique<FPBDCollisionContainerSolver>(*this, Priority),
 					MakeUnique<Private::FPBDCollisionContainerSolverSimd>(*this, Priority),
 					Priority,
-					&ABTestCollisionContainerSolver);
+					&ABTestSimdCollisionSolver);
 			}
 #endif
 			return MakeUnique<Private::FPBDCollisionContainerSolverSimd>(*this, Priority);
diff --git a/Engine/Source/Runtime/Experimental/Chaos/Public/Chaos/Collision/PBDCollisionConstraint.h b/Engine/Source/Runtime/Experimental/Chaos/Public/Chaos/Collision/PBDCollisionConstraint.h
--- a/Engine/Source/Runtime/Experimental/Chaos/Public/Chaos/Collision/PBDCollisionConstraint.h
+++ b/Engine/Source/Runtime/Experimental/Chaos/Public/Chaos/Collision/PBDCollisionConstraint.h
@@ -427,6 +427,7 @@ namespace Chaos
 		FManifoldPoint& GetManifoldPoint(const int32 PointIndex) { return ManifoldPoints[PointIndex]; }
 		const FManifoldPoint& GetManifoldPoint(const int32 PointIndex) const { return ManifoldPoints[PointIndex]; }
 		const FManifoldPoint* GetClosestManifoldPoint() const { return (ClosestManifoldPointIndex != INDEX_NONE) ? &ManifoldPoints[ClosestManifoldPointIndex] : nullptr; }
+		bool IsManifoldPointActive(const int32 PointIndex) const { return IsEnabled() && (PointIndex < NumManifoldPoints()) && !ManifoldPoints[PointIndex].Flags.bDisabled; }
 
 		const FManifoldPointResult& GetManifoldPointResult(const int32 PointIndex) const
 		{
diff --git a/Engine/Source/Runtime/Experimental/Chaos/Public/Chaos/Collision/PBDCollisionContainerSolverSimd.h b/Engine/Source/Runtime/Experimental/Chaos/Public/Chaos/Collision/PBDCollisionContainerSolverSimd.h
--- a/Engine/Source/Runtime/Experimental/Chaos/Public/Chaos/Collision/PBDCollisionContainerSolverSimd.h
+++ b/Engine/Source/Runtime/Experimental/Chaos/Public/Chaos/Collision/PBDCollisionContainerSolverSimd.h
@@ -19,6 +19,9 @@ namespace Chaos
 		class FPBDIsland;
 		class FPBDIslandConstraint;
 
+		template<int TNumLanes>
+		using TConstraintPtrSimd = TSimdValue<FPBDCollisionConstraint*, TNumLanes>;
+
 		/**
 		 * The solver for a set of collision constraints. This collects all the data required to solve a set of collision
 		 * constraints into a contiguous, ordered buffer.
@@ -51,62 +54,50 @@ namespace Chaos
 			virtual void ApplyVelocityConstraints(const FReal Dt, const int32 It, const int32 NumIts) override final;
 			virtual void ApplyProjectionConstraints(const FReal Dt, const int32 It, const int32 NumIts) override final;
 
-			template<int TNumLanes>
-			struct TConstraintIndexSimd
-			{
-				TConstraintIndexSimd()
-				{
-					for (int32 LaneIndex = 0; LaneIndex < TNumLanes; ++LaneIndex)
-					{
-						ConstraintIndex[LaneIndex] = INDEX_NONE;
-					}
-				}
-
-				int32 ConstraintIndex[TNumLanes];
-			};
-
 			// Rows of collision constraints and solvers. 
 			// We solve each row using SIMD. 
 			// NumLanes is the width of a row and will be the same as the float SIMD register width (or less).
 			template<int TNumLanes>
 			struct FDataSimd
 			{
-				FDataSimd()
-					: SimdNumConstraints{ 0, }
-					, SimdNumManifoldPoints{ 0, }
-				{
-				}
-
-				int32 SimdNumConstraints[TNumLanes];
-				int32 SimdNumManifoldPoints[TNumLanes];
-
+				TSimdInt32<TNumLanes> SimdNumConstraints;
 				TArray<TSolverBodyPtrPairSimd<TNumLanes>> SimdSolverBodies;
-				TArray<TConstraintIndexSimd<TNumLanes>> SimdConstraintIndices;
+				TArray<TPBDCollisionSolverSimd<TNumLanes>> SimdSolvers;
 				TArray<TPBDCollisionSolverManifoldPointsSimd<TNumLanes>> SimdManifoldPoints;
+				TArray<TConstraintPtrSimd<TNumLanes>> SimdConstraints;
+				FSolverBody DummySolverBody0[TNumLanes];
+				FSolverBody DummySolverBody1[TNumLanes];
+			};
+
+			// Used to recover the solver for a constraint
+			// @todo(chaos): this is only really needed for debug validation
+			struct FConstraintSolverId
+			{
+				int32 SolverIndex;
+				int32 LaneIndex;
 			};
 
 			// For testing
-			const FPBDCollisionSolverSimd& GetConstraintSolver(const int32 ConstraintIndex) const { return Solvers[ConstraintIndex]; }
-			const TPBDCollisionSolverManifoldPointsSimd<4>& GetManifoldPointSolver(const int32 RowIndex) const { return SimdData.SimdManifoldPoints[RowIndex]; }
+			const FConstraintSolverId& GetConstraintSolverId(const int32 ConstraintIndex) const { return ConstraintSolverIds[ConstraintIndex]; }
+			const TPBDCollisionSolverSimd<4>& GetConstraintSolver(const int32 SolverIndex) const { return SimdData.SimdSolvers[SolverIndex]; }
+			const TPBDCollisionSolverManifoldPointsSimd<4>& GetManifoldPointSolver(const int32 PointIndex) const { return SimdData.SimdManifoldPoints[PointIndex]; }
+			TArrayView<const TPBDCollisionSolverManifoldPointsSimd<4>> GetManifoldPointBuffer() const { return MakeArrayView(SimdData.SimdManifoldPoints); }
 
 		private:
 			int32 GetNumLanes() const { return 4; }
 			void CreateSolvers();
-			void UpdatePositionShockPropagation(const FReal Dt, const int32 It, const int32 NumIts, const FPBDCollisionSolverSettings& SolverSettings);
-			void UpdateVelocityShockPropagation(const FReal Dt, const int32 It, const int32 NumIts, const FPBDCollisionSolverSettings& SolverSettings);
-			void UpdateCollisions(const FReal InDt);
+			void UpdatePositionShockPropagation(const int32 It, const int32 NumIts, const FPBDCollisionSolverSettings& SolverSettings);
+			void UpdateVelocityShockPropagation(const int32 It, const int32 NumIts, const FPBDCollisionSolverSettings& SolverSettings);
+			void ApplyShockPropagation(const FSolverReal ShockPropagation);
+			void UpdateCollisions(const FSolverReal InDt);
 
 			const FPBDCollisionConstraints& ConstraintContainer;
 
-			TArray<FPBDCollisionConstraint*> Constraints;
-			TArray<FPBDCollisionSolverSimd> Solvers;
 			FDataSimd<4> SimdData;
-			TArray<bool> bCollisionConstraintPerIterationCollisionDetection;
+			TArray<FConstraintSolverId> ConstraintSolverIds;
 			int32 NumConstraints;
+			FSolverReal AppliedShockPropagation;
 			bool bPerIterationCollisionDetection;
-
-			FSolverBody DummySolverBody;
-			FConstraintSolverBody DummyConstraintSolverBody;
 		};
 
 	}	// namespace Private
diff --git a/Engine/Source/Runtime/Experimental/Chaos/Public/Chaos/Collision/PBDCollisionSolverSimd.h b/Engine/Source/Runtime/Experimental/Chaos/Public/Chaos/Collision/PBDCollisionSolverSimd.h
--- a/Engine/Source/Runtime/Experimental/Chaos/Public/Chaos/Collision/PBDCollisionSolverSimd.h
+++ b/Engine/Source/Runtime/Experimental/Chaos/Public/Chaos/Collision/PBDCollisionSolverSimd.h
@@ -7,6 +7,12 @@
 #include "Chaos/Evolution/SolverBody.h"
 #include "Chaos/Simd4.h"
 
+// Set to 1 to use a dummy solver body instead of a nullptr in unused lanes, and an aligned Simd gather operation
+// Currently this is slower, but once the object sizes are reduced may be better
+#ifndef CHAOS_SIMDCOLLISIONSOLVER_USEDUMMYSOLVERBODY
+#define CHAOS_SIMDCOLLISIONSOLVER_USEDUMMYSOLVERBODY 0
+#endif
+
 namespace Chaos
 {
 	class FManifoldPoint;
@@ -20,14 +26,147 @@ namespace Chaos
 
 	namespace Private
 	{
-		class FPBDCollisionSolverSimd;
-	}
+		template<int TNumLanes>
+		class TPBDCollisionSolverSimd;
 
-	namespace Private
-	{
 		template<int TNumLanes>
 		using TSolverBodyPtrSimd = TSimdValue<FSolverBody*, TNumLanes>;
 
+		inline void GatherBodyPositionCorrections(
+			const TSolverBodyPtrSimd<4>& Body0,
+			const TSolverBodyPtrSimd<4>& Body1,
+			TSimdVec3f<4>& DP0,
+			TSimdVec3f<4>& DQ0,
+			TSimdVec3f<4>& DP1,
+			TSimdVec3f<4>& DQ1)
+		{
+#if !CHAOS_SIMDCOLLISIONSOLVER_USEDUMMYSOLVERBODY
+			// Non-SIMD gather
+			for (int32 LaneIndex = 0; LaneIndex < 4; ++LaneIndex)
+			{
+				const FSolverBody* LaneBody0 = Body0.GetValue(LaneIndex);
+				const FSolverBody* LaneBody1 = Body1.GetValue(LaneIndex);
+				if (LaneBody0 != nullptr)
+				{
+					DP0.SetValue(LaneIndex, LaneBody0->DP());
+					DQ0.SetValue(LaneIndex, LaneBody0->DQ());
+				}
+				if (LaneBody1 != nullptr)
+				{
+					DP1.SetValue(LaneIndex, LaneBody1->DP());
+					DQ1.SetValue(LaneIndex, LaneBody1->DQ());
+				}
+			}
+#else
+			// SIMD gather (valid ptrs to dummy object for unused lanes)
+			const FSolverBody* Body00 = Body0.GetValue(0);
+			const FSolverBody* Body01 = Body0.GetValue(1);
+			const FSolverBody* Body02 = Body0.GetValue(2);
+			const FSolverBody* Body03 = Body0.GetValue(3);
+			DP0 = SimdGatherAligned(
+				Body00->DP(),
+				Body01->DP(),
+				Body02->DP(),
+				Body03->DP());
+			DQ0 = SimdGatherAligned(
+				Body00->DQ(),
+				Body01->DQ(),
+				Body02->DQ(),
+				Body03->DQ());
+
+			const FSolverBody* Body10 = Body1.GetValue(0);
+			const FSolverBody* Body11 = Body1.GetValue(1);
+			const FSolverBody* Body12 = Body1.GetValue(2);
+			const FSolverBody* Body13 = Body1.GetValue(3);
+			DP1 = SimdGatherAligned(
+				Body10->DP(),
+				Body11->DP(),
+				Body12->DP(),
+				Body13->DP());
+			DQ1 = SimdGatherAligned(
+				Body10->DQ(),
+				Body11->DQ(),
+				Body12->DQ(),
+				Body13->DQ());
+#endif
+		}
+
+		inline void ScatterBodyPositionCorrections(
+			const TSimdVec3f<4>& DP0,
+			const TSimdVec3f<4>& DQ0,
+			const TSimdVec3f<4>& DP1,
+			const TSimdVec3f<4>& DQ1,
+			const TSolverBodyPtrSimd<4>& Body0,
+			const TSolverBodyPtrSimd<4>& Body1)
+		{
+			for (int32 LaneIndex = 0; LaneIndex < 4; ++LaneIndex)
+			{
+				FSolverBody* LaneBody0 = Body0.GetValue(LaneIndex);
+				FSolverBody* LaneBody1 = Body1.GetValue(LaneIndex);
+				if (LaneBody0 != nullptr)
+				{
+					LaneBody0->SetDP(DP0.GetValue(LaneIndex));
+					LaneBody0->SetDQ(DQ0.GetValue(LaneIndex));
+				}
+				if (LaneBody1 != nullptr)
+				{
+					LaneBody1->SetDP(DP1.GetValue(LaneIndex));
+					LaneBody1->SetDQ(DQ1.GetValue(LaneIndex));
+				}
+			}
+		}
+
+
+		inline void GatherBodyVelocities(
+			const TSolverBodyPtrSimd<4>& Body0,
+			const TSolverBodyPtrSimd<4>& Body1,
+			TSimdVec3f<4>& V0,
+			TSimdVec3f<4>& W0,
+			TSimdVec3f<4>& V1,
+			TSimdVec3f<4>& W1)
+		{
+			for (int32 LaneIndex = 0; LaneIndex < 4; ++LaneIndex)
+			{
+				const FSolverBody* LaneBody0 = Body0.GetValue(LaneIndex);
+				const FSolverBody* LaneBody1 = Body1.GetValue(LaneIndex);
+				if (LaneBody0 != nullptr)
+				{
+					V0.SetValue(LaneIndex, LaneBody0->V());
+					W0.SetValue(LaneIndex, LaneBody0->W());
+				}
+				if (LaneBody1 != nullptr)
+				{
+					V1.SetValue(LaneIndex, LaneBody1->V());
+					W1.SetValue(LaneIndex, LaneBody1->W());
+				}
+			}
+		}
+
+		inline void ScatterBodyVelocities(
+			const TSimdVec3f<4>& V0,
+			const TSimdVec3f<4>& W0,
+			const TSimdVec3f<4>& V1,
+			const TSimdVec3f<4>& W1,
+			const TSolverBodyPtrSimd<4>& Body0,
+			const TSolverBodyPtrSimd<4>& Body1)
+		{
+			for (int32 LaneIndex = 0; LaneIndex < 4; ++LaneIndex)
+			{
+				FSolverBody* LaneBody0 = Body0.GetValue(LaneIndex);
+				FSolverBody* LaneBody1 = Body1.GetValue(LaneIndex);
+				if (LaneBody0 != nullptr)
+				{
+					LaneBody0->SetV(V0.GetValue(LaneIndex));
+					LaneBody0->SetW(W0.GetValue(LaneIndex));
+				}
+				if (LaneBody0 != nullptr)
+				{
+					LaneBody1->SetV(V1.GetValue(LaneIndex));
+					LaneBody1->SetW(W1.GetValue(LaneIndex));
+				}
+			}
+		}
+
 		template<int TNumLanes>
 		struct TSolverBodyPtrPairSimd
 		{
@@ -48,26 +187,144 @@ namespace Chaos
 			using FSimdSelector = TSimdSelector<TNumLanes>;
 			using FSimdSolverBodyPtr = TSolverBodyPtrSimd<TNumLanes>;
 
-			TPBDCollisionSolverManifoldPointsSimd()
+			// Whether the manifold point in each lane is set up for solving
+			FSimdSelector IsValid;
+
+			// World-space contact point relative to each particle's center of mass
+			FSimdVec3f SimdRelativeContactPoint0;
+			FSimdVec3f SimdRelativeContactPoint1;
+
+			// Normal PushOut
+			FSimdVec3f SimdContactNormal;
+			FSimdRealf SimdContactDeltaNormal;
+			FSimdRealf SimdNetPushOutNormal;
+			FSimdRealf SimdContactMassNormal;
+			FSimdVec3f SimdContactNormalAngular0;
+			FSimdVec3f SimdContactNormalAngular1;
+
+			// Tangential PushOut
+			FSimdVec3f SimdContactTangentU;
+			FSimdVec3f SimdContactTangentV;
+			FSimdRealf SimdContactDeltaTangentU;
+			FSimdRealf SimdContactDeltaTangentV;
+			FSimdRealf SimdNetPushOutTangentU;
+			FSimdRealf SimdNetPushOutTangentV;
+			FSimdRealf SimdStaticFrictionRatio;
+			FSimdRealf SimdContactMassTangentU;
+			FSimdRealf SimdContactMassTangentV;
+			FSimdVec3f SimdContactTangentUAngular0;
+			FSimdVec3f SimdContactTangentVAngular0;
+			FSimdVec3f SimdContactTangentUAngular1;
+			FSimdVec3f SimdContactTangentVAngular1;
+
+			// Normal Impulse
+			FSimdRealf SimdContactTargetVelocityNormal;
+			FSimdRealf SimdNetImpulseNormal;
+			FSimdRealf SimdNetImpulseTangentU;
+			FSimdRealf SimdNetImpulseTangentV;
+		};
+
+		/**
+		* Holds the solver
+		*/
+		template<int TNumLanes>
+		class TPBDCollisionSolverSimd
+		{
+		public:
+			using FSimdVec3f = TSimdVec3f<TNumLanes>;
+			using FSimdRealf = TSimdRealf<TNumLanes>;
+			using FSimdInt32 = TSimdInt32<TNumLanes>;
+			using FSimdSelector = TSimdSelector<TNumLanes>;
+			using FSimdSolverBodyPtr = TSolverBodyPtrSimd<TNumLanes>;
+			using FSimdManifoldPoint = TPBDCollisionSolverManifoldPointsSimd<TNumLanes>;
+
+			static const int32 MaxConstrainedBodies = 2;
+			static const int32 MaxPointsPerConstraint = 4;
+
+			// Create a solver that is initialized to safe defaults
+			static TPBDCollisionSolverSimd<TNumLanes> MakeInitialized()
 			{
-				SimdConstraintIndex.SetValues(INDEX_NONE);
-				SimdManifoldPointIndex.SetValues(INDEX_NONE);
+				TPBDCollisionSolverSimd<TNumLanes> Solver;
+				Solver.Init();
+				return Solver;
 			}
 
-			void SetSharedData(
-				const int32 LaneIndex,
-				const int32 InConstraintIndex,
-				const int32 InManifoldPointIndex)
+			// Create a solver with no initialization
+			static TPBDCollisionSolverSimd<TNumLanes> MakeUninitialized()
+			{
+				return TPBDCollisionSolverSimd<TNumLanes>();
+			}
+
+			// NOTE: Does not initialize any properties. See MakeInitialized
+			TPBDCollisionSolverSimd() {}
+
+			FSimdInt32 NumManifoldPoints() const
+			{
+				return SimdNumManifoldPoints;
+			}
+
+			void SetNumManifoldPoints(const FSimdInt32& InNum)
+			{
+				SimdNumManifoldPoints = InNum;
+			}
+
+			void SetManifoldPointsBuffer(const int32 InBeginIndex, const int32 InMax)
+			{
+				ManifoldPointBeginIndex = InBeginIndex;
+				MaxManifoldPoints = InMax;
+				SimdNumManifoldPoints = FSimdInt32::Zero();
+			}
+
+			int32 GetMaxManifoldPoints() const
+			{
+				return MaxManifoldPoints;
+			}
+
+			const FSimdManifoldPoint& GetManifoldPoint(
+				const int32 ManifoldPointIndex,
+				const TArrayView<const FSimdManifoldPoint>& ManifoldPointsBuffer) const
+			{
+				return ManifoldPointsBuffer[GetBufferIndex(ManifoldPointIndex)];
+			}
+
+			/** Reset the state of the collision solver */
+			void Reset()
 			{
-				SimdConstraintIndex.SetValue(LaneIndex, InConstraintIndex);
-				SimdManifoldPointIndex.SetValue(LaneIndex, InManifoldPointIndex);
+				ManifoldPointBeginIndex = 0;
+				MaxManifoldPoints = 0;
+				SimdNumManifoldPoints = FSimdInt32::Zero();
+			}
+
+			void ResetManifold()
+			{
+				SimdNumManifoldPoints = FSimdInt32::Zero();
+			}
+
+			void SetFriction(const FSimdRealf InStaticFriction, const FSimdRealf InDynamicFriction, const FSimdRealf InVelocityFriction)
+			{
+				SimdStaticFriction = InStaticFriction;
+				SimdDynamicFriction = InDynamicFriction;
+				SimdVelocityFriction = InVelocityFriction;
+			}
+
+			void SetStiffness(const FSimdRealf InStiffness)
+			{
+				SimdStiffness = InStiffness;
+			}
+
+			void InitManifoldPoints(const TArrayView<FSimdManifoldPoint>& ManifoldPointsBuffer)
+			{
+				// Not required as long as we are zeroing the full points array on resize
+				// Otherwise we need to reset all outputs and some properties here as not all lanes will be used by all points
 			}
 
 			/**
-			 * @brief Initialize contact geometry (contact positions, normal, etc)
+			 * Set up a manifold point ready for solving
 			*/
-			void InitContact(
+			void SetManifoldPoint(
+				const int32 ManifoldPointIndex,
 				const int32 LaneIndex,
+				const TArrayView<FSimdManifoldPoint>& ManifoldPointsBuffer,
 				const FSolverVec3& InRelativeContactPosition0,
 				const FSolverVec3& InRelativeContactPosition1,
 				const FSolverVec3& InWorldContactNormal,
@@ -77,442 +334,446 @@ namespace Chaos
 				const FSolverReal InWorldContactDeltaTangentU,
 				const FSolverReal InWorldContactDeltaTangentV,
 				const FSolverReal InWorldContactVelocityTargetNormal,
-				const FSolverReal InStiffness,
-				const FSolverReal InStaticFriction,
-				const FSolverReal InDynamicFriction,
-				const FSolverReal InVelocityFriction)
+				const FSolverBody& Body0,
+				const FSolverBody& Body1,
+				const FSolverReal InvMScale0,
+				const FSolverReal InvIScale0,
+				const FSolverReal InvMScale1,
+				const FSolverReal InvIScale1)
 			{
-				SimdRelativeContactPoint0.SetValue(LaneIndex, InRelativeContactPosition0);
-				SimdRelativeContactPoint1.SetValue(LaneIndex, InRelativeContactPosition1);
-				SimdContactNormal.SetValue(LaneIndex, InWorldContactNormal);
-				SimdContactTangentU.SetValue(LaneIndex, InWorldContactTangentU);
-				SimdContactTangentV.SetValue(LaneIndex, InWorldContactTangentV);
-				SimdContactDeltaNormal.SetValue(LaneIndex, InWorldContactDeltaNormal);
-				SimdContactDeltaTangentU.SetValue(LaneIndex, InWorldContactDeltaTangentU);
-				SimdContactDeltaTangentV.SetValue(LaneIndex, InWorldContactDeltaTangentV);
-				SimdContactTargetVelocityNormal.SetValue(LaneIndex, InWorldContactVelocityTargetNormal);
-
-				SimdStiffness.SetValue(LaneIndex, InStiffness);
-				SimdStaticFriction.SetValue(LaneIndex, InStaticFriction);
-				SimdDynamicFriction.SetValue(LaneIndex, InDynamicFriction);
-				SimdVelocityFriction.SetValue(LaneIndex, InVelocityFriction);
-			}
-
-			/**
-			 * @brief Initialize contact state (net impulses, masses, etc). Must come after InitContact
-			 * because the mass properties depend on contact offsets etc.
-			*/
-			void FinalizeContact(
-				const int32 LaneIndex,
-				const FConstraintSolverBody& Body0,
-				const FConstraintSolverBody& Body1)
-			{
-				SimdNetPushOutNormal.SetValue(LaneIndex, 0);
-				SimdNetPushOutTangentU.SetValue(LaneIndex, 0);
-				SimdNetPushOutTangentV.SetValue(LaneIndex, 0);
-				SimdNetImpulseNormal.SetValue(LaneIndex, 0);
-				SimdNetImpulseTangentU.SetValue(LaneIndex, 0);
-				SimdNetImpulseTangentV.SetValue(LaneIndex, 0);
-				SimdStaticFrictionRatio.SetValue(LaneIndex, 0);
-
-				UpdateMass(LaneIndex, Body0, Body1);
+				const int32 BufferIndex = GetBufferIndex(ManifoldPointIndex);
+				FSimdManifoldPoint& ManifoldPoint = ManifoldPointsBuffer[BufferIndex];
+
+				ManifoldPoint.IsValid.SetValue(LaneIndex, true);
+
+				ManifoldPoint.SimdRelativeContactPoint0.SetValue(LaneIndex, InRelativeContactPosition0);
+				ManifoldPoint.SimdRelativeContactPoint1.SetValue(LaneIndex, InRelativeContactPosition1);
+				ManifoldPoint.SimdContactNormal.SetValue(LaneIndex, InWorldContactNormal);
+				ManifoldPoint.SimdContactTangentU.SetValue(LaneIndex, InWorldContactTangentU);
+				ManifoldPoint.SimdContactTangentV.SetValue(LaneIndex, InWorldContactTangentV);
+				ManifoldPoint.SimdContactDeltaNormal.SetValue(LaneIndex, InWorldContactDeltaNormal);
+				ManifoldPoint.SimdContactDeltaTangentU.SetValue(LaneIndex, InWorldContactDeltaTangentU);
+				ManifoldPoint.SimdContactDeltaTangentV.SetValue(LaneIndex, InWorldContactDeltaTangentV);
+				ManifoldPoint.SimdContactTargetVelocityNormal.SetValue(LaneIndex, InWorldContactVelocityTargetNormal);
+
+				UpdateManifoldPointMass(
+					ManifoldPointIndex, 
+					LaneIndex, 
+					ManifoldPointsBuffer, 
+					Body0, 
+					Body1,
+					InvMScale0,
+					InvIScale0,
+					InvMScale1,
+					InvIScale1);
 			}
 
 			/**
 			 * @brief Update the cached mass properties based on the current body transforms
 			*/
-			void UpdateMass(
+			void UpdateManifoldPointMass(
+				const int32 ManifoldPointIndex,
 				const int32 LaneIndex,
-				const FConstraintSolverBody& Body0,
-				const FConstraintSolverBody& Body1)
+				const TArrayView<FSimdManifoldPoint>& ManifoldPointsBuffer,
+				const FSolverBody& Body0,
+				const FSolverBody& Body1,
+				const FSolverReal InvMScale0,
+				const FSolverReal InvIScale0,
+				const FSolverReal InvMScale1,
+				const FSolverReal InvIScale1)
 			{
+				const int32 BufferIndex = GetBufferIndex(ManifoldPointIndex);
+				FSimdManifoldPoint& ManifoldPoint = ManifoldPointsBuffer[BufferIndex];
+
+				if (!ManifoldPoint.IsValid.GetValue(LaneIndex))
+				{
+					return;
+				}
+
 				FSolverReal ContactMassInvNormal = FSolverReal(0);
 				FSolverReal ContactMassInvTangentU = FSolverReal(0);
 				FSolverReal ContactMassInvTangentV = FSolverReal(0);
 
 				SimdInvM0.SetValue(LaneIndex, 0);
 				SimdInvM1.SetValue(LaneIndex, 0);
-				SimdContactNormalAngular0.SetValue(LaneIndex, FVec3f(0));
-				SimdContactTangentUAngular0.SetValue(LaneIndex, FVec3f(0));
-				SimdContactTangentVAngular0.SetValue(LaneIndex, FVec3f(0));
-				SimdContactNormalAngular1.SetValue(LaneIndex, FVec3f(0));
-				SimdContactTangentUAngular1.SetValue(LaneIndex, FVec3f(0));
-				SimdContactTangentVAngular1.SetValue(LaneIndex, FVec3f(0));
+				ManifoldPoint.SimdContactNormalAngular0.SetValue(LaneIndex, FVec3f(0));
+				ManifoldPoint.SimdContactTangentUAngular0.SetValue(LaneIndex, FVec3f(0));
+				ManifoldPoint.SimdContactTangentVAngular0.SetValue(LaneIndex, FVec3f(0));
+				ManifoldPoint.SimdContactNormalAngular1.SetValue(LaneIndex, FVec3f(0));
+				ManifoldPoint.SimdContactTangentUAngular1.SetValue(LaneIndex, FVec3f(0));
+				ManifoldPoint.SimdContactTangentVAngular1.SetValue(LaneIndex, FVec3f(0));
 
-				const FVec3f ContactNormal = SimdContactNormal.GetValue(LaneIndex);
-				const FVec3f ContactTangentU = SimdContactTangentU.GetValue(LaneIndex);
-				const FVec3f ContactTangentV = SimdContactTangentV.GetValue(LaneIndex);
+				const FVec3f ContactNormal = ManifoldPoint.SimdContactNormal.GetValue(LaneIndex);
+				const FVec3f ContactTangentU = ManifoldPoint.SimdContactTangentU.GetValue(LaneIndex);
+				const FVec3f ContactTangentV = ManifoldPoint.SimdContactTangentV.GetValue(LaneIndex);
 
 				if (Body0.IsDynamic())
 				{
-					const FVec3f RelativeContactPoint0 = SimdRelativeContactPoint0.GetValue(LaneIndex);
+					const FSolverReal InvM0 = InvMScale0 * Body0.InvM();
+					const FSolverMatrix33 InvI0 = InvIScale0 * Body0.InvI();
 
+					const FVec3f RelativeContactPoint0 = ManifoldPoint.SimdRelativeContactPoint0.GetValue(LaneIndex);
 					const FSolverVec3 R0xN = FSolverVec3::CrossProduct(RelativeContactPoint0, ContactNormal);
 					const FSolverVec3 R0xU = FSolverVec3::CrossProduct(RelativeContactPoint0, ContactTangentU);
 					const FSolverVec3 R0xV = FSolverVec3::CrossProduct(RelativeContactPoint0, ContactTangentV);
-
-					const FSolverMatrix33 InvI0 = Body0.InvI();
 					const FSolverVec3 IR0xN = InvI0 * R0xN;
 					const FSolverVec3 IR0xU = InvI0 * R0xU;
 					const FSolverVec3 IR0xV = InvI0 * R0xV;
 
-					SimdInvM0.SetValue(LaneIndex, Body0.InvM());
+					SimdInvM0.SetValue(LaneIndex, InvM0);
 
-					SimdContactNormalAngular0.SetValue(LaneIndex, IR0xN);
-					SimdContactTangentUAngular0.SetValue(LaneIndex, IR0xU);
-					SimdContactTangentVAngular0.SetValue(LaneIndex, IR0xV);
+					ManifoldPoint.SimdContactNormalAngular0.SetValue(LaneIndex, IR0xN);
+					ManifoldPoint.SimdContactTangentUAngular0.SetValue(LaneIndex, IR0xU);
+					ManifoldPoint.SimdContactTangentVAngular0.SetValue(LaneIndex, IR0xV);
 
-					ContactMassInvNormal += FSolverVec3::DotProduct(R0xN, IR0xN) + Body0.InvM();
-					ContactMassInvTangentU += FSolverVec3::DotProduct(R0xU, IR0xU) + Body0.InvM();
-					ContactMassInvTangentV += FSolverVec3::DotProduct(R0xV, IR0xV) + Body0.InvM();
+					ContactMassInvNormal += FSolverVec3::DotProduct(R0xN, IR0xN) + InvM0;
+					ContactMassInvTangentU += FSolverVec3::DotProduct(R0xU, IR0xU) + InvM0;
+					ContactMassInvTangentV += FSolverVec3::DotProduct(R0xV, IR0xV) + InvM0;
 				}
 				if (Body1.IsDynamic())
 				{
-					const FVec3f RelativeContactPoint1 = SimdRelativeContactPoint1.GetValue(LaneIndex);
+					const FSolverReal InvM1 = InvMScale1 * Body1.InvM();
+					const FSolverMatrix33 InvI1 = InvIScale1 * Body1.InvI();
 
+					const FVec3f RelativeContactPoint1 = ManifoldPoint.SimdRelativeContactPoint1.GetValue(LaneIndex);
 					const FSolverVec3 R1xN = FSolverVec3::CrossProduct(RelativeContactPoint1, ContactNormal);
 					const FSolverVec3 R1xU = FSolverVec3::CrossProduct(RelativeContactPoint1, ContactTangentU);
 					const FSolverVec3 R1xV = FSolverVec3::CrossProduct(RelativeContactPoint1, ContactTangentV);
-
-					const FSolverMatrix33 InvI1 = Body1.InvI();
 					const FSolverVec3 IR1xN = InvI1 * R1xN;
 					const FSolverVec3 IR1xU = InvI1 * R1xU;
 					const FSolverVec3 IR1xV = InvI1 * R1xV;
 
-					SimdInvM1.SetValue(LaneIndex, Body1.InvM());
+					SimdInvM1.SetValue(LaneIndex, InvM1);
 
-					SimdContactNormalAngular1.SetValue(LaneIndex, IR1xN);
-					SimdContactTangentUAngular1.SetValue(LaneIndex, IR1xU);
-					SimdContactTangentVAngular1.SetValue(LaneIndex, IR1xV);
+					ManifoldPoint.SimdContactNormalAngular1.SetValue(LaneIndex, IR1xN);
+					ManifoldPoint.SimdContactTangentUAngular1.SetValue(LaneIndex, IR1xU);
+					ManifoldPoint.SimdContactTangentVAngular1.SetValue(LaneIndex, IR1xV);
 
-					ContactMassInvNormal += FSolverVec3::DotProduct(R1xN, IR1xN) + Body1.InvM();
-					ContactMassInvTangentU += FSolverVec3::DotProduct(R1xU, IR1xU) + Body1.InvM();
-					ContactMassInvTangentV += FSolverVec3::DotProduct(R1xV, IR1xV) + Body1.InvM();
+					ContactMassInvNormal += FSolverVec3::DotProduct(R1xN, IR1xN) + InvM1;
+					ContactMassInvTangentU += FSolverVec3::DotProduct(R1xU, IR1xU) + InvM1;
+					ContactMassInvTangentV += FSolverVec3::DotProduct(R1xV, IR1xV) + InvM1;
 				}
 
-				SimdContactMassNormal.SetValue(LaneIndex, (ContactMassInvNormal > FSolverReal(UE_SMALL_NUMBER)) ? FSolverReal(1) / ContactMassInvNormal : FSolverReal(0));
-				SimdContactMassTangentU.SetValue(LaneIndex, (ContactMassInvTangentU > FSolverReal(UE_SMALL_NUMBER)) ? FSolverReal(1) / ContactMassInvTangentU : FSolverReal(0));
-				SimdContactMassTangentV.SetValue(LaneIndex, (ContactMassInvTangentV > FSolverReal(UE_SMALL_NUMBER)) ? FSolverReal(1) / ContactMassInvTangentV : FSolverReal(0));
+				ManifoldPoint.SimdContactMassNormal.SetValue(LaneIndex, (ContactMassInvNormal > FSolverReal(UE_SMALL_NUMBER)) ? FSolverReal(1) / ContactMassInvNormal : FSolverReal(0));
+				ManifoldPoint.SimdContactMassTangentU.SetValue(LaneIndex, (ContactMassInvTangentU > FSolverReal(UE_SMALL_NUMBER)) ? FSolverReal(1) / ContactMassInvTangentU : FSolverReal(0));
+				ManifoldPoint.SimdContactMassTangentV.SetValue(LaneIndex, (ContactMassInvTangentV > FSolverReal(UE_SMALL_NUMBER)) ? FSolverReal(1) / ContactMassInvTangentV : FSolverReal(0));
 			}
 
-			/**
-			 * @brief Update the contact mass for the normal correction
-			 * This is used by shock propagation.
-			*/
-			void UpdateMassNormal(
+			void UpdateManifoldPointMassNormal(
+				const int32 ManifoldPointIndex,
 				const int32 LaneIndex,
-				FConstraintSolverBody& Body0,
-				FConstraintSolverBody& Body1)
+				const TArrayView<FSimdManifoldPoint>& ManifoldPointsBuffer,
+				const FSolverBody& Body0,
+				const FSolverBody& Body1,
+				const FSolverReal InvMScale0,
+				const FSolverReal InvIScale0,
+				const FSolverReal InvMScale1,
+				const FSolverReal InvIScale1)
 			{
+				const int32 BufferIndex = GetBufferIndex(ManifoldPointIndex);
+				FSimdManifoldPoint& ManifoldPoint = ManifoldPointsBuffer[BufferIndex];
+
+				if (!ManifoldPoint.IsValid.GetValue(LaneIndex))
+				{
+					return;
+				}
+
 				FSolverReal ContactMassInvNormal = FSolverReal(0);
 
 				SimdInvM0.SetValue(LaneIndex, 0);
 				SimdInvM1.SetValue(LaneIndex, 0);
-				SimdContactNormalAngular0.SetValue(LaneIndex, FVec3f(0));
-				SimdContactNormalAngular1.SetValue(LaneIndex, FVec3f(0));
+				ManifoldPoint.SimdContactNormalAngular0.SetValue(LaneIndex, FVec3f(0));
+				ManifoldPoint.SimdContactNormalAngular1.SetValue(LaneIndex, FVec3f(0));
 
-				const FVec3f ContactNormal = SimdContactNormal.GetValue(LaneIndex);
+				const FVec3f ContactNormal = ManifoldPoint.SimdContactNormal.GetValue(LaneIndex);
 
 				if (Body0.IsDynamic())
 				{
-					const FVec3f RelativeContactPoint0 = SimdRelativeContactPoint0.GetValue(LaneIndex);
+					const FSolverReal InvM0 = InvMScale0 * Body0.InvM();
+					const FSolverMatrix33 InvI0 = InvIScale0 * Body0.InvI();
+
+					const FVec3f RelativeContactPoint0 = ManifoldPoint.SimdRelativeContactPoint0.GetValue(LaneIndex);
 					const FSolverVec3 R0xN = FSolverVec3::CrossProduct(RelativeContactPoint0, ContactNormal);
-					const FSolverMatrix33 InvI0 = Body0.InvI();
 					const FSolverVec3 IR0xN = InvI0 * R0xN;
 
-					SimdInvM0.SetValue(LaneIndex, Body0.InvM());
+					SimdInvM0.SetValue(LaneIndex, InvM0);
 
-					SimdContactNormalAngular0.SetValue(LaneIndex, IR0xN);
+					ManifoldPoint.SimdContactNormalAngular0.SetValue(LaneIndex, IR0xN);
 
-					ContactMassInvNormal += FSolverVec3::DotProduct(R0xN, IR0xN) + Body0.InvM();
+					ContactMassInvNormal += FSolverVec3::DotProduct(R0xN, IR0xN) + InvM0;
 				}
 				if (Body1.IsDynamic())
 				{
-					const FVec3f RelativeContactPoint1 = SimdRelativeContactPoint1.GetValue(LaneIndex);
+					const FSolverReal InvM1 = InvMScale1 * Body1.InvM();
+					const FSolverMatrix33 InvI1 = InvIScale1 * Body1.InvI();
+
+					const FVec3f RelativeContactPoint1 = ManifoldPoint.SimdRelativeContactPoint1.GetValue(LaneIndex);
 					const FSolverVec3 R1xN = FSolverVec3::CrossProduct(RelativeContactPoint1, ContactNormal);
-					const FSolverMatrix33 InvI1 = Body1.InvI();
 					const FSolverVec3 IR1xN = InvI1 * R1xN;
 
-					SimdInvM1.SetValue(LaneIndex, Body1.InvM());
-
-					SimdContactNormalAngular1.SetValue(LaneIndex, IR1xN);
-
-					ContactMassInvNormal += FSolverVec3::DotProduct(R1xN, IR1xN) + Body1.InvM();
-				}
-
-				SimdContactMassNormal.SetValue(LaneIndex, (ContactMassInvNormal > FSolverReal(UE_SMALL_NUMBER)) ? FSolverReal(1) / ContactMassInvNormal : FSolverReal(0));
-			}
+					SimdInvM1.SetValue(LaneIndex, InvM1);
 
-			FORCEINLINE void GatherBodyPositionCorrections(
-				const FSimdSelector& IsValid,
-				const FSimdSolverBodyPtr& Body0,
-				const FSimdSolverBodyPtr& Body1,
-				FSimdVec3f& DP0,
-				FSimdVec3f& DQ0,
-				FSimdVec3f& DP1,
-				FSimdVec3f& DQ1)
-			{
-				for (int32 LaneIndex = 0; LaneIndex < TNumLanes; ++LaneIndex)
-				{
-					if (IsValid.GetValue(LaneIndex))
-					{
-						const FSolverBody* LaneBody0 = Body0.GetValue(LaneIndex);
-						const FSolverBody* LaneBody1 = Body1.GetValue(LaneIndex);
-						DP0.SetValue(LaneIndex, LaneBody0->DP());
-						DQ0.SetValue(LaneIndex, LaneBody0->DQ());
-						DP1.SetValue(LaneIndex, LaneBody1->DP());
-						DQ1.SetValue(LaneIndex, LaneBody1->DQ());
-					}
-				}
-			}
+					ManifoldPoint.SimdContactNormalAngular1.SetValue(LaneIndex, IR1xN);
 
-			FORCEINLINE void ScatterBodyPositionCorrections(
-				const FSimdSelector& IsValid,
-				const FSimdVec3f& DP0,
-				const FSimdVec3f& DQ0,
-				const FSimdVec3f& DP1,
-				const FSimdVec3f& DQ1,
-				const FSimdSolverBodyPtr& Body0,
-				const FSimdSolverBodyPtr& Body1)
-			{
-				for (int32 LaneIndex = 0; LaneIndex < TNumLanes; ++LaneIndex)
-				{
-					if (IsValid.GetValue(LaneIndex))
-					{
-						FSolverBody* LaneBody0 = Body0.GetValue(LaneIndex);
-						FSolverBody* LaneBody1 = Body1.GetValue(LaneIndex);
-						LaneBody0->SetDP(DP0.GetValue(LaneIndex));
-						LaneBody0->SetDQ(DQ0.GetValue(LaneIndex));
-						LaneBody1->SetDP(DP1.GetValue(LaneIndex));
-						LaneBody1->SetDQ(DQ1.GetValue(LaneIndex));
-					}
+					ContactMassInvNormal += FSolverVec3::DotProduct(R1xN, IR1xN) + InvM1;
 				}
-			}
 
-
-			FORCEINLINE void GatherBodyVelocities(
-				const FSimdSelector& IsValid,
-				const FSimdSolverBodyPtr& Body0,
-				const FSimdSolverBodyPtr& Body1,
-				FSimdVec3f& V0,
-				FSimdVec3f& W0,
-				FSimdVec3f& V1,
-				FSimdVec3f& W1)
-			{
-				for (int32 LaneIndex = 0; LaneIndex < TNumLanes; ++LaneIndex)
-				{
-					if (IsValid.GetValue(LaneIndex))
-					{
-						const FSolverBody* LaneBody0 = Body0.GetValue(LaneIndex);
-						const FSolverBody* LaneBody1 = Body1.GetValue(LaneIndex);
-						V0.SetValue(LaneIndex, LaneBody0->V());
-						W0.SetValue(LaneIndex, LaneBody0->W());
-						V1.SetValue(LaneIndex, LaneBody1->V());
-						W1.SetValue(LaneIndex, LaneBody1->W());
-					}
-				}
+				const FSolverReal ContactMassNormal = (ContactMassInvNormal > FSolverReal(UE_SMALL_NUMBER)) ? FSolverReal(1) / ContactMassInvNormal : FSolverReal(0);
+				ManifoldPoint.SimdContactMassNormal.SetValue(LaneIndex, ContactMassNormal);
 			}
 
-			FORCEINLINE void ScatterBodyVelocities(
-				const FSimdSelector& IsValid,
-				const FSimdVec3f& V0,
-				const FSimdVec3f& W0,
-				const FSimdVec3f& V1,
-				const FSimdVec3f& W1,
-				const FSimdSolverBodyPtr& Body0,
-				const FSimdSolverBodyPtr& Body1)
+			void UpdateMassNormal(
+				const int32 LaneIndex,
+				const TArrayView<FSimdManifoldPoint>& ManifoldPointsBuffer,
+				const FSolverBody& Body0,
+				const FSolverBody& Body1,
+				const FSolverReal InvMScale0,
+				const FSolverReal InvIScale0,
+				const FSolverReal InvMScale1,
+				const FSolverReal InvIScale1)
 			{
-				for (int32 LaneIndex = 0; LaneIndex < TNumLanes; ++LaneIndex)
+				const int32 NumManifoldPoints = SimdNumManifoldPoints.GetValue(LaneIndex);
+				for (int32 ManifoldPointIndex = 0; ManifoldPointIndex < NumManifoldPoints; ++ManifoldPointIndex)
 				{
-					if (IsValid.GetValue(LaneIndex))
-					{
-						FSolverBody* LaneBody0 = Body0.GetValue(LaneIndex);
-						FSolverBody* LaneBody1 = Body1.GetValue(LaneIndex);
-						LaneBody0->SetV(V0.GetValue(LaneIndex));
-						LaneBody0->SetW(W0.GetValue(LaneIndex));
-						LaneBody1->SetV(V1.GetValue(LaneIndex));
-						LaneBody1->SetW(W1.GetValue(LaneIndex));
-					}
+					UpdateManifoldPointMassNormal(
+						ManifoldPointIndex,
+						LaneIndex,
+						ManifoldPointsBuffer,
+						Body0,
+						Body1,
+						InvMScale0,
+						InvIScale0,
+						InvMScale1,
+						InvIScale1);
 				}
 			}
 
 			void SolvePositionNoFriction(
+				const TArrayView<FSimdManifoldPoint>& ManifoldPointsBuffer,
 				const FSimdSolverBodyPtr& Body0,
 				const FSimdSolverBodyPtr& Body1,
-				const FSimd4Realf& MaxPushOut)
+				const FSimdRealf& MaxPushOut)
 			{
-				// Which lanes a have a valid constraint setup?
-				const FSimdSelector IsValid = SimdGreaterEqual(SimdConstraintIndex, FSimdInt32::Zero());
-
 				// Get the current corrections for each body. 
 				// NOTE: This is a gather operation
 				FSimdVec3f DP0, DQ0, DP1, DQ1;
-				GatherBodyPositionCorrections(IsValid, Body0, Body1, DP0, DQ0, DP1, DQ1);
-
-				// Calculate the contact error
-				const FSimdVec3f DQ0xR0 = SimdCrossProduct(DQ0, SimdRelativeContactPoint0);
-				const FSimdVec3f DQ1xR1 = SimdCrossProduct(DQ1, SimdRelativeContactPoint1);
-				const FSimdVec3f ContactDelta0 = SimdAdd(DP0, DQ0xR0);
-				const FSimdVec3f ContactDelta1 = SimdAdd(DP1, DQ1xR1);
-				const FSimdVec3f ContactDelta = SimdSubtract(ContactDelta0, ContactDelta1);
-				FSimdRealf ContactErrorNormal = SimdAdd(SimdContactDeltaNormal, SimdDotProduct(ContactDelta, SimdContactNormal));
-
-				// Apply MaxPushOut clamping if required
-				//if ((MaxPushOut > 0) && (ContactErrorNormal < -MaxPushOut)) { ContactErrorNormal = -MaxPushOut; }
-				const FSimdRealf NegMaxPushOut = SimdNegate(MaxPushOut);
-				const FSimdSelector ShouldClampError = SimdAnd(SimdLess(NegMaxPushOut, FSimdRealf::Zero()), SimdLess(ContactErrorNormal, NegMaxPushOut));
-				ContactErrorNormal = SimdSelect(ShouldClampError, NegMaxPushOut, ContactErrorNormal);
-
-				// Determine which lanes to process: only those with an overlap or that applied a pushout on a prior iteration
-				const FSimdSelector IsErrorNegative = SimdLess(ContactErrorNormal, FSimdRealf::Zero());
-				const FSimdSelector IsNetPushOutPositive = SimdGreater(SimdNetPushOutNormal, FSimdRealf::Zero());
-				const FSimdSelector ShouldProcess = SimdAnd(IsValid, SimdOr(IsErrorNegative, IsNetPushOutPositive));
-
-				// If all lanes are to be skipped, early-out
-				if (!SimdAnyTrue(ShouldProcess))
-				{
-					return;
-				}
+				GatherBodyPositionCorrections(Body0, Body1, DP0, DQ0, DP1, DQ1);
 
-				// Zero out the error for points we should not process so we don't apply a correction for them
-				ContactErrorNormal = SimdSelect(ShouldProcess, ContactErrorNormal, FSimdRealf::Zero());
+				for (int32 ManifoldPointIndex = 0; ManifoldPointIndex < MaxManifoldPoints; ++ManifoldPointIndex)
+				{
+					FSimdManifoldPoint& ManifoldPoint = ManifoldPointsBuffer[GetBufferIndex(ManifoldPointIndex)];
+
+					// Which lanes require this point be simulated?
+					//if (!SimdAnyTrue(ManifoldPoint.IsValid))
+					//{
+					//	continue;
+					//}
+
+					// Calculate the contact error
+					const FSimdVec3f DQ0xR0 = SimdCrossProduct(DQ0, ManifoldPoint.SimdRelativeContactPoint0);
+					const FSimdVec3f DQ1xR1 = SimdCrossProduct(DQ1, ManifoldPoint.SimdRelativeContactPoint1);
+					const FSimdVec3f ContactDelta0 = SimdAdd(DP0, DQ0xR0);
+					const FSimdVec3f ContactDelta1 = SimdAdd(DP1, DQ1xR1);
+					const FSimdVec3f ContactDelta = SimdSubtract(ContactDelta0, ContactDelta1);
+					FSimdRealf ContactErrorNormal = SimdAdd(ManifoldPoint.SimdContactDeltaNormal, SimdDotProduct(ContactDelta, ManifoldPoint.SimdContactNormal));
+
+					// Apply MaxPushOut clamping if required
+					//if ((MaxPushOut > 0) && (ContactErrorNormal < -MaxPushOut)) { ContactErrorNormal = -MaxPushOut; }
+					const FSimdRealf NegMaxPushOut = SimdNegate(MaxPushOut);
+					const FSimdSelector ShouldClampError = SimdAnd(SimdLess(NegMaxPushOut, FSimdRealf::Zero()), SimdLess(ContactErrorNormal, NegMaxPushOut));
+					ContactErrorNormal = SimdSelect(ShouldClampError, NegMaxPushOut, ContactErrorNormal);
+
+					// Determine which lanes to process: only those with an overlap or that applied a pushout on a prior iteration
+					const FSimdSelector IsErrorNegative = SimdLess(ContactErrorNormal, FSimdRealf::Zero());
+					const FSimdSelector IsNetPushOutPositive = SimdGreater(ManifoldPoint.SimdNetPushOutNormal, FSimdRealf::Zero());
+					const FSimdSelector ShouldProcess = SimdAnd(ManifoldPoint.IsValid, SimdOr(IsErrorNegative, IsNetPushOutPositive));
+
+					// If all lanes are to be skipped, early-out
+					//if (!SimdAnyTrue(ShouldProcess))
+					//{
+					//	continue;
+					//}
 
-				FSimdRealf PushOutNormal = SimdNegate(SimdMultiply(ContactErrorNormal, SimdMultiply(SimdStiffness, SimdContactMassNormal)));
+					// Zero out the error for points we should not process so we don't apply a correction for them
+					ContactErrorNormal = SimdSelect(ShouldProcess, ContactErrorNormal, FSimdRealf::Zero());
 
-				// Unilateral constraint: Net-negative pushout is not allowed, but
-				// PushOutNormal may be negative on any iteration as long as the net is positive
-				// If the net goes negative, apply a pushout to make it zero
-				const FSimdSelector IsPositive = SimdGreater(SimdAdd(SimdNetPushOutNormal, PushOutNormal), FSimdRealf::Zero());
-				PushOutNormal = SimdSelect(IsPositive, PushOutNormal, SimdNegate(SimdNetPushOutNormal));
+					FSimdRealf PushOutNormal = SimdNegate(SimdMultiply(ContactErrorNormal, SimdMultiply(SimdStiffness, ManifoldPoint.SimdContactMassNormal)));
 
-				// New net pushout
-				SimdNetPushOutNormal = SimdAdd(PushOutNormal, SimdNetPushOutNormal);
+					// Unilateral constraint: Net-negative pushout is not allowed, but
+					// PushOutNormal may be negative on any iteration as long as the net is positive
+					// If the net goes negative, apply a pushout to make it zero
+					const FSimdSelector IsPositive = SimdGreater(SimdAdd(ManifoldPoint.SimdNetPushOutNormal, PushOutNormal), FSimdRealf::Zero());
+					PushOutNormal = SimdSelect(IsPositive, PushOutNormal, SimdNegate(ManifoldPoint.SimdNetPushOutNormal));
 
-				// Convert the positional impulse into position and rotation corrections for each body
-				const FSimdVec3f DDP0 = SimdMultiply(SimdContactNormal, SimdMultiply(SimdInvM0, PushOutNormal));
-				const FSimdVec3f DDQ0 = SimdMultiply(SimdContactNormalAngular0, PushOutNormal);
-				const FSimdVec3f DDP1 = SimdMultiply(SimdContactNormal, SimdMultiply(SimdInvM1, PushOutNormal));
-				const FSimdVec3f DDQ1 = SimdMultiply(SimdContactNormalAngular1, PushOutNormal);
+					// New net pushout
+					ManifoldPoint.SimdNetPushOutNormal = SimdAdd(PushOutNormal, ManifoldPoint.SimdNetPushOutNormal);
 
-				DP0 = SimdAdd(DP0, DDP0);
-				DQ0 = SimdAdd(DQ0, DDQ0);
-				DP1 = SimdSubtract(DP1, DDP1);
-				DQ1 = SimdSubtract(DQ1, DDQ1);
+					// Convert the positional impulse into position and rotation corrections for each body
+					// NOTE: order of operations matches FPBDCollisionSolver::ApplyPositionCorrectionNormal so we can AB test
+					const FSimdVec3f DDP0 = SimdMultiply(ManifoldPoint.SimdContactNormal, SimdMultiply(SimdInvM0, PushOutNormal));
+					const FSimdVec3f DDQ0 = SimdMultiply(ManifoldPoint.SimdContactNormalAngular0, PushOutNormal);
+					const FSimdVec3f DDP1 = SimdMultiply(ManifoldPoint.SimdContactNormal, SimdMultiply(SimdInvM1, PushOutNormal));
+					const FSimdVec3f DDQ1 = SimdMultiply(ManifoldPoint.SimdContactNormalAngular1, PushOutNormal);
+					DP0 = SimdAdd(DP0, DDP0);
+					DQ0 = SimdAdd(DQ0, DDQ0);
+					DP1 = SimdSubtract(DP1, DDP1);
+					DQ1 = SimdSubtract(DQ1, DDQ1);
+				}
 
 				// Update the corrections on the bodies. 
 				// NOTE: This is a scatter operation
-				ScatterBodyPositionCorrections(IsValid, DP0, DQ0, DP1, DQ1, Body0, Body1);
+				ScatterBodyPositionCorrections(DP0, DQ0, DP1, DQ1, Body0, Body1);
 			}
 
 			void SolvePositionWithFriction(
+				const TArrayView<FSimdManifoldPoint>& ManifoldPointsBuffer,
 				const FSimdSolverBodyPtr& Body0,
 				const FSimdSolverBodyPtr& Body1,
 				const FSimd4Realf& MaxPushOut,
 				const FSimd4Realf& FrictionStiffnessScale)
 			{
-				// Which lanes a have a valid constraint setup?
-				const FSimdSelector IsValid = SimdGreaterEqual(SimdConstraintIndex, FSimdInt32::Zero());
-
 				// Get the current corrections for each body. 
 				// NOTE: This is a gather operation
 				FSimdVec3f DP0, DQ0, DP1, DQ1;
-				GatherBodyPositionCorrections(IsValid, Body0, Body1, DP0, DQ0, DP1, DQ1);
-
-				// Calculate the contact error
-				const FSimdVec3f DQ0xR0 = SimdCrossProduct(DQ0, SimdRelativeContactPoint0);
-				const FSimdVec3f DQ1xR1 = SimdCrossProduct(DQ1, SimdRelativeContactPoint1);
-				const FSimdVec3f ContactDelta0 = SimdAdd(DP0, DQ0xR0);
-				const FSimdVec3f ContactDelta1 = SimdAdd(DP1, DQ1xR1);
-				const FSimdVec3f ContactDelta = SimdSubtract(ContactDelta0, ContactDelta1);
-				FSimdRealf ContactErrorNormal = SimdAdd(SimdContactDeltaNormal, SimdDotProduct(ContactDelta, SimdContactNormal));
-
-				// Apply MaxPushOut clamping if required
-				//if ((MaxPushOut > 0) && (ContactErrorNormal < -MaxPushOut)) { ContactErrorNormal = -MaxPushOut; }
-				const FSimdRealf NegMaxPushOut = SimdNegate(MaxPushOut);
-				const FSimdSelector ShouldClampError = SimdAnd(SimdLess(NegMaxPushOut, FSimdRealf::Zero()), SimdLess(ContactErrorNormal, NegMaxPushOut));
-				ContactErrorNormal = SimdSelect(ShouldClampError, NegMaxPushOut, ContactErrorNormal);
-
-				// Determine which lanes to process: only those with an overlap or that applied a pushout on a prior iteration
-				const FSimdSelector IsErrorNegative = SimdLess(ContactErrorNormal, FSimdRealf::Zero());
-				const FSimdSelector IsNetPushOutPositive = SimdGreater(SimdNetPushOutNormal, FSimdRealf::Zero());
-				const FSimdSelector ShouldProcess = SimdAnd(IsValid, SimdOr(IsErrorNegative, IsNetPushOutPositive));
-
-				// If all lanes are to be skipped, early-out
-				//if (SimdAnyTrue(ShouldProcess))
+				GatherBodyPositionCorrections(Body0, Body1, DP0, DQ0, DP1, DQ1);
+
+				const FSimdRealf Zero = FSimdRealf::Zero();
+				const FSimdRealf One = FSimdRealf::One();
+
+				FSimdRealf NumFrictionPoints = Zero;
+				FSimdRealf MaxFrictionPushOut = Zero;
+
+				// @todo(chaos): can these be zero (and make non-simd version match)?
+				const FSimdRealf PushOutNormalTolerance = FSimdRealf::Make(UE_SMALL_NUMBER);
+				const FSimdRealf MaxFrictionPushOutTolerance = FSimdRealf::Make(UE_KINDA_SMALL_NUMBER);
+
+				// Apply the normal pushout and calculate the net normal pushout for the friction limit
+				for (int32 ManifoldPointIndex = 0; ManifoldPointIndex < MaxManifoldPoints; ++ManifoldPointIndex)
 				{
+					FSimdManifoldPoint& ManifoldPoint = ManifoldPointsBuffer[GetBufferIndex(ManifoldPointIndex)];
+
+					// Which lanes require this point be simulated?
+					//if (!SimdAnyTrue(ManifoldPoint.IsValid))
+					//{
+					//	continue;
+					//}
+
+					// Calculate the contact error
+					const FSimdVec3f DQ0xR0 = SimdCrossProduct(DQ0, ManifoldPoint.SimdRelativeContactPoint0);
+					const FSimdVec3f DQ1xR1 = SimdCrossProduct(DQ1, ManifoldPoint.SimdRelativeContactPoint1);
+					const FSimdVec3f ContactDelta0 = SimdAdd(DP0, DQ0xR0);
+					const FSimdVec3f ContactDelta1 = SimdAdd(DP1, DQ1xR1);
+					const FSimdVec3f ContactDelta = SimdSubtract(ContactDelta0, ContactDelta1);
+					FSimdRealf ContactErrorNormal = SimdAdd(ManifoldPoint.SimdContactDeltaNormal, SimdDotProduct(ContactDelta, ManifoldPoint.SimdContactNormal));
+
+					// Apply MaxPushOut clamping if required
+					//if ((MaxPushOut > 0) && (ContactErrorNormal < -MaxPushOut)) { ContactErrorNormal = -MaxPushOut; }
+					const FSimdRealf NegMaxPushOut = SimdNegate(MaxPushOut);
+					const FSimdSelector ShouldClampError = SimdAnd(SimdLess(NegMaxPushOut, Zero), SimdLess(ContactErrorNormal, NegMaxPushOut));
+					ContactErrorNormal = SimdSelect(ShouldClampError, NegMaxPushOut, ContactErrorNormal);
+
+					// Determine which lanes to process: only those with an overlap or that applied a pushout on a prior iteration
+					const FSimdSelector IsErrorNegative = SimdLess(ContactErrorNormal, Zero);
+					const FSimdSelector IsNetPushOutPositive = SimdGreater(ManifoldPoint.SimdNetPushOutNormal, PushOutNormalTolerance);
+					const FSimdSelector ProcessManifoldPoint = SimdAnd(ManifoldPoint.IsValid, SimdOr(IsErrorNegative, IsNetPushOutPositive));
+
+					// If all lanes are to be skipped, early-out
+					//if (!SimdAnyTrue(ProcessManifoldPoint))
+					//{
+					//	continue;
+					//}
+
 					// Zero out the error for points we should not process so we don't apply a correction for them
-					ContactErrorNormal = SimdSelect(ShouldProcess, ContactErrorNormal, FSimdRealf::Zero());
+					ContactErrorNormal = SimdSelect(ProcessManifoldPoint, ContactErrorNormal, Zero);
 
-					FSimdRealf PushOutNormal = SimdNegate(SimdMultiply(SimdStiffness, SimdMultiply(ContactErrorNormal, SimdContactMassNormal)));
+					FSimdRealf PushOutNormal = SimdNegate(SimdMultiply(ContactErrorNormal, SimdMultiply(SimdStiffness, ManifoldPoint.SimdContactMassNormal)));
 
 					// Unilateral constraint: Net-negative pushout is not allowed, but
 					// PushOutNormal may be negative on any iteration as long as the net is positive
 					// If the net goes negative, apply a pushout to make it zero
-					const FSimdSelector IsPositive = SimdGreater(SimdAdd(SimdNetPushOutNormal, PushOutNormal), FSimdRealf::Zero());
-					PushOutNormal = SimdSelect(IsPositive, PushOutNormal, SimdNegate(SimdNetPushOutNormal));
+					const FSimdSelector IsPositive = SimdGreater(SimdAdd(ManifoldPoint.SimdNetPushOutNormal, PushOutNormal), Zero);
+					PushOutNormal = SimdSelect(IsPositive, PushOutNormal, SimdNegate(ManifoldPoint.SimdNetPushOutNormal));
 
-					// New net normal pushout
-					SimdNetPushOutNormal = SimdAdd(PushOutNormal, SimdNetPushOutNormal);
-					SimdStaticFrictionRatio = FSimdRealf::One();
+					// New net pushout
+					ManifoldPoint.SimdNetPushOutNormal = SimdAdd(PushOutNormal, ManifoldPoint.SimdNetPushOutNormal);
 
 					// Convert the positional impulse into position and rotation corrections for each body
-					FSimdVec3f DDP0 = SimdMultiply(SimdContactNormal, SimdMultiply(SimdInvM0, PushOutNormal));
-					FSimdVec3f DDQ0 = SimdMultiply(SimdContactNormalAngular0, PushOutNormal);
-					FSimdVec3f DDP1 = SimdMultiply(SimdContactNormal, SimdMultiply(SimdInvM1, PushOutNormal));
-					FSimdVec3f DDQ1 = SimdMultiply(SimdContactNormalAngular1, PushOutNormal);
+					// NOTE: order of operations matches FPBDCollisionSolver::ApplyPositionCorrectionNormal so we can AB test
+					const FSimdVec3f DDP0 = SimdMultiply(ManifoldPoint.SimdContactNormal, SimdMultiply(SimdInvM0, PushOutNormal));
+					const FSimdVec3f DDQ0 = SimdMultiply(ManifoldPoint.SimdContactNormalAngular0, PushOutNormal);
+					const FSimdVec3f DDP1 = SimdMultiply(ManifoldPoint.SimdContactNormal, SimdMultiply(SimdInvM1, PushOutNormal));
+					const FSimdVec3f DDQ1 = SimdMultiply(ManifoldPoint.SimdContactNormalAngular1, PushOutNormal);
+					DP0 = SimdAdd(DP0, DDP0);
+					DQ0 = SimdAdd(DQ0, DDQ0);
+					DP1 = SimdSubtract(DP1, DDP1);
+					DQ1 = SimdSubtract(DQ1, DDQ1);
+
+					MaxFrictionPushOut = SimdSelect(IsPositive, SimdAdd(MaxFrictionPushOut, ManifoldPoint.SimdNetPushOutNormal), MaxFrictionPushOut);
+					NumFrictionPoints = SimdSelect(IsPositive, SimdAdd(NumFrictionPoints, One), NumFrictionPoints);
+				}
+
+				MaxFrictionPushOut = SimdSelect(SimdGreater(NumFrictionPoints, Zero), SimdDivide(MaxFrictionPushOut, NumFrictionPoints), Zero);
+
+				for (int32 ManifoldPointIndex = 0; ManifoldPointIndex < MaxManifoldPoints; ++ManifoldPointIndex)
+				{
+					FSimdManifoldPoint& ManifoldPoint = ManifoldPointsBuffer[GetBufferIndex(ManifoldPointIndex)];
+
+					// Which lanes require this point be simulated?
+					//if (!SimdAnyTrue(ManifoldPoint.IsValid))
+					//{
+					//	continue;
+					//}
+
+					// Calculate the contact error
+					const FSimdVec3f DQ0xR0 = SimdCrossProduct(DQ0, ManifoldPoint.SimdRelativeContactPoint0);
+					const FSimdVec3f DQ1xR1 = SimdCrossProduct(DQ1, ManifoldPoint.SimdRelativeContactPoint1);
+					const FSimdVec3f ContactDelta0 = SimdAdd(DP0, DQ0xR0);
+					const FSimdVec3f ContactDelta1 = SimdAdd(DP1, DQ1xR1);
+					const FSimdVec3f ContactDelta = SimdSubtract(ContactDelta0, ContactDelta1);
 
 					// Should we apply tangential corrections for friction? 
-					// bUpdateFriction = ((NetPushOutNormal > 0) || (NetPushOutTangentU != 0) || (NetPushOutTangentV != 0))
-					const FSimdSelector HasFriction = SimdAnd(IsValid, SimdGreater(SimdStaticFriction, FSimdRealf::Zero()));
-					const FSimdSelector HasNormalPushout = SimdAnd(HasFriction, SimdGreater(SimdNetPushOutNormal, FSimdRealf::Zero()));
-					const FSimdSelector HasTangentUPushout = SimdNotEqual(SimdNetPushOutTangentU, FSimdRealf::Zero());
-					const FSimdSelector HasTangentVPushout = SimdNotEqual(SimdNetPushOutTangentV, FSimdRealf::Zero());
-					//const FSimdSelector UpdateFriction = SimdAnd(HasFriction, SimdOr(HasNormalPushout, SimdOr(HasTangentUPushout, HasTangentVPushout)));
-					//if (SimdAnyTrue(UpdateFriction))
+					// bUpdateFriction = ((TotalPushOutNormal > 0) || (NetPushOutTangentU != 0) || (NetPushOutTangentV != 0))
+					const FSimdSelector HasFriction = SimdGreater(SimdStaticFriction, Zero);
+					const FSimdSelector HasNormalPushout = SimdGreater(ManifoldPoint.SimdNetPushOutNormal, Zero);
+					const FSimdSelector HasTangentUPushout = SimdNotEqual(ManifoldPoint.SimdNetPushOutTangentU, Zero);
+					const FSimdSelector HasTangentVPushout = SimdNotEqual(ManifoldPoint.SimdNetPushOutTangentV, Zero);
+					const FSimdSelector ApplyPointFriction = SimdAnd(HasFriction, SimdOr(HasNormalPushout, SimdOr(HasTangentUPushout, HasTangentVPushout)));
+					//if (SimdAnyTrue(ApplyPointFriction))
 					{
 						// Calculate tangential errors
-						const FSimdRealf ContactErrorTangentU = SimdAdd(SimdContactDeltaTangentU, SimdDotProduct(ContactDelta, SimdContactTangentU));
-						const FSimdRealf ContactErrorTangentV = SimdAdd(SimdContactDeltaTangentV, SimdDotProduct(ContactDelta, SimdContactTangentV));
+						const FSimdRealf ContactErrorTangentU = SimdAdd(ManifoldPoint.SimdContactDeltaTangentU, SimdDotProduct(ContactDelta, ManifoldPoint.SimdContactTangentU));
+						const FSimdRealf ContactErrorTangentV = SimdAdd(ManifoldPoint.SimdContactDeltaTangentV, SimdDotProduct(ContactDelta, ManifoldPoint.SimdContactTangentV));
 
 						// A stiffness multiplier on the impulses
 						const FSimdRealf FrictionStiffness = SimdMultiply(FrictionStiffnessScale, SimdStiffness);
 
 						// Calculate tangential correction
-						FSimdRealf PushOutTangentU = SimdNegate(SimdMultiply(FrictionStiffness, SimdMultiply(ContactErrorTangentU, SimdContactMassTangentU)));
-						FSimdRealf PushOutTangentV = SimdNegate(SimdMultiply(FrictionStiffness, SimdMultiply(ContactErrorTangentV, SimdContactMassTangentV)));
+						FSimdRealf PushOutTangentU = SimdNegate(SimdMultiply(FrictionStiffness, SimdMultiply(ContactErrorTangentU, ManifoldPoint.SimdContactMassTangentU)));
+						FSimdRealf PushOutTangentV = SimdNegate(SimdMultiply(FrictionStiffness, SimdMultiply(ContactErrorTangentV, ManifoldPoint.SimdContactMassTangentV)));
 
 						// New net tangential pushouts
-						FSimdRealf NetPushOutTangentU = SimdAdd(SimdNetPushOutTangentU, PushOutTangentU);
-						FSimdRealf NetPushOutTangentV = SimdAdd(SimdNetPushOutTangentV, PushOutTangentV);
-						FSimdRealf StaticFrictionRatio = SimdStaticFrictionRatio;
+						FSimdRealf NetPushOutTangentU = SimdAdd(ManifoldPoint.SimdNetPushOutTangentU, PushOutTangentU);
+						FSimdRealf NetPushOutTangentV = SimdAdd(ManifoldPoint.SimdNetPushOutTangentV, PushOutTangentV);
+						FSimdRealf StaticFrictionRatio = One;
 
-						// If a lane has a zero normal pushout, we need to undo all previously applied friction
-						//if (!SimdAllTrue(HasNormalPushout))
-						{
-							PushOutTangentU = SimdSelect(HasNormalPushout, PushOutTangentU, SimdNegate(SimdNetPushOutTangentU));
-							PushOutTangentV = SimdSelect(HasNormalPushout, PushOutTangentV, SimdNegate(SimdNetPushOutTangentV));
-							NetPushOutTangentU = SimdSelect(HasNormalPushout, NetPushOutTangentU, FSimdRealf::Zero());
-							NetPushOutTangentV = SimdSelect(HasNormalPushout, NetPushOutTangentV, FSimdRealf::Zero());
-							StaticFrictionRatio = SimdSelect(HasNormalPushout, StaticFrictionRatio, FSimdRealf::Zero());
-						}
+						// Should we clamp to the friction cone or reset the friction?
+						const FSimdSelector ApplyFrictionCone = SimdGreaterEqual(MaxFrictionPushOut, MaxFrictionPushOutTolerance);
 
 						// Apply cone limits to tangential pushouts on lanes that exceed the limit
 						// NOTE: if HasNormalPushout is false in any lane, we have already zeroed the net pushout and don't need to do anything here
-						const FSimdSelector ApplyFrictionCone = SimdAnd(HasFriction, SimdAnd(HasNormalPushout, SimdOr(HasTangentUPushout, HasTangentVPushout)));
 						//if (SimdAnyTrue(ApplyFrictionCone))
 						{
-							const FSimdRealf MaxPushOutTangent = SimdNetPushOutNormal;
-							const FSimdRealf MaxStaticPushOutTangentSq = SimdSquare(SimdMultiply(SimdStaticFriction, MaxPushOutTangent));
+							const FSimdRealf MaxStaticPushOutTangentSq = SimdSquare(SimdMultiply(SimdStaticFriction, MaxFrictionPushOut));
 							const FSimdRealf NetPushOutTangentSq = SimdAdd(SimdSquare(NetPushOutTangentU), SimdSquare(NetPushOutTangentV));
 							const FSimdSelector ExceededFrictionCone = SimdAnd(ApplyFrictionCone, SimdGreater(NetPushOutTangentSq, MaxStaticPushOutTangentSq));
 							//if (SimdAnyTrue(ExceededFrictionCone))
 							{
-								const FSimdRealf MaxDynamicPushOutTangent = SimdMultiply(SimdDynamicFriction, MaxPushOutTangent);
+								const FSimdRealf MaxDynamicPushOutTangent = SimdMultiply(SimdDynamicFriction, MaxFrictionPushOut);
 								const FSimdRealf FrictionMultiplier = SimdMultiply(MaxDynamicPushOutTangent, SimdInvSqrt(NetPushOutTangentSq));
 								const FSimdRealf ClampedNetPushOutTangentU = SimdMultiply(FrictionMultiplier, NetPushOutTangentU);
 								const FSimdRealf ClampedNetPushOutTangentV = SimdMultiply(FrictionMultiplier, NetPushOutTangentV);
-								const FSimdRealf ClampedPushOutTangentU = SimdSubtract(ClampedNetPushOutTangentU, SimdNetPushOutTangentU);
-								const FSimdRealf ClampedPushOutTangentV = SimdSubtract(ClampedNetPushOutTangentV, SimdNetPushOutTangentV);
+								const FSimdRealf ClampedPushOutTangentU = SimdSubtract(ClampedNetPushOutTangentU, ManifoldPoint.SimdNetPushOutTangentU);
+								const FSimdRealf ClampedPushOutTangentV = SimdSubtract(ClampedNetPushOutTangentV, ManifoldPoint.SimdNetPushOutTangentV);
 
 								PushOutTangentU = SimdSelect(ExceededFrictionCone, ClampedPushOutTangentU, PushOutTangentU);
 								PushOutTangentV = SimdSelect(ExceededFrictionCone, ClampedPushOutTangentV, PushOutTangentV);
@@ -522,178 +783,214 @@ namespace Chaos
 							}
 						}
 
-						SimdNetPushOutTangentU = NetPushOutTangentU;
-						SimdNetPushOutTangentV = NetPushOutTangentV;
-						SimdStaticFrictionRatio = StaticFrictionRatio;
+						// If we did not apply the friction cone because we have no normal impulse, apply a pushout to cancel previously applied friction
+						//if (!SimdAllTrue(ApplyFrictionCone))
+						{
+							PushOutTangentU = SimdSelect(ApplyFrictionCone, PushOutTangentU, SimdNegate(ManifoldPoint.SimdNetPushOutTangentU));
+							PushOutTangentV = SimdSelect(ApplyFrictionCone, PushOutTangentV, SimdNegate(ManifoldPoint.SimdNetPushOutTangentV));
+							NetPushOutTangentU = SimdSelect(ApplyFrictionCone, NetPushOutTangentU, Zero);
+							NetPushOutTangentV = SimdSelect(ApplyFrictionCone, NetPushOutTangentV, Zero);
+							StaticFrictionRatio = SimdSelect(ApplyFrictionCone, StaticFrictionRatio, Zero);
+						}
+
+						// Undo all our good work for lanes that should not apply friction at all
+						//if (!SimdAllTrue(ApplyPointFriction))
+						{
+							PushOutTangentU = SimdSelect(ApplyPointFriction, PushOutTangentU, Zero);
+							PushOutTangentV = SimdSelect(ApplyPointFriction, PushOutTangentV, Zero);
+							NetPushOutTangentU = SimdSelect(ApplyPointFriction, NetPushOutTangentU, ManifoldPoint.SimdNetPushOutTangentU);
+							NetPushOutTangentV = SimdSelect(ApplyPointFriction, NetPushOutTangentV, ManifoldPoint.SimdNetPushOutTangentV);
+							StaticFrictionRatio = SimdSelect(ApplyPointFriction, StaticFrictionRatio, One);
+						}
 
-						// Add the tangential corrections to the applied correction
-						DDP0 = SimdAdd(DDP0, SimdMultiply(SimdContactTangentU, SimdMultiply(SimdInvM0, PushOutTangentU)));
-						DDP0 = SimdAdd(DDP0, SimdMultiply(SimdContactTangentV, SimdMultiply(SimdInvM0, PushOutTangentV)));
-						DDQ0 = SimdAdd(DDQ0, SimdMultiply(SimdContactTangentUAngular0, PushOutTangentU));
-						DDQ0 = SimdAdd(DDQ0, SimdMultiply(SimdContactTangentVAngular0, PushOutTangentV));
-						DDP1 = SimdAdd(DDP1, SimdMultiply(SimdContactTangentU, SimdMultiply(SimdInvM1, PushOutTangentU)));
-						DDP1 = SimdAdd(DDP1, SimdMultiply(SimdContactTangentV, SimdMultiply(SimdInvM1, PushOutTangentV)));
-						DDQ1 = SimdAdd(DDQ1, SimdMultiply(SimdContactTangentUAngular1, PushOutTangentU));
-						DDQ1 = SimdAdd(DDQ1, SimdMultiply(SimdContactTangentVAngular1, PushOutTangentV));
+						ManifoldPoint.SimdNetPushOutTangentU = NetPushOutTangentU;
+						ManifoldPoint.SimdNetPushOutTangentV = NetPushOutTangentV;
+						ManifoldPoint.SimdStaticFrictionRatio = StaticFrictionRatio;
 
+						// Add the tangential corrections to the applied correction
+						// NOTE: The order of operations here matches FPBDCollisionSolver::ApplyPositionCorrectionTangential
+						// so that we can do an AB test of the SIMD solver versus the standard solver
+						const FSimdVec3f PushOut = SimdAdd(SimdMultiply(PushOutTangentU, ManifoldPoint.SimdContactTangentU), SimdMultiply(PushOutTangentV, ManifoldPoint.SimdContactTangentV));
+						const FSimdVec3f DDP0 = SimdMultiply(SimdInvM0, PushOut);
+						const FSimdVec3f DDP1 = SimdMultiply(SimdInvM1, PushOut);
+						const FSimdVec3f DDQ0 = SimdAdd(SimdMultiply(ManifoldPoint.SimdContactTangentUAngular0, PushOutTangentU), SimdMultiply(ManifoldPoint.SimdContactTangentVAngular0, PushOutTangentV));
+						const FSimdVec3f DDQ1 = SimdAdd(SimdMultiply(ManifoldPoint.SimdContactTangentUAngular1, PushOutTangentU), SimdMultiply(ManifoldPoint.SimdContactTangentVAngular1, PushOutTangentV));
 						DP0 = SimdAdd(DP0, DDP0);
 						DQ0 = SimdAdd(DQ0, DDQ0);
 						DP1 = SimdSubtract(DP1, DDP1);
 						DQ1 = SimdSubtract(DQ1, DDQ1);
 					}
-
-					// Update the corrections on the bodies. 
-					// NOTE: This is a scatter operation
-					ScatterBodyPositionCorrections(IsValid, DP0, DQ0, DP1, DQ1, Body0, Body1);
 				}
+
+				// Update the corrections on the bodies. 
+				// NOTE: This is a scatter operation
+				ScatterBodyPositionCorrections(DP0, DQ0, DP1, DQ1, Body0, Body1);
 			}
 
 			void SolveVelocityNoFriction(
+				const TArrayView<FSimdManifoldPoint>& ManifoldPointsBuffer,
 				const FSimdSolverBodyPtr& Body0,
 				const FSimdSolverBodyPtr& Body1,
 				const FSimdRealf& Dt)
 			{
-				// Which lanes a have a valid constraint setup?
-				const FSimdSelector IsValid = SimdGreaterEqual(SimdConstraintIndex, FSimdInt32::Zero());
+				// Gather the body data we need
+				FSimdVec3f V0, W0, V1, W1;
+				GatherBodyVelocities(Body0, Body1, V0, W0, V1, W1);
 
-				// Only lanes that applied a position correction or that started with an overlap require a velocity correction
-				FSimdSelector ShouldSolveVelocity = SimdOr(SimdGreater(SimdNetPushOutNormal, FSimdRealf::Zero()), SimdLess(SimdContactDeltaNormal, FSimdRealf::Zero()));
-				ShouldSolveVelocity = SimdAnd(IsValid, ShouldSolveVelocity);
-				if (!SimdAnyTrue(ShouldSolveVelocity))
+				for (int32 ManifoldPointIndex = 0; ManifoldPointIndex < MaxManifoldPoints; ++ManifoldPointIndex)
 				{
-					return;
+					FSimdManifoldPoint& ManifoldPoint = ManifoldPointsBuffer[GetBufferIndex(ManifoldPointIndex)];
+
+					// Which lanes require this point be simulated?
+					//if (!SimdAnyTrue(ManifoldPoint.IsValid))
+					//{
+					//	continue;
+					//}
+
+					// Only lanes that applied a position correction or that started with an overlap require a velocity correction
+					FSimdSelector ShouldSolveVelocity = SimdOr(SimdGreater(ManifoldPoint.SimdNetPushOutNormal, FSimdRealf::Zero()), SimdLess(ManifoldPoint.SimdContactDeltaNormal, FSimdRealf::Zero()));
+					ShouldSolveVelocity = SimdAnd(ManifoldPoint.IsValid, ShouldSolveVelocity);
+					//if (!SimdAnyTrue(ShouldSolveVelocity))
+					//{
+					//	continue;
+					//}
+
+					// Calculate the velocity error we need to correct
+					const FSimdVec3f ContactVelocity0 = SimdAdd(V0, SimdCrossProduct(W0, ManifoldPoint.SimdRelativeContactPoint0));
+					const FSimdVec3f ContactVelocity1 = SimdAdd(V1, SimdCrossProduct(W1, ManifoldPoint.SimdRelativeContactPoint1));
+					const FSimdVec3f ContactVelocity = SimdSubtract(ContactVelocity0, ContactVelocity1);
+					const FSimdRealf ContactVelocityNormal = SimdDotProduct(ContactVelocity, ManifoldPoint.SimdContactNormal);
+					FSimdRealf ContactVelocityErrorNormal = SimdSubtract(ContactVelocityNormal, ManifoldPoint.SimdContactTargetVelocityNormal);
+
+					// Calculate the velocity correction for the error
+					FSimdRealf ImpulseNormal = SimdNegate(SimdMultiply(SimdMultiply(SimdStiffness, ManifoldPoint.SimdContactMassNormal), ContactVelocityErrorNormal));
+
+					// The minimum normal impulse we can apply. We are allowed to apply a negative impulse 
+					// up to an amount that would conteract the implciit velocity applied by the pushout
+					// MinImpulseNormal = FMath::Min(0, -NetPushOutNormal / Dt)
+					// if (NetImpulseNormal + ImpulseNormal < MinImpulseNormal) {...}
+					const FSimdRealf MinImpulseNormal = SimdMin(SimdDivide(SimdNegate(ManifoldPoint.SimdNetPushOutNormal), Dt), FSimdRealf::Zero());
+					const FSimdSelector ShouldClampImpulse = SimdLess(SimdAdd(ManifoldPoint.SimdNetImpulseNormal, ImpulseNormal), MinImpulseNormal);
+					const FSimdRealf ClampedImpulseNormal = SimdSubtract(MinImpulseNormal, ManifoldPoint.SimdNetImpulseNormal);
+					ImpulseNormal = SimdSelect(ShouldClampImpulse, ClampedImpulseNormal, ImpulseNormal);
+
+					// Clear the impulse for lanes that should not be solving for velocity
+					ImpulseNormal = SimdSelect(ShouldSolveVelocity, ImpulseNormal, FSimdRealf::Zero());
+
+					ManifoldPoint.SimdNetImpulseNormal = SimdAdd(ManifoldPoint.SimdNetImpulseNormal, ImpulseNormal);
+
+					// NOTE: order of operations matches FPBDCollisionSolver::ApplyVelocityCorrectionNormal for AB Testing
+					const FSimdVec3f Impulse = SimdMultiply(ImpulseNormal, ManifoldPoint.SimdContactNormal);
+					const FSimdVec3f DV0 = SimdMultiply(Impulse, SimdInvM0);
+					const FSimdVec3f DV1 = SimdMultiply(Impulse, SimdInvM1);
+					const FSimdVec3f DW0 = SimdMultiply(ImpulseNormal, ManifoldPoint.SimdContactNormalAngular0);
+					const FSimdVec3f DW1 = SimdMultiply(ImpulseNormal, ManifoldPoint.SimdContactNormalAngular1);
+					V0 = SimdAdd(V0, DV0);
+					W0 = SimdAdd(W0, DW0);
+					V1 = SimdSubtract(V1, DV1);
+					W1 = SimdSubtract(W1, DW1);
 				}
 
-				// Gather the body data we need
-				FSimdVec3f V0, W0, V1, W1;
-				GatherBodyVelocities(IsValid, Body0, Body1, V0, W0, V1, W1);
-
-				// Calculate the velocity error we need to correct
-				const FSimdVec3f ContactVelocity0 = SimdAdd(V0, SimdCrossProduct(W0, SimdRelativeContactPoint0));
-				const FSimdVec3f ContactVelocity1 = SimdAdd(V1, SimdCrossProduct(W1, SimdRelativeContactPoint1));
-				const FSimdVec3f ContactVelocity = SimdSubtract(ContactVelocity0, ContactVelocity1);
-				const FSimdRealf ContactVelocityNormal = SimdDotProduct(ContactVelocity, SimdContactNormal);
-				FSimdRealf ContactVelocityErrorNormal = SimdSubtract(ContactVelocityNormal, SimdContactTargetVelocityNormal);
-
-				// Make sure we don't add a correction to the lanes we should ignore
-				ContactVelocityErrorNormal = SimdSelect(ShouldSolveVelocity, ContactVelocityErrorNormal, FSimdRealf::Zero());
-
-				// Calculate the velocity correction for the error
-				FSimdRealf ImpulseNormal = SimdNegate(SimdMultiply(SimdMultiply(SimdStiffness, SimdContactMassNormal), ContactVelocityErrorNormal));
-
-				// The minimum normal impulse we can apply. We are allowed to apply a negative impulse 
-				// up to an amount that would conteract the implciit velocity applied by the pushout
-				// MinImpulseNormal = FMath::Min(0, -NetPushOutNormal / Dt)
-				// if (NetImpulseNormal + ImpulseNormal < MinImpulseNormal) {...}
-				const FSimdRealf MinImpulseNormal = SimdMin(SimdNegate(SimdDivide(SimdNetPushOutNormal, Dt)), FSimdRealf::Zero());
-				const FSimdSelector ShouldClampImpulse = SimdLess(SimdAdd(SimdNetImpulseNormal, ImpulseNormal), MinImpulseNormal);
-				const FSimdRealf ClampedImpulseNormal = SimdSubtract(MinImpulseNormal, SimdNetImpulseNormal);
-				ImpulseNormal = SimdSelect(ShouldClampImpulse, ClampedImpulseNormal, ImpulseNormal);
-
-				SimdNetImpulseNormal = SimdAdd(SimdNetImpulseNormal, ImpulseNormal);
-
-				const FSimdVec3f DV0 = SimdMultiply(SimdMultiply(ImpulseNormal, SimdInvM0), SimdContactNormal);
-				const FSimdVec3f DV1 = SimdMultiply(SimdMultiply(ImpulseNormal, SimdInvM1), SimdContactNormal);
-				const FSimdVec3f DW0 = SimdMultiply(ImpulseNormal, SimdContactNormalAngular0);
-				const FSimdVec3f DW1 = SimdMultiply(ImpulseNormal, SimdContactNormalAngular1);
-				V0 = SimdAdd(V0, DV0);
-				W0 = SimdAdd(W0, DW0);
-				V1 = SimdSubtract(V1, DV1);
-				W1 = SimdSubtract(W1, DW1);
-
-				ScatterBodyVelocities(IsValid, V0, W0, V1, W1, Body0, Body1);
+				ScatterBodyVelocities(V0, W0, V1, W1, Body0, Body1);
 			}
 
 			void SolveVelocityWithFrictionImpl(
+				const TArrayView<FSimdManifoldPoint>& ManifoldPointsBuffer,
 				const FSimdSolverBodyPtr& Body0,
 				const FSimdSolverBodyPtr& Body1,
 				const FSimdRealf& Dt,
 				const FSimd4Realf& FrictionStiffnessScale)
 			{
-				// Which lanes a have a valid constraint setup?
-				const FSimdSelector IsValid = SimdGreaterEqual(SimdConstraintIndex, FSimdInt32::Zero());
+				// Gather the body data we need
+				FSimdVec3f V0, W0, V1, W1;
+				GatherBodyVelocities(Body0, Body1, V0, W0, V1, W1);
 
-				// Only lanes that applied a position correction or that started with an overlap require a velocity correction
-				FSimdSelector ShouldSolveVelocity = SimdOr(SimdGreater(SimdNetPushOutNormal, FSimdRealf::Zero()), SimdLess(SimdContactDeltaNormal, FSimdRealf::Zero()));
-				ShouldSolveVelocity = SimdAnd(IsValid, ShouldSolveVelocity);
-				if (!SimdAnyTrue(ShouldSolveVelocity))
+				for (int32 ManifoldPointIndex = 0; ManifoldPointIndex < MaxManifoldPoints; ++ManifoldPointIndex)
 				{
-					return;
+					FSimdManifoldPoint& ManifoldPoint = ManifoldPointsBuffer[GetBufferIndex(ManifoldPointIndex)];
+
+					// Which lanes require this point be simulated?
+					//if (!SimdAnyTrue(ManifoldPoint.IsValid))
+					//{
+					//	continue;
+					//}
+
+					// Only lanes that applied a position correction or that started with an overlap require a velocity correction
+					FSimdSelector ShouldSolveVelocity = SimdOr(SimdGreater(ManifoldPoint.SimdNetPushOutNormal, FSimdRealf::Zero()), SimdLess(ManifoldPoint.SimdContactDeltaNormal, FSimdRealf::Zero()));
+					ShouldSolveVelocity = SimdAnd(ManifoldPoint.IsValid, ShouldSolveVelocity);
+					//if (!SimdAnyTrue(ShouldSolveVelocity))
+					//{
+					//	continue;
+					//}
+
+					// Calculate the velocity error we need to correct
+					const FSimdVec3f ContactVelocity0 = SimdAdd(V0, SimdCrossProduct(W0, ManifoldPoint.SimdRelativeContactPoint0));
+					const FSimdVec3f ContactVelocity1 = SimdAdd(V1, SimdCrossProduct(W1, ManifoldPoint.SimdRelativeContactPoint1));
+					const FSimdVec3f ContactVelocity = SimdSubtract(ContactVelocity0, ContactVelocity1);
+					FSimdRealf ContactVelocityErrorNormal = SimdSubtract(SimdDotProduct(ContactVelocity, ManifoldPoint.SimdContactNormal), ManifoldPoint.SimdContactTargetVelocityNormal);
+					FSimdRealf ContactVelocityErrorTangentU = SimdDotProduct(ContactVelocity, ManifoldPoint.SimdContactTangentU);
+					FSimdRealf ContactVelocityErrorTangentV = SimdDotProduct(ContactVelocity, ManifoldPoint.SimdContactTangentV);
+
+					// A stiffness multiplier on the impulses (zeroed if we have no friction in a lane)
+					const FSimdRealf FrictionStiffness = SimdMultiply(FrictionStiffnessScale, SimdStiffness);
+
+					// Calculate the impulses
+					FSimdRealf ImpulseNormal = SimdNegate(SimdMultiply(SimdMultiply(SimdStiffness, ManifoldPoint.SimdContactMassNormal), ContactVelocityErrorNormal));
+					FSimdRealf ImpulseTangentU = SimdNegate(SimdMultiply(SimdMultiply(FrictionStiffness, ManifoldPoint.SimdContactMassTangentU), ContactVelocityErrorTangentU));
+					FSimdRealf ImpulseTangentV = SimdNegate(SimdMultiply(SimdMultiply(FrictionStiffness, ManifoldPoint.SimdContactMassTangentV), ContactVelocityErrorTangentV));
+
+					// Zero out friction for lanes with no friction
+					const FSimdSelector HasFriction = SimdAnd(ShouldSolveVelocity, SimdGreater(SimdVelocityFriction, FSimdRealf::Zero()));
+					ImpulseTangentU = SimdSelect(HasFriction, ImpulseTangentU, FSimdRealf::Zero());
+					ImpulseTangentV = SimdSelect(HasFriction, ImpulseTangentV, FSimdRealf::Zero());
+
+					// The minimum normal impulse we can apply. We are allowed to apply a negative impulse 
+					// up to an amount that would conteract the implciit velocity applied by the pushout
+					// MinImpulseNormal = FMath::Min(0, -NetPushOutNormal / Dt)
+					// if (NetImpulseNormal + ImpulseNormal < MinImpulseNormal) {...}
+					const FSimdRealf MinImpulseNormal = SimdMin(SimdNegate(SimdDivide(ManifoldPoint.SimdNetPushOutNormal, Dt)), FSimdRealf::Zero());
+					const FSimdSelector ShouldClampImpulse = SimdLess(SimdAdd(ManifoldPoint.SimdNetImpulseNormal, ImpulseNormal), MinImpulseNormal);
+					const FSimdRealf ClampedImpulseNormal = SimdSubtract(MinImpulseNormal, ManifoldPoint.SimdNetImpulseNormal);
+					ImpulseNormal = SimdSelect(ShouldClampImpulse, ClampedImpulseNormal, ImpulseNormal);
+
+					// Calculate the impulse multipler if clamped to the dynamic friction cone
+					const FSimdRealf MaxNetImpulseAndPushOutTangent = SimdAdd(ManifoldPoint.SimdNetImpulseNormal, SimdAdd(ImpulseNormal, SimdDivide(ManifoldPoint.SimdNetPushOutNormal, Dt)));
+					const FSimdRealf MaxImpulseTangent = SimdMax(FSimdRealf::Zero(), SimdMultiply(SimdVelocityFriction, MaxNetImpulseAndPushOutTangent));
+					const FSimdRealf MaxImpulseTangentSq = SimdSquare(MaxImpulseTangent);
+					const FSimdRealf ImpulseTangentSq = SimdAdd(SimdSquare(ImpulseTangentU), SimdSquare(ImpulseTangentV));
+					const FSimdRealf ImpulseTangentScale = SimdMultiply(MaxImpulseTangent, SimdInvSqrt(ImpulseTangentSq));
+
+					// Apply the multiplier to lanes that exceeded the friction limit
+					const FSimdSelector ExceededFrictionCone = SimdGreater(ImpulseTangentSq, SimdAdd(MaxImpulseTangentSq, FSimdRealf::Make(UE_SMALL_NUMBER)));
+					ImpulseTangentU = SimdSelect(ExceededFrictionCone, SimdMultiply(ImpulseTangentScale, ImpulseTangentU), ImpulseTangentU);
+					ImpulseTangentV = SimdSelect(ExceededFrictionCone, SimdMultiply(ImpulseTangentScale, ImpulseTangentV), ImpulseTangentV);
+
+					// Clear the impulse for lanes that should not be solving for velocity
+					ImpulseNormal = SimdSelect(ShouldSolveVelocity, ImpulseNormal, FSimdRealf::Zero());
+					ImpulseTangentU = SimdSelect(ShouldSolveVelocity, ImpulseTangentU, FSimdRealf::Zero());
+					ImpulseTangentV = SimdSelect(ShouldSolveVelocity, ImpulseTangentV, FSimdRealf::Zero());
+
+					ManifoldPoint.SimdNetImpulseNormal = SimdAdd(ManifoldPoint.SimdNetImpulseNormal, ImpulseNormal);
+					ManifoldPoint.SimdNetImpulseTangentU = SimdAdd(ManifoldPoint.SimdNetImpulseTangentU, ImpulseTangentU);
+					ManifoldPoint.SimdNetImpulseTangentV = SimdAdd(ManifoldPoint.SimdNetImpulseTangentU, ImpulseTangentV);
+
+					// NOTE: order of operations matches FPBDCollisionSolver::ApplyVelocityCorrection for AB Testing
+					const FSimdVec3f Impulse = SimdAdd(SimdMultiply(ImpulseNormal, ManifoldPoint.SimdContactNormal), SimdAdd(SimdMultiply(ImpulseTangentU, ManifoldPoint.SimdContactTangentU), SimdMultiply(ImpulseTangentV, ManifoldPoint.SimdContactTangentV)));
+					const FSimdVec3f DV0 = SimdMultiply(Impulse, SimdInvM0);
+					const FSimdVec3f DV1 = SimdMultiply(Impulse, SimdInvM1);
+					const FSimdVec3f DW0 = SimdAdd(SimdMultiply(ImpulseNormal, ManifoldPoint.SimdContactNormalAngular0), SimdAdd(SimdMultiply(ImpulseTangentU, ManifoldPoint.SimdContactTangentUAngular0), SimdMultiply(ImpulseTangentV, ManifoldPoint.SimdContactTangentVAngular0)));
+					const FSimdVec3f DW1 = SimdAdd(SimdMultiply(ImpulseNormal, ManifoldPoint.SimdContactNormalAngular1), SimdAdd(SimdMultiply(ImpulseTangentU, ManifoldPoint.SimdContactTangentUAngular1), SimdMultiply(ImpulseTangentV, ManifoldPoint.SimdContactTangentVAngular1)));
+					V0 = SimdAdd(V0, DV0);
+					W0 = SimdAdd(W0, DW0);
+					V1 = SimdSubtract(V1, DV1);
+					W1 = SimdSubtract(W1, DW1);
 				}
 
-				// Gather the body data we need
-				FSimdVec3f V0, W0, V1, W1;
-				GatherBodyVelocities(IsValid, Body0, Body1, V0, W0, V1, W1);
-
-				// Calculate the velocity error we need to correct
-				const FSimdVec3f ContactVelocity0 = SimdAdd(V0, SimdCrossProduct(W0, SimdRelativeContactPoint0));
-				const FSimdVec3f ContactVelocity1 = SimdAdd(V1, SimdCrossProduct(W1, SimdRelativeContactPoint1));
-				const FSimdVec3f ContactVelocity = SimdSubtract(ContactVelocity0, ContactVelocity1);
-				FSimdRealf ContactVelocityErrorNormal = SimdSubtract(SimdDotProduct(ContactVelocity, SimdContactNormal), SimdContactTargetVelocityNormal);
-				FSimdRealf ContactVelocityErrorTangentU = SimdDotProduct(ContactVelocity, SimdContactTangentU);
-				FSimdRealf ContactVelocityErrorTangentV = SimdDotProduct(ContactVelocity, SimdContactTangentV);
-
-				// Zero out error for lanes which should not generate friction
-				ContactVelocityErrorNormal = SimdSelect(ShouldSolveVelocity, ContactVelocityErrorNormal, FSimdRealf::Zero());
-				ContactVelocityErrorTangentU = SimdSelect(ShouldSolveVelocity, ContactVelocityErrorTangentU, FSimdRealf::Zero());
-				ContactVelocityErrorTangentV = SimdSelect(ShouldSolveVelocity, ContactVelocityErrorTangentV, FSimdRealf::Zero());
-
-				// A stiffness multiplier on the impulses (zeroed if we have no friction in a lane)
-				const FSimdRealf FrictionStiffness = SimdMultiply(FrictionStiffnessScale, SimdStiffness);
-
-				// Calculate the impulses
-				FSimdRealf ImpulseNormal = SimdNegate(SimdMultiply(SimdMultiply(SimdStiffness, SimdContactMassNormal), ContactVelocityErrorNormal));
-				FSimdRealf ImpulseTangentU = SimdNegate(SimdMultiply(SimdMultiply(FrictionStiffness, SimdContactMassTangentU), ContactVelocityErrorTangentU));
-				FSimdRealf ImpulseTangentV = SimdNegate(SimdMultiply(SimdMultiply(FrictionStiffness, SimdContactMassTangentV), ContactVelocityErrorTangentV));
-
-				// Zero out friction for lanes with no friction
-				const FSimdSelector HasFriction = SimdAnd(IsValid, SimdGreater(SimdVelocityFriction, FSimdRealf::Zero()));
-				ImpulseTangentU = SimdSelect(HasFriction, ImpulseTangentU, FSimdRealf::Zero());
-				ImpulseTangentV = SimdSelect(HasFriction, ImpulseTangentV, FSimdRealf::Zero());
-
-				// The minimum normal impulse we can apply. We are allowed to apply a negative impulse 
-				// up to an amount that would conteract the implciit velocity applied by the pushout
-				// MinImpulseNormal = FMath::Min(0, -NetPushOutNormal / Dt)
-				// if (NetImpulseNormal + ImpulseNormal < MinImpulseNormal) {...}
-				const FSimdRealf MinImpulseNormal = SimdMin(SimdNegate(SimdDivide(SimdNetPushOutNormal, Dt)), FSimdRealf::Zero());
-				const FSimdSelector ShouldClampImpulse = SimdLess(SimdAdd(SimdNetImpulseNormal, ImpulseNormal), MinImpulseNormal);
-				const FSimdRealf ClampedImpulseNormal = SimdSubtract(MinImpulseNormal, SimdNetImpulseNormal);
-				ImpulseNormal = SimdSelect(ShouldClampImpulse, ClampedImpulseNormal, ImpulseNormal);
-
-				// Calculate the impulse multipler if clamped to the dynamic friction cone
-				const FSimdRealf MaxNetImpulseAndPushOutTangent = SimdAdd(SimdNetImpulseNormal, SimdAdd(ImpulseNormal, SimdDivide(SimdNetPushOutNormal, Dt)));
-				const FSimdRealf MaxImpulseTangent = SimdMax(FSimdRealf::Zero(), SimdMultiply(SimdVelocityFriction, MaxNetImpulseAndPushOutTangent));
-				const FSimdRealf MaxImpulseTangentSq = SimdSquare(MaxImpulseTangent);
-				const FSimdRealf ImpulseTangentSq = SimdAdd(SimdSquare(ImpulseTangentU), SimdSquare(ImpulseTangentV));
-				const FSimdRealf ImpulseTangentScale = SimdMultiply(MaxImpulseTangent, SimdInvSqrt(ImpulseTangentSq));
-
-				// Apply the multiplier to lanes that exceeded the friction limit
-				const FSimdSelector ExceededFrictionCone = SimdGreater(ImpulseTangentSq, SimdAdd(MaxImpulseTangentSq, FSimdRealf::Make(UE_SMALL_NUMBER)));
-				ImpulseTangentU = SimdSelect(ExceededFrictionCone, SimdMultiply(ImpulseTangentScale, ImpulseTangentU), ImpulseTangentU);
-				ImpulseTangentV = SimdSelect(ExceededFrictionCone, SimdMultiply(ImpulseTangentScale, ImpulseTangentV), ImpulseTangentV);
-
-				SimdNetImpulseNormal = SimdAdd(SimdNetImpulseNormal, ImpulseNormal);
-				SimdNetImpulseTangentU = SimdAdd(SimdNetImpulseTangentU, ImpulseTangentU);
-				SimdNetImpulseTangentV = SimdAdd(SimdNetImpulseTangentU, ImpulseTangentV);
-
-				V0 = SimdAdd(V0, SimdMultiply(SimdMultiply(ImpulseNormal, SimdInvM0), SimdContactNormal));
-				V0 = SimdAdd(V0, SimdMultiply(SimdMultiply(ImpulseTangentU, SimdInvM0), SimdContactTangentU));
-				V0 = SimdAdd(V0, SimdMultiply(SimdMultiply(ImpulseTangentV, SimdInvM0), SimdContactTangentV));
-				W0 = SimdAdd(W0, SimdMultiply(ImpulseNormal, SimdContactNormalAngular0));
-				W0 = SimdAdd(W0, SimdMultiply(ImpulseTangentU, SimdContactTangentUAngular0));
-				W0 = SimdAdd(W0, SimdMultiply(ImpulseTangentV, SimdContactTangentVAngular0));
-
-				V1 = SimdSubtract(V1, SimdMultiply(SimdMultiply(ImpulseNormal, SimdInvM1), SimdContactNormal));
-				V1 = SimdSubtract(V1, SimdMultiply(SimdMultiply(ImpulseTangentU, SimdInvM1), SimdContactTangentU));
-				V1 = SimdSubtract(V1, SimdMultiply(SimdMultiply(ImpulseTangentV, SimdInvM1), SimdContactTangentV));
-				W1 = SimdSubtract(W1, SimdMultiply(ImpulseNormal, SimdContactNormalAngular1));
-				W1 = SimdSubtract(W1, SimdMultiply(ImpulseTangentU, SimdContactTangentUAngular1));
-				W1 = SimdSubtract(W1, SimdMultiply(ImpulseTangentV, SimdContactTangentVAngular1));
-
-				ScatterBodyVelocities(IsValid, V0, W0, V1, W1, Body0, Body1);
+				ScatterBodyVelocities(V0, W0, V1, W1, Body0, Body1);
 			}
 
 			void SolveVelocityWithFriction(
+				const TArrayView<FSimdManifoldPoint>& ManifoldPointsBuffer,
 				const FSimdSolverBodyPtr& Body0,
 				const FSimdSolverBodyPtr& Body1,
 				const FSimdRealf& Dt,
@@ -701,463 +998,84 @@ namespace Chaos
 			{
 				// If all the lanes have zero velocity friction, run the zero-friction path
 				// (only spheres and capsules use the velocity-based dynamic friction path)
-				const FSimdSelector IsValid = SimdGreaterEqual(SimdConstraintIndex, FSimdInt32::Zero());
 				const FSimdSelector HasNonZeroFriction = SimdGreater(SimdVelocityFriction, FSimdRealf::Zero());
-				if (!SimdAnyTrue(SimdAnd(IsValid, HasNonZeroFriction)))
+				if (!SimdAnyTrue(HasNonZeroFriction))
 				{
-					SolveVelocityNoFriction(Body0, Body1, Dt);
+					SolveVelocityNoFriction(ManifoldPointsBuffer, Body0, Body1, Dt);
 					return;
 				}
 
-				SolveVelocityWithFrictionImpl(Body0, Body1, Dt, FrictionStiffnessScale);
-			}
-
-			FSolverVec3 GetNetPushOut(const int32 LaneIndex) const
-			{
-				return SimdNetPushOutNormal.GetValue(LaneIndex) * SimdContactNormal.GetValue(LaneIndex) +
-					SimdNetPushOutTangentU.GetValue(LaneIndex) * SimdContactTangentU.GetValue(LaneIndex) +
-					SimdNetPushOutTangentV.GetValue(LaneIndex) * SimdContactTangentV.GetValue(LaneIndex);
-			}
-
-			FSolverVec3 GetNetImpulse(const int32 LaneIndex) const
-			{
-				return SimdNetImpulseNormal.GetValue(LaneIndex) * SimdContactNormal.GetValue(LaneIndex) +
-					SimdNetImpulseTangentU.GetValue(LaneIndex) * SimdContactTangentU.GetValue(LaneIndex) +
-					SimdNetImpulseTangentV.GetValue(LaneIndex) * SimdContactTangentV.GetValue(LaneIndex);
-			}
-
-			FSolverReal GetStaticFrictionRatio(const int32 LaneIndex) const
-			{
-				return SimdStaticFrictionRatio.GetValue(LaneIndex);
+				SolveVelocityWithFrictionImpl(ManifoldPointsBuffer, Body0, Body1, Dt, FrictionStiffnessScale);
 			}
 
-			// @todo(chaos): make private
-		public:
-			friend class FPBDCollisionSolverSimd;
-
-			/**
-			 * @brief Whether we need to solve velocity for this manifold point (only if we were penetrating or applied a pushout)
-			*/
-			bool ShouldSolveVelocity() const
-			{
-				return true;
-			}
-
-			FSimdInt32 SimdConstraintIndex;
-			FSimdInt32 SimdManifoldPointIndex;
-
-			FSimdRealf SimdStiffness;
-
-			// World-space contact point relative to each particle's center of mass
-			FSimdVec3f SimdRelativeContactPoint0;
-			FSimdVec3f SimdRelativeContactPoint1;
-
-			// Normal PushOut
-			FSimdVec3f SimdContactNormal;
-			FSimdRealf SimdContactDeltaNormal;
-			FSimdRealf SimdNetPushOutNormal;
-			FSimdRealf SimdContactMassNormal;
-			FSimdVec3f SimdContactNormalAngular0;
-			FSimdVec3f SimdContactNormalAngular1;
-
-			// M^-1 for each body
-			FSimdRealf SimdInvM0;
-			FSimdRealf SimdInvM1;
-
-			// Tangential PushOut
-			FSimdVec3f SimdContactTangentU;
-			FSimdVec3f SimdContactTangentV;
-			FSimdRealf SimdContactDeltaTangentU;
-			FSimdRealf SimdContactDeltaTangentV;
-			FSimdRealf SimdNetPushOutTangentU;
-			FSimdRealf SimdNetPushOutTangentV;
-			FSimdRealf SimdStaticFrictionRatio;
-			FSimdRealf SimdContactMassTangentU;
-			FSimdRealf SimdContactMassTangentV;
-			FSimdRealf SimdStaticFriction;
-			FSimdRealf SimdDynamicFriction;
-			FSimdRealf SimdVelocityFriction;
-			FSimdVec3f SimdContactTangentUAngular0;
-			FSimdVec3f SimdContactTangentVAngular0;
-			FSimdVec3f SimdContactTangentUAngular1;
-			FSimdVec3f SimdContactTangentVAngular1;
-
-			// Normal Impulse
-			FSimdRealf SimdContactTargetVelocityNormal;
-			FSimdRealf SimdNetImpulseNormal;
-			FSimdRealf SimdNetImpulseTangentU;
-			FSimdRealf SimdNetImpulseTangentV;
-		};
-
-		/**
-		 * @brief
-		 * @todo(chaos): Make this solver operate on a single contact point rather than all points in a manifold.
-		 * This would be beneficial if we have many contacts with less than 4 points in the manifold. However this
-		 * is dificult to do while we are still supporting non-manifold collisions.
-		*/
-		class FPBDCollisionSolverSimd
-		{
-		public:
-			static const int32 MaxConstrainedBodies = 2;
-			static const int32 MaxPointsPerConstraint = 4;
-
-			// Create a solver that is initialized to safe defaults
-			static FPBDCollisionSolverSimd MakeInitialized()
-			{
-				FPBDCollisionSolverSimd Solver;
-				Solver.State.Init();
-				return Solver;
-			}
-
-			// Create a solver with no initialization
-			static FPBDCollisionSolverSimd MakeUninitialized()
-			{
-				return FPBDCollisionSolverSimd();
-			}
-
-			// NOTE: Does not initialize any properties. See MakeInitialized
-			FPBDCollisionSolverSimd() {}
-
-			int32 GetManifoldPointBeginIndex() const
-			{
-				return State.ManifoldPointBeginIndex;
-			}
-
-			int32 GetLaneIndex() const
-			{
-				return State.LaneIndex;
-			}
-
-			/** Reset the state of the collision solver */
-			void Reset()
-			{
-				State.SolverBody0.Reset();
-				State.SolverBody1.Reset();
-				State.LaneIndex = 0;
-				State.ManifoldPointBeginIndex = 0;
-				State.NumManifoldPoints = 0;
-				State.MaxManifoldPoints = 0;
-			}
-
-			void ResetManifold()
-			{
-				State.NumManifoldPoints = 0;
-			}
-
-			FSolverReal StaticFriction() const { return State.StaticFriction; }
-			FSolverReal DynamicFriction() const { return State.DynamicFriction; }
-			FSolverReal VelocityFriction() const { return State.VelocityFriction; }
-
-			void SetFriction(const FSolverReal InStaticFriction, const FSolverReal InDynamicFriction, const FSolverReal InVelocityFriction)
-			{
-				State.StaticFriction = InStaticFriction;
-				State.DynamicFriction = InDynamicFriction;
-				State.VelocityFriction = InVelocityFriction;
-			}
-
-			void SetStiffness(const FSolverReal InStiffness)
-			{
-				State.Stiffness = InStiffness;
-			}
-
-			void SetSolverBodies(FSolverBody& SolverBody0, FSolverBody& SolverBody1)
-			{
-				State.SolverBody0.SetSolverBody(SolverBody0);
-				State.SolverBody1.SetSolverBody(SolverBody1);
-			}
-
-			void SetManifoldPointsBuffer(const int32 ConstraintIndex, const int32 LaneIndex, const int32 BeginIndex, const int32 Num)
-			{
-				State.ConstraintIndex = ConstraintIndex;
-				State.LaneIndex = LaneIndex;
-				State.ManifoldPointBeginIndex = BeginIndex;
-				State.NumManifoldPoints = 0;
-				State.MaxManifoldPoints = Num;
-			}
-
-			int32 NumManifoldPoints() const
-			{
-				return State.NumManifoldPoints;
-			}
-
-			int32 AddManifoldPoint()
-			{
-				check(State.NumManifoldPoints < State.MaxManifoldPoints);
-				return State.NumManifoldPoints++;
-			}
-
-			/**
-			 * Set up a manifold point (calls InitManifodlPoint and FinalizeManifoldPoint)
-			*/
-			template<int TNumLanes>
-			void SetManifoldPoint(
-				const TArrayView<TPBDCollisionSolverManifoldPointsSimd<TNumLanes>>& ManifoldPointsBuffer,
-				const TArrayView<TSolverBodyPtrPairSimd<4>>& SolverBodiesBuffer,
-				const int32 ManifoldPointIndex,
-				const FSolverVec3& InRelativeContactPosition0,
-				const FSolverVec3& InRelativeContactPosition1,
-				const FSolverVec3& InWorldContactNormal,
-				const FSolverVec3& InWorldContactTangentU,
-				const FSolverVec3& InWorldContactTangentV,
-				const FSolverReal InWorldContactDeltaNormal,
-				const FSolverReal InWorldContactDeltaTangentU,
-				const FSolverReal InWorldContactDeltaTangentV,
-				const FSolverReal InWorldContactVelocityTargetNormal)
-			{
-				InitManifoldPoint(
-					ManifoldPointsBuffer,
-					SolverBodiesBuffer,
-					ManifoldPointIndex,
-					InRelativeContactPosition0,
-					InRelativeContactPosition1,
-					InWorldContactNormal,
-					InWorldContactTangentU,
-					InWorldContactTangentV,
-					InWorldContactDeltaNormal,
-					InWorldContactDeltaTangentU,
-					InWorldContactDeltaTangentV,
-					InWorldContactVelocityTargetNormal);
-
-				FinalizeManifoldPoint(
-					ManifoldPointsBuffer,
-					ManifoldPointIndex);
-			}
-
-			/**
-			 * Initialize the contact data for a manifold point (must call FinalizeManifoldPoint before use)
-			*/
-			template<int TNumLanes>
-			void InitManifoldPoint(
-				const TArrayView<TPBDCollisionSolverManifoldPointsSimd<TNumLanes>>& ManifoldPointsBuffer,
-				const TArrayView<TSolverBodyPtrPairSimd<4>>& SolverBodiesBuffer,
+			FSolverVec3 GetNetPushOut(
 				const int32 ManifoldPointIndex,
-				const FSolverVec3& InRelativeContactPosition0,
-				const FSolverVec3& InRelativeContactPosition1,
-				const FSolverVec3& InWorldContactNormal,
-				const FSolverVec3& InWorldContactTangentU,
-				const FSolverVec3& InWorldContactTangentV,
-				const FSolverReal InWorldContactDeltaNormal,
-				const FSolverReal InWorldContactDeltaTangentU,
-				const FSolverReal InWorldContactDeltaTangentV,
-				const FSolverReal InWorldContactVelocityTargetNormal)
+				const int32 LaneIndex,
+				const TArrayView<FSimdManifoldPoint>& ManifoldPointsBuffer) const
 			{
 				const int32 BufferIndex = GetBufferIndex(ManifoldPointIndex);
+				const FSimdManifoldPoint& ManifoldPoint = ManifoldPointsBuffer[BufferIndex];
 
-				SolverBodiesBuffer[BufferIndex].Body0.SetValue(
-					State.LaneIndex,
-					&SolverBody0().SolverBody());
-
-				SolverBodiesBuffer[BufferIndex].Body1.SetValue(
-					State.LaneIndex,
-					&SolverBody1().SolverBody());
-
-				ManifoldPointsBuffer[BufferIndex].SetSharedData(
-					State.LaneIndex,
-					State.ConstraintIndex,
-					ManifoldPointIndex);
-
-				ManifoldPointsBuffer[BufferIndex].InitContact(
-					State.LaneIndex,
-					InRelativeContactPosition0,
-					InRelativeContactPosition1,
-					InWorldContactNormal,
-					InWorldContactTangentU,
-					InWorldContactTangentV,
-					InWorldContactDeltaNormal,
-					InWorldContactDeltaTangentU,
-					InWorldContactDeltaTangentV,
-					InWorldContactVelocityTargetNormal,
-					State.Stiffness,
-					State.StaticFriction,
-					State.DynamicFriction,
-					State.VelocityFriction);
-			}
-
-			/** 
-			 * Finish manifold point setup.
-			 * NOTE: Can only be called after the InitManifoldPoint has called
-			*/
-			template<int TNumLanes>
-			void FinalizeManifoldPoint(
-				const TArrayView<TPBDCollisionSolverManifoldPointsSimd<TNumLanes>>& ManifoldPointsBuffer,
-				const int32 ManifoldPointIndex)
-			{
-				ManifoldPointsBuffer[GetBufferIndex(ManifoldPointIndex)].FinalizeContact(
-					State.LaneIndex,
-					SolverBody0(),
-					SolverBody1());
-			}
-
-			/**
-			 * @brief Get the first (decorated) solver body
-			 * The decorator add a possible mass scale
-			*/
-			FConstraintSolverBody& SolverBody0() { return State.SolverBody0; }
-			const FConstraintSolverBody& SolverBody0() const { return State.SolverBody0; }
-
-			/**
-			 * @brief Get the second (decorated) solver body
-			 * The decorator add a possible mass scale
-			*/
-			FConstraintSolverBody& SolverBody1() { return State.SolverBody1; }
-			const FConstraintSolverBody& SolverBody1() const { return State.SolverBody1; }
-
-			/**
-			 * @brief Set up the mass scaling for shock propagation, using the position-phase mass scale
-			*/
-			template<int TNumLanes>
-			void EnablePositionShockPropagation(
-				const TArrayView<TPBDCollisionSolverManifoldPointsSimd<TNumLanes>>& ManifoldPointsBuffer)
-			{
-				SetShockPropagationInvMassScale(ManifoldPointsBuffer, CVars::Chaos_PBDCollisionSolver_Position_MinInvMassScale);
-			}
-
-			/**
-			 * @brief Set up the mass scaling for shock propagation, using the velocity-phase mass scale
-			*/
-			template<int TNumLanes>
-			void EnableVelocityShockPropagation(
-				const TArrayView<TPBDCollisionSolverManifoldPointsSimd<TNumLanes>>& ManifoldPointsBuffer)
-			{
-				SetShockPropagationInvMassScale(ManifoldPointsBuffer, CVars::Chaos_PBDCollisionSolver_Velocity_MinInvMassScale);
-			}
-
-			/**
-			 * @brief Disable mass scaling
-			*/
-			template<int TNumLanes>
-			void DisableShockPropagation(
-				const TArrayView<TPBDCollisionSolverManifoldPointsSimd<TNumLanes>>& ManifoldPointsBuffer)
-			{
-				SetShockPropagationInvMassScale(ManifoldPointsBuffer, FReal(1));
-			}
-
-			template<int TNumLanes>
-			FSolverVec3 GetNetPushOut(
-				const TArrayView<TPBDCollisionSolverManifoldPointsSimd<TNumLanes>>& ManifoldPointsBuffer,
-				const int32 ManifoldPointIndex) const
-			{
-				return ManifoldPointsBuffer[GetBufferIndex(ManifoldPointIndex)].GetNetPushOut(State.LaneIndex);
+				return ManifoldPoint.SimdNetPushOutNormal.GetValue(LaneIndex) * ManifoldPoint.SimdContactNormal.GetValue(LaneIndex) +
+					ManifoldPoint.SimdNetPushOutTangentU.GetValue(LaneIndex) * ManifoldPoint.SimdContactTangentU.GetValue(LaneIndex) +
+					ManifoldPoint.SimdNetPushOutTangentV.GetValue(LaneIndex) * ManifoldPoint.SimdContactTangentV.GetValue(LaneIndex);
 			}
 
-			template<int TNumLanes>
 			FSolverVec3 GetNetImpulse(
-				const TArrayView<TPBDCollisionSolverManifoldPointsSimd<TNumLanes>>& ManifoldPointsBuffer,
-				const int32 ManifoldPointIndex) const
+				const int32 ManifoldPointIndex,
+				const int32 LaneIndex,
+				const TArrayView<FSimdManifoldPoint>& ManifoldPointsBuffer) const
 			{
-				return ManifoldPointsBuffer[GetBufferIndex(ManifoldPointIndex)].GetNetImpulse(State.LaneIndex);
+				const int32 BufferIndex = GetBufferIndex(ManifoldPointIndex);
+				const FSimdManifoldPoint& ManifoldPoint = ManifoldPointsBuffer[BufferIndex];
+
+				return ManifoldPoint.SimdNetImpulseNormal.GetValue(LaneIndex) * ManifoldPoint.SimdContactNormal.GetValue(LaneIndex) +
+					ManifoldPoint.SimdNetImpulseTangentU.GetValue(LaneIndex) * ManifoldPoint.SimdContactTangentU.GetValue(LaneIndex) +
+					ManifoldPoint.SimdNetImpulseTangentV.GetValue(LaneIndex) * ManifoldPoint.SimdContactTangentV.GetValue(LaneIndex);
 			}
 
-			template<int TNumLanes>
 			FSolverReal GetStaticFrictionRatio(
-				const TArrayView<TPBDCollisionSolverManifoldPointsSimd<TNumLanes>>& ManifoldPointsBuffer,
-				const int32 ManifoldPointIndex) const
+				const int32 ManifoldPointIndex,
+				const int32 LaneIndex,
+				const TArrayView<FSimdManifoldPoint>& ManifoldPointsBuffer) const
 			{
-				return ManifoldPointsBuffer[GetBufferIndex(ManifoldPointIndex)].GetStaticFrictionRatio(State.LaneIndex);
+				const int32 BufferIndex = GetBufferIndex(ManifoldPointIndex);
+				const FSimdManifoldPoint& ManifoldPoint = ManifoldPointsBuffer[BufferIndex];
+
+				return ManifoldPoint.SimdStaticFrictionRatio.GetValue(LaneIndex);
 			}
 
-		private:
+		public:
 			int32 GetBufferIndex(const int32 ManifoldPointIndex) const
 			{
-				check(ManifoldPointIndex < State.NumManifoldPoints);
-				return State.ManifoldPointBeginIndex + ManifoldPointIndex;
+				return ManifoldPointBeginIndex + ManifoldPointIndex;
 			}
 
-			/**
-			 * @brief Apply the inverse mass scale the body with the lower level
-			 * @param InvMassScale
-			*/
-			template<int TNumLanes>
-			void SetShockPropagationInvMassScale(
-				const TArrayView<TPBDCollisionSolverManifoldPointsSimd<TNumLanes>>& ManifoldPointsBuffer,
-				const FSolverReal InvMassScale)
+			void Init()
 			{
-				{
-					FConstraintSolverBody& Body0 = SolverBody0();
-					FConstraintSolverBody& Body1 = SolverBody1();
-
-					// Shock propagation decreases the inverse mass of bodies that are lower in the pile
-					// of objects. This significantly improves stability of heaps and stacks. Height in the pile is indictaed by the "level". 
-					// No need to set an inverse mass scale if the other body is kinematic (with inv mass of 0).
-					// Bodies at the same level do not take part in shock propagation.
-					if (Body0.IsDynamic() && Body1.IsDynamic() && (Body0.Level() != Body1.Level()))
-					{
-						// Set the inv mass scale of the "lower" body to make it heavier
-						bool bInvMassUpdated = false;
-						if (Body0.Level() < Body1.Level())
-						{
-							if (Body0.ShockPropagationScale() != InvMassScale)
-							{
-								Body0.SetShockPropagationScale(InvMassScale);
-								bInvMassUpdated = true;
-							}
-						}
-						else
-						{
-							if (Body1.ShockPropagationScale() != InvMassScale)
-							{
-								Body1.SetShockPropagationScale(InvMassScale);
-								bInvMassUpdated = true;
-							}
-						}
-
-						// If the masses changed, we need to rebuild the contact mass for each manifold point
-						if (bInvMassUpdated)
-						{
-							for (int32 PointIndex = 0; PointIndex < NumManifoldPoints(); ++PointIndex)
-							{
-								ManifoldPointsBuffer[GetBufferIndex(PointIndex)].UpdateMassNormal(State.LaneIndex, Body0, Body1);
-							}
-						}
-					}
-				}
+				MaxManifoldPoints = 0;
+				ManifoldPointBeginIndex = INDEX_NONE;
+				SimdNumManifoldPoints.SetValues(0);
+				SimdStaticFriction.SetValues(0);
+				SimdDynamicFriction.SetValues(0);
+				SimdVelocityFriction.SetValues(0);
+				SimdStiffness.SetValues(1);
+				SimdInvM0.SetValues(0);
+				SimdInvM1.SetValues(0);
 			}
 
-			struct FState
-			{
-				FState()
-				{
-				}
-
-				void Init()
-				{
-					ConstraintIndex = INDEX_NONE;
-					LaneIndex = INDEX_NONE;
-					ManifoldPointBeginIndex = INDEX_NONE;
-					NumManifoldPoints = 0;
-					MaxManifoldPoints = 0;
-					StaticFriction = 0;
-					DynamicFriction = 0;
-					VelocityFriction = 0;
-					Stiffness = 1;
-					SolverBody0.Init();
-					SolverBody1.Init();
-				}
-
-				int32 ConstraintIndex;
-				int32 LaneIndex;
-				int32 ManifoldPointBeginIndex;
-				int32 NumManifoldPoints;
-				int32 MaxManifoldPoints;
-
-				// Static Friction in the position-solve phase
-				FSolverReal StaticFriction;
-
-				// Dynamic Friction in the position-solve phase
-				FSolverReal DynamicFriction;
-
-				// Dynamic Friction in the velocity-solve phase
-				FSolverReal VelocityFriction;
-
-				// Solver stiffness (scales all pushout and impulses)
-				FSolverReal Stiffness;
+			// Each lane has space for MaxManifoldPoints, but not all will be used
+			// Unused manifold points may be in the middle of the list (if disabled for example)
+			int32 MaxManifoldPoints;
+			int32 ManifoldPointBeginIndex;
+			FSimdInt32 SimdNumManifoldPoints;
 
-				// Bodies and contacts
-				FConstraintSolverBody SolverBody0;
-				FConstraintSolverBody SolverBody1;
-			};
-
-			FState State;
+			FSimdRealf SimdStaticFriction;
+			FSimdRealf SimdDynamicFriction;
+			FSimdRealf SimdVelocityFriction;
+			FSimdRealf SimdStiffness;
+			FSimdRealf SimdInvM0;
+			FSimdRealf SimdInvM1;
 		};
 
 		//////////////////////////////////////////////////////////////////////////////////////////////////
@@ -1175,6 +1093,7 @@ namespace Chaos
 		public:
 			template<int TNumLanes>
 			static CHAOS_API void SolvePositionNoFriction(
+				const TArrayView<TPBDCollisionSolverSimd<TNumLanes>>& Solvers,
 				const TArrayView<TPBDCollisionSolverManifoldPointsSimd<TNumLanes>>& ManifoldPoints,
 				const TArrayView<TSolverBodyPtrPairSimd<TNumLanes>>& SolverBodies,
 				const FSolverReal Dt,
@@ -1182,6 +1101,7 @@ namespace Chaos
 
 			template<int TNumLanes>
 			static CHAOS_API void SolvePositionWithFriction(
+				const TArrayView<TPBDCollisionSolverSimd<TNumLanes>>& Solvers,
 				const TArrayView<TPBDCollisionSolverManifoldPointsSimd<TNumLanes>>& ManifoldPoints,
 				const TArrayView<TSolverBodyPtrPairSimd<TNumLanes>>& SolverBodies,
 				const FSolverReal Dt,
@@ -1189,12 +1109,14 @@ namespace Chaos
 
 			template<int TNumLanes>
 			static CHAOS_API void SolveVelocityNoFriction(
+				const TArrayView<TPBDCollisionSolverSimd<TNumLanes>>& Solvers,
 				const TArrayView<TPBDCollisionSolverManifoldPointsSimd<TNumLanes>>& ManifoldPoints,
 				const TArrayView<TSolverBodyPtrPairSimd<TNumLanes>>& SolverBodies,
 				const FSolverReal Dt);
 
 			template<int TNumLanes>
 			static CHAOS_API void SolveVelocityWithFriction(
+				const TArrayView<TPBDCollisionSolverSimd<TNumLanes>>& Solvers,
 				const TArrayView<TPBDCollisionSolverManifoldPointsSimd<TNumLanes>>& ManifoldPoints,
 				const TArrayView<TSolverBodyPtrPairSimd<TNumLanes>>& SolverBodies,
 				const FSolverReal Dt);
@@ -1204,6 +1126,7 @@ namespace Chaos
 
 		template<> 
 		void FPBDCollisionSolverHelperSimd::SolvePositionNoFriction(
+			const TArrayView<TPBDCollisionSolverSimd<4>>& Solvers,
 			const TArrayView<TPBDCollisionSolverManifoldPointsSimd<4>>& ManifoldPoints,
 			const TArrayView<TSolverBodyPtrPairSimd<4>>& SolverBodies,
 			const FSolverReal Dt,
@@ -1211,6 +1134,7 @@ namespace Chaos
 
 		template<>
 		void FPBDCollisionSolverHelperSimd::SolvePositionWithFriction(
+			const TArrayView<TPBDCollisionSolverSimd<4>>& Solvers,
 			const TArrayView<TPBDCollisionSolverManifoldPointsSimd<4>>& ManifoldPoints,
 			const TArrayView<TSolverBodyPtrPairSimd<4>>& SolverBodies,
 			const FSolverReal Dt,
@@ -1218,12 +1142,14 @@ namespace Chaos
 
 		template<>
 		void FPBDCollisionSolverHelperSimd::SolveVelocityNoFriction(
+			const TArrayView<TPBDCollisionSolverSimd<4>>& Solvers,
 			const TArrayView<TPBDCollisionSolverManifoldPointsSimd<4>>& ManifoldPoints,
 			const TArrayView<TSolverBodyPtrPairSimd<4>>& SolverBodies,
 			const FSolverReal Dt);
 
 		template<>
 		void FPBDCollisionSolverHelperSimd::SolveVelocityWithFriction(
+			const TArrayView<TPBDCollisionSolverSimd<4>>& Solvers,
 			const TArrayView<TPBDCollisionSolverManifoldPointsSimd<4>>& ManifoldPoints,
 			const TArrayView<TSolverBodyPtrPairSimd<4>>& SolverBodies,
 			const FSolverReal Dt);
diff --git a/Engine/Source/Runtime/Experimental/Chaos/Public/Chaos/Evolution/ABTestingConstraintContainerSolver.h b/Engine/Source/Runtime/Experimental/Chaos/Public/Chaos/Evolution/ABTestingConstraintContainerSolver.h
--- a/Engine/Source/Runtime/Experimental/Chaos/Public/Chaos/Evolution/ABTestingConstraintContainerSolver.h
+++ b/Engine/Source/Runtime/Experimental/Chaos/Public/Chaos/Evolution/ABTestingConstraintContainerSolver.h
@@ -42,7 +42,9 @@ namespace Chaos
 			{
 				PreApplyPositionConstraints,
 				PostApplyPositionConstraints,
+				PreApplyVelocityConstraints,
 				PostApplyVelocityConstraints,
+				PreApplyProjectionConstraints,
 				PostApplyProjectionConstraints,
 			};
 
@@ -157,6 +159,8 @@ namespace Chaos
 
 				SolverA->PreApplyVelocityConstraints(Dt);
 				SolverB->PreApplyVelocityConstraints(Dt);
+
+				CallABTestFunctor(ESolverPhase::PreApplyProjectionConstraints);
 			}
 
 			virtual void ApplyVelocityConstraints(const FReal Dt, const int32 It, const int32 NumIts) override final
@@ -173,6 +177,8 @@ namespace Chaos
 
 				SolverA->PreApplyProjectionConstraints(Dt);
 				SolverB->PreApplyProjectionConstraints(Dt);
+
+				CallABTestFunctor(ESolverPhase::PreApplyProjectionConstraints);
 			}
 
 			virtual void ApplyProjectionConstraints(const FReal Dt, const int32 It, const int32 NumIts) override final
diff --git a/Engine/Source/Runtime/Experimental/Chaos/Public/Chaos/VectorUtility.h b/Engine/Source/Runtime/Experimental/Chaos/Public/Chaos/VectorUtility.h
--- a/Engine/Source/Runtime/Experimental/Chaos/Public/Chaos/VectorUtility.h
+++ b/Engine/Source/Runtime/Experimental/Chaos/Public/Chaos/VectorUtility.h
@@ -58,6 +58,24 @@ FORCEINLINE VectorRegister4Float VectorUnpackLo(const VectorRegister4Float& A, c
 #endif
 }
 
+/**
+ * Selects and interleaves the higher two SP FP values from A and B.
+ *
+ * @param A	1st vector
+ * @param B	2nd vector
+ * @return		VectorRegister4Float( A.z, B.z, A.w, B.w)
+ */
+FORCEINLINE VectorRegister4Float VectorUnpackHi(const VectorRegister4Float& A, const VectorRegister4Float& B)
+{
+#if PLATFORM_ENABLE_VECTORINTRINSICS_NEON
+	return vzip2q_f32(A, B);
+#elif PLATFORM_ENABLE_VECTORINTRINSICS
+	return _mm_unpackhi_ps(A, B);
+#else
+	return MakeVectorRegisterFloat(A.V[2], B.V[2], A.V[3], B.V[3]);
+#endif
+}
+
 /**
  * Moves the lower 2 SP FP values of b to the upper 2 SP FP values of the result. The lower 2 SP FP values of a are passed through to the result.
  *
diff --git a/Engine/Source/Runtime/Experimental/ChaosCore/Public/Chaos/Simd4.h b/Engine/Source/Runtime/Experimental/ChaosCore/Public/Chaos/Simd4.h
--- a/Engine/Source/Runtime/Experimental/ChaosCore/Public/Chaos/Simd4.h
+++ b/Engine/Source/Runtime/Experimental/ChaosCore/Public/Chaos/Simd4.h
@@ -31,7 +31,7 @@ namespace Chaos
 		FORCEINLINE TSimdSelector<4> TSimdSelector<4>::True()
 		{
 			TSimdSelector<4> Selector;
-			VectorIntStoreAligned(MakeVectorRegisterIntConstant(0xFFFFFFFF, 0xFFFFFFFF, 0xFFFFFFFF, 0xFFFFFFFF), Selector.V);
+			VectorIntStoreAligned(GlobalVectorConstants::IntAllMask, Selector.V);
 			return Selector;
 		}
 
@@ -39,7 +39,7 @@ namespace Chaos
 		FORCEINLINE TSimdSelector<4> TSimdSelector<4>::False()
 		{
 			TSimdSelector<4> Selector;
-			VectorIntStoreAligned(MakeVectorRegisterIntConstant(0, 0, 0, 0), Selector.V);
+			VectorIntStoreAligned(GlobalVectorConstants::IntZero, Selector.V);
 			return Selector;
 		}
 
@@ -55,7 +55,7 @@ namespace Chaos
 		FORCEINLINE TSimdInt32<4> TSimdInt32<4>::Zero()
 		{
 			TSimdInt32<4> I;
-			VectorIntStoreAligned(MakeVectorRegisterIntConstant(0, 0, 0, 0), I.V);
+			VectorIntStoreAligned(GlobalVectorConstants::IntZero, I.V);
 			return I;
 		}
 
@@ -138,6 +138,17 @@ namespace Chaos
 			return ((VectorMaskBits(L) & 0xF) == 0xF);
 		}
 
+		FORCEINLINE FSimd4Selector SimdNot(const FSimd4Selector& InL)
+		{
+			VectorRegister4f L = VectorLoadAligned(InL.V);
+		
+			VectorRegister4f Mask = VectorBitwiseNotAnd(L, GlobalVectorConstants::AllMask());
+
+			FSimd4Selector Out;
+			VectorStoreAligned(Mask, Out.V);
+			return Out;
+		}
+
 		FORCEINLINE FSimd4Selector SimdOr(const FSimd4Selector& InL, const FSimd4Selector& InR)
 		{
 			VectorRegister4f L = VectorLoadAligned(InL.V);
@@ -199,6 +210,18 @@ namespace Chaos
 			return Out;
 		}
 
+		FORCEINLINE FSimd4Selector SimdGreaterEqual(const FSimd4Realf& InL, const FSimd4Realf& InR)
+		{
+			VectorRegister4f L = VectorLoadAligned(InL.V);
+			VectorRegister4f R = VectorLoadAligned(InR.V);
+
+			VectorRegister4f Mask = VectorCompareGE(L, R);
+
+			FSimd4Selector Out;
+			VectorStoreAligned(Mask, Out.V);
+			return Out;
+		}
+
 		FORCEINLINE FSimd4Selector SimdGreater(const FSimd4Int32& InL, const FSimd4Int32& InR)
 		{
 			VectorRegister4i L = VectorIntLoadAligned(InL.V);
@@ -510,5 +533,38 @@ namespace Chaos
 			return Out;
 		}
 
+		///////////////////////////////////////////////////////////////////////
+		//
+		// 4-wide Gather/Scatter operations
+		//
+		///////////////////////////////////////////////////////////////////////
+
+		// Convert 4 row-vectors into 3 column-vectors
+		// NOTE: The input vectors must be 16-byte aligned and padded to 16 bytes to avoid reading past valid memory
+		FORCEINLINE FSimd4Vec3f SimdGatherAligned(const FVec3f& InA, const FVec3f& InB, const FVec3f& InC, const FVec3f& InD)
+		{
+			VectorRegister4f A = VectorLoadAligned(&InA.X);	// Ax Ay Az Aw
+			VectorRegister4f B = VectorLoadAligned(&InB.X);	// Bx By Bz Bw
+			VectorRegister4f C = VectorLoadAligned(&InC.X);	// Cx Cy Cz Cw
+			VectorRegister4f D = VectorLoadAligned(&InD.X);	// Dx Dy Dz Dw
+
+			// This can be done with fewer registers, but the compiler should figure that out
+			// and its much easier to follow when left like this...
+			VectorRegister4f P = VectorUnpackLo(A, C);		// Ax Cx Ay Cy
+			VectorRegister4f Q = VectorUnpackLo(B, D);		// Bx Dx By Dy
+			VectorRegister4f R = VectorUnpackHi(A, C);		// Az Cz Aw Cw
+			VectorRegister4f S = VectorUnpackHi(B, D);		// Bz Dz Bw Dw
+			
+			VectorRegister4f X = VectorUnpackLo(P, Q);		// Ax Bx Cx Dx
+			VectorRegister4f Y = VectorUnpackHi(P, Q);		// Ay By Cy Dy 
+			VectorRegister4f Z = VectorUnpackLo(R, S);		// Az Bz Cz Dz
+
+			FSimd4Vec3f Out;
+			VectorStoreAligned(X, Out.VX);
+			VectorStoreAligned(Y, Out.VY);
+			VectorStoreAligned(Z, Out.VZ);
+			return Out;
+		}
+
 	}
 }
\ No newline at end of file
diff --git a/Engine/Source/Runtime/Experimental/ChaosCore/Public/Chaos/SimdTypes.h b/Engine/Source/Runtime/Experimental/ChaosCore/Public/Chaos/SimdTypes.h
--- a/Engine/Source/Runtime/Experimental/ChaosCore/Public/Chaos/SimdTypes.h
+++ b/Engine/Source/Runtime/Experimental/ChaosCore/Public/Chaos/SimdTypes.h
@@ -2,6 +2,8 @@
 #pragma once
 
 #include "Chaos/Core.h"
+#include "Chaos/VectorUtility.h"
+#include "Math/VectorRegister.h"
 
 namespace Chaos
 {
@@ -144,6 +146,16 @@ namespace Chaos
 				*this = Make(I);
 			}
 
+			FORCEINLINE int32 GetMaxValue() const
+			{
+				int32 MaxValue = V[0];
+				for (int32 LaneIndex = 1; LaneIndex < TNumLanes; ++LaneIndex)
+				{
+					MaxValue = FMath::Max(MaxValue, V[LaneIndex]);
+				}
+				return MaxValue;
+			}
+
 			static TSimdInt32<TNumLanes> Make(const int32 I);
 			static TSimdInt32<TNumLanes> Zero();
 		};
