Chaos
- add CCD modifier callback to allow the user to accept/reject contacts before CCD rewind is applied
- move Contact Modifier callback to run after CCD rewind so that the callback has access to the actual contact details that will be used in the constraint solving phase
- add a cvar to call the ContactModifier to be called in its original place before CCD rewind for projects that have behaviour dependent on that

diff --git a/Engine/Source/Runtime/Experimental/Chaos/Private/Chaos/CCDModification.cpp b/Engine/Source/Runtime/Experimental/Chaos/Private/Chaos/CCDModification.cpp
new file mode 100644
--- /dev/null
+++ b/Engine/Source/Runtime/Experimental/Chaos/Private/Chaos/CCDModification.cpp
@@ -0,0 +1,213 @@
+// Copyright Epic Games, Inc. All Rights Reserved.
+
+#include "Chaos/CCDModification.h"
+#include "Chaos/CCDUtilities.h"
+#include "Chaos/Collision/PBDCollisionConstraint.h"
+#include "Chaos/Collision/ParticlePairMidPhase.h"
+#include "Chaos/Particle/ParticleUtilities.h"
+#include "Chaos/PBDCollisionConstraints.h"
+
+namespace Chaos
+{
+	bool FCCDModifier::IsSweepHit() const
+	{
+		if (Constraint != nullptr)
+		{
+			return Constraint->GetCCDSweepEnabled() && (Constraint->NumManifoldPoints() > 0);
+		}
+		return false;
+	}
+
+	FReal FCCDModifier::GetSweepHitTOI() const
+	{
+		if (Constraint != nullptr)
+		{
+			return Constraint->GetCCDTimeOfImpact();
+		}
+		return FReal(1);
+	}
+
+	FVec3 FCCDModifier::GetWorldSweepHitLocation(const int32 ParticleIndex) const
+	{
+		check(ParticleIndex >= 0);
+		check(ParticleIndex < 2);
+
+		if (Constraint != nullptr)
+		{
+			// NOTE: sweeps only generate 1 manifold point
+			if (Constraint->NumManifoldPoints() > 0)
+			{
+				const FManifoldPoint& ManifoldPoint = Constraint->GetManifoldPoint(0);
+				const FRigidTransform3 ParticleWorldTransform = FCCDHelpers::GetParticleTransformAtTOI(GetParticle(ParticleIndex), GetSweepHitTOI(), Accessor->GetDt());
+				const FRigidTransform3 ShapeWorldTransform = Constraint->GetShapeRelativeTransform(ParticleIndex) * ParticleWorldTransform;
+				return ShapeWorldTransform.TransformPositionNoScale(FVec3(ManifoldPoint.ContactPoint.ShapeContactPoints[ParticleIndex]));
+			}
+		}
+		return FVec3(0);
+	}
+
+	FVec3 FCCDModifier::GetWorldSweepHitNormal() const
+	{
+		if (Constraint != nullptr)
+		{
+			// NOTE: sweeps only generate 1 manifold point
+			if (Constraint->NumManifoldPoints() > 0)
+			{
+				const FManifoldPoint& ManifoldPoint = Constraint->GetManifoldPoint(0);
+				const FRigidTransform3& ShapeWorldTransform = Constraint->GetShapeWorldTransform1();
+				return ShapeWorldTransform.TransformVectorNoScale(FVec3(ManifoldPoint.ContactPoint.ShapeContactNormal));
+			}
+		}
+		return FVec3(0,0,1);
+	}
+
+	const FGeometryParticleHandle* FCCDModifier::GetParticle(const int32 ParticleIndex) const
+	{
+		check(ParticleIndex >= 0);
+		check(ParticleIndex < 2);
+
+		if (Constraint != nullptr)
+		{
+			if (ParticleIndex == 0)
+			{
+				return Constraint->GetParticle0();
+			}
+			else
+			{
+				return Constraint->GetParticle1();
+			}
+		}
+		return nullptr;
+	}
+
+	const FGeometryParticleHandle* FCCDModifier::GetOtherParticle(const FGeometryParticleHandle* InParticle) const
+	{
+		if (Constraint != nullptr)
+		{
+			const FGeometryParticleHandle* Particle0 = Constraint->GetParticle0();
+			const FGeometryParticleHandle* Particle1 = Constraint->GetParticle1();
+			if (InParticle == Particle0)
+			{
+				return Particle1;
+			}
+			else if (InParticle == Particle1)
+			{
+				return Particle0;
+			}
+		}
+		return nullptr;
+	}
+
+	void FCCDModifier::Enable()
+	{
+		if (Constraint != nullptr)
+		{
+			Constraint->SetModifierApplied();
+			Constraint->SetDisabled(false);
+		}
+	}
+
+	void FCCDModifier::Disable()
+	{
+		if (Constraint != nullptr)
+		{
+			Constraint->SetModifierApplied();
+			Constraint->SetDisabled(true);
+		}
+	}
+
+	void FCCDModifier::ConvertToProbe()
+	{
+		if (Constraint != nullptr)
+		{
+			Constraint->SetModifierApplied();
+			Constraint->SetIsProbe(true);
+		}
+	}
+
+	//////////////////////////////////////////////////////////////////////////////////////////////////////////////////////
+	//////////////////////////////////////////////////////////////////////////////////////////////////////////////////////
+	//////////////////////////////////////////////////////////////////////////////////////////////////////////////////////
+
+	FCCDModifierParticleIterator& FCCDModifierParticleIterator::MoveToBegin()
+	{ 
+		ConstraintIndex = -1;
+		MoveToNext();
+		return *this;
+	}
+	
+	FCCDModifierParticleIterator& FCCDModifierParticleIterator::MoveToEnd()
+	{ 
+		ConstraintIndex = Range->GetNumConstraints();
+		PairModifier.Reset();
+		return *this;
+	}
+	
+	FCCDModifierParticleIterator& FCCDModifierParticleIterator::MoveToNext()
+	{
+		check(ConstraintIndex < Range->GetNumConstraints());
+
+		++ConstraintIndex;
+		if (ConstraintIndex < Range->GetNumConstraints())
+		{
+			PairModifier = FCCDModifier(Range->GetAccessor(), Range->GetConstraint(ConstraintIndex));
+		}
+		else
+		{
+			PairModifier.Reset();
+		}
+		return *this;
+	}
+
+	//////////////////////////////////////////////////////////////////////////////////////////////////////////////////////
+	//////////////////////////////////////////////////////////////////////////////////////////////////////////////////////
+	//////////////////////////////////////////////////////////////////////////////////////////////////////////////////////
+
+	FCCDModifierParticleRange::FCCDModifierParticleRange(
+		FCCDModifierAccessor* InAccessor, 
+		FGeometryParticleHandle* InParticle)
+		: Accessor(InAccessor)
+		, Particle(InParticle)
+	{
+		CollectConstraints();
+	}
+
+	FCCDModifierParticleIterator FCCDModifierParticleRange::begin()
+	{
+		return FCCDModifierParticleIterator(this).MoveToBegin();
+	}
+
+	FCCDModifierParticleIterator FCCDModifierParticleRange::end()
+	{
+		return FCCDModifierParticleIterator(this).MoveToEnd();
+	}
+
+	void FCCDModifierParticleRange::CollectConstraints()
+	{
+		Constraints.Reset();
+
+		Particle->ParticleCollisions().VisitCollisions(
+			[this](FPBDCollisionConstraint& Constraint) -> ECollisionVisitorResult
+			{
+				if (Constraint.GetCCDEnabled())
+				{
+					Constraints.Add(&Constraint);
+				}
+				return ECollisionVisitorResult::Continue;
+			});
+	}
+
+	//////////////////////////////////////////////////////////////////////////////////////////////////////////////////////
+	//////////////////////////////////////////////////////////////////////////////////////////////////////////////////////
+	//////////////////////////////////////////////////////////////////////////////////////////////////////////////////////
+
+	FCCDModifierAccessor::FCCDModifierAccessor(const FReal InDt)
+		: Dt(InDt)
+	{
+	}
+
+	FCCDModifierParticleRange FCCDModifierAccessor::GetModifiers(FGeometryParticleHandle* Particle)
+	{
+		return FCCDModifierParticleRange(this, Particle);
+	}
+}
diff --git a/Engine/Source/Runtime/Experimental/Chaos/Private/Chaos/CCDUtilities.cpp b/Engine/Source/Runtime/Experimental/Chaos/Private/Chaos/CCDUtilities.cpp
--- a/Engine/Source/Runtime/Experimental/Chaos/Private/Chaos/CCDUtilities.cpp
+++ b/Engine/Source/Runtime/Experimental/Chaos/Private/Chaos/CCDUtilities.cpp
@@ -85,7 +85,7 @@ namespace Chaos
 			const TPBDRigidParticleHandle<FReal, 3>* Rigid = Constraint->GetParticle(i)->CastToRigidParticle();
 			if (Rigid && Rigid->ObjectState() == EObjectStateType::Kinematic)
 			{
-				if (CCDHelpers::DeltaExceedsThreshold(Rigid->CCDAxisThreshold(), Displacements[i], Rigid->Q()))
+				if (FCCDHelpers::DeltaExceedsThreshold(Rigid->CCDAxisThreshold(), Displacements[i], Rigid->Q()))
 				{
 					return i;
 				}
@@ -201,7 +201,7 @@ namespace Chaos
 			// Determine if this particle pair should trigger CCD
 			const auto Particle0 = Constraint->GetParticle(0);
 			const auto Particle1 = Constraint->GetParticle(1);
-			bNeedCCDSolve = CCDHelpers::DeltaExceedsThreshold(*Particle0, *Particle1, Dt);
+			bNeedCCDSolve = FCCDHelpers::DeltaExceedsThreshold(*Particle0, *Particle1, Dt);
 
 			// make sure we ignore pairs that don't include any dynamics
 			if (CCDParticlePair[0] != nullptr || CCDParticlePair[1] != nullptr)
@@ -1076,13 +1076,34 @@ namespace Chaos
 		}
 	}
 
-	bool CCDHelpers::DeltaExceedsThreshold(const FVec3& AxisThreshold, const FVec3& DeltaX, const FQuat& R)
+
+	//////////////////////////////////////////////////////////////////////////////////////////////////////////////////////
+	//////////////////////////////////////////////////////////////////////////////////////////////////////////////////////
+	//////////////////////////////////////////////////////////////////////////////////////////////////////////////////////
+
+
+	FRigidTransform3 FCCDHelpers::GetParticleTransformAtTOI(const FGeometryParticleHandle* InParticle, const FReal TOI, const FReal Dt)
+	{
+		// Particles are swept using their latest rotation. We interpolate using velocity rather that Lerp(X,P,TOI)
+		// to handle the case where the center of mass is offset from the particle position.
+		const FConstGenericParticleHandle Particle = InParticle;
+		if (Particle->CCDEnabled())
+		{
+			const FReal InvClampedTOI = FMath::Clamp(FReal(1) - TOI, FReal(0), FReal(1));
+			return FRigidTransform3(
+				Particle->P() - (InvClampedTOI * Dt) * Particle->V(),
+				Particle->Q());
+		}
+		return Particle->GetTransformXR();
+	}
+
+	bool FCCDHelpers::DeltaExceedsThreshold(const FVec3& AxisThreshold, const FVec3& DeltaX, const FQuat& R)
 	{
 		FVec3 AbsLocalDelta, AxisThresholdScaled, AxisThresholdDiff;
 		return DeltaExceedsThreshold(AxisThreshold, DeltaX, R, AbsLocalDelta, AxisThresholdScaled, AxisThresholdDiff);
 	}
 
-	bool CCDHelpers::DeltaExceedsThreshold(const FVec3& AxisThreshold, const FVec3& DeltaX, const FQuat& R, FVec3& OutAbsLocalDelta, FVec3& OutAxisThresholdScaled, FVec3& OutAxisThresholdDiff)
+	bool FCCDHelpers::DeltaExceedsThreshold(const FVec3& AxisThreshold, const FVec3& DeltaX, const FQuat& R, FVec3& OutAbsLocalDelta, FVec3& OutAxisThresholdScaled, FVec3& OutAxisThresholdDiff)
 	{
 		if (CVars::CCDEnableThresholdBoundsScale < 0.f) { return false; }
 		if (CVars::CCDEnableThresholdBoundsScale == 0.f) { return true; }
@@ -1104,11 +1125,11 @@ namespace Chaos
 		return OutAxisThresholdDiff.GetMax() > 0.f;
 	}
 
-	bool CCDHelpers::DeltaExceedsThreshold(
+	bool FCCDHelpers::DeltaExceedsThreshold(
 		const FVec3& AxisThreshold0, const FVec3& DeltaX0, const FQuat& R0,
 		const FVec3& AxisThreshold1, const FVec3& DeltaX1, const FQuat& R1)
 	{
-		return CCDHelpers::DeltaExceedsThreshold(
+		return FCCDHelpers::DeltaExceedsThreshold(
 
 			// To combine axis thresholds:
 			// * transform particle1's threshold into particle0's local space
@@ -1134,7 +1155,7 @@ namespace Chaos
 			R0);
 	}
 
-	bool CCDHelpers::DeltaExceedsThreshold(const FGeometryParticleHandle& Particle0, const FGeometryParticleHandle& Particle1)
+	bool FCCDHelpers::DeltaExceedsThreshold(const FGeometryParticleHandle& Particle0, const FGeometryParticleHandle& Particle1)
 	{
 		// For rigids, compute DeltaX from the X - P diff and use Q for the rotation.
 		// For non-rigids, DeltaX is zero and use R for rotation.
@@ -1149,7 +1170,7 @@ namespace Chaos
 			Particle1.CCDAxisThreshold(), DeltaX1, R1);
 	}
 
-	bool CCDHelpers::DeltaExceedsThreshold(const FGeometryParticleHandle& Particle0, const FGeometryParticleHandle& Particle1, const FReal Dt)
+	bool FCCDHelpers::DeltaExceedsThreshold(const FGeometryParticleHandle& Particle0, const FGeometryParticleHandle& Particle1, const FReal Dt)
 	{
 		// For rigids, compute DeltaX from the V * Dt and use Q for the rotation.
 		// For non-rigids, DeltaX is zero and use R for rotation.
diff --git a/Engine/Source/Runtime/Experimental/Chaos/Private/Chaos/ChaosDebugDraw.cpp b/Engine/Source/Runtime/Experimental/Chaos/Private/Chaos/ChaosDebugDraw.cpp
--- a/Engine/Source/Runtime/Experimental/Chaos/Private/Chaos/ChaosDebugDraw.cpp
+++ b/Engine/Source/Runtime/Experimental/Chaos/Private/Chaos/ChaosDebugDraw.cpp
@@ -1853,7 +1853,7 @@ namespace Chaos
 			// Call the version of DeltaExceedsThreshold that provides access to its
 			// intermediate variables so we can debug draw them
 			FVec3 AbsLocalDelta, AxisThresholdScaled, AxisThresholdDiff;
-			const bool bEnableCCD = CCDHelpers::DeltaExceedsThreshold(
+			const bool bEnableCCD = FCCDHelpers::DeltaExceedsThreshold(
 				AxisThreshold, DeltaX, R, AbsLocalDelta, AxisThresholdScaled, AxisThresholdDiff);
 			const FVector AxisPercents = AbsLocalDelta / AxisThresholdScaled;
 
diff --git a/Engine/Source/Runtime/Experimental/Chaos/Private/Chaos/CollisionResolution.cpp b/Engine/Source/Runtime/Experimental/Chaos/Private/Chaos/CollisionResolution.cpp
--- a/Engine/Source/Runtime/Experimental/Chaos/Private/Chaos/CollisionResolution.cpp
+++ b/Engine/Source/Runtime/Experimental/Chaos/Private/Chaos/CollisionResolution.cpp
@@ -147,7 +147,7 @@ namespace Chaos
 				return false;
 			}
 
-			bool bUseCCD = CCDHelpers::DeltaExceedsThreshold(
+			bool bUseCCD = FCCDHelpers::DeltaExceedsThreshold(
 				Particle0->CCDAxisThreshold(), DeltaX0, FConstGenericParticleHandle(Particle0)->Q(),
 				Particle1->CCDAxisThreshold(), DeltaX1, FConstGenericParticleHandle(Particle1)->Q());
 
diff --git a/Engine/Source/Runtime/Experimental/Chaos/Private/Chaos/PBDCollisionConstraints.cpp b/Engine/Source/Runtime/Experimental/Chaos/Private/Chaos/PBDCollisionConstraints.cpp
--- a/Engine/Source/Runtime/Experimental/Chaos/Private/Chaos/PBDCollisionConstraints.cpp
+++ b/Engine/Source/Runtime/Experimental/Chaos/Private/Chaos/PBDCollisionConstraints.cpp
@@ -6,6 +6,7 @@
 #include "Chaos/ContactModification.h"
 #include "Chaos/MidPhaseModification.h"
 #include "Chaos/PBDCollisionConstraintsContact.h"
+#include "Chaos/CCDModification.h"
 #include "Chaos/CollisionResolution.h"
 #include "Chaos/Collision/CollisionPruning.h"
 #include "Chaos/Collision/PBDCollisionContainerSolver.h"
@@ -597,6 +598,15 @@ namespace Chaos
 		}
 	}
 
+	void FPBDCollisionConstraints::ApplyCCDModifier(const TArray<ISimCallbackObject*>& CCDModifiers, FReal Dt)
+	{
+		FCCDModifierAccessor ModifierAccessor(Dt);
+		for (ISimCallbackObject* ModifierCallback : CCDModifiers)
+		{
+			ModifierCallback->CCDModification_Internal(ModifierAccessor);
+		}
+	}
+
 	void FPBDCollisionConstraints::ApplyCollisionModifier(const TArray<ISimCallbackObject*>& CollisionModifiers, FReal Dt)
 	{
 		if (GetConstraints().Num() > 0)
diff --git a/Engine/Source/Runtime/Experimental/Chaos/Private/Chaos/PBDRigidsEvolutionGBF.cpp b/Engine/Source/Runtime/Experimental/Chaos/Private/Chaos/PBDRigidsEvolutionGBF.cpp
--- a/Engine/Source/Runtime/Experimental/Chaos/Private/Chaos/PBDRigidsEvolutionGBF.cpp
+++ b/Engine/Source/Runtime/Experimental/Chaos/Private/Chaos/PBDRigidsEvolutionGBF.cpp
@@ -110,11 +110,21 @@ namespace Chaos
 		FRealSingle ChaosSolverMaxInvInertiaComponentRatio = 0;
 		FAutoConsoleVariableRef  CVarChaosSolverInertiaConditioningMaxInvInertiaComponentRatio(TEXT("p.Chaos.Solver.InertiaConditioning.MaxInvInertiaComponentRatio"), ChaosSolverMaxInvInertiaComponentRatio, TEXT("An input to inertia conditioning system. The largest inertia component must be at least least multiple of the smallest component"));
 
+		// Collision modifiers are run after the CCD rewind and manifold regeneration because the CCD "contact" data is only used to rewind the particle 
+		// and we regenerate the manifold at that new position. Running the modifier callback before CCD rewind would mean that the user does not have 
+		// data related to the actual contacts that will be resolved, and the positions and normals may be misleading. 
+		// This cvar allows use to run the modifiers before CCD rewind which is how it used to be, just in case.
+		// The new way to disable CCD contacts (pre-rewind) is to do it in a MidPhase- or CCD-Modifier callbacks. See FCCDModifier.
+		bool bChaosCollisionModiferBeforeCCD = false;
+		FAutoConsoleVariableRef  CVarChaosCollisionModiferBeforeCCD(TEXT("p.Chaos.Solver.CollisionModifiersBeforeCCD"), bChaosCollisionModiferBeforeCCD, TEXT("True: run the collision modifiers before CCD rewind is applied; False(default): run modifiers after CCD rewind. See comments in code."));
+
+
 		DECLARE_CYCLE_STAT(TEXT("FPBDRigidsEvolutionGBF::AdvanceOneTimeStep"), STAT_Evolution_AdvanceOneTimeStep, STATGROUP_Chaos);
 		DECLARE_CYCLE_STAT(TEXT("FPBDRigidsEvolutionGBF::UnclusterUnions"), STAT_Evolution_UnclusterUnions, STATGROUP_Chaos);
 		DECLARE_CYCLE_STAT(TEXT("FPBDRigidsEvolutionGBF::Integrate"), STAT_Evolution_Integrate, STATGROUP_Chaos);
 		DECLARE_CYCLE_STAT(TEXT("FPBDRigidsEvolutionGBF::KinematicTargets"), STAT_Evolution_KinematicTargets, STATGROUP_Chaos);
 		DECLARE_CYCLE_STAT(TEXT("FPBDRigidsEvolutionGBF::PostIntegrateCallback"), STAT_Evolution_PostIntegrateCallback, STATGROUP_Chaos);
+		DECLARE_CYCLE_STAT(TEXT("FPBDRigidsEvolutionGBF::CCDModifierCallback"), STAT_Evolution_CCDModifierCallback, STATGROUP_Chaos);
 		DECLARE_CYCLE_STAT(TEXT("FPBDRigidsEvolutionGBF::CollisionModifierCallback"), STAT_Evolution_CollisionModifierCallback, STATGROUP_Chaos);
 		DECLARE_CYCLE_STAT(TEXT("FPBDRigidsEvolutionGBF::MidPhaseModifierCallback"), STAT_Evolution_MidPhaseModifierCallback, STATGROUP_Chaos);
 		DECLARE_CYCLE_STAT(TEXT("FPBDRigidsEvolutionGBF::CCD"), STAT_Evolution_CCD, STATGROUP_Chaos);
@@ -425,7 +435,13 @@ void FPBDRigidsEvolutionGBF::AdvanceOneTimeStepImpl(const FReal Dt, const FSubSt
 		TransferJointConstraintCollisions();
 	}
 
-	if(CollisionModifiers)
+	if (CCDModifiers)
+	{
+		SCOPE_CYCLE_COUNTER(STAT_Evolution_CCDModifierCallback);
+		CollisionConstraints.ApplyCCDModifier(*CCDModifiers, Dt);
+	}
+
+	if(CollisionModifiers && CVars::bChaosCollisionModiferBeforeCCD)
 	{
 		SCOPE_CYCLE_COUNTER(STAT_Evolution_CollisionModifierCallback);
 		CollisionConstraints.ApplyCollisionModifier(*CollisionModifiers, Dt);
@@ -442,6 +458,12 @@ void FPBDRigidsEvolutionGBF::AdvanceOneTimeStepImpl(const FReal Dt, const FSubSt
 		CCDManager.ApplyConstraintsPhaseCCD(Dt, &CollisionConstraints.GetConstraintAllocator(), Particles.GetActiveParticlesView().Num());
 	}
 
+	if (CollisionModifiers && !CVars::bChaosCollisionModiferBeforeCCD)
+	{
+		SCOPE_CYCLE_COUNTER(STAT_Evolution_CollisionModifierCallback);
+		CollisionConstraints.ApplyCollisionModifier(*CollisionModifiers, Dt);
+	}
+
 	{
 		SCOPE_CYCLE_COUNTER(STAT_Evolution_CreateConstraintGraph);
 		CSV_SCOPED_TIMING_STAT(PhysicsVerbose, StepSolver_CreateConstraintGraph);
@@ -672,7 +694,13 @@ void FPBDRigidsEvolutionGBF::TestModeResetCollisions()
 	}
 }
 
-FPBDRigidsEvolutionGBF::FPBDRigidsEvolutionGBF(FPBDRigidsSOAs& InParticles, THandleArray<FChaosPhysicsMaterial>& SolverPhysicsMaterials, const TArray<ISimCallbackObject*>* InMidPhaseModifiers, const TArray<ISimCallbackObject*>* InCollisionModifiers, bool InIsSingleThreaded)
+FPBDRigidsEvolutionGBF::FPBDRigidsEvolutionGBF(
+	FPBDRigidsSOAs& InParticles, 
+	THandleArray<FChaosPhysicsMaterial>& SolverPhysicsMaterials, 
+	const TArray<ISimCallbackObject*>* InMidPhaseModifiers,
+	const TArray<ISimCallbackObject*>* InCCDModifiers,
+	const TArray<ISimCallbackObject*>* InCollisionModifiers,
+	bool InIsSingleThreaded)
 	: Base(InParticles, SolverPhysicsMaterials, InIsSingleThreaded)
 	, Clustering(*this, Particles.GetClusteredParticles())
 	, CollisionConstraints(InParticles, Collided, PhysicsMaterials, PerParticlePhysicsMaterials, &SolverPhysicsMaterials, CalculateNumCollisionsPerBlock(), DefaultRestitutionThreshold)
@@ -682,6 +710,7 @@ FPBDRigidsEvolutionGBF::FPBDRigidsEvolutionGBF(FPBDRigidsSOAs& InParticles, THan
 	, PreApplyCallback(nullptr)
 	, CurrentStepResimCacheImp(nullptr)
 	, MidPhaseModifiers(InMidPhaseModifiers)
+	, CCDModifiers(InCCDModifiers)
 	, CollisionModifiers(InCollisionModifiers)
 	, CCDManager()
 	, bIsDeterministic(false)
diff --git a/Engine/Source/Runtime/Experimental/Chaos/Private/PBDRigidsSolver.cpp b/Engine/Source/Runtime/Experimental/Chaos/Private/PBDRigidsSolver.cpp
--- a/Engine/Source/Runtime/Experimental/Chaos/Private/PBDRigidsSolver.cpp
+++ b/Engine/Source/Runtime/Experimental/Chaos/Private/PBDRigidsSolver.cpp
@@ -494,7 +494,7 @@ namespace Chaos
 		, FloorHeight(0.f)
 		, bIsDeterministic(false)
 		, Particles(UniqueIndices)
-		, MEvolution(new FPBDRigidsEvolution(Particles, SimMaterials, &MidPhaseModifiers, &ContactModifiers, BufferingModeIn == Chaos::EMultiBufferMode::Single))
+		, MEvolution(new FPBDRigidsEvolution(Particles, SimMaterials, &MidPhaseModifiers, &CCDModifiers, &ContactModifiers, BufferingModeIn == Chaos::EMultiBufferMode::Single))
 		, MEventManager(new FEventManager(BufferingModeIn))
 		, MSolverEventFilters(new FSolverEventFilters())
 		, MDirtyParticlesBuffer(new FDirtyParticlesBuffer(BufferingModeIn, BufferingModeIn == Chaos::EMultiBufferMode::Single))
@@ -822,7 +822,7 @@ namespace Chaos
 		SetMaxDeltaTime_External(1.0f);
 		SetMinDeltaTime_External(UE_SMALL_NUMBER);
 		SetMaxSubSteps_External(1);
-		MEvolution = TUniquePtr<FPBDRigidsEvolution>(new FPBDRigidsEvolution(Particles, SimMaterials, &MidPhaseModifiers, &ContactModifiers, BufferMode == EMultiBufferMode::Single)); 
+		MEvolution = TUniquePtr<FPBDRigidsEvolution>(new FPBDRigidsEvolution(Particles, SimMaterials, &MidPhaseModifiers, &CCDModifiers, &ContactModifiers, BufferMode == EMultiBufferMode::Single)); 
 
 		PerSolverField = MakeUnique<FPerSolverFieldSystem>();
 
@@ -1364,6 +1364,11 @@ namespace Chaos
 				MidPhaseModifiers.Add(SimCallbackObject);
 			}
 
+			if (SimCallbackObject->HasOption(ESimCallbackOptions::CCDModification))
+			{
+				CCDModifiers.Add(SimCallbackObject);
+			}
+
 			if (SimCallbackObject->HasOption(ESimCallbackOptions::ContactModification))
 			{
 				ContactModifiers.Add(SimCallbackObject);
diff --git a/Engine/Source/Runtime/Experimental/Chaos/Public/Chaos/CCDModification.h b/Engine/Source/Runtime/Experimental/Chaos/Public/Chaos/CCDModification.h
new file mode 100644
--- /dev/null
+++ b/Engine/Source/Runtime/Experimental/Chaos/Public/Chaos/CCDModification.h
@@ -0,0 +1,196 @@
+// Copyright Epic Games, Inc. All Rights Reserved.
+#pragma once
+
+#include "Chaos/Core.h"
+#include "ParticleHandleFwd.h"
+
+namespace Chaos
+{
+	class FCCDModifier;
+	class FCCDModifierAccessor;
+	class FCCDModifierParticleIterator;
+	class FCCDModifierParticleRange;
+	class FPBDCollisionConstraint;
+
+	/*
+	 * Allows the user to modify the results of CCD collision detection prior to the CCD rewind being
+	 * applied. This is in addition to the midphase modification phase which happens before collision
+	 * detection has been run, and the contact modifier phase which happens after we have rewound
+	 * 
+	 * The CCD "contact" details reported here do not necessarily represent the contact details
+	 * that will be used in the contact solving phase. The CCD data provides raw first-touch data from
+	 * a swept collision detection test, but we will be rebuilding a full contact manifold at the 
+	 * post-CCD position. The CCD normal and position here are not used again. For example, it is not
+	 * a good idea to use the CCD data to categorize collisions as wall or floor contacts. If you need
+	 * to know the actual contact positions and normals that will be used in the contact reoslution
+	 * phase, then see FCollisionContactModifier.
+	 * 
+	 * @see FMidPhaseModifier, FContactPairModifier
+	 */
+	class CHAOS_API FCCDModifier
+	{
+	public:
+		FCCDModifier()
+			: Constraint(nullptr)
+		{}
+
+		bool IsValid() const
+		{
+			return Constraint != nullptr;
+		}
+
+		//
+		// Accessor functions
+		//
+
+		// Whether the two particle actually hit each other in the sweep
+		bool IsSweepHit() const;
+
+		// Get the time of impact. This is as a fraction of total movement and will be in the range [0,1] if a hit was found
+		FReal GetSweepHitTOI() const;
+
+		// Get the sweep impact position for the specified body
+		// NOTE: See class comments.
+		FVec3 GetWorldSweepHitLocation(const int32 ParticleIndex) const;
+
+		// Get the sweep impact normal
+		// NOTE: See class comments.
+		FVec3 GetWorldSweepHitNormal() const;
+
+		const FGeometryParticleHandle* GetParticle(const int32 ParticleIndex) const;
+		const FGeometryParticleHandle* GetOtherParticle(const FGeometryParticleHandle* InParticle) const;
+
+		//
+		// Modifying functions
+		//
+
+		// Re-enable this contact (if it was previously disabled)
+		void Enable();
+
+		// Disable this contact
+		void Disable();
+
+		// Convert this contact to a probe
+		void ConvertToProbe();
+
+	private:
+		FCCDModifier(
+			FCCDModifierAccessor* InAccessor,
+			FPBDCollisionConstraint* InConstraint)
+			: Accessor(InAccessor)
+			, Constraint(InConstraint)
+		{}
+
+		void Reset()
+		{
+			Accessor = nullptr;
+			Constraint = nullptr;
+		}
+
+		FCCDModifierAccessor* Accessor;
+		FPBDCollisionConstraint* Constraint;
+
+		friend class FCCDModifierParticleIterator;
+	};
+
+	/*
+	 * Class for iterating over ccd results involving a specific particle
+	 */
+	class CHAOS_API FCCDModifierParticleIterator
+	{
+	public:
+		FCCDModifier& operator*()
+		{
+			return PairModifier;
+		}
+
+		FCCDModifier* operator->()
+		{
+			return &PairModifier;
+		}
+
+		FCCDModifierParticleIterator& operator++()
+		{
+			return MoveToNext();
+		}
+
+		explicit operator bool() const
+		{
+			return PairModifier.IsValid();
+		}
+
+		friend bool operator==(const FCCDModifierParticleIterator& L, const FCCDModifierParticleIterator& R)
+		{
+			return (L.Range == R.Range) && (L.ConstraintIndex == R.ConstraintIndex);
+		}
+
+		friend bool operator!=(const FCCDModifierParticleIterator& L, const FCCDModifierParticleIterator& R)
+		{
+			return !(L == R);
+		}
+
+	private:
+		FCCDModifierParticleIterator(FCCDModifierParticleRange* InRange)
+			: Range(InRange)
+			, ConstraintIndex(INDEX_NONE)
+		{ 
+		}
+
+		FCCDModifierParticleIterator& MoveToBegin();
+		FCCDModifierParticleIterator& MoveToEnd();
+		FCCDModifierParticleIterator& MoveToNext();
+
+		FCCDModifierParticleRange* Range;
+		int32 ConstraintIndex;
+		FCCDModifier PairModifier;
+
+		friend class FCCDModifierParticleRange;
+	};
+
+	/*
+	 * Represents all the swept constraints on a particle to allow for iteration.
+	 */
+	class CHAOS_API FCCDModifierParticleRange
+	{
+	public:
+		FCCDModifierParticleRange(
+			FCCDModifierAccessor* InAccessor, 
+			FGeometryParticleHandle* InParticle);
+
+		FCCDModifierParticleIterator begin();
+		FCCDModifierParticleIterator end();
+
+	private:
+		FCCDModifierAccessor* GetAccessor() const { return Accessor; }
+		FGeometryParticleHandle* GetParticle() const { return Particle; }
+		int32 GetNumConstraints() const { return Constraints.Num(); }
+		FPBDCollisionConstraint* GetConstraint(const int32 Index) const { return Constraints[Index]; }
+		void CollectConstraints();
+
+		FCCDModifierAccessor* Accessor;
+		FGeometryParticleHandle* Particle;
+		TArray<FPBDCollisionConstraint*> Constraints;
+
+		friend class FCCDModifierParticleIterator;
+	};
+
+	/*
+	 * Provides access to CCD modifiers
+	 */
+	class CHAOS_API FCCDModifierAccessor
+	{
+	public:
+		FCCDModifierAccessor(const FReal InDt);
+
+		// Get an object which allows for range iteration over the CCD modifiers for a particle
+		FCCDModifierParticleRange GetModifiers(FGeometryParticleHandle* Particle);
+
+		FReal GetDt() const
+		{
+			return Dt;
+		}
+
+	private:
+		FReal Dt;
+	};
+}
diff --git a/Engine/Source/Runtime/Experimental/Chaos/Public/Chaos/CCDUtilities.h b/Engine/Source/Runtime/Experimental/Chaos/Public/Chaos/CCDUtilities.h
--- a/Engine/Source/Runtime/Experimental/Chaos/Public/Chaos/CCDUtilities.h
+++ b/Engine/Source/Runtime/Experimental/Chaos/Public/Chaos/CCDUtilities.h
@@ -122,8 +122,10 @@ namespace Chaos
 		TArray<int32> IslandConstraintStart, IslandConstraintNum, IslandConstraintEnd; // The duplication of information in IslandConstraintNum and IslandConstraintEnd is needed when we group constraints using islands.
 	};
 
-	struct CCDHelpers
+	struct FCCDHelpers
 	{
+		static FRigidTransform3 GetParticleTransformAtTOI(const FGeometryParticleHandle* Particle, const FReal TOI, const FReal Dt);
+
 		// Return true if DeltaX indicates a movement beyond a set of threshold
 		// distances for each local axis.
 		static bool DeltaExceedsThreshold(const FVec3& AxisThreshold, const FVec3& DeltaX, const FQuat& R);
@@ -146,4 +148,7 @@ namespace Chaos
 		// R: For rigid particles, use Q. For non-rigids, use R.
 		static bool DeltaExceedsThreshold(const FGeometryParticleHandle& Particle0, const FGeometryParticleHandle& Particle1, const FReal Dt);
 	};
+
+
+	using CCDHelpers UE_DEPRECATED(5.3, "Renamed to FCCDHelpers to meet naming conventions") = FCCDHelpers;
 }
diff --git a/Engine/Source/Runtime/Experimental/Chaos/Public/Chaos/Collision/ParticleCollisions.h b/Engine/Source/Runtime/Experimental/Chaos/Public/Chaos/Collision/ParticleCollisions.h
--- a/Engine/Source/Runtime/Experimental/Chaos/Public/Chaos/Collision/ParticleCollisions.h
+++ b/Engine/Source/Runtime/Experimental/Chaos/Public/Chaos/Collision/ParticleCollisions.h
@@ -107,7 +107,7 @@ namespace Chaos
 
 		/**
 		 * @brief Visit all the collisions on this particle
-		 * @tparam TLambda visitor type with signature void(FPBDCollisionParticle&)
+		 * @tparam TLambda visitor type with signature void(FPBDCollisionConstraint&)
 		 * 
 		 * @note do not delete constraint from the lambda. You may disable them though.
 		*/
@@ -116,7 +116,7 @@ namespace Chaos
 
 		/**
 		 * @brief Visit all the collisions on this particle
-		 * @tparam TLambda visitor type with signature void(const FPBDCollisionParticle&)
+		 * @tparam TLambda visitor type with signature void(const FPBDCollisionConstraint&)
 		 * 
 		 * @note do not delete constraint from the lambda. You may disable them though.
 		*/
diff --git a/Engine/Source/Runtime/Experimental/Chaos/Public/Chaos/Framework/PhysicsSolverBase.h b/Engine/Source/Runtime/Experimental/Chaos/Public/Chaos/Framework/PhysicsSolverBase.h
--- a/Engine/Source/Runtime/Experimental/Chaos/Public/Chaos/Framework/PhysicsSolverBase.h
+++ b/Engine/Source/Runtime/Experimental/Chaos/Public/Chaos/Framework/PhysicsSolverBase.h
@@ -616,6 +616,7 @@ namespace Chaos
 
 	TArray<ISimCallbackObject*> SimCallbackObjects;
 	TArray<ISimCallbackObject*> MidPhaseModifiers;
+	TArray<ISimCallbackObject*> CCDModifiers;
 	TArray<ISimCallbackObject*> ContactModifiers;
 	TArray<ISimCallbackObject*> RegistrationWatchers;
 	TArray<ISimCallbackObject*> UnregistrationWatchers;
diff --git a/Engine/Source/Runtime/Experimental/Chaos/Public/Chaos/MidPhaseModification.h b/Engine/Source/Runtime/Experimental/Chaos/Public/Chaos/MidPhaseModification.h
--- a/Engine/Source/Runtime/Experimental/Chaos/Public/Chaos/MidPhaseModification.h
+++ b/Engine/Source/Runtime/Experimental/Chaos/Public/Chaos/MidPhaseModification.h
@@ -147,9 +147,6 @@ namespace Chaos
 	class CHAOS_API FMidPhaseModifierAccessor
 	{
 	public:
-		friend class FMidPhaseModifier;
-		friend class FMidPhasePairModifierIterator;
-
 		// Get an object which allows for range iteration over the list of
 		// midphases for a particle
 		FMidPhaseModifierParticleRange GetMidPhases(FGeometryParticleHandle* Particle);
@@ -157,6 +154,4 @@ namespace Chaos
 		// Get a midphase modifier for a particular object pair
 		FMidPhaseModifier GetMidPhase(FGeometryParticleHandle* Particle0, FGeometryParticleHandle* Particle1);
 	};
-
-	using FMidPhaseModifierCallback = TFunction<void(const FCollisionContactModifier&)>;
 }
diff --git a/Engine/Source/Runtime/Experimental/Chaos/Public/Chaos/PBDCollisionConstraints.h b/Engine/Source/Runtime/Experimental/Chaos/Public/Chaos/PBDCollisionConstraints.h
--- a/Engine/Source/Runtime/Experimental/Chaos/Public/Chaos/PBDCollisionConstraints.h
+++ b/Engine/Source/Runtime/Experimental/Chaos/Public/Chaos/PBDCollisionConstraints.h
@@ -120,6 +120,12 @@ public:
 	 */
 	void ApplyMidPhaseModifier(const TArray<ISimCallbackObject*>& MidPhaseModifiers, FReal Dt);
 
+	/**
+	 * Apply modifiers to CCD results
+	 */
+	void ApplyCCDModifier(const TArray<ISimCallbackObject*>& CCDModifiers, FReal Dt);
+
+
 	/**
 	 * Apply modifiers to the constraints and specify which constraints should be disabled.
 	 * You would probably call this in the PostComputeCallback. Prefer this to calling RemoveConstraints in a loop,
diff --git a/Engine/Source/Runtime/Experimental/Chaos/Public/Chaos/PBDRigidsEvolutionGBF.h b/Engine/Source/Runtime/Experimental/Chaos/Public/Chaos/PBDRigidsEvolutionGBF.h
--- a/Engine/Source/Runtime/Experimental/Chaos/Public/Chaos/PBDRigidsEvolutionGBF.h
+++ b/Engine/Source/Runtime/Experimental/Chaos/Public/Chaos/PBDRigidsEvolutionGBF.h
@@ -66,7 +66,13 @@ namespace Chaos
 		static constexpr FRealSingle DefaultCollisionMaxPushOutVelocity = 1000.0f;
 		static constexpr int32 DefaultRestitutionThreshold = 1000;
 
-		CHAOS_API FPBDRigidsEvolutionGBF(FPBDRigidsSOAs& InParticles, THandleArray<FChaosPhysicsMaterial>& SolverPhysicsMaterials, const TArray<ISimCallbackObject*>* InMidPhaseModifiers = nullptr, const TArray<ISimCallbackObject*>* InCollisionModifiers = nullptr, bool InIsSingleThreaded = false);
+		CHAOS_API FPBDRigidsEvolutionGBF(
+			FPBDRigidsSOAs& InParticles, 
+			THandleArray<FChaosPhysicsMaterial>& SolverPhysicsMaterials, 
+			const TArray<ISimCallbackObject*>* InMidPhaseModifiers = nullptr,
+			const TArray<ISimCallbackObject*>* InCCDModifiers = nullptr,
+			const TArray<ISimCallbackObject*>* InCollisionModifiers = nullptr,
+			bool InIsSingleThreaded = false);
 		CHAOS_API ~FPBDRigidsEvolutionGBF();
 
 		FORCEINLINE void SetPostIntegrateCallback(const FPBDRigidsEvolutionCallback& Cb)
@@ -283,7 +289,7 @@ namespace Chaos
 						}
 #endif
 
-						if (CCDHelpers::DeltaExceedsThreshold(Particle.CCDAxisThreshold(), Particle.P() - Particle.X(), Particle.Q()))
+						if (FCCDHelpers::DeltaExceedsThreshold(Particle.CCDAxisThreshold(), Particle.P() - Particle.X(), Particle.Q()))
 						{
 							// We sweep the bounds from P back along the velocity and expand by a small amount.
 							// If not using tight bounds we also expand the bounds in all directions by Velocity. This is necessary only for secondary CCD collisions
@@ -369,7 +375,10 @@ namespace Chaos
 		FPBDRigidsEvolutionCallback PreApplyCallback;
 		FPBDRigidsEvolutionInternalHandleCallback InternalParticleInitilization;
 		FEvolutionResimCache* CurrentStepResimCacheImp;
+
+		// @todo(chaos): evolution and collision constraints should not know about ISimCallbackObject. Fix this.
 		const TArray<ISimCallbackObject*>* MidPhaseModifiers;
+		const TArray<ISimCallbackObject*>* CCDModifiers;
 		const TArray<ISimCallbackObject*>* CollisionModifiers;
 
 		FCCDManager CCDManager;
diff --git a/Engine/Source/Runtime/Experimental/Chaos/Public/Chaos/SimCallbackObject.h b/Engine/Source/Runtime/Experimental/Chaos/Public/Chaos/SimCallbackObject.h
--- a/Engine/Source/Runtime/Experimental/Chaos/Public/Chaos/SimCallbackObject.h
+++ b/Engine/Source/Runtime/Experimental/Chaos/Public/Chaos/SimCallbackObject.h
@@ -11,6 +11,7 @@ namespace Chaos
 {
 class FPhysicsSolverBase;
 class FMidPhaseModifierAccessor;
+class FCCDModifierAccessor;
 class FCollisionContactModifier;
 class FSingleParticlePhysicsProxy;
 
@@ -23,11 +24,12 @@ enum class ESimCallbackOptions : uint8
 {
 	Presimulate				= 1 << 0,
 	MidPhaseModification	= 1 << 1,
-	ContactModification		= 1 << 2,
-	ParticleRegister		= 1 << 3,
-	ParticleUnregister		= 1 << 4,
-	RunOnFrozenGameThread	= 1 << 5,
-	Rewind					= 1 << 6
+	CCDModification			= 1 << 2,
+	ContactModification		= 1 << 3,
+	ParticleRegister		= 1 << 4,
+	ParticleUnregister		= 1 << 5,
+	RunOnFrozenGameThread	= 1 << 6,
+	Rewind					= 1 << 7
 };
 ENUM_CLASS_FLAGS(ESimCallbackOptions)
 
@@ -72,6 +74,11 @@ public:
 		OnMidPhaseModification_Internal(Modifier);
 	}
 
+	void CCDModification_Internal(FCCDModifierAccessor& Modifier)
+	{
+		OnCCDModification_Internal(Modifier);
+	}
+
 	void ContactModification_Internal(FCollisionContactModifier& Modifier)
 	{
 		OnContactModification_Internal(Modifier);
@@ -173,6 +180,17 @@ private:
 		check(false);
 	}
 
+
+	/**
+	* Called once per simulation step. Allows user to modify CCD results
+	*
+	* NOTE: you must explicitly request CCD modification when registering the callback for this to be called
+	*/
+	virtual void OnCCDModification_Internal(FCCDModifierAccessor& Modifier)
+	{
+		check(false);
+	}
+
 	/**
 	* Called once per simulation step. Allows user to modify contacts
 	*
