geometry collection : expose IsRootBroken function to blueprints

diff --git a/Engine/Source/Runtime/Experimental/GeometryCollectionEngine/Public/GeometryCollection/GeometryCollectionComponent.h b/Engine/Source/Runtime/Experimental/GeometryCollectionEngine/Public/GeometryCollection/GeometryCollectionComponent.h
--- a/Engine/Source/Runtime/Experimental/GeometryCollectionEngine/Public/GeometryCollection/GeometryCollectionComponent.h
+++ b/Engine/Source/Runtime/Experimental/GeometryCollectionEngine/Public/GeometryCollection/GeometryCollectionComponent.h
@@ -1239,6 +1239,7 @@ private:
 	void RefreshISMPoolInstances();
 
 	/** return true if the root cluster is not longer active at runtime */
+	UFUNCTION(BlueprintCallable, Category = "ChaosPhysics")
 	bool IsRootBroken() const;
 
 	void IncrementSleepTimer(float DeltaTime);
