geometry collection : add cvar to enable/disable collision particles on geometry collection

diff --git a/Engine/Source/Runtime/Experimental/Chaos/Private/GeometryCollection/GeometryCollection.cpp b/Engine/Source/Runtime/Experimental/Chaos/Private/GeometryCollection/GeometryCollection.cpp
--- a/Engine/Source/Runtime/Experimental/Chaos/Private/GeometryCollection/GeometryCollection.cpp
+++ b/Engine/Source/Runtime/Experimental/Chaos/Private/GeometryCollection/GeometryCollection.cpp
@@ -21,6 +21,12 @@
 #include "Chaos/ChaosArchive.h"
 #include "Voronoi/Voronoi.h"
 
+bool bChaosGeometryCollectionEnableCollisionParticles = true;
+FAutoConsoleVariableRef CVarChaosGeometryCollectionEnableCollisionParticles(
+	TEXT("p.Chaos.GC.EnableCollisionParticles"),
+	bChaosGeometryCollectionEnableCollisionParticles,
+	TEXT("Enable use of collision particles for collision [def:true]"));
+
 DEFINE_LOG_CATEGORY_STATIC(FGeometryCollectionLogging, Log, All);
 
 // @todo: update names 
@@ -35,6 +41,13 @@ const FName FGeometryCollection::SimulatableParticlesAttribute("SimulatableParti
 const FName FGeometryCollection::SimulationTypeAttribute("SimulationType");
 const FName FGeometryCollection::StatusFlagsAttribute("StatusFlags");
 
+
+bool FGeometryCollection::AreCollisionParticlesEnabled()
+{
+	return bChaosGeometryCollectionEnableCollisionParticles;
+}
+
+
 FGeometryCollection::FGeometryCollection()
 	: FTransformCollection()
 	, FGeometryCollectionConvexPropertiesInterface(this)
diff --git a/Engine/Source/Runtime/Experimental/Chaos/Private/PhysicsProxy/GeometryCollectionPhysicsProxy.cpp b/Engine/Source/Runtime/Experimental/Chaos/Private/PhysicsProxy/GeometryCollectionPhysicsProxy.cpp
--- a/Engine/Source/Runtime/Experimental/Chaos/Private/PhysicsProxy/GeometryCollectionPhysicsProxy.cpp
+++ b/Engine/Source/Runtime/Experimental/Chaos/Private/PhysicsProxy/GeometryCollectionPhysicsProxy.cpp
@@ -669,7 +669,8 @@ void FGeometryCollectionPhysicsProxy::Initialize(Chaos::FPBDRigidsEvolutionBase
 	// TODO: Ryan - Should we just transfer ownership of the SimplicialsAttribute from the DynamicCollection to
 	// the PhysicsThreadCollection?
 	{
-		if (DynamicCollection.HasAttribute(DynamicCollection.SimplicialsAttribute, FTransformCollection::TransformGroup))
+		if (FGeometryCollection::AreCollisionParticlesEnabled()
+			&& DynamicCollection.HasAttribute(DynamicCollection.SimplicialsAttribute, FTransformCollection::TransformGroup))
 		{
 			const auto& SourceSimplicials = DynamicCollection.GetAttribute<TUniquePtr<FSimplicial>>(
 				DynamicCollection.SimplicialsAttribute, FTransformCollection::TransformGroup);
@@ -747,7 +748,8 @@ void FGeometryCollectionPhysicsProxy::InitializeDynamicCollection(FGeometryDynam
 		static IConsoleVariable* AnalyticDisableCVar = IConsoleManager::Get().FindConsoleVariable(TEXT("p.IgnoreAnalyticCollisionsOverride"));
 		static const bool bAnalyticsDisabled = (AnalyticDisableCVar && AnalyticDisableCVar->GetBool());
 
-		if (RestCollection.HasAttribute(DynamicCollection.SimplicialsAttribute, FTransformCollection::TransformGroup)
+		if (FGeometryCollection::AreCollisionParticlesEnabled()
+			&& RestCollection.HasAttribute(DynamicCollection.SimplicialsAttribute, FTransformCollection::TransformGroup)
 			&& Params.Shared.SizeSpecificData[0].CollisionShapesData.Num()
 			&& (Params.Shared.SizeSpecificData[0].CollisionShapesData[0].CollisionType == ECollisionTypeEnum::Chaos_Surface_Volumetric || bAnalyticsDisabled))
 		{
diff --git a/Engine/Source/Runtime/Experimental/Chaos/Public/GeometryCollection/GeometryCollection.h b/Engine/Source/Runtime/Experimental/Chaos/Public/GeometryCollection/GeometryCollection.h
--- a/Engine/Source/Runtime/Experimental/Chaos/Public/GeometryCollection/GeometryCollection.h
+++ b/Engine/Source/Runtime/Experimental/Chaos/Public/GeometryCollection/GeometryCollection.h
@@ -122,6 +122,7 @@ public:
 
 	};
 
+	static bool AreCollisionParticlesEnabled();
 
 	/**
 	 * Create a GeometryCollection from Vertex and Indices arrays
diff --git a/Engine/Source/Runtime/Experimental/GeometryCollectionEngine/Private/GeometryCollection/GeometryCollectionObject.cpp b/Engine/Source/Runtime/Experimental/GeometryCollectionEngine/Private/GeometryCollection/GeometryCollectionObject.cpp
--- a/Engine/Source/Runtime/Experimental/GeometryCollectionEngine/Private/GeometryCollection/GeometryCollectionObject.cpp
+++ b/Engine/Source/Runtime/Experimental/GeometryCollectionEngine/Private/GeometryCollection/GeometryCollectionObject.cpp
@@ -772,6 +772,15 @@ void UGeometryCollection::Serialize(FArchive& Ar)
 				ArchiveGeometryCollection = CopyCollectionAndRemoveGeometry(GeometryCollection);
 			}
 		}
+		else
+		{
+			// do we need to remove the simplicial attribute ? 
+			if (false == FGeometryCollection::AreCollisionParticlesEnabled())
+			{
+				ArchiveGeometryCollection = TSharedPtr<FGeometryCollection, ESPMode::ThreadSafe>(GeometryCollection->NewCopy<FGeometryCollection>());
+				ArchiveGeometryCollection->RemoveAttribute(FGeometryDynamicCollection::SimplicialsAttribute, FTransformCollection::TransformGroup);
+			}
+		}
 
 		// The dataflow asset is only needed for the editor, so we just remove it when cooking 
 		StrippedDataflowAsset = DataflowAsset;
@@ -1267,14 +1276,17 @@ TSharedPtr<FGeometryCollection, ESPMode::ThreadSafe> UGeometryCollection::CopyCo
 
 	CollectionToCopy->CopyTo(GeometryCollectionToReturn.Get(), GroupsToSkip, AttributesToSkip);
 
-	// recreate the simplicial attribute since we cannot copy it and we skipped it 
-	using FSimplicialUniquePtr = TUniquePtr<FCollisionStructureManager::FSimplicial>;
-	if (const TManagedArray<FSimplicialUniquePtr>* SourceSimplicials = CollectionToCopy->FindAttribute<FSimplicialUniquePtr>(FGeometryDynamicCollection::SimplicialsAttribute, FTransformCollection::TransformGroup))
+	if (FGeometryCollection::AreCollisionParticlesEnabled())
 	{
-		TManagedArray<FSimplicialUniquePtr>& SimplicialsToWrite = GeometryCollectionToReturn->AddAttribute<FSimplicialUniquePtr>(FGeometryDynamicCollection::SimplicialsAttribute, FTransformCollection::TransformGroup);
-		for (int32 Index = SourceSimplicials->Num()-1; 0 <= Index; Index--)
+		// recreate the simplicial attribute since we cannot copy it and we skipped it 
+		using FSimplicialUniquePtr = TUniquePtr<FCollisionStructureManager::FSimplicial>;
+		if (const TManagedArray<FSimplicialUniquePtr>* SourceSimplicials = CollectionToCopy->FindAttribute<FSimplicialUniquePtr>(FGeometryDynamicCollection::SimplicialsAttribute, FTransformCollection::TransformGroup))
 		{
-			SimplicialsToWrite[Index].Reset((*SourceSimplicials)[Index] ? (*SourceSimplicials)[Index]->NewCopy() : nullptr);
+			TManagedArray<FSimplicialUniquePtr>& SimplicialsToWrite = GeometryCollectionToReturn->AddAttribute<FSimplicialUniquePtr>(FGeometryDynamicCollection::SimplicialsAttribute, FTransformCollection::TransformGroup);
+			for (int32 Index = SourceSimplicials->Num() - 1; 0 <= Index; Index--)
+			{
+				SimplicialsToWrite[Index].Reset((*SourceSimplicials)[Index] ? (*SourceSimplicials)[Index]->NewCopy() : nullptr);
+			}
 		}
 	}
 
