Remove unnecessary copies in Update Geometry in Cluster Algo

diff --git a/Engine/Source/Runtime/Experimental/Chaos/Private/Chaos/PBDRigidClusteringAlgo.cpp b/Engine/Source/Runtime/Experimental/Chaos/Private/Chaos/PBDRigidClusteringAlgo.cpp
--- a/Engine/Source/Runtime/Experimental/Chaos/Private/Chaos/PBDRigidClusteringAlgo.cpp
+++ b/Engine/Source/Runtime/Experimental/Chaos/Private/Chaos/PBDRigidClusteringAlgo.cpp
@@ -408,7 +408,17 @@ namespace Chaos
 			};
 			//ensureMsgf(false, TEXT("Checking usage with proxy"));
 			//@coverage {production}
-			Parent->SetSharedGeometry(TSharedPtr<FImplicitObject, ESPMode::ThreadSafe>(DeepCopyImplicit(ProxyGeometry).Release()));
+
+			Chaos::EImplicitObjectType GeometryType = ProxyGeometry->GetType();
+			// Don't copy if it is not a level set and scale is one
+			if (GeometryType != Chaos::ImplicitObjectType::LevelSet && Scale.Equals(FVector::OneVector))
+			{
+				Parent->SetSharedGeometry(TSharedPtr<FImplicitObject, ESPMode::ThreadSafe>(ProxyGeometry));
+			}
+			else
+			{
+				Parent->SetSharedGeometry(TSharedPtr<FImplicitObject, ESPMode::ThreadSafe>(DeepCopyImplicit(ProxyGeometry).Release()));
+			}
 		}
 		else if (Objects.Num() == 0)
 		{
