Dataflow:
- Added VertexNormal, VertexColor to CollectionRenderingFacade
- they get initialized in FRenderingFacade::AddSurface()
- updated FGeometryCollection::SetDefaults() to initialize the color attr in the Vertices group

diff --git a/Engine/Plugins/Experimental/Dataflow/Source/DataflowEnginePlugin/Private/Dataflow/DataflowEngineRendering.cpp b/Engine/Plugins/Experimental/Dataflow/Source/DataflowEnginePlugin/Private/Dataflow/DataflowEngineRendering.cpp
--- a/Engine/Plugins/Experimental/Dataflow/Source/DataflowEnginePlugin/Private/Dataflow/DataflowEngineRendering.cpp
+++ b/Engine/Plugins/Experimental/Dataflow/Source/DataflowEnginePlugin/Private/Dataflow/DataflowEngineRendering.cpp
@@ -90,20 +90,43 @@ namespace Dataflow
 							// Maybe these buffers should be shrunk, but there are unused vertices in the buffer. 
 							for (int i = 0; i < Visited.Num(); i++) if (!Visited[i]) Vertices[i] = FVector3f(0);
 
-							RenderCollection.AddSurface(MoveTemp(Vertices), MoveTemp(Tris));
-							if (const TManagedArray<FLinearColor>* VertexColor = Collection.FindAttribute<FLinearColor>("Color", FGeometryCollection::VerticesGroup))
+							// Copy VertexNormals from the Collection if exists otherwise compute and set it
+							TArray<FVector3f> VertexNormals; VertexNormals.AddUninitialized(Vertex.Num());
+							if (const TManagedArray<FVector3f>* VertexNormal = Collection.FindAttribute<FVector3f>("Normal", FGeometryCollection::VerticesGroup))
 							{
-								TManagedArray<FLinearColor>& RenderCollectionVertexColor = RenderCollection.ModifyVertexColor();
-								for (int32 Idx = 0; Idx < Vertex.Num(); ++Idx)
+								for (int32 VertexIdx = 0; VertexIdx < VertexNormals.Num(); ++VertexIdx)
 								{
-									RenderCollectionVertexColor[Idx] = (*VertexColor)[Idx];
+									VertexNormals[VertexIdx] = (*VertexNormal)[VertexIdx];
 								}
 							}
 							else
 							{
-								TManagedArray<FLinearColor>& RenderCollectionVertexColor = RenderCollection.ModifyVertexColor();
-								RenderCollectionVertexColor.Fill(IDataflowEnginePlugin::SurfaceColor);
+								for (int32 VertexIdx = 0; VertexIdx < VertexNormals.Num(); ++VertexIdx)
+								{
+									// TODO: Compute the normal
+									VertexNormals[VertexIdx] = FVector3f(0.f);
+								}
+							}
+
+							// Copy VertexColors from the Collection if exists otherwise set it to IDataflowEnginePlugin::SurfaceColor
+							TArray<FLinearColor> VertexColors; VertexColors.AddUninitialized(Vertex.Num());
+							if (const TManagedArray<FLinearColor>* VertexColorManagedArray = Collection.FindAttribute<FLinearColor>("Color", FGeometryCollection::VerticesGroup))
+							{
+								for (int32 VertexIdx = 0; VertexIdx < VertexColors.Num(); ++VertexIdx)
+								{
+									VertexColors[VertexIdx] = (*VertexColorManagedArray)[VertexIdx];
+								}
 							}
+							else
+							{
+								for (int32 VertexIdx = 0; VertexIdx < VertexColors.Num(); ++VertexIdx)
+								{
+									VertexColors[VertexIdx] = FLinearColor(IDataflowEnginePlugin::SurfaceColor);
+								}
+							}
+
+							// Set the data on the RenderCollection
+							RenderCollection.AddSurface(MoveTemp(Vertices), MoveTemp(Tris), MoveTemp(VertexNormals), MoveTemp(VertexColors));
 						}
 					}
 				}
@@ -153,7 +176,17 @@ namespace Dataflow
 								Vertices[VertexID] = (FVector3f)DynamicMesh.GetVertex(VertexID);
 							}
 
-							RenderCollection.AddSurface(MoveTemp(Vertices), MoveTemp(Tris));
+							// Add VertexNormal and VertexColor
+							TArray<FVector3f> VertexNormals; VertexNormals.AddUninitialized(Vertices.Num());
+							TArray<FLinearColor> VertexColors; VertexColors.AddUninitialized(Vertices.Num());
+							for (int32 VertexIdx = 0; VertexIdx < VertexNormals.Num(); ++VertexIdx)
+							{
+								// TODO: Get the normal from FDynamicMesh3
+								VertexNormals[VertexIdx] = FVector3f(0.f);
+								VertexColors[VertexIdx] = FLinearColor(IDataflowEnginePlugin::SurfaceColor);
+							}
+
+							RenderCollection.AddSurface(MoveTemp(Vertices), MoveTemp(Tris), MoveTemp(VertexNormals), MoveTemp(VertexColors));
 						}
 					}
 				}
@@ -204,7 +237,18 @@ namespace Dataflow
 					Tris[8] = FIntVector(0, 3, 7); Tris[9] = FIntVector(4, 0, 7);
 					Tris[10] = FIntVector(5, 4, 7); Tris[11] = FIntVector(5, 7, 6);
 
-					RenderCollection.AddSurface(MoveTemp(Vertices), MoveTemp(Tris));
+					TArray<FVector3f> VertexNormals; VertexNormals.AddUninitialized(NumVertices);
+					// TODO: Compute vertex normals
+					
+					// Add VertexNormal and VertexColor
+					TArray<FLinearColor> VertexColors; VertexColors.AddUninitialized(Vertices.Num());
+					for (int32 VertexIdx = 0; VertexIdx < VertexNormals.Num(); ++VertexIdx)
+					{
+						VertexNormals[VertexIdx] = FVector3f(0.f);
+						VertexColors[VertexIdx] = FLinearColor(IDataflowEnginePlugin::SurfaceColor);
+					}
+
+					RenderCollection.AddSurface(MoveTemp(Vertices), MoveTemp(Tris), MoveTemp(VertexNormals), MoveTemp(VertexColors));
 				}
 			});
 
diff --git a/Engine/Source/Runtime/Experimental/Chaos/Private/GeometryCollection/Facades/CollectionRenderingFacade.cpp b/Engine/Source/Runtime/Experimental/Chaos/Private/GeometryCollection/Facades/CollectionRenderingFacade.cpp
--- a/Engine/Source/Runtime/Experimental/Chaos/Private/GeometryCollection/Facades/CollectionRenderingFacade.cpp
+++ b/Engine/Source/Runtime/Experimental/Chaos/Private/GeometryCollection/Facades/CollectionRenderingFacade.cpp
@@ -17,6 +17,8 @@ namespace GeometryCollection::Facades
 		, VertexToGeometryIndexAttribute(InCollection, "GeometryIndex", FGeometryCollection::VerticesGroup, FGeometryCollection::GeometryGroup)
 		, VertexSelectionAttribute(InCollection, "SelectionState", FGeometryCollection::VerticesGroup)
 		, VertexHitProxyIndexAttribute(InCollection, "HitIndex", FGeometryCollection::VerticesGroup)
+		, VertexNormalAttribute(InCollection, "Normal", FGeometryCollection::VerticesGroup)
+		, VertexColorAttribute(InCollection, "Color", FGeometryCollection::VerticesGroup)
 		, IndicesAttribute(InCollection, "Indices", FGeometryCollection::FacesGroup, FGeometryCollection::VerticesGroup)
 		, MaterialIDAttribute(InCollection, "MaterialID", FGeometryCollection::FacesGroup)
 		, TriangleSectionAttribute(InCollection, "Sections", FGeometryCollection::MaterialGroup)
@@ -27,7 +29,7 @@ namespace GeometryCollection::Facades
 		, IndicesStartAttribute(InCollection, "IndicesStart", FGeometryCollection::GeometryGroup)
 		, IndicesCountAttribute(InCollection, "IndicesCount", FGeometryCollection::GeometryGroup)
 		, GeometrySelectionAttribute(InCollection, "SelectionState", FGeometryCollection::GeometryGroup)
-		, VertexColorAttribute(InCollection, "Color", FGeometryCollection::VerticesGroup)
+
 	{}
 
 	FRenderingFacade::FRenderingFacade(const FManagedArrayCollection& InCollection)
@@ -37,6 +39,8 @@ namespace GeometryCollection::Facades
 		, VertexToGeometryIndexAttribute(InCollection, "GeometryIndex", FGeometryCollection::VerticesGroup, FGeometryCollection::GeometryGroup)
 		, VertexSelectionAttribute(InCollection, "SelectionState", FGeometryCollection::VerticesGroup)
 		, VertexHitProxyIndexAttribute(InCollection, "HitIndex", FGeometryCollection::VerticesGroup)
+		, VertexNormalAttribute(InCollection, "Normal", FGeometryCollection::VerticesGroup)
+		, VertexColorAttribute(InCollection, "Color", FGeometryCollection::VerticesGroup)
 		, IndicesAttribute(InCollection, "Indices", FGeometryCollection::FacesGroup, FGeometryCollection::VerticesGroup)
 		, MaterialIDAttribute(InCollection, "MaterialID", FGeometryCollection::FacesGroup)
 		, TriangleSectionAttribute(InCollection, "Sections", FGeometryCollection::MaterialGroup)
@@ -47,7 +51,6 @@ namespace GeometryCollection::Facades
 		, IndicesStartAttribute(InCollection, "IndicesStart", FGeometryCollection::GeometryGroup, FGeometryCollection::FacesGroup)
 		, IndicesCountAttribute(InCollection, "IndicesCount", FGeometryCollection::GeometryGroup)
 		, GeometrySelectionAttribute(InCollection, "SelectionState", FGeometryCollection::GeometryGroup)
-		, VertexColorAttribute(InCollection, "Color", FGeometryCollection::VerticesGroup)
 	{}
 
 	//
@@ -61,6 +64,8 @@ namespace GeometryCollection::Facades
 		VertexSelectionAttribute.Add();
 		VertexToGeometryIndexAttribute.Add();
 		VertexHitProxyIndexAttribute.Add();
+		VertexNormalAttribute.Add();
+		VertexColorAttribute.Add();
 		IndicesAttribute.Add();
 		MaterialIDAttribute.Add();
 		TriangleSectionAttribute.Add();
@@ -71,7 +76,6 @@ namespace GeometryCollection::Facades
 		IndicesStartAttribute.Add();
 		IndicesCountAttribute.Add();
 		GeometrySelectionAttribute.Add();
-		VertexColorAttribute.Add();
 	}
 
 	bool FRenderingFacade::CanRenderSurface( ) const
@@ -88,7 +92,8 @@ namespace GeometryCollection::Facades
 			GeometryNameAttribute.IsValid() && GeometryHitProxyIndexAttribute.IsValid() &&
 			VertexStartAttribute.IsValid() && VertexCountAttribute.IsValid() &&
 			IndicesStartAttribute.IsValid() && IndicesCountAttribute.IsValid() &&
-			GeometrySelectionAttribute.IsValid() && VertexColorAttribute.IsValid();
+			GeometrySelectionAttribute.IsValid() && VertexColorAttribute.IsValid() &&
+			VertexNormalAttribute.IsValid();
 	}
 
 	int32 FRenderingFacade::NumTriangles() const
@@ -123,7 +128,7 @@ namespace GeometryCollection::Facades
 	}
 
 
-	void FRenderingFacade::AddSurface(TArray<FVector3f>&& InVertices, TArray<FIntVector>&& InIndices)
+	void FRenderingFacade::AddSurface(TArray<FVector3f>&& InVertices, TArray<FIntVector>&& InIndices, TArray<FVector3f>&& InNormals, TArray<FLinearColor>&& InColors)
 	{
 		check(!IsConst());
 		if (IsValid())
@@ -148,6 +153,18 @@ namespace GeometryCollection::Facades
 
 			const FVector3f * VerticesDest = Vertices.GetData() + VertexStart;
 			FMemory::Memmove((void*)VerticesDest, InVertices.GetData(), sizeof(FVector3f) * InVertices.Num());
+
+			// Add VertexNormals
+			TManagedArray<FVector3f>& Normals = VertexNormalAttribute.Modify();
+
+			const FVector3f* VertexNormalsDest = Normals.GetData() + VertexStart;
+			FMemory::Memmove((void*)VertexNormalsDest, InNormals.GetData(), sizeof(FVector3f) * InNormals.Num());
+
+			// Add VertexColors
+			TManagedArray<FLinearColor>& VertexColors = VertexColorAttribute.Modify();
+
+			const FLinearColor* VertexColorsDest = VertexColors.GetData() + VertexStart;
+			FMemory::Memmove((void*)VertexColorsDest, InColors.GetData(), sizeof(FLinearColor) * InColors.Num());
 		}
 	}
 
diff --git a/Engine/Source/Runtime/Experimental/Chaos/Private/GeometryCollection/GeometryCollection.cpp b/Engine/Source/Runtime/Experimental/Chaos/Private/GeometryCollection/GeometryCollection.cpp
--- a/Engine/Source/Runtime/Experimental/Chaos/Private/GeometryCollection/GeometryCollection.cpp
+++ b/Engine/Source/Runtime/Experimental/Chaos/Private/GeometryCollection/GeometryCollection.cpp
@@ -145,7 +145,6 @@ void FGeometryCollection::Construct()
 	InitializeInterfaces();
 }
 
-
 void FGeometryCollection::SetDefaults(FName Group, uint32 StartSize, uint32 NumElements)
 {
 	if (Group == FTransformCollection::TransformGroup)
@@ -162,6 +161,13 @@ void FGeometryCollection::SetDefaults(FName Group, uint32 StartSize, uint32 NumE
 
 		FGeometryCollectionConvexUtility::SetDefaults(this, Group, StartSize, NumElements);
 	}
+	else if (Group == FGeometryCollection::VerticesGroup)
+	{
+		for (uint32 Idx = StartSize; Idx < StartSize + NumElements; ++Idx)
+		{
+			Color[Idx] = FLinearColor::White;
+		}
+	}
 }
 
 // MaterialIDOffset is based on the number of materials added by this append geometry call
diff --git a/Engine/Source/Runtime/Experimental/Chaos/Public/GeometryCollection/Facades/CollectionRenderingFacade.h b/Engine/Source/Runtime/Experimental/Chaos/Public/GeometryCollection/Facades/CollectionRenderingFacade.h
--- a/Engine/Source/Runtime/Experimental/Chaos/Public/GeometryCollection/Facades/CollectionRenderingFacade.h
+++ b/Engine/Source/Runtime/Experimental/Chaos/Public/GeometryCollection/Facades/CollectionRenderingFacade.h
@@ -53,7 +53,7 @@ namespace GeometryCollection::Facades
 		void AddTriangle(const Chaos::FTriangle& InTriangle);
 
 		/** Add a surface to the rendering view.*/
-		void AddSurface(TArray<FVector3f>&& InVertices, TArray<FIntVector>&& InIndices);
+		void AddSurface(TArray<FVector3f>&& InVertices, TArray<FIntVector>&& InIndices, TArray<FVector3f>&& InNormals, TArray<FLinearColor>&& InColors);
 		
 		/** GetIndices */
 		const TManagedArray< FIntVector >& GetIndices() const { return IndicesAttribute.Get(); }
@@ -142,6 +142,8 @@ namespace GeometryCollection::Facades
 		TManagedArrayAccessor<int32> VertexToGeometryIndexAttribute;
 		TManagedArrayAccessor<int32> VertexSelectionAttribute;
 		TManagedArrayAccessor<int32> VertexHitProxyIndexAttribute;
+		TManagedArrayAccessor<FVector3f> VertexNormalAttribute;
+		TManagedArrayAccessor<FLinearColor> VertexColorAttribute;
 
 		TManagedArrayAccessor<FIntVector> IndicesAttribute;
 		TManagedArrayAccessor<int32> MaterialIDAttribute;
@@ -154,9 +156,6 @@ namespace GeometryCollection::Facades
 		TManagedArrayAccessor<int32> IndicesStartAttribute;
 		TManagedArrayAccessor<int32> IndicesCountAttribute;
 		TManagedArrayAccessor<int32> GeometrySelectionAttribute;
-		TManagedArrayAccessor<FLinearColor> VertexColorAttribute;
-
-
 	};
 
 }
\ No newline at end of file
