Unshelved from pending changelist '23649281': Plus some nullptr fixes and includes latest PIE fixes.

Sequencer: Constraints: Handle duplication of contraints so it recreates proper constraint duplicate with the new duplicated objects and add preanimated state for component constraints. Put this all in one review since the fixes are interrelated.

For duplication, copy/paste of tracks, levels we need to not replace the ChildTRS since we can set that in the system evaluation, and it would resolve the original child trs to the new copy on paste which is not what we want, we want to keep the original.
Also added a PostDuplicate so we don't actually use a copy of the Constraint.
Fixed up places were would we get the constraint from the manager by it's FName, which may not be unique. Also make sure the Component Handles are owned by the constraint, not the manager since it duplication could return an existing handle if the parameters sent to DuplicateObject match.

diff --git a/Engine/Source/Runtime/Experimental/Animation/Constraints/Private/TransformConstraint.cpp b/Engine/Source/Runtime/Experimental/Animation/Constraints/Private/TransformConstraint.cpp
--- a/Engine/Source/Runtime/Experimental/Animation/Constraints/Private/TransformConstraint.cpp
+++ b/Engine/Source/Runtime/Experimental/Animation/Constraints/Private/TransformConstraint.cpp
@@ -243,13 +243,13 @@ UTickableConstraint* UTickableTransformConstraint::Duplicate(UObject* NewOuter)
 
 FString UTickableTransformConstraint::GetLabel() const
 {
-	if (!ChildTRSHandle->IsValid())
+	if (!ChildTRSHandle || !ChildTRSHandle->IsValid())
 	{
 		static const FString DummyLabel;
 		return DummyLabel;
 	}
 	
-	if (ParentTRSHandle->IsValid())
+	if (ParentTRSHandle && ParentTRSHandle->IsValid())
 	{
 		return FString::Printf(TEXT("%s.%s"), *ParentTRSHandle->GetLabel(), *ChildTRSHandle->GetLabel() );		
 	}
@@ -259,7 +259,7 @@ FString UTickableTransformConstraint::GetLabel() const
 
 FString UTickableTransformConstraint::GetFullLabel() const
 {
-	if (!ChildTRSHandle->IsValid())
+	if (!ChildTRSHandle || !ChildTRSHandle->IsValid())
 	{
 		static const FString DummyLabel;
 		return DummyLabel;
@@ -429,18 +429,18 @@ void UTickableTransformConstraint::PostDuplicate(bool bDuplicateForPIE)
 
 uint32 UTickableTransformConstraint::GetTargetHash() const
 {
-	return ChildTRSHandle->IsValid() ? ChildTRSHandle->GetHash() : 0;
+	return (ChildTRSHandle && ChildTRSHandle->IsValid()) ? ChildTRSHandle->GetHash() : 0;
 }
 
 bool UTickableTransformConstraint::ReferencesObject(TWeakObjectPtr<UObject> InObject) const
 {
-	const TWeakObjectPtr<UObject> ChildTarget = ChildTRSHandle->IsValid() ? ChildTRSHandle->GetTarget() : nullptr;
+	const TWeakObjectPtr<UObject> ChildTarget = (ChildTRSHandle && ChildTRSHandle->IsValid())? ChildTRSHandle->GetTarget() : nullptr;
 	if (ChildTarget == InObject)
 	{
 		return true;
 	}
 
-	const TWeakObjectPtr<UObject> ParentTarget = ParentTRSHandle->IsValid() ? ParentTRSHandle->GetTarget() : nullptr;
+	const TWeakObjectPtr<UObject> ParentTarget = (ParentTRSHandle && ParentTRSHandle->IsValid()) ? ParentTRSHandle->GetTarget() : nullptr;
 	if (ParentTarget == InObject)
 	{
 		return true;	
