Fix crash when simulating physics with a geometry collection with no rest collection assigned

diff --git a/Engine/Source/Runtime/Experimental/GeometryCollectionEngine/Private/GeometryCollection/GeometryCollectionComponent.cpp b/Engine/Source/Runtime/Experimental/GeometryCollectionEngine/Private/GeometryCollection/GeometryCollectionComponent.cpp
--- a/Engine/Source/Runtime/Experimental/GeometryCollectionEngine/Private/GeometryCollection/GeometryCollectionComponent.cpp
+++ b/Engine/Source/Runtime/Experimental/GeometryCollectionEngine/Private/GeometryCollection/GeometryCollectionComponent.cpp
@@ -850,7 +850,7 @@ void UGeometryCollectionComponent::SetSimulatePhysics(bool bEnabled)
 		BodyInstance.OwnerComponent = PreviousOwnerCOmponent;
 	}
 
-	if (bEnabled && !PhysicsProxy)
+	if (bEnabled && !PhysicsProxy && RestCollection)
 	{
 		RegisterAndInitializePhysicsProxy();
 	}
