Geometry collection : allow convex merging for cluster to use imported geometry
- support convex, box, sphere and capsule imported shapes ( as well as transformed , scaled and instanced )
- add display of imported collision in the convex tool if they are available and enabled

diff --git a/Engine/Plugins/Experimental/ChaosEditor/Source/FractureEditor/Private/FractureToolConvex.cpp b/Engine/Plugins/Experimental/ChaosEditor/Source/FractureEditor/Private/FractureToolConvex.cpp
--- a/Engine/Plugins/Experimental/ChaosEditor/Source/FractureEditor/Private/FractureToolConvex.cpp
+++ b/Engine/Plugins/Experimental/ChaosEditor/Source/FractureEditor/Private/FractureToolConvex.cpp
@@ -225,40 +225,58 @@ TArray<UObject*> UFractureToolConvex::GetSettingsObjects() const
 	return Settings;
 }
 
-void UFractureToolConvex::AddHullVisualizationData(FGeometryCollection& Collection, const FTransform& OuterTransform, int32 TransformIdx, bool bIsUnionOfHulls)
+void UFractureToolConvex::AddConvexHullVizualizationData(const Chaos::FConvex& ConvexHull, const FTransform& WorldSpaceTransform, bool bIsCustom, bool bIsUnionOfHulls)
 {
+	int32 HullPtsStart = HullPoints.Num();
+	for (Chaos::FVec3 Pt : ConvexHull.GetVertices())
+	{
+		HullPoints.Add(WorldSpaceTransform.TransformPosition(Pt));
+	}
+	int32 NumPlanes = ConvexHull.NumPlanes();
+	const Chaos::FConvexStructureData& HullData = ConvexHull.GetStructureData();
+	for (int PlaneIdx = 0; PlaneIdx < NumPlanes; PlaneIdx++)
+	{
+		int32 NumPlaneVerts = HullData.NumPlaneVertices(PlaneIdx);
+		for (int32 PlaneVertexIdx = 0; PlaneVertexIdx < NumPlaneVerts; PlaneVertexIdx++)
+		{
+			int32 V0 = HullPtsStart + HullData.GetPlaneVertex(PlaneIdx, PlaneVertexIdx);
+			int32 V1 = HullPtsStart + HullData.GetPlaneVertex(PlaneIdx, (PlaneVertexIdx + 1) % NumPlaneVerts);
+			HullEdges.Add({ V0, V1, bIsCustom, bIsUnionOfHulls });
+		}
+	}
+}
+
+void UFractureToolConvex::AddVisualizationData(FGeometryCollection& Collection, const FTransform& OuterTransform, int32 TransformIdx, bool bUseExternalCollision, bool bIsUnionOfHulls)
+{
+	using FImplicitGeom = TSharedPtr<Chaos::FImplicitObject, ESPMode::ThreadSafe>;
+
 	const TManagedArray<int32>* HasCustomConvex = FGeometryCollectionConvexUtility::GetCustomConvexFlags(&Collection, false);
-	const TManagedArray<TSet<int32>>& TransformToConvexIndices = Collection.GetAttribute<TSet<int32>>("TransformToConvexIndices", FTransformCollection::TransformGroup);
-	const TManagedArray<TUniquePtr<Chaos::FConvex>>& ConvexHull = Collection.GetAttribute<TUniquePtr<Chaos::FConvex>>("ConvexHull", "Convex");
+	const TManagedArray<TSet<int32>>* TransformToConvexIndices = Collection.FindAttribute<TSet<int32>>("TransformToConvexIndices", FTransformCollection::TransformGroup);
+	const TManagedArray<TUniquePtr<Chaos::FConvex>>* ConvexHull = Collection.FindAttribute<TUniquePtr<Chaos::FConvex>>("ConvexHull", "Convex");
+	const TManagedArray<FImplicitGeom>* ExternalCollisions = Collection.FindAttribute<FImplicitGeom>("ExternalCollisions", FGeometryCollection::TransformGroup);
 
 	const FTransform ComponentSpaceTransform = GeometryCollectionAlgo::GlobalMatrix(Collection.Transform, Collection.Parent, TransformIdx);
 	const FTransform WorldSpaceTransform = ComponentSpaceTransform * OuterTransform;
 
-	const TSet<int32>& ConvexIndices = TransformToConvexIndices[TransformIdx];
-	const bool bHasAnyConvex = (ConvexIndices.Num() > 0);
+	const bool bHasAnyConvex = TransformToConvexIndices && ConvexHull && ((*TransformToConvexIndices)[TransformIdx].Num() > 0);
 	const bool bIsCustom = bHasAnyConvex && (HasCustomConvex ? bool((*HasCustomConvex)[TransformIdx]) : false);
+	const bool bHasExternalCollision = ExternalCollisions && ((*ExternalCollisions)[TransformIdx].IsValid());
 
-	if (bHasAnyConvex)
+	// order of priority : external collision / generated convex / union of children shapes
+	if (bHasExternalCollision && bUseExternalCollision)
 	{
-		for (int32 ConvexIdx : TransformToConvexIndices[TransformIdx])
+		TArray<FGeometryCollectionConvexUtility::FTransformedConvex> ExternalHulls;
+		FGeometryCollectionConvexUtility::ConvertImplicitToConvexArray(*(*ExternalCollisions)[TransformIdx], WorldSpaceTransform, ExternalHulls);
+		for (const FGeometryCollectionConvexUtility::FTransformedConvex& ExternalHull : ExternalHulls)
 		{
-			int32 HullPtsStart = HullPoints.Num();
-			for (Chaos::FVec3 Pt : ConvexHull[ConvexIdx]->GetVertices())
-			{
-				HullPoints.Add(WorldSpaceTransform.TransformPosition(Pt));
-			}
-			int32 NumPlanes = ConvexHull[ConvexIdx]->NumPlanes();
-			const Chaos::FConvexStructureData& HullData = ConvexHull[ConvexIdx]->GetStructureData();
-			for (int PlaneIdx = 0; PlaneIdx < NumPlanes; PlaneIdx++)
-			{
-				int32 NumPlaneVerts = HullData.NumPlaneVertices(PlaneIdx);
-				for (int32 PlaneVertexIdx = 0; PlaneVertexIdx < NumPlaneVerts; PlaneVertexIdx++)
-				{
-					int32 V0 = HullPtsStart + HullData.GetPlaneVertex(PlaneIdx, PlaneVertexIdx);
-					int32 V1 = HullPtsStart + HullData.GetPlaneVertex(PlaneIdx, (PlaneVertexIdx + 1) % NumPlaneVerts);
-					HullEdges.Add({ V0, V1, bIsCustom, bIsUnionOfHulls });
-				}
-			}
+			AddConvexHullVizualizationData(*ExternalHull.Convex, ExternalHull.Transform, bIsCustom, bIsUnionOfHulls);
+		}
+	}
+	else if (bHasAnyConvex)
+	{
+		for (int32 ConvexIdx : (*TransformToConvexIndices)[TransformIdx])
+		{
+			AddConvexHullVizualizationData(*(*ConvexHull)[ConvexIdx], WorldSpaceTransform, bIsCustom, bIsUnionOfHulls);
 		}
 	}
 	else
@@ -266,7 +284,7 @@ void UFractureToolConvex::AddHullVisualizationData(FGeometryCollection& Collecti
 		// no convex generated, we are using a union of the children convexes
 		for (int32 ChildTransformIndex : Collection.Children[TransformIdx])
 		{
-			AddHullVisualizationData(Collection, OuterTransform, ChildTransformIndex, /* bIsUnionOfHulls */ true);
+			AddVisualizationData(Collection, OuterTransform, ChildTransformIndex, bUseExternalCollision, /* bIsUnionOfHulls */ true);
 		}
 	}
 }
@@ -279,26 +297,21 @@ void UFractureToolConvex::FractureContextChanged()
 
 	for (const FFractureToolContext& FractureContext : FractureContexts)
 	{
-		FGeometryCollection& Collection = *FractureContext.GetGeometryCollection();
-
-		if (!Collection.HasAttribute("ConvexHull", "Convex") ||
-			!Collection.HasAttribute("TransformToConvexIndices", FGeometryCollection::TransformGroup))
+		if (UGeometryCollection* RestCollection = FractureContext.GetFracturedGeometryCollection())
 		{
-			continue;
-		}
+			FGeometryCollection& Collection = *FractureContext.GetGeometryCollection();
 
-		TManagedArray<int32>* HasCustomConvex = FGeometryCollectionConvexUtility::GetCustomConvexFlags(&Collection, false);
-		const TManagedArray<TSet<int32>>& TransformToConvexIndices = Collection.GetAttribute<TSet<int32>>("TransformToConvexIndices", FTransformCollection::TransformGroup);
-		const TManagedArray<TUniquePtr<Chaos::FConvex>>& ConvexHull = Collection.GetAttribute<TUniquePtr<Chaos::FConvex>>("ConvexHull", "Convex");
+			const int32 CollectionIdx = VisualizedCollections.Add(FractureContext.GetGeometryCollectionComponent());
 
-		int32 CollectionIdx = VisualizedCollections.Add(FractureContext.GetGeometryCollectionComponent());
+			//FTransform OuterTransform = FractureContext.GetTransform();
+			const FTransform OuterTransform = FractureContext.GetGeometryCollectionComponent()->GetComponentTransform();
 
-		FTransform OuterTransform = FractureContext.GetTransform();
-		for (int32 TransformIdx : FractureContext.GetSelection())
-		{
-			EdgesMappings.AddMapping(CollectionIdx, TransformIdx, HullEdges.Num());
+			for (int32 TransformIdx : FractureContext.GetSelection())
+			{
+				EdgesMappings.AddMapping(CollectionIdx, TransformIdx, HullEdges.Num());
 
-			AddHullVisualizationData(Collection, OuterTransform, TransformIdx);
+				AddVisualizationData(Collection, OuterTransform, TransformIdx, RestCollection->bImportCollisionFromSource);
+			}
 		}
 	}
 }
diff --git a/Engine/Plugins/Experimental/ChaosEditor/Source/FractureEditor/Private/FractureToolConvex.h b/Engine/Plugins/Experimental/ChaosEditor/Source/FractureEditor/Private/FractureToolConvex.h
--- a/Engine/Plugins/Experimental/ChaosEditor/Source/FractureEditor/Private/FractureToolConvex.h
+++ b/Engine/Plugins/Experimental/ChaosEditor/Source/FractureEditor/Private/FractureToolConvex.h
@@ -148,8 +148,8 @@ protected:
 		EdgesMappings.Empty();
 	}
 
-	/* return true if the convex is a union of the children */
-	void AddHullVisualizationData(FGeometryCollection& Collection, const FTransform& OuterTransform, int32 TransformIdx, bool bIsUnionOfHulls = false);
+	void AddConvexHullVizualizationData(const Chaos::FConvex& ConvexHull, const FTransform& WorldSpaceTransform, bool bIsCustom, bool bIsUnionOfHulls);
+	void AddVisualizationData(FGeometryCollection& Collection, const FTransform& OuterTransform, int32 TransformIdx, bool bUseExternalCollision, bool bIsUnionOfHulls = false);
 
 	struct FEdgeVisInfo
 	{
diff --git a/Engine/Plugins/Experimental/ChaosEditor/Source/FractureEditor/Private/SGeometryCollectionOutliner.cpp b/Engine/Plugins/Experimental/ChaosEditor/Source/FractureEditor/Private/SGeometryCollectionOutliner.cpp
--- a/Engine/Plugins/Experimental/ChaosEditor/Source/FractureEditor/Private/SGeometryCollectionOutliner.cpp
+++ b/Engine/Plugins/Experimental/ChaosEditor/Source/FractureEditor/Private/SGeometryCollectionOutliner.cpp
@@ -543,6 +543,12 @@ void SGeometryCollectionOutliner::SetBoneSelection(UGeometryCollectionComponent*
 
 void SGeometryCollectionOutliner::OnSelectionChanged(FGeometryCollectionTreeItemPtr Item, ESelectInfo::Type SelectInfo)
 {
+	// object may just being deleted, we need to bail if that's the case to avoid a crash accessing teh component
+	if (!Item || !IsValid(Item->GetComponent()))
+	{
+		return;
+	}
+
 	if(!bPerformingSelection && BoneSelectionChangedDelegate.IsBound())
 	{
 		TMap<UGeometryCollectionComponent*, TArray<int32>> ComponentToBoneSelectionMap;
diff --git a/Engine/Source/Runtime/Experimental/Chaos/Private/GeometryCollection/GeometryCollectionConvexUtility.cpp b/Engine/Source/Runtime/Experimental/Chaos/Private/GeometryCollection/GeometryCollectionConvexUtility.cpp
--- a/Engine/Source/Runtime/Experimental/Chaos/Private/GeometryCollection/GeometryCollectionConvexUtility.cpp
+++ b/Engine/Source/Runtime/Experimental/Chaos/Private/GeometryCollection/GeometryCollectionConvexUtility.cpp
@@ -20,6 +20,101 @@ FAutoConsoleVariableRef CVarUseVolumeToComputeRelativeSize(TEXT("p.gc.UseVolumeT
 bool UseLargestClusterToComputeRelativeSize = false;
 FAutoConsoleVariableRef CVarUseMaxClusterToComputeRelativeSize(TEXT("p.gc.UseLargestClusterToComputeRelativeSize"), UseVolumeToComputeRelativeSize, TEXT("Use the largest Cluster as reference for the releative size instead of the largest child (def: false)"));
 
+static const Chaos::FVec3f IcoSphere_Subdiv0[] =
+{
+	{  0.000000f,  0.000000f, -1.000000f },
+	{  0.525720f,  0.723600f, -0.447215f },
+	{  0.850640f, -0.276385f, -0.447215f },
+	{  0.000000f, -0.894425f, -0.447215f },
+	{ -0.850640f, -0.276385f, -0.447215f },
+	{ -0.525720f,  0.723600f, -0.447215f },
+	{  0.850640f,  0.276385f,  0.447215f },
+	{  0.525720f, -0.723600f,  0.447215f },
+	{ -0.525720f, -0.723600f,  0.447215f },
+	{ -0.850640f,  0.276385f,  0.447215f },
+	{  0.000000f,  0.894425f,  0.447215f },
+	{  0.000000f,  0.000000f,  1.000000f },
+};
+constexpr int32 IcoSphere_Subdiv0_Num = sizeof(IcoSphere_Subdiv0) / sizeof(Chaos::FVec3f);
+
+static const Chaos::FVec3f IcoSphere_Subdiv1[] =
+{
+	{  0.000000f,  0.000000f, -1.000000f },
+	{  0.525725f,  0.723607f, -0.447220f },
+	{  0.850649f, -0.276388f, -0.447220f },
+	{  0.000000f, -0.894426f, -0.447216f },
+	{ -0.850649f, -0.276388f, -0.447220f },
+	{ -0.525725f,  0.723607f, -0.447220f },
+	{  0.850649f,  0.276388f,  0.447220f },
+	{  0.525725f, -0.723607f,  0.447220f },
+	{ -0.525725f, -0.723607f,  0.447220f },
+	{ -0.850649f,  0.276388f,  0.447220f },
+	{  0.000000f,  0.894426f,  0.447216f },
+	{  0.000000f,  0.000000f,  1.000000f },
+	{  0.499995f, -0.162456f, -0.850654f },
+	{  0.309011f,  0.425323f, -0.850654f },
+	{  0.809012f,  0.262869f, -0.525738f },
+	{  0.000000f,  0.850648f, -0.525736f },
+	{ -0.309011f,  0.425323f, -0.850654f },
+	{  0.000000f, -0.525730f, -0.850652f },
+	{  0.499997f, -0.688189f, -0.525736f },
+	{ -0.499995f, -0.162456f, -0.850654f },
+	{ -0.499997f, -0.688189f, -0.525736f },
+	{ -0.809012f,  0.262869f, -0.525738f },
+	{  0.309013f,  0.951058f,  0.000000f },
+	{ -0.309013f,  0.951058f,  0.000000f },
+	{  1.000000f,  0.000000f,  0.000000f },
+	{  0.809017f,  0.587786f,  0.000000f },
+	{  0.309013f, -0.951058f,  0.000000f },
+	{  0.809017f, -0.587786f,  0.000000f },
+	{ -0.809017f, -0.587786f,  0.000000f },
+	{ -0.309013f, -0.951058f,  0.000000f },
+	{ -0.809017f,  0.587786f,  0.000000f },
+	{ -1.000000f,  0.000000f,  0.000000f },
+	{  0.499997f,  0.688189f,  0.525736f },
+	{  0.809012f, -0.262869f,  0.525738f },
+	{  0.000000f, -0.850648f,  0.525736f },
+	{ -0.809012f, -0.262869f,  0.525738f },
+	{ -0.499997f,  0.688189f,  0.525736f },
+	{  0.499995f,  0.162456f,  0.850654f },
+	{  0.000000f,  0.525730f,  0.850652f },
+	{  0.309011f, -0.425323f,  0.850654f },
+	{ -0.309011f, -0.425323f,  0.850654f },
+	{ -0.499995f,  0.162456f,  0.850654f },
+};
+constexpr int32 IcoSphere_Subdiv1_Num = sizeof(IcoSphere_Subdiv1) / sizeof(Chaos::FVec3f);
+
+static const Chaos::FVec3f IcoHemisphere_Subdiv1[] =
+{
+	{  0.850649f,  0.27638f, 0.447220f },
+	{  0.525725f, -0.72360f, 0.447220f },
+	{ -0.525725f, -0.72360f, 0.447220f },
+	{ -0.850649f,  0.27638f, 0.447220f },
+	{  0.000000f,  0.89442f, 0.447216f },
+	{  0.000000f,  0.00000f, 1.000000f },
+	{  0.309013f,  0.95105f, 0.000000f },
+	{ -0.309013f,  0.95105f, 0.000000f },
+	{  1.000000f,  0.00000f, 0.000000f },
+	{  0.809017f,  0.58778f, 0.000000f },
+	{  0.309013f, -0.95105f, 0.000000f },
+	{  0.809017f, -0.58778f, 0.000000f },
+	{ -0.809017f, -0.58778f, 0.000000f },
+	{ -0.309013f, -0.95105f, 0.000000f },
+	{ -0.809017f,  0.58778f, 0.000000f },
+	{ -1.000000f,  0.00000f, 0.000000f },
+	{  0.499997f,  0.68818f, 0.525736f },
+	{  0.809012f, -0.26286f, 0.525738f },
+	{  0.000000f, -0.85064f, 0.525736f },
+	{ -0.809012f, -0.26286f, 0.525738f },
+	{ -0.499997f,  0.68818f, 0.525736f },
+	{  0.499995f,  0.16245f, 0.850654f },
+	{  0.000000f,  0.52573f, 0.850652f },
+	{  0.309011f, -0.42532f, 0.850654f },
+	{ -0.309011f, -0.42532f, 0.850654f },
+	{ -0.499995f,  0.16245f, 0.850654f },
+};
+constexpr int32 IcoHemisphere_Subdiv1_Num = sizeof(IcoHemisphere_Subdiv1) / sizeof(Chaos::FVec3f);
+
 
 TOptional<FGeometryCollectionConvexUtility::FGeometryCollectionConvexData> FGeometryCollectionConvexUtility::GetConvexHullDataIfPresent(FGeometryCollection* GeometryCollection)
 {
@@ -1371,8 +1466,11 @@ void FGeometryCollectionConvexUtility::GenerateClusterConvexHullsFromLeafHulls(F
 	static FName ConvexHullAttributeName("ConvexHull");
 	static FName TransformToConvexIndicesName("TransformToConvexIndices");
 
+	using FSharedImplicit = TSharedPtr<Chaos::FImplicitObject, ESPMode::ThreadSafe>;
+
 	TManagedArrayAccessor<TUniquePtr<Chaos::FConvex>> ConvexHullAttribute(Collection, ConvexHullAttributeName, ConvexGroupName);
 	TManagedArrayAccessor<TSet<int32>> TransformToConvexIndicesAttribute(Collection, TransformToConvexIndicesName, FGeometryCollection::TransformGroup);
+	const TManagedArrayAccessor<FSharedImplicit> ExternalCollisionAttribute(Collection, "ExternalCollisions", FGeometryCollection::TransformGroup);
 
 	// if any of the leaves hulls has been computed those two attributes should be valid
 	if (!ConvexHullAttribute.IsValid() || !TransformToConvexIndicesAttribute.IsValid())
@@ -1419,21 +1517,41 @@ void FGeometryCollectionConvexUtility::GenerateClusterConvexHullsFromLeafHulls(F
 
 				// build information required for merging hulls 
 				TArray<FHullInfo> Hulls;
+				TArray<FTransformedConvex> ExternalHulls;
 				TMap<int32, TArray<int32>> TransformToHullIdx;
 				TArray<TPair<int32, int32>> HullProximity;
-
+	
 				for (int32 SourceTransformIndex : SourceTransformIndices)
 				{
 					const FTransform InnerTransform = GlobalTransforms[SourceTransformIndex];
 					FTransform ChildToParentTransform = InnerTransform.GetRelativeTransform(ParentTransform);
-					int32 HullIdxStart = Hulls.Num();
-					for (int32 SourceConvexIdx : TransformToConvexIndices[SourceTransformIndex])
+
+					const FSharedImplicit ExternalCollision = (ExternalCollisionAttribute.IsValid()) ? ExternalCollisionAttribute.Get()[SourceTransformIndex] : FSharedImplicit();
+					if (ExternalCollision)
+					{
+						const int32 ExternalHullsStart = ExternalHulls.Num();
+						ConvertImplicitToConvexArray(*ExternalCollision, FTransform::Identity, ExternalHulls);
+						for (int32 HullIndex = ExternalHullsStart; HullIndex < ExternalHulls.Num();  HullIndex++)
+						{
+							const FTransformedConvex& ExternalHull = ExternalHulls[HullIndex];
+
+							FHullInfo HullInfo;
+							HullInfo.Convex = ExternalHull.Convex.Get();
+							HullInfo.Transform = ExternalHull.Transform * ChildToParentTransform;
+							const int32 HullIdx = Hulls.Emplace(HullInfo);
+							TransformToHullIdx.FindOrAdd(SourceTransformIndex).Add(HullIdx);
+						}
+					}
+					else
 					{
-						FHullInfo HullInfo;
-						HullInfo.Convex = ConvexHull[SourceConvexIdx].Get();
-						HullInfo.Transform = ChildToParentTransform;
-						const int32 HullIdx = Hulls.Emplace(HullInfo);
-						TransformToHullIdx.FindOrAdd(SourceTransformIndex).Add(HullIdx);
+						for (int32 SourceConvexIdx : TransformToConvexIndices[SourceTransformIndex])
+						{
+							FHullInfo HullInfo;
+							HullInfo.Convex = ConvexHull[SourceConvexIdx].Get();
+							HullInfo.Transform = ChildToParentTransform;
+							const int32 HullIdx = Hulls.Emplace(HullInfo);
+							TransformToHullIdx.FindOrAdd(SourceTransformIndex).Add(HullIdx);
+						}
 					}
 				}
 				// get hull proximity out of geometry proximity.  (Note: Only works if merging leaf convexes ... need a more general source of proximity!)
@@ -1837,3 +1955,237 @@ void FGeometryCollectionConvexUtility::SetVolumeAttributes(FGeometryCollection*
 	}
 }
 
+void FGeometryCollectionConvexUtility::ConvertImplicitToConvexArray(const Chaos::FImplicitObject& InImplicit, const FTransform& Transform, TArray<FTransformedConvex>& InOutConvex)
+{
+	const Chaos::EImplicitObjectType PackedType = InImplicit.GetType(); // Type includes scaling and instancing data
+	const Chaos::EImplicitObjectType InnerType = Chaos::GetInnerType(InImplicit.GetType());
+
+	// Unwrap the wrapper/aggregating shapes
+	if (Chaos::IsScaled(PackedType))
+	{
+		ConvertScaledImplicitToConvexArray(InImplicit, Transform, Chaos::IsInstanced(PackedType), InOutConvex);
+		return;
+	}
+
+	if (Chaos::IsInstanced(PackedType))
+	{
+		ConvertInstancedImplicitToConvexArray(InImplicit, Transform, InOutConvex);
+		return;
+	}
+
+	if (InnerType == Chaos::ImplicitObjectType::Transformed)
+	{
+		const Chaos::TImplicitObjectTransformed<Chaos::FReal, 3>* Transformed = InImplicit.template GetObject<Chaos::TImplicitObjectTransformed<Chaos::FReal, 3>>();
+		const Chaos::FRigidTransform3 ImplicitTransform(
+			Transform.TransformPosition(Transformed->GetTransform().GetLocation()),
+			Transform.GetRotation() * Transformed->GetTransform().GetRotation()
+		);
+		if (const Chaos::FImplicitObject* TransformedImplicit = Transformed->GetTransformedObject())
+		{
+			ConvertImplicitToConvexArray(*TransformedImplicit, ImplicitTransform, InOutConvex);
+		}
+		return;
+	}
+
+	if (InnerType == Chaos::ImplicitObjectType::Union)
+	{
+		const Chaos::FImplicitObjectUnion* Union = InImplicit.template GetObject<Chaos::FImplicitObjectUnion>();
+		int32 UnionIdx = 0;
+		for (const TUniquePtr<Chaos::FImplicitObject>& UnionImplicit : Union->GetObjects())
+		{
+			if (UnionImplicit)
+			{
+				ConvertImplicitToConvexArray(*UnionImplicit, Transform, InOutConvex);
+			}
+		}
+		return;
+	}
+
+	if (InnerType == Chaos::ImplicitObjectType::UnionClustered)
+	{
+		// unsupported - quiet exit
+		return;
+	}
+
+	// If we get here, we have an actual shape to render
+	switch (InnerType)
+	{
+	case Chaos::ImplicitObjectType::Sphere:
+	{
+		if (const Chaos::TSphere<Chaos::FReal, 3>*Sphere = InImplicit.template GetObject<Chaos::TSphere<Chaos::FReal, 3>>())
+		{
+			const FTransform SphereTransform(FQuat::Identity, Sphere->GetCenter(), FVector(Sphere->GetRadius()));
+			const TArray<Chaos::FConvex::FVec3Type> Vertices(IcoSphere_Subdiv1, IcoSphere_Subdiv1_Num);
+
+			FTransformedConvex& TransformedConvex = InOutConvex.Emplace_GetRef();
+			TransformedConvex.Convex = MakeShared<Chaos::FConvex>(Vertices, 0);
+			TransformedConvex.Transform = SphereTransform * Transform;
+		}
+		break;
+	}
+	case Chaos::ImplicitObjectType::Box:
+	{
+		if (const Chaos::TBox<Chaos::FReal, 3>*Box = InImplicit.template GetObject<Chaos::TBox<Chaos::FReal, 3>>())
+		{
+			Chaos::FVec3f A(Box->Min());
+			Chaos::FVec3f B(Box->Max());
+			TArray<Chaos::FConvex::FVec3Type> Vertices;
+			Vertices.Add({ A.X, A.Y, A.Z });
+			Vertices.Add({ B.X, A.Y, A.Z });
+			Vertices.Add({ A.X, B.Y, A.Z });
+			Vertices.Add({ B.X, B.Y, A.Z });
+			Vertices.Add({ A.X, A.Y, B.Z });
+			Vertices.Add({ B.X, A.Y, B.Z });
+			Vertices.Add({ A.X, B.Y, B.Z });
+			Vertices.Add({ B.X, B.Y, B.Z });
+
+			FTransformedConvex& TransformedConvex = InOutConvex.Emplace_GetRef();
+			TransformedConvex.Convex = MakeShared<Chaos::FConvex>(Vertices, 0);
+			TransformedConvex.Transform = Transform;
+		}
+		break;
+	}
+	case Chaos::ImplicitObjectType::Capsule:
+	{
+		if (const Chaos::FCapsule* Capsule = InImplicit.template GetObject<Chaos::FCapsule>())
+		{
+			const FTransform CapsuleTransform(FRotationMatrix::MakeFromZ(Capsule->GetAxis()).ToQuat(), Capsule->GetCenter());
+
+			const FVector3f HalfHeightOffset(0, 0, static_cast<float>(Capsule->GetHeight()) * 0.5f);
+
+			TArray<Chaos::FConvex::FVec3Type> Vertices;
+			for (int32 VtxIndex = 0; VtxIndex < IcoHemisphere_Subdiv1_Num; VtxIndex++)
+			{
+				const float CapsuleRadius = static_cast<float>(Capsule->GetRadius());
+				Vertices.Add(IcoHemisphere_Subdiv1[VtxIndex] * CapsuleRadius + HalfHeightOffset); // top hemisphere
+				Vertices.Add(IcoHemisphere_Subdiv1[VtxIndex] * -CapsuleRadius - HalfHeightOffset); // bottom hemisphere
+			}
+
+			FTransformedConvex& TransformedConvex = InOutConvex.Emplace_GetRef();
+			TransformedConvex.Convex = MakeShared<Chaos::FConvex>(Vertices, 0);
+			TransformedConvex.Transform = CapsuleTransform * Transform;
+		}
+		break;
+	}
+	break;
+	case Chaos::ImplicitObjectType::Convex:
+	{
+		if (const Chaos::FConvex* Convex = InImplicit.template GetObject<Chaos::FConvex>())
+		{
+			FTransformedConvex& TransformedConvex = InOutConvex.Emplace_GetRef();
+			TUniquePtr<Chaos::FConvex> CopiedConvex = Convex->CopyAsConvex();
+			TransformedConvex.Convex = TSharedPtr<Chaos::FConvex>(CopiedConvex.Release());
+			TransformedConvex.Transform = Transform;
+		}
+		break;
+	}
+	// Unsuported shape types
+	case Chaos::ImplicitObjectType::Plane:
+	case Chaos::ImplicitObjectType::LevelSet:
+	case Chaos::ImplicitObjectType::TaperedCylinder:
+	case Chaos::ImplicitObjectType::Cylinder:
+	case Chaos::ImplicitObjectType::TriangleMesh:
+	case Chaos::ImplicitObjectType::HeightField:
+	default:
+		break;
+	}
+}
+
+void FGeometryCollectionConvexUtility::ConvertScaledImplicitToConvexArray(
+	const Chaos::FImplicitObject& Implicit, 
+	const FTransform& WorldSpaceTransform, bool bInstanced, 
+	TArray<FTransformedConvex>& InOutConvex)
+{
+	const Chaos::EImplicitObjectType InnerType = Chaos::GetInnerType(Implicit.GetType());
+	switch (InnerType)
+	{
+		// we only support scaled / instanced convex
+	case Chaos::ImplicitObjectType::Convex:
+	{
+		Chaos::FRigidTransform3 ScaleTM{ Chaos::FRigidTransform3::Identity };
+		const Chaos::FConvex* Convex = nullptr;
+		if (bInstanced)
+		{
+			const auto ScaledInstancedConvex = Implicit.template GetObject<Chaos::TImplicitObjectScaled<Chaos::FConvex, true>>();
+			ScaleTM.SetScale3D(ScaledInstancedConvex->GetScale());
+			if (ScaledInstancedConvex)
+			{
+				Convex = ScaledInstancedConvex->GetUnscaledObject();
+			}
+		}
+		else
+		{
+			const auto ScaledConvex = Implicit.template GetObject<Chaos::TImplicitObjectScaled<Chaos::FConvex, false>>();
+			ScaleTM.SetScale3D(ScaledConvex->GetScale());
+			if (ScaledConvex)
+			{
+				Convex = ScaledConvex->GetUnscaledObject();
+			}
+		}
+		if (Convex)
+		{
+			FTransformedConvex& TransformedConvex = InOutConvex.Emplace_GetRef();
+			TUniquePtr<Chaos::FConvex> CopiedConvex = Convex->CopyAsConvex();
+			TransformedConvex.Convex = TSharedPtr<Chaos::FConvex>(CopiedConvex.Release());
+			TransformedConvex.Transform = ScaleTM * WorldSpaceTransform;
+		}
+		break;
+	}
+	// unsupported types for scaled implicit
+	case Chaos::ImplicitObjectType::Sphere:
+	case Chaos::ImplicitObjectType::Box:
+	case Chaos::ImplicitObjectType::Plane:
+	case Chaos::ImplicitObjectType::Capsule:
+	case Chaos::ImplicitObjectType::Transformed:
+	case Chaos::ImplicitObjectType::Union:
+	case Chaos::ImplicitObjectType::LevelSet:
+	case Chaos::ImplicitObjectType::Unknown:
+	case Chaos::ImplicitObjectType::TaperedCylinder:
+	case Chaos::ImplicitObjectType::Cylinder:
+	case Chaos::ImplicitObjectType::TriangleMesh:
+	case Chaos::ImplicitObjectType::HeightField:
+	default:
+		break;
+	}
+}
+
+void FGeometryCollectionConvexUtility::ConvertInstancedImplicitToConvexArray(
+	const Chaos::FImplicitObject& Implicit, 
+	const FTransform& Transform, 
+	TArray<FTransformedConvex>& InOutConvex)
+{
+	const Chaos::EImplicitObjectType InnerType = Chaos::GetInnerType(Implicit.GetType());
+	switch (InnerType)
+	{
+		// we only support instanced convex
+	case Chaos::ImplicitObjectType::Convex:
+	{
+		const Chaos::TImplicitObjectInstanced<Chaos::FConvex>* Instanced = Implicit.template GetObject< Chaos::TImplicitObjectInstanced< Chaos::FConvex>>();
+		if (const Chaos::FConvex* Convex = Instanced->GetInstancedObject())
+		{
+			FTransformedConvex& TransformedConvex = InOutConvex.Emplace_GetRef();
+			TUniquePtr<Chaos::FConvex> CopiedConvex = Convex->CopyAsConvex();
+			TransformedConvex.Convex = TSharedPtr<Chaos::FConvex>(CopiedConvex.Release());
+			TransformedConvex.Transform = Transform;
+		}
+		break;
+	}
+
+	// unsupported types for scaled implicit
+	case Chaos::ImplicitObjectType::Sphere:
+	case Chaos::ImplicitObjectType::Box:
+	case Chaos::ImplicitObjectType::Plane:
+	case Chaos::ImplicitObjectType::Capsule:
+	case Chaos::ImplicitObjectType::Transformed:
+	case Chaos::ImplicitObjectType::Union:
+	case Chaos::ImplicitObjectType::LevelSet:
+	case Chaos::ImplicitObjectType::Unknown:
+	case Chaos::ImplicitObjectType::TaperedCylinder:
+	case Chaos::ImplicitObjectType::Cylinder:
+	case Chaos::ImplicitObjectType::TriangleMesh:
+	case Chaos::ImplicitObjectType::HeightField:
+	default:
+		break;
+	}
+}
+
diff --git a/Engine/Source/Runtime/Experimental/Chaos/Public/Chaos/Convex.h b/Engine/Source/Runtime/Experimental/Chaos/Public/Chaos/Convex.h
--- a/Engine/Source/Runtime/Experimental/Chaos/Public/Chaos/Convex.h
+++ b/Engine/Source/Runtime/Experimental/Chaos/Public/Chaos/Convex.h
@@ -159,6 +159,11 @@ namespace Chaos
 			return *this;
 		}
 
+		TUniquePtr<FConvex> CopyAsConvex() const
+		{
+			return TUniquePtr<FConvex>(new FConvex(*this));
+		}
+
 		virtual TUniquePtr<FImplicitObject> Copy() const
 		{
 			return TUniquePtr<FImplicitObject>(new FConvex(*this));
diff --git a/Engine/Source/Runtime/Experimental/Chaos/Public/GeometryCollection/GeometryCollectionConvexUtility.h b/Engine/Source/Runtime/Experimental/Chaos/Public/GeometryCollection/GeometryCollectionConvexUtility.h
--- a/Engine/Source/Runtime/Experimental/Chaos/Public/GeometryCollection/GeometryCollectionConvexUtility.h
+++ b/Engine/Source/Runtime/Experimental/Chaos/Public/GeometryCollection/GeometryCollectionConvexUtility.h
@@ -121,5 +121,28 @@ public:
 	// @param GlobalTransformArray		GeometryCollection's transforms to global space, as computed by GeometryCollectionAlgo::GlobalMatrices
 	static UE::GeometryCollectionConvexUtility::FConvexHulls ComputeLeafHulls(FGeometryCollection* GeometryCollection, const TArray<FTransform>& GlobalTransformArray, double SimplificationDistanceThreshold = 0.0, double OverlapRemovalShrinkPercent = 0.0);
 
+	struct FTransformedConvex
+	{
+		TSharedPtr<Chaos::FConvex> Convex;
+		FTransform Transform;
+	};
+
+	// generate a list of convex out of a hierarchy of implciit shapes
+	// suported shapes are scaled / transformed implicits as well as Boxes, convexes, spheres and capsules
+	// levelset, tapered capsule, cylinder, heightfields and trimesh are not supported
+	// array will not be erased and convex will be added if any already exists
+	static void ConvertImplicitToConvexArray(const Chaos::FImplicitObject& InImplicit, const FTransform& Transform, TArray<FTransformedConvex>& InOutConvex);
+
+private:
+	static void ConvertScaledImplicitToConvexArray(
+		const Chaos::FImplicitObject& Implicit,
+		const FTransform& WorldSpaceTransform, bool bInstanced,
+		TArray<FTransformedConvex>& InOutConvex);
+
+	static void ConvertInstancedImplicitToConvexArray(
+		const Chaos::FImplicitObject& Implicit,
+		const FTransform& Transform,
+		TArray<FTransformedConvex>& InOutConvex);
+
 };
 
