Geometry collection : GenerateClusterConvexHullsFromLeafHulls can now create the convex attributes if necessary

diff --git a/Engine/Source/Runtime/Experimental/Chaos/Private/GeometryCollection/GeometryCollectionConvexUtility.cpp b/Engine/Source/Runtime/Experimental/Chaos/Private/GeometryCollection/GeometryCollectionConvexUtility.cpp
--- a/Engine/Source/Runtime/Experimental/Chaos/Private/GeometryCollection/GeometryCollectionConvexUtility.cpp
+++ b/Engine/Source/Runtime/Experimental/Chaos/Private/GeometryCollection/GeometryCollectionConvexUtility.cpp
@@ -140,23 +140,9 @@ bool FGeometryCollectionConvexUtility::HasConvexHullData(FGeometryCollection* Ge
 
 FGeometryCollectionConvexUtility::FGeometryCollectionConvexData FGeometryCollectionConvexUtility::GetValidConvexHullData(FGeometryCollection* GeometryCollection)
 {
-	check (GeometryCollection)
+	check(GeometryCollection)
 
-	if (!GeometryCollection->HasGroup("Convex"))
-	{
-		GeometryCollection->AddGroup("Convex");
-	}
-
-	if (!GeometryCollection->HasAttribute("TransformToConvexIndices", FTransformCollection::TransformGroup))
-	{
-		FManagedArrayCollection::FConstructionParameters ConvexDependency("Convex");
-		GeometryCollection->AddAttribute<TSet<int32>>("TransformToConvexIndices", FTransformCollection::TransformGroup, ConvexDependency);
-	}
-
-	if (!GeometryCollection->HasAttribute("ConvexHull", "Convex"))
-	{
-		GeometryCollection->AddAttribute<TUniquePtr<Chaos::FConvex>>("ConvexHull", "Convex");
-	}
+	CreateConvexHullAttributesIfNeeded(*GeometryCollection);
 
 	// Check for correct population. Make sure all rigid nodes should have a convex associated; leave convex hulls for transform nodes alone for now
 	const TManagedArray<int32>& SimulationType = GeometryCollection->GetAttribute<int32>("SimulationType", FTransformCollection::TransformGroup);
@@ -1370,6 +1356,24 @@ double ComputeGeometryVolume(
 
 }
 
+void FGeometryCollectionConvexUtility::CreateConvexHullAttributesIfNeeded(FManagedArrayCollection& Collection)
+{
+	if (!Collection.HasGroup("Convex"))
+	{
+		Collection.AddGroup("Convex");
+	}
+
+	if (!Collection.HasAttribute("TransformToConvexIndices", FTransformCollection::TransformGroup))
+	{
+		FManagedArrayCollection::FConstructionParameters ConvexDependency("Convex");
+		Collection.AddAttribute<TSet<int32>>("TransformToConvexIndices", FTransformCollection::TransformGroup, ConvexDependency);
+	}
+
+	if (!Collection.HasAttribute("ConvexHull", "Convex"))
+	{
+		Collection.AddAttribute<TUniquePtr<Chaos::FConvex>>("ConvexHull", "Convex");
+	}
+}
 
 UE::GeometryCollectionConvexUtility::FConvexHulls
 FGeometryCollectionConvexUtility::ComputeLeafHulls(FGeometryCollection* GeometryCollection, const TArray<FTransform>& GlobalTransformArray, double SimplificationDistanceThreshold, double OverlapRemovalShrinkPercent)
@@ -1428,21 +1432,7 @@ FGeometryCollectionConvexUtility::FGeometryCollectionConvexData FGeometryCollect
 
 	TransformHullsToLocal(GlobalTransformArray, UseLeafHulls->Hulls, UseLeafHulls->Pivots, UseLeafHulls->TransformToHullsIndices, UseLeafHulls->OverlapRemovalShrinkPercent);
 
-	if (!GeometryCollection->HasGroup("Convex"))
-	{
-		GeometryCollection->AddGroup("Convex");
-	}
-
-	if (!GeometryCollection->HasAttribute("TransformToConvexIndices", FTransformCollection::TransformGroup))
-	{
-		FManagedArrayCollection::FConstructionParameters ConvexDependency("Convex");
-		GeometryCollection->AddAttribute<TSet<int32>>("TransformToConvexIndices", FTransformCollection::TransformGroup, ConvexDependency);
-	}
-
-	if (!GeometryCollection->HasAttribute("ConvexHull", "Convex"))
-	{
-		GeometryCollection->AddAttribute<TUniquePtr<Chaos::FConvex>>("ConvexHull", "Convex");
-	}
+	CreateConvexHullAttributesIfNeeded(*GeometryCollection);
 
 	TManagedArray<TSet<int32>>& TransformToConvexIndices = GeometryCollection->ModifyAttribute<TSet<int32>>("TransformToConvexIndices", FTransformCollection::TransformGroup);
 	TManagedArray<TUniquePtr<Chaos::FConvex>>& ConvexHull = GeometryCollection->ModifyAttribute<TUniquePtr<Chaos::FConvex>>("ConvexHull", "Convex");
@@ -1468,16 +1458,12 @@ void FGeometryCollectionConvexUtility::GenerateClusterConvexHullsFromLeafHulls(F
 
 	using FSharedImplicit = TSharedPtr<Chaos::FImplicitObject, ESPMode::ThreadSafe>;
 
+	CreateConvexHullAttributesIfNeeded(Collection);
+
 	TManagedArrayAccessor<TUniquePtr<Chaos::FConvex>> ConvexHullAttribute(Collection, ConvexHullAttributeName, ConvexGroupName);
 	TManagedArrayAccessor<TSet<int32>> TransformToConvexIndicesAttribute(Collection, TransformToConvexIndicesName, FGeometryCollection::TransformGroup);
 	const TManagedArrayAccessor<FSharedImplicit> ExternalCollisionAttribute(Collection, "ExternalCollisions", FGeometryCollection::TransformGroup);
 
-	// if any of the leaves hulls has been computed those two attributes should be valid
-	if (!ConvexHullAttribute.IsValid() || !TransformToConvexIndicesAttribute.IsValid())
-	{
-		return;
-	}
-
 	GeometryCollection::Facades::FCollectionTransformFacade TransformFacade(Collection);
 	Chaos::Facades::FCollectionHierarchyFacade HierarchyFacade(Collection);
 
diff --git a/Engine/Source/Runtime/Experimental/Chaos/Public/GeometryCollection/GeometryCollectionConvexUtility.h b/Engine/Source/Runtime/Experimental/Chaos/Public/GeometryCollection/GeometryCollectionConvexUtility.h
--- a/Engine/Source/Runtime/Experimental/Chaos/Public/GeometryCollection/GeometryCollectionConvexUtility.h
+++ b/Engine/Source/Runtime/Experimental/Chaos/Public/GeometryCollection/GeometryCollectionConvexUtility.h
@@ -144,5 +144,6 @@ private:
 		const FTransform& Transform,
 		TArray<FTransformedConvex>& InOutConvex);
 
+	static void CreateConvexHullAttributesIfNeeded(FManagedArrayCollection& GeometryCollection);
 };
 
