[Chaos]
* Moved VectorUtility from Chaos to ChaosCore to solve include violations

diff --git a/Engine/Source/Runtime/Experimental/Chaos/Public/Chaos/VectorUtility.h b/Engine/Source/Runtime/Experimental/Chaos/Public/Chaos/VectorUtility.h
deleted file mode 100644
diff --git a/Engine/Source/Runtime/Experimental/ChaosCore/Public/Chaos/VectorUtility.h b/Engine/Source/Runtime/Experimental/ChaosCore/Public/Chaos/VectorUtility.h
new file mode 100644
--- /dev/null
+++ b/Engine/Source/Runtime/Experimental/ChaosCore/Public/Chaos/VectorUtility.h
@@ -0,0 +1,140 @@
+// Copyright Epic Games, Inc. All Rights Reserved.
+#pragma once
+
+#include "Math/VectorRegister.h"
+
+/**
+ * Cast VectorRegister4Int in VectorRegister4Float
+ *
+ * @param V	vector
+ * @return		VectorRegister4Float( B.x, A.y, A.z, A.w)
+ */
+
+FORCEINLINE VectorRegister4Float VectorCast4IntTo4Float(const VectorRegister4Int& V)
+{
+#if (!defined(_MSC_VER) || PLATFORM_ENABLE_VECTORINTRINSICS_NEON) && PLATFORM_ENABLE_VECTORINTRINSICS
+	return VectorRegister4Float(V);
+#elif PLATFORM_ENABLE_VECTORINTRINSICS
+	return _mm_castsi128_ps(V);
+#else
+	return VectorCastIntToFloat(Vec);
+#endif
+}
+
+
+/**
+ * Cast VectorRegister4Float in VectorRegister4Int
+ *
+ * @param V	vector
+ * @return		VectorCast4FloatTo4Int( B.x, A.y, A.z, A.w)
+ */
+FORCEINLINE VectorRegister4Int VectorCast4FloatTo4Int(const VectorRegister4Float& V)
+{
+#if (!defined(_MSC_VER) || PLATFORM_ENABLE_VECTORINTRINSICS_NEON) && PLATFORM_ENABLE_VECTORINTRINSICS
+	return VectorRegister4Int(V);
+#elif PLATFORM_ENABLE_VECTORINTRINSICS
+	return _mm_castps_si128(V);
+#else
+	return VectorCastFloatToInt(Vec);
+#endif
+
+}
+
+/**
+ * Selects and interleaves the lower two SP FP values from A and B.
+ *
+ * @param A	1st vector
+ * @param B	2nd vector
+ * @return		VectorRegister4Float( A.x, B.x, A.y, B.y)
+ */
+FORCEINLINE VectorRegister4Float VectorUnpackLo(const VectorRegister4Float& A, const VectorRegister4Float& B)
+{
+#if PLATFORM_ENABLE_VECTORINTRINSICS_NEON
+	return vzip1q_f32(A, B);
+#elif PLATFORM_ENABLE_VECTORINTRINSICS
+	return _mm_unpacklo_ps(A, B);
+#else
+	return MakeVectorRegisterFloat(A.V[0], B.V[0], A.V[1], B.V[1]);
+#endif
+}
+
+/**
+ * Selects and interleaves the higher two SP FP values from A and B.
+ *
+ * @param A	1st vector
+ * @param B	2nd vector
+ * @return		VectorRegister4Float( A.z, B.z, A.w, B.w)
+ */
+FORCEINLINE VectorRegister4Float VectorUnpackHi(const VectorRegister4Float& A, const VectorRegister4Float& B)
+{
+#if PLATFORM_ENABLE_VECTORINTRINSICS_NEON
+	return vzip2q_f32(A, B);
+#elif PLATFORM_ENABLE_VECTORINTRINSICS
+	return _mm_unpackhi_ps(A, B);
+#else
+	return MakeVectorRegisterFloat(A.V[2], B.V[2], A.V[3], B.V[3]);
+#endif
+}
+
+/**
+ * Moves the lower 2 SP FP values of b to the upper 2 SP FP values of the result. The lower 2 SP FP values of a are passed through to the result.
+ *
+ * @param A	1st vector
+ * @param B	2nd vector
+ * @return		VectorRegister4Float( A.x, A.y, B.x, B.y)
+  */
+FORCEINLINE VectorRegister4Float VectorMoveLh(const VectorRegister4Float& A, const VectorRegister4Float& B)
+{
+#if PLATFORM_ENABLE_VECTORINTRINSICS_NEON
+	return vzip1q_f64(A, B);
+#else
+	return VectorCombineLow(A, B);
+#endif
+}
+
+
+/**
+ * Combines two vectors using bitwise NOT AND (treating each vector as a 128 bit field)
+ *
+ * @param Vec1	1st vector
+ * @param Vec2	2nd vector
+ * @return		VectorRegister4Float( for each bit i: !Vec1[i] & Vec2[i] )
+ */
+FORCEINLINE VectorRegister4Float VectorBitwiseNotAnd(const VectorRegister4Float& A, const VectorRegister4Float& B)
+{
+#if PLATFORM_ENABLE_VECTORINTRINSICS_NEON
+	return (VectorRegister4Float)vandq_u32(vmvnq_u32((VectorRegister4Int)A), (VectorRegister4Int)B);
+#elif PLATFORM_ENABLE_VECTORINTRINSICS
+	return _mm_andnot_ps(A, B);
+#else
+	return MakeVectorRegisterFloat(
+		uint32(~((uint32*)(A.V))[0] & ((uint32*)(B.V))[0]),
+		uint32(~((uint32*)(A.V))[1] & ((uint32*)(B.V))[1]),
+		uint32(~((uint32*)(A.V))[2] & ((uint32*)(B.V))[2]),
+		uint32(~((uint32*)(A.V))[3] & ((uint32*)(B.V))[3]));
+#endif
+}
+
+FORCEINLINE VectorRegister4Double VectorBitwiseNotAnd(const VectorRegister4Double& A, const VectorRegister4Double& B)
+{
+	VectorRegister4Double Result;
+#if PLATFORM_ENABLE_VECTORINTRINSICS_NEON	
+	Result.XY = (VectorRegister2Double)vandq_u32(vmvnq_u32((VectorRegister2Double)A.XY), (VectorRegister2Double)B.XY);
+	Result.ZW = (VectorRegister2Double)vandq_u32(vmvnq_u32((VectorRegister2Double)A.ZW), (VectorRegister2Double)B.ZW);
+#elif PLATFORM_ENABLE_VECTORINTRINSICS
+#if !UE_PLATFORM_MATH_USE_AVX
+	Result.XY = _mm_cvtps_pd(_mm_andnot_ps(_mm_cvtpd_ps(A.XY), _mm_cvtpd_ps(B.XY)));
+	Result.ZW = _mm_cvtps_pd(_mm_andnot_ps(_mm_cvtpd_ps(A.ZW), _mm_cvtpd_ps(B.ZW)));
+#else
+	Result = _mm256_andnot_pd(A, B);
+#endif
+#else
+	Result = MakeVectorRegisterDouble(
+		uint64(~((uint64*)(A.V))[0] & ((uint64*)(B.V))[0]),
+		uint64(~((uint64*)(A.V))[1] & ((uint64*)(B.V))[1]),
+		uint64(~((uint64*)(A.V))[2] & ((uint64*)(B.V))[2]),
+		uint64(~((uint64*)(A.V))[3] & ((uint64*)(B.V))[3]));
+#endif
+	return Result;
+}
+
