ChaosFlesh : Add threading arguments into XPBD corotated constraints. Add stiffness tolerance for XPBD weak constraints

diff --git a/Engine/Source/Runtime/Experimental/Chaos/Private/Chaos/Deformable/ChaosDeformableSolver.cpp b/Engine/Source/Runtime/Experimental/Chaos/Private/Chaos/Deformable/ChaosDeformableSolver.cpp
--- a/Engine/Source/Runtime/Experimental/Chaos/Private/Chaos/Deformable/ChaosDeformableSolver.cpp
+++ b/Engine/Source/Runtime/Experimental/Chaos/Private/Chaos/Deformable/ChaosDeformableSolver.cpp
@@ -69,6 +69,9 @@ namespace Chaos::Softs
 	FAutoConsoleVariableRef CVarDeformableDebugParamsDrawTetrahedralParticles(TEXT("p.Chaos.DebugDraw.Deformable.TetrahedralParticle"), GDeformableDebugParams.bDoDrawTetrahedralParticles, TEXT("Debug draw the deformable solvers tetrahedron. [def: false]"));
 	FAutoConsoleVariableRef CVarDeformableDebugParamsDrawKinematicParticles(TEXT("p.Chaos.DebugDraw.Deformable.KinematicParticle"), GDeformableDebugParams.bDoDrawKinematicParticles, TEXT("Debug draw the deformables kinematic particles. [def: false]"));
 
+	FDeformableXPBDCorotatedParams GDeformableXPBDCorotatedParams;
+	FAutoConsoleVariableRef CVarDeformableXPBDCorotatedBatchSize(TEXT("p.Chaos.Deformable.XPBDBatchSize"), GDeformableXPBDCorotatedParams.XPBDCorotatedBatchSize, TEXT("Batch size for physics parallel for. [def: 5]"));
+	FAutoConsoleVariableRef CVarDeformableXPBDCorotatedBatchThreshold(TEXT("p.Chaos.Deformable.XPBDBatchThreshold"), GDeformableXPBDCorotatedParams.XPBDCorotatedBatchThreshold, TEXT("Batch threshold for physics parallel for. [def: 5]"));
 
 	FCriticalSection FDeformableSolver::InitializationMutex;
 	FCriticalSection FDeformableSolver::RemovalMutex;
@@ -516,7 +519,7 @@ namespace Chaos::Softs
 					{
 						FXPBDCorotatedConstraints<FSolverReal, FSolverParticles>* CorotatedConstraint =
 							new FXPBDCorotatedConstraints<FSolverReal, FSolverParticles>(
-								Evolution->Particles(), Elements, TetStiffness);
+								Evolution->Particles(), Elements, TetStiffness, GDeformableXPBDCorotatedParams);
 
 						Evolution->ConstraintInits()[InitIndex] =
 							[CorotatedConstraint](FSolverParticles& InParticles, const FSolverReal Dt)
diff --git a/Engine/Source/Runtime/Experimental/Chaos/Public/Chaos/Deformable/ChaosDeformableSolverTypes.h b/Engine/Source/Runtime/Experimental/Chaos/Public/Chaos/Deformable/ChaosDeformableSolverTypes.h
--- a/Engine/Source/Runtime/Experimental/Chaos/Public/Chaos/Deformable/ChaosDeformableSolverTypes.h
+++ b/Engine/Source/Runtime/Experimental/Chaos/Public/Chaos/Deformable/ChaosDeformableSolverTypes.h
@@ -155,5 +155,12 @@ namespace Chaos::Softs
 		}
 	};
 
+	struct CHAOS_API FDeformableXPBDCorotatedParams
+	{
+		int32 XPBDCorotatedBatchSize = 5;
+		int32 XPBDCorotatedBatchThreshold = 5;
+
+	};
+
 
 }; // namesapce Chaos::Softs
\ No newline at end of file
diff --git a/Engine/Source/Runtime/Experimental/Chaos/Public/Chaos/XPBDCorotatedConstraints.h b/Engine/Source/Runtime/Experimental/Chaos/Public/Chaos/XPBDCorotatedConstraints.h
--- a/Engine/Source/Runtime/Experimental/Chaos/Public/Chaos/XPBDCorotatedConstraints.h
+++ b/Engine/Source/Runtime/Experimental/Chaos/Public/Chaos/XPBDCorotatedConstraints.h
@@ -6,6 +6,7 @@
 #include "Chaos/ImplicitQRSVD.h"
 #include "Chaos/GraphColoring.h"
 #include "Chaos/Framework/Parallel.h"
+#include "Chaos/Deformable/ChaosDeformableSolverTypes.h"
 
 DECLARE_CYCLE_STAT(TEXT("Chaos XPBD Corotated Constraint"), STAT_ChaosXPBDCorotated, STATGROUP_Chaos);
 DECLARE_CYCLE_STAT(TEXT("Chaos XPBD Corotated Constraint Polar Compute"), STAT_ChaosXPBDCorotatedPolar, STATGROUP_Chaos);
@@ -13,11 +14,11 @@ DECLARE_CYCLE_STAT(TEXT("Chaos XPBD Corotated Constraint Det Compute"), STAT_Cha
 
 namespace Chaos::Softs
 {
+
 	template <typename T, typename ParticleType>
 	class FXPBDCorotatedConstraints 
 	{
 
-
 	public:
 		//this one only accepts tetmesh input and mesh
 		FXPBDCorotatedConstraints(
@@ -75,6 +76,14 @@ namespace Chaos::Softs
 			Measure.Init((T)0., MeshConstraints.Num());
 			LambdaElementArray.Init((T)0., MeshConstraints.Num());
 			MuElementArray.Init((T)0., MeshConstraints.Num());
+
+			for (int e = 0; e < InMesh.Num(); e++)
+			{
+				for (int32 j = 0; j < 4; j++)
+				{
+					ensure(MeshConstraints[e][j] > -1 && MeshConstraints[e][j] < int32(InParticles.Size()));
+				}
+			}
 			
 			for (int e = 0; e < InMesh.Num(); e++)
 			{
@@ -100,6 +109,55 @@ namespace Chaos::Softs
 			InitColor(InParticles);
 		}
 
+		FXPBDCorotatedConstraints(
+			const ParticleType& InParticles,
+			const TArray<TVector<int32, 4>>& InMesh,
+			const TArray<T>& EMeshArray,
+			const FDeformableXPBDCorotatedParams& InParams,
+			const T& NuMesh = (T).3,
+			const bool bRecordMetricIn = false
+		)
+			: CorotatedParams(InParams), bRecordMetric(bRecordMetricIn), MeshConstraints(InMesh)
+		{
+			ensureMsgf(EMeshArray.Num() == InMesh.Num(), TEXT("Input Young Modulus Array Size is wrong"));
+			LambdaArray.Init((T)0., 2 * MeshConstraints.Num());
+			DmInverse.Init((T)0., 9 * MeshConstraints.Num());
+			Measure.Init((T)0., MeshConstraints.Num());
+			LambdaElementArray.Init((T)0., MeshConstraints.Num());
+			MuElementArray.Init((T)0., MeshConstraints.Num());
+
+			for (int e = 0; e < InMesh.Num(); e++)
+			{
+				for (int32 j = 0; j < 4; j++)
+				{
+					ensure(MeshConstraints[e][j] > -1 && MeshConstraints[e][j] < int32(InParticles.Size()));
+				}
+			}
+
+			for (int e = 0; e < InMesh.Num(); e++)
+			{
+				LambdaElementArray[e] = EMeshArray[e] * NuMesh / (((T)1. + NuMesh) * ((T)1. - (T)2. * NuMesh));
+				MuElementArray[e] = EMeshArray[e] / ((T)2. * ((T)1. + NuMesh));
+
+				PMatrix<T, 3, 3> Dm = DsInit(e, InParticles);
+				PMatrix<T, 3, 3> DmInv = Dm.Inverse();
+				for (int r = 0; r < 3; r++) {
+					for (int c = 0; c < 3; c++) {
+						DmInverse[(3 * 3) * e + 3 * r + c] = DmInv.GetAt(r, c);
+					}
+				}
+
+				Measure[e] = Dm.Determinant() / (T)6.;
+
+				if (Measure[e] < (T)0.)
+				{
+					Measure[e] = -Measure[e];
+				}
+			}
+
+			InitColor(InParticles);
+		}
+
 		FXPBDCorotatedConstraints(
 			const ParticleType& InParticles,
 			const TArray<TVector<int32, 4>>& InMesh,
@@ -180,25 +238,19 @@ namespace Chaos::Softs
 		{
 			TRACE_CPUPROFILER_EVENT_SCOPE(TEXT("STAT_ChaosXPBDCorotatedApplySingle"));
 
-			//if (MuElementArray[ElementIndex] > 300.f)
-			//{
+			TVec4<TVector<T, 3>> PolarDelta = GetPolarDelta(Particles, Dt, ElementIndex);
 
+			for (int i = 0; i < 4; i++)
+			{
+				Particles.P(MeshConstraints[ElementIndex][i]) += PolarDelta[i];
+			}
 
-				TVec4<TVector<T, 3>> PolarDelta = GetPolarDelta(Particles, Dt, ElementIndex);
-
-				for (int i = 0; i < 4; i++)
-				{
-					Particles.P(MeshConstraints[ElementIndex][i]) += PolarDelta[i];
-				}
-
-				TVec4<TVector<T, 3>> DetDelta = GetDeterminantDelta(Particles, Dt, ElementIndex);
-
-				for (int i = 0; i < 4; i++)
-				{
-					Particles.P(MeshConstraints[ElementIndex][i]) += DetDelta[i];
-				}
-			//}
+			TVec4<TVector<T, 3>> DetDelta = GetDeterminantDelta(Particles, Dt, ElementIndex);
 
+			for (int i = 0; i < 4; i++)
+			{
+				Particles.P(MeshConstraints[ElementIndex][i]) += DetDelta[i];
+			}
 
 		}
 
@@ -206,9 +258,16 @@ namespace Chaos::Softs
 		{
 			SCOPE_CYCLE_COUNTER(STAT_ChaosXPBDCorotated);
 			TRACE_CPUPROFILER_EVENT_SCOPE(TEXT("STAT_ChaosXPBDCorotatedApplySerial"));
-			for (int32 ElementIndex = 0; ElementIndex < MeshConstraints.Num(); ++ElementIndex)
+			const int32 ConstraintColorNum = ConstraintsPerColorStartIndex.Num() - 1;
+			for (int32 ConstraintColorIndex = 0; ConstraintColorIndex < ConstraintColorNum; ++ConstraintColorIndex)
 			{
-				ApplyInSerial(Particles, Dt, ElementIndex);
+				const int32 ColorStart = ConstraintsPerColorStartIndex[ConstraintColorIndex];
+				const int32 ColorSize = ConstraintsPerColorStartIndex[ConstraintColorIndex + 1] - ColorStart;
+				for (int32 Index = 0; Index < ColorSize; Index++)
+				{
+					const int32 ConstraintIndex = ColorStart + Index;
+					ApplyInSerial(Particles, Dt, ConstraintIndex);
+				}
 			}
 
 
@@ -234,11 +293,24 @@ namespace Chaos::Softs
 					{
 						const int32 ColorStart = ConstraintsPerColorStartIndex[ConstraintColorIndex];
 						const int32 ColorSize = ConstraintsPerColorStartIndex[ConstraintColorIndex + 1] - ColorStart;
-						PhysicsParallelFor(ColorSize, [&](const int32 Index)
+
+						int32 NumBatch = ColorSize / CorotatedParams.XPBDCorotatedBatchSize;
+						if (ColorSize % CorotatedParams.XPBDCorotatedBatchSize != 0)
+						{
+							NumBatch += 1;
+						}
+
+						PhysicsParallelFor(NumBatch, [&](const int32 BatchIndex)
 							{
-								const int32 ConstraintIndex = ColorStart + Index;
-								ApplyInSerial(Particles, Dt, ConstraintIndex);
-							});
+								for (int32 BatchSubIndex = 0; BatchSubIndex < CorotatedParams.XPBDCorotatedBatchSize; BatchSubIndex++) {
+									int32 TaskIndex = CorotatedParams.XPBDCorotatedBatchSize * BatchIndex + BatchSubIndex;
+									const int32 ConstraintIndex = ColorStart + TaskIndex;
+									if (ConstraintIndex < ColorStart + ColorSize)
+									{
+										ApplyInSerial(Particles, Dt, ConstraintIndex);
+									}
+								}
+							}, NumBatch < CorotatedParams.XPBDCorotatedBatchThreshold);
 					}
 				}
 			}
@@ -499,6 +571,9 @@ namespace Chaos::Softs
 		mutable TArray<T> LambdaArray;
 		mutable TArray<T> DmInverse;
 
+		//parallel data:
+		FDeformableXPBDCorotatedParams CorotatedParams;
+
 		//material constants calculated from E:
 		T Mu;
 		T Lambda;
diff --git a/Engine/Source/Runtime/Experimental/Chaos/Public/Chaos/XPBDWeakConstraints.h b/Engine/Source/Runtime/Experimental/Chaos/Public/Chaos/XPBDWeakConstraints.h
--- a/Engine/Source/Runtime/Experimental/Chaos/Public/Chaos/XPBDWeakConstraints.h
+++ b/Engine/Source/Runtime/Experimental/Chaos/Public/Chaos/XPBDWeakConstraints.h
@@ -124,8 +124,13 @@ namespace Chaos::Softs
 			Chaos::ComputeWeakConstraintsColoring(Indices, SecondIndices, Particles, ConstraintsPerColor);
 		}
 
-		void ApplySingleConstraintWithoutSelfTarget(ParticleType& Particles, const T Dt, const int32 ConstraintIndex) const
+		void ApplySingleConstraintWithoutSelfTarget(ParticleType& Particles, const T Dt, const int32 ConstraintIndex, const T Tol = UE_KINDA_SMALL_NUMBER) const
 		{
+			if (Stiffness[ConstraintIndex] < Tol)
+			{
+				return;
+			}
+
 			T AlphaTilde = T(2) / (Stiffness[ConstraintIndex] * Dt * Dt);
 
 			if (Stiffness[ConstraintIndex] > 1e14)
@@ -161,8 +166,13 @@ namespace Chaos::Softs
 			}
 		}
 
-		void ApplySingleConstraintWithSelfTarget(ParticleType& Particles, const T Dt, const int32 ConstraintIndex) const
+		void ApplySingleConstraintWithSelfTarget(ParticleType& Particles, const T Dt, const int32 ConstraintIndex, const T Tol = UE_KINDA_SMALL_NUMBER) const
 		{
+			if (Stiffness[ConstraintIndex] < Tol)
+			{
+				return;
+			}
+
 			ensure(SecondIndices.Num() > 0);
 
 			T AlphaTilde = T(2) / (Stiffness[ConstraintIndex] * Dt * Dt);
