Fix issues with plane collision and geometry collection by replacing plane with a solver actor hoirizontally centered box

diff --git a/Engine/Source/Runtime/Experimental/ChaosSolverEngine/Private/Chaos/ChaosSolverActor.cpp b/Engine/Source/Runtime/Experimental/ChaosSolverEngine/Private/Chaos/ChaosSolverActor.cpp
--- a/Engine/Source/Runtime/Experimental/ChaosSolverEngine/Private/Chaos/ChaosSolverActor.cpp
+++ b/Engine/Source/Runtime/Experimental/ChaosSolverEngine/Private/Chaos/ChaosSolverActor.cpp
@@ -404,8 +404,14 @@ void AChaosSolverActor::MakeFloor()
 	if(bHasFloor)
 	{
 		TUniquePtr<Chaos::FGeometryParticle> FloorParticle = Chaos::FGeometryParticle::CreateParticle();
-		FloorParticle->SetGeometry(TUniquePtr<Chaos::TPlane<Chaos::FReal, 3>>(new Chaos::TPlane<Chaos::FReal, 3>(FVector(0), FVector(0, 0, 1))));
-		FloorParticle->SetX(Chaos::FVec3(0.f, 0.f, FloorHeight));
+		// todo(chaos) : Changing the floor to be a box for now because there's few cases where this may fail to collide with geometry collection
+		//				 since the box is finite, let's center it at the actor position for better user control
+		//FloorParticle->SetGeometry(TUniquePtr<Chaos::TPlane<Chaos::FReal, 3>>(new Chaos::TPlane<Chaos::FReal, 3>(FVector(0), FVector(0, 0, 1))));
+		const FVector SolverLocation = GetActorLocation();
+		const FVector BoxMin(-100000, -100000, -1000);
+		const FVector BoxMax(100000, 100000, 0);
+		FloorParticle->SetGeometry(TUniquePtr<Chaos::TBox<Chaos::FReal, 3>>(new Chaos::TBox<Chaos::FReal, 3>(BoxMin, BoxMax)));
+		FloorParticle->SetX(Chaos::FVec3(SolverLocation.X, SolverLocation.Y, FloorHeight));
 		FCollisionFilterData FilterData;
 		FilterData.Word1 = 0xFFFF;
 		FilterData.Word3 = 0xFFFF;
