Geometry collection : fix warning when cooking GC with p.Chaos.GC.EnableCollisionParticles set to 0
- use skip attribute method instead of post copy remove attribute one

diff --git a/Engine/Source/Runtime/Experimental/GeometryCollectionEngine/Private/GeometryCollection/GeometryCollectionObject.cpp b/Engine/Source/Runtime/Experimental/GeometryCollectionEngine/Private/GeometryCollection/GeometryCollectionObject.cpp
--- a/Engine/Source/Runtime/Experimental/GeometryCollectionEngine/Private/GeometryCollection/GeometryCollectionObject.cpp
+++ b/Engine/Source/Runtime/Experimental/GeometryCollectionEngine/Private/GeometryCollection/GeometryCollectionObject.cpp
@@ -777,8 +777,10 @@ void UGeometryCollection::Serialize(FArchive& Ar)
 			// do we need to remove the simplicial attribute ? 
 			if (false == FGeometryCollection::AreCollisionParticlesEnabled())
 			{
-				ArchiveGeometryCollection = TSharedPtr<FGeometryCollection, ESPMode::ThreadSafe>(GeometryCollection->NewCopy<FGeometryCollection>());
-				ArchiveGeometryCollection->RemoveAttribute(FGeometryDynamicCollection::SimplicialsAttribute, FTransformCollection::TransformGroup);
+				ArchiveGeometryCollection = TSharedPtr<FGeometryCollection, ESPMode::ThreadSafe>(new FGeometryCollection);
+				const TArray<FName> NoGroupsToSkip;
+				const TArray<TTuple<FName, FName>> AttributesToSkip{ { FGeometryDynamicCollection::SimplicialsAttribute, FTransformCollection::TransformGroup } };
+				GeometryCollection->CopyTo(ArchiveGeometryCollection.Get(), NoGroupsToSkip, AttributesToSkip);
 			}
 		}
 
