Chaos Cloth - Renamed PropertyCollectionAdapter into CollectionPropertyFacade, and updated the API to better conform to the rest of the geometry collection facades.

(Review 24122949 v3 committed without the 5.3 specific additions in 5.2.)

#rnx

diff --git a/Engine/Plugins/Enterprise/DatasmithImporter/Source/DatasmithImporter/Private/DatasmithImporter.cpp b/Engine/Plugins/Enterprise/DatasmithImporter/Source/DatasmithImporter/Private/DatasmithImporter.cpp
--- a/Engine/Plugins/Enterprise/DatasmithImporter/Source/DatasmithImporter/Private/DatasmithImporter.cpp
+++ b/Engine/Plugins/Enterprise/DatasmithImporter/Source/DatasmithImporter/Private/DatasmithImporter.cpp
@@ -440,7 +440,7 @@ void FDatasmithImporter::ImportClothes(FDatasmithImportContext& ImportContext)
 
 		UChaosClothAsset* ClothAsset = nullptr;
 
-		// TODO: Use PropertyCollectionAdapter instead, see \Engine\Source\Runtime\Experimental\Chaos\Public\Chaos\PropertyCollectionAdapter.h
+		// TODO: Use FCollectionPropertyFacade instead, see \Engine\Source\Runtime\Experimental\Chaos\Public\Chaos\CollectionPropertyFacade.h
 		//TArray<UObject*> ClothPresetAssets;
 
 		//for (const FDatasmithClothPresetPropertySet& PropertySet : DsCloth.PropertySets)
@@ -554,7 +554,7 @@ void FDatasmithImporter::ImportClothes(FDatasmithImportContext& ImportContext)
 
 		ImportContext.ImportedClothes.Add(ClothElement, ClothAsset);
 
-		// TODO: Use PropertyCollectionAdapter instead, see \Engine\Source\Runtime\Experimental\Chaos\Public\Chaos\PropertyCollectionAdapter.h
+		// TODO: Use FCollectionPropertyFacade instead, see \Engine\Source\Runtime\Experimental\Chaos\Public\Chaos\CollectionPropertyFacade.h
 		//for (UObject* ClothPresetAsset : ClothPresetAssets)
 		//{
 		//	ImportContext.ImportedClothPresets.Add(ClothPresetAsset);
diff --git a/Engine/Plugins/Experimental/ChaosCloth/Source/Chaos/Private/ChaosCloth/ChaosClothConstraints.cpp b/Engine/Plugins/Experimental/ChaosCloth/Source/Chaos/Private/ChaosCloth/ChaosClothConstraints.cpp
--- a/Engine/Plugins/Experimental/ChaosCloth/Source/Chaos/Private/ChaosCloth/ChaosClothConstraints.cpp
+++ b/Engine/Plugins/Experimental/ChaosCloth/Source/Chaos/Private/ChaosCloth/ChaosClothConstraints.cpp
@@ -16,7 +16,7 @@
 #include "Chaos/PBDTriangleMeshCollisions.h"
 #include "Chaos/PBDTriangleMeshIntersections.h"
 #include "Chaos/PBDEvolution.h"
-#include "Chaos/PropertyCollectionAdapter.h"
+#include "Chaos/CollectionPropertyFacade.h"
 
 namespace Chaos {
 
@@ -96,7 +96,7 @@ void FClothConstraints::Enable(bool bEnable)
 }
 
 void FClothConstraints::AddRules(
-	const Softs::FPropertyCollectionConstAdapter& ConfigProperties,
+	const Softs::FCollectionPropertyConstFacade& ConfigProperties,
 	const FTriangleMesh& TriangleMesh,
 	const TArray<TConstArrayView<FRealSingle>>& WeightMaps,
 	const TArray<TConstArrayView<TTuple<int32, int32, FRealSingle>>>& Tethers,
@@ -136,7 +136,7 @@ PRAGMA_ENABLE_DEPRECATION_WARNINGS
 	Enable(bEnabled);
 }
 
-void FClothConstraints::CreateSelfCollisionConstraints(const Softs::FPropertyCollectionConstAdapter& ConfigProperties, const FTriangleMesh& TriangleMesh)
+void FClothConstraints::CreateSelfCollisionConstraints(const Softs::FCollectionPropertyConstFacade& ConfigProperties, const FTriangleMesh& TriangleMesh)
 {
 	const bool bUseSelfCollisions = ConfigProperties.GetValue<bool>(TEXT("UseSelfCollisions"));
 
@@ -185,7 +185,7 @@ void FClothConstraints::CreateSelfCollisionConstraints(const Softs::FPropertyCol
 }
 
 void FClothConstraints::CreateEdgeConstraints(
-	const Softs::FPropertyCollectionConstAdapter& ConfigProperties,
+	const Softs::FCollectionPropertyConstFacade& ConfigProperties,
 	const TArray<TConstArrayView<FRealSingle>>& WeightMaps,
 	const FTriangleMesh& TriangleMesh)
 {
@@ -226,7 +226,7 @@ void FClothConstraints::CreateEdgeConstraints(
 }
 
 void FClothConstraints::CreateBendingConstraints(
-	const Softs::FPropertyCollectionConstAdapter& ConfigProperties,
+	const Softs::FCollectionPropertyConstFacade& ConfigProperties,
 	const TArray<TConstArrayView<FRealSingle>>& WeightMaps,
 	const FTriangleMesh& TriangleMesh)
 {
@@ -314,7 +314,7 @@ void FClothConstraints::CreateBendingConstraints(
 }
 
 void FClothConstraints::CreateAreaConstraints(
-	const Softs::FPropertyCollectionConstAdapter& ConfigProperties,
+	const Softs::FCollectionPropertyConstFacade& ConfigProperties,
 	const TArray<TConstArrayView<FRealSingle>>& WeightMaps,
 	const FTriangleMesh& TriangleMesh)
 {
@@ -353,7 +353,7 @@ void FClothConstraints::CreateAreaConstraints(
 }
 
 void FClothConstraints::CreateLongRangeConstraints(
-	const Softs::FPropertyCollectionConstAdapter& ConfigProperties,
+	const Softs::FCollectionPropertyConstFacade& ConfigProperties,
 	const TArray<TConstArrayView<FRealSingle>>& WeightMaps,
 	const TArray<TConstArrayView<TTuple<int32, int32, FRealSingle>>>& Tethers,
 	Softs::FSolverReal MeshScale)
@@ -379,7 +379,7 @@ void FClothConstraints::CreateLongRangeConstraints(
 }
 
 void FClothConstraints::CreateMaxDistanceConstraints(
-	const Softs::FPropertyCollectionConstAdapter& ConfigProperties,
+	const Softs::FCollectionPropertyConstFacade& ConfigProperties,
 	const TArray<TConstArrayView<FRealSingle>>& WeightMaps,
 	Softs::FSolverReal MeshScale)
 {
@@ -399,7 +399,7 @@ void FClothConstraints::CreateMaxDistanceConstraints(
 }
 
 void FClothConstraints::CreateBackstopConstraints(
-	const Softs::FPropertyCollectionConstAdapter& ConfigProperties,
+	const Softs::FCollectionPropertyConstFacade& ConfigProperties,
 	const TArray<TConstArrayView<FRealSingle>>& WeightMaps,
 	Softs::FSolverReal MeshScale)
 {
@@ -425,7 +425,7 @@ void FClothConstraints::CreateBackstopConstraints(
 	}
 }
 
-void FClothConstraints::CreateAnimDriveConstraints(const Softs::FPropertyCollectionConstAdapter& ConfigProperties, const TArray<TConstArrayView<FRealSingle>>& WeightMaps)
+void FClothConstraints::CreateAnimDriveConstraints(const Softs::FCollectionPropertyConstFacade& ConfigProperties, const TArray<TConstArrayView<FRealSingle>>& WeightMaps)
 {
 	if (Softs::FPBDAnimDriveConstraint::IsEnabled(ConfigProperties))
 	{
@@ -1113,7 +1113,7 @@ void FClothConstraints::SetSelfCollisionConstraints(const FTriangleMesh& Triangl
 }
 
 void FClothConstraints::Update(
-	const Softs::FPropertyCollectionConstAdapter& ConfigProperties,
+	const Softs::FCollectionPropertyConstFacade& ConfigProperties,
 	Softs::FSolverReal MeshScale,
 	Softs::FSolverReal MaxDistancesScale)
 {
diff --git a/Engine/Plugins/Experimental/ChaosCloth/Source/Chaos/Private/ChaosCloth/ChaosClothingSimulationCloth.cpp b/Engine/Plugins/Experimental/ChaosCloth/Source/Chaos/Private/ChaosCloth/ChaosClothingSimulationCloth.cpp
--- a/Engine/Plugins/Experimental/ChaosCloth/Source/Chaos/Private/ChaosCloth/ChaosClothingSimulationCloth.cpp
+++ b/Engine/Plugins/Experimental/ChaosCloth/Source/Chaos/Private/ChaosCloth/ChaosClothingSimulationCloth.cpp
@@ -6,7 +6,7 @@
 #include "ChaosCloth/ChaosClothingSimulationConfig.h"
 #include "ChaosCloth/ChaosWeightMapTarget.h"
 #include "ChaosCloth/ChaosClothPrivate.h"
-#include "Chaos/PropertyCollectionAdapter.h"
+#include "Chaos/CollectionPropertyFacade.h"
 #include "Containers/ArrayView.h"
 #include "GeometryCollection/ManagedArrayCollection.h"
 #include "HAL/IConsoleManager.h"
@@ -149,7 +149,7 @@ PRAGMA_ENABLE_DEPRECATION_WARNINGS
 		};
 
 	check(Cloth->Config);
-	const Softs::FPropertyCollectionConstAdapter& ConfigProperties = Cloth->Config->GetProperties();
+	const Softs::FCollectionPropertyFacade& ConfigProperties = Cloth->Config->GetProperties();
 
 	const int32 MassMode = ConfigProperties.GetValue<int32>(TEXT("MassMode"), ClothingSimulationClothDefault::MassMode);
 	const FRealSingle MassValue = (FRealSingle)ConfigProperties.GetValue<float>(TEXT("MassValue"), ClothingSimulationClothDefault::MassValue);
@@ -334,7 +334,8 @@ PRAGMA_ENABLE_DEPRECATION_WARNINGS
 {
 	// Turn parameters into a config
 	PropertyCollection = MakeShared<FManagedArrayCollection>();
-	Softs::FPropertyCollectionMutableAdapter Properties(PropertyCollection);
+	Softs::FCollectionPropertyMutableFacade Properties(PropertyCollection);
+	Properties.DefineSchema();
 
 	constexpr bool bEnable = true;
 	constexpr bool bAnimatable = true;
@@ -697,7 +698,7 @@ TVec3<FRealSingle> FClothingSimulationCloth::GetGravity(const FClothingSimulatio
 {
 	check(Solver);
 	check(Config);
-	const Softs::FPropertyCollectionConstAdapter& ConfigProperties = Config->GetProperties();
+	const Softs::FCollectionPropertyFacade& ConfigProperties = Config->GetProperties();
 
 	const bool bUseGravityOverride = ConfigProperties.GetValue<bool>(TEXT("UseGravityOverride"));
 	const TVec3<FRealSingle> GravityOverride = (TVec3<FRealSingle>)ConfigProperties.GetValue<FVector3f>(TEXT("GravityOverride"), FVector3f(0.f, 0.f, ClothingSimulationClothDefault::GravityZOverride));
@@ -868,7 +869,7 @@ PRAGMA_ENABLE_DEPRECATION_WARNINGS
 	if (LODIndex != INDEX_NONE)
 	{
 		check(Config);
-		const Softs::FPropertyCollectionConstAdapter& ConfigProperties = Config->GetProperties();
+		const Softs::FCollectionPropertyFacade& ConfigProperties = Config->GetProperties();
 
 		// Update Cloth group parameters  TODO: Cloth groups should exist as their own node object so that they can be used by several cloth objects
 		LODData[LODIndex]->Update(Solver, this);
diff --git a/Engine/Plugins/Experimental/ChaosCloth/Source/Chaos/Private/ChaosCloth/ChaosClothingSimulationConfig.cpp b/Engine/Plugins/Experimental/ChaosCloth/Source/Chaos/Private/ChaosCloth/ChaosClothingSimulationConfig.cpp
--- a/Engine/Plugins/Experimental/ChaosCloth/Source/Chaos/Private/ChaosCloth/ChaosClothingSimulationConfig.cpp
+++ b/Engine/Plugins/Experimental/ChaosCloth/Source/Chaos/Private/ChaosCloth/ChaosClothingSimulationConfig.cpp
@@ -1,7 +1,7 @@
 // Copyright Epic Games, Inc. All Rights Reserved.
 #include "ChaosCloth/ChaosClothingSimulationConfig.h"
 #include "ChaosCloth/ChaosClothConfig.h"
-#include "Chaos/PropertyCollectionAdapter.h"
+#include "Chaos/CollectionPropertyFacade.h"
 #include "Chaos/PBDLongRangeConstraints.h"  // For Tether modes
 #include "GeometryCollection/ManagedArrayCollection.h"
 
@@ -9,13 +9,13 @@ namespace Chaos
 {
 	FClothingSimulationConfig::FClothingSimulationConfig()
 		: PropertyCollection(MakeShared<FManagedArrayCollection>())
-		, Properties(MakeUnique<Softs::FPropertyCollectionMutableAdapter>(PropertyCollection))
+		, Properties(MakeUnique<Softs::FCollectionPropertyMutableFacade>(PropertyCollection))
 	{
 	}
 
 	FClothingSimulationConfig::FClothingSimulationConfig(const TSharedPtr<const FManagedArrayCollection>& InPropertyCollection)
 		: PropertyCollection(MakeShared<FManagedArrayCollection>())
-		, Properties(MakeUnique<Softs::FPropertyCollectionMutableAdapter>(PropertyCollection))
+		, Properties(MakeUnique<Softs::FCollectionPropertyMutableFacade>(PropertyCollection))
 	{
 		Initialize(InPropertyCollection);
 	}
@@ -26,7 +26,8 @@ namespace Chaos
 		constexpr bool bAnimatable = true;
 
 		// Clear all properties
-		Properties->Reset();
+		PropertyCollection->Reset();
+		Properties->DefineSchema();
 
 		// Solver properties
 		if (ClothSharedConfig)
@@ -205,15 +206,15 @@ namespace Chaos
 
 	void FClothingSimulationConfig::Initialize(const TSharedPtr<const FManagedArrayCollection>& InPropertyCollection)
 	{
-		Properties->Copy(InPropertyCollection);
+		Properties->Copy(*InPropertyCollection);
 	}
 
-	const Softs::FPropertyCollectionConstAdapter& FClothingSimulationConfig::GetProperties() const
+	const Softs::FCollectionPropertyConstFacade& FClothingSimulationConfig::GetProperties() const
 	{
 		return *Properties;
 	}
 
-	Softs::FPropertyCollectionAdapter& FClothingSimulationConfig::GetProperties()
+	Softs::FCollectionPropertyFacade& FClothingSimulationConfig::GetProperties()
 	{
 		return *Properties;
 	}
diff --git a/Engine/Plugins/Experimental/ChaosCloth/Source/Chaos/Private/ChaosCloth/ChaosClothingSimulationSolver.cpp b/Engine/Plugins/Experimental/ChaosCloth/Source/Chaos/Private/ChaosCloth/ChaosClothingSimulationSolver.cpp
--- a/Engine/Plugins/Experimental/ChaosCloth/Source/Chaos/Private/ChaosCloth/ChaosClothingSimulationSolver.cpp
+++ b/Engine/Plugins/Experimental/ChaosCloth/Source/Chaos/Private/ChaosCloth/ChaosClothingSimulationSolver.cpp
@@ -1118,7 +1118,7 @@ void FClothingSimulationSolver::Update(Softs::FSolverReal InDeltaTime)
 			int32 ConfigNumSubsteps;
 			if (Config)
 			{
-				const Softs::FPropertyCollectionConstAdapter& Properties = Config->GetProperties();
+				const Softs::FCollectionPropertyFacade& Properties = Config->GetProperties();
 
 				ConfigMaxNumIterations = FMath::Max(
 					Properties.GetValue<int32>(TEXT("MaxNumIterations"), ClothingSimulationSolverDefault::MaxNumIterations),
diff --git a/Engine/Plugins/Experimental/ChaosCloth/Source/Chaos/Public/ChaosCloth/ChaosClothConstraints.h b/Engine/Plugins/Experimental/ChaosCloth/Source/Chaos/Public/ChaosCloth/ChaosClothConstraints.h
--- a/Engine/Plugins/Experimental/ChaosCloth/Source/Chaos/Public/ChaosCloth/ChaosClothConstraints.h
+++ b/Engine/Plugins/Experimental/ChaosCloth/Source/Chaos/Public/ChaosCloth/ChaosClothConstraints.h
@@ -39,7 +39,7 @@ namespace Chaos
 
 		// ---- Cloth interface ----
 		void AddRules(
-			const Softs::FPropertyCollectionConstAdapter& ConfigProperties,
+			const Softs::FCollectionPropertyConstFacade& ConfigProperties,
 			const FTriangleMesh& TriangleMesh,
 			const TArray<TConstArrayView<FRealSingle>>& WeightMaps,
 			const TArray<TConstArrayView<TTuple<int32, int32, FRealSingle>>>& Tethers,
@@ -47,7 +47,7 @@ namespace Chaos
 			bool bEnabled);
 
 		void Update(
-			const Softs::FPropertyCollectionConstAdapter& ConfigProperties,
+			const Softs::FCollectionPropertyConstFacade& ConfigProperties,
 			Softs::FSolverReal MeshScale,
 			Softs::FSolverReal MaxDistancesScale = (Softs::FSolverReal)1.);
 
@@ -172,35 +172,35 @@ namespace Chaos
 
 	private:
 		void CreateSelfCollisionConstraints(
-			const Softs::FPropertyCollectionConstAdapter& ConfigProperties,
+			const Softs::FCollectionPropertyConstFacade& ConfigProperties,
 			const FTriangleMesh& TriangleMesh);
 		void CreateEdgeConstraints(
-			const Softs::FPropertyCollectionConstAdapter& ConfigProperties,
+			const Softs::FCollectionPropertyConstFacade& ConfigProperties,
 			const TArray<TConstArrayView<FRealSingle>>& WeightMaps,
 			const FTriangleMesh& TriangleMesh);
 		void CreateBendingConstraints(
-			const Softs::FPropertyCollectionConstAdapter& ConfigProperties,
+			const Softs::FCollectionPropertyConstFacade& ConfigProperties,
 			const TArray<TConstArrayView<FRealSingle>>& WeightMaps,
 			const FTriangleMesh& TriangleMesh);
 		void CreateAreaConstraints(
-			const Softs::FPropertyCollectionConstAdapter& ConfigProperties,
+			const Softs::FCollectionPropertyConstFacade& ConfigProperties,
 			const TArray<TConstArrayView<FRealSingle>>& WeightMaps,
 			const FTriangleMesh& TriangleMesh);
 		void CreateLongRangeConstraints(
-			const Softs::FPropertyCollectionConstAdapter& ConfigProperties,
+			const Softs::FCollectionPropertyConstFacade& ConfigProperties,
 			const TArray<TConstArrayView<FRealSingle>>& WeightMaps,
 			const TArray<TConstArrayView<TTuple<int32, int32, FRealSingle>>>& Tethers,
 			Softs::FSolverReal MeshScale);
 		void CreateMaxDistanceConstraints(
-			const Softs::FPropertyCollectionConstAdapter& ConfigProperties,
+			const Softs::FCollectionPropertyConstFacade& ConfigProperties,
 			const TArray<TConstArrayView<FRealSingle>>& WeightMaps,
 			Softs::FSolverReal MeshScale);
 		void CreateBackstopConstraints(
-			const Softs::FPropertyCollectionConstAdapter& ConfigProperties,
+			const Softs::FCollectionPropertyConstFacade& ConfigProperties,
 			const TArray<TConstArrayView<FRealSingle>>& WeightMaps,
 			Softs::FSolverReal MeshScale);
 		void CreateAnimDriveConstraints(
-			const Softs::FPropertyCollectionConstAdapter& ConfigProperties,
+			const Softs::FCollectionPropertyConstFacade& ConfigProperties,
 			const TArray<TConstArrayView<FRealSingle>>& WeightMaps);
 
 		//~ Begin deprecated constraints
diff --git a/Engine/Plugins/Experimental/ChaosCloth/Source/Chaos/Public/ChaosCloth/ChaosClothingSimulationConfig.h b/Engine/Plugins/Experimental/ChaosCloth/Source/Chaos/Public/ChaosCloth/ChaosClothingSimulationConfig.h
--- a/Engine/Plugins/Experimental/ChaosCloth/Source/Chaos/Public/ChaosCloth/ChaosClothingSimulationConfig.h
+++ b/Engine/Plugins/Experimental/ChaosCloth/Source/Chaos/Public/ChaosCloth/ChaosClothingSimulationConfig.h
@@ -12,9 +12,9 @@ namespace Chaos
 {
 	namespace Softs
 	{
-		class FPropertyCollectionConstAdapter;
-		class FPropertyCollectionAdapter;
-		class FPropertyCollectionMutableAdapter;
+		class FCollectionPropertyConstFacade;
+		class FCollectionPropertyFacade;
+		class FCollectionPropertyMutableFacade;
 	}
 
 	// Cloth simulation properties
@@ -37,17 +37,17 @@ namespace Chaos
 		/** Initialize config from a property collection. */
 		void Initialize(const TSharedPtr<const FManagedArrayCollection>& InPropertyCollection);
 
-		/** Return a property collection adapter for reading properties from this configuration. */
-		const Softs::FPropertyCollectionConstAdapter& GetProperties() const;
+		/** Return a property collection facade for reading properties from this configuration. */
+		const Softs::FCollectionPropertyConstFacade& GetProperties() const;
 
-		/** Return a property collection adapter for setting properties for this configuration. */
-		Softs::FPropertyCollectionAdapter& GetProperties();
+		/** Return a property collection facade for setting properties for this configuration. */
+		Softs::FCollectionPropertyFacade& GetProperties();
 
 		/** Return this configuration's internal property collection. */
 		TSharedPtr<const FManagedArrayCollection> GetPropertyCollection() const { return TSharedPtr<const FManagedArrayCollection>(PropertyCollection); }
 
 	private:
 		TSharedPtr<FManagedArrayCollection> PropertyCollection;
-		TUniquePtr<Softs::FPropertyCollectionMutableAdapter> Properties;
+		TUniquePtr<Softs::FCollectionPropertyMutableFacade> Properties;
 	};
 } // namespace Chaos
diff --git a/Engine/Plugins/Experimental/ChaosClothAsset/Source/ChaosClothAssetEngine/Private/ChaosClothAsset/ClothComponent.cpp b/Engine/Plugins/Experimental/ChaosClothAsset/Source/ChaosClothAssetEngine/Private/ChaosClothAsset/ClothComponent.cpp
--- a/Engine/Plugins/Experimental/ChaosClothAsset/Source/ChaosClothAssetEngine/Private/ChaosClothAsset/ClothComponent.cpp
+++ b/Engine/Plugins/Experimental/ChaosClothAsset/Source/ChaosClothAssetEngine/Private/ChaosClothAsset/ClothComponent.cpp
@@ -6,7 +6,7 @@
 #include "ChaosClothAsset/ClothCollection.h"
 #include "ChaosClothAsset/ClothSimulationModel.h"
 #include "ChaosClothAsset/ClothSimulationProxy.h"
-#include "Chaos/PropertyCollectionAdapter.h"
+#include "Chaos/CollectionPropertyFacade.h"
 #include "HAL/IConsoleManager.h"
 #include "Rendering/SkeletalMeshRenderData.h"
 #include "SkeletalRenderPublic.h"
@@ -128,19 +128,21 @@ void UChaosClothComponent::ResetConfigProperties()
 	using namespace UE::Chaos::ClothAsset;
 	using namespace ::Chaos::Softs;
 
-	PropertyCollectionAdapter.Reset();
+	// Create a mutable facade for our component's property collection, and use it to copy the properties from the cloth collection
+	FCollectionPropertyMutableFacade CollectionPropertyMutableFacade(PropertyCollection);
 
-	if (GetClothAsset())
+	if (GetClothAsset() && GetClothAsset()->GetClothCollection())
 	{
-		if (TSharedPtr<const FClothCollection> ClothCollection = GetClothAsset()->GetClothCollection())
-		{
-			// Create a mutable adapter for our component's property collection, and use it to copy the properties from the cloth collection
-			FPropertyCollectionMutableAdapter(PropertyCollection).Copy(TSharedPtr<const FManagedArrayCollection>(ClothCollection));
-
-			// The Adapter used to access the properties shouldn't be mutable, since adding/removing properties would cause issues
-			PropertyCollectionAdapter = MakeUnique<FPropertyCollectionAdapter>(PropertyCollection);
-		}
+		CollectionPropertyMutableFacade.Copy(*GetClothAsset()->GetClothCollection());
 	}
+	else
+	{
+		PropertyCollection->Reset();
+		CollectionPropertyMutableFacade.DefineSchema();
+	}
+
+	// The facade used to access the properties shouldn't be mutable, since adding/removing properties would cause issues
+	CollectionPropertyFacade = MakeUnique<FCollectionPropertyFacade>(PropertyCollection);
 }
 
 void UChaosClothComponent::OnUnregister()
diff --git a/Engine/Plugins/Experimental/ChaosClothAsset/Source/ChaosClothAssetEngine/Public/ChaosClothAsset/ClothComponent.h b/Engine/Plugins/Experimental/ChaosClothAsset/Source/ChaosClothAssetEngine/Public/ChaosClothAsset/ClothComponent.h
--- a/Engine/Plugins/Experimental/ChaosClothAsset/Source/ChaosClothAssetEngine/Public/ChaosClothAsset/ClothComponent.h
+++ b/Engine/Plugins/Experimental/ChaosClothAsset/Source/ChaosClothAssetEngine/Public/ChaosClothAsset/ClothComponent.h
@@ -11,7 +11,7 @@ struct FManagedArrayCollection;
 
 namespace Chaos::Softs
 {
-	class FPropertyCollectionAdapter;
+	class FCollectionPropertyFacade;
 }
 
 namespace UE::Chaos::ClothAsset
@@ -159,7 +159,7 @@ private:
 	float BlendWeight = 1.f;
 
 	TSharedPtr<FManagedArrayCollection> PropertyCollection;
-	TUniquePtr<::Chaos::Softs::FPropertyCollectionAdapter> PropertyCollectionAdapter;
+	TUniquePtr<::Chaos::Softs::FCollectionPropertyFacade> CollectionPropertyFacade;
 
 	TUniquePtr<UE::Chaos::ClothAsset::FClothSimulationProxy> ClothSimulationProxy;
 };
diff --git a/Engine/Source/Runtime/Experimental/Chaos/Private/Chaos/CollectionPropertyFacade.cpp b/Engine/Source/Runtime/Experimental/Chaos/Private/Chaos/CollectionPropertyFacade.cpp
new file mode 100644
--- /dev/null
+++ b/Engine/Source/Runtime/Experimental/Chaos/Private/Chaos/CollectionPropertyFacade.cpp
@@ -0,0 +1,245 @@
+// Copyright Epic Games, Inc. All Rights Reserved.
+
+#include "Chaos/CollectionPropertyFacade.h"
+#include "GeometryCollection/ManagedArrayCollection.h"
+
+namespace Chaos::Softs
+{
+	namespace PropertyFacadeNames
+	{
+		// Attribute groups, predefined data member of the collection
+		static const FName PropertyGroup("Property");
+		static const FName KeyName("Key");  // Property key, name to look for
+		static const FName LowValueName("LowValue");  // Boolean, 24 bit integer (max 16777215), float, or vector value, or value of the lowest weight on the weight map if any
+		static const FName HighValueName("HighValue");  // Boolean, 24 bit integer (max 16777215), float, or vector value of the highest weight on the weight map if any
+		static const FName StringValueName("StringValue");  // String value, or weight map name, ...etc.
+		static const FName FlagsName("Flags");  // Whether this property is enabled, animatable, ...etc.
+	}
+	
+	FCollectionPropertyConstFacade::FCollectionPropertyConstFacade(const TSharedPtr<const FManagedArrayCollection>& InManagedArrayCollection)
+		: ManagedArrayCollection(InManagedArrayCollection)
+	{
+		UpdateArrays();
+		RebuildKeyIndices();
+	}
+
+	FCollectionPropertyConstFacade::FCollectionPropertyConstFacade(const TSharedPtr<const FManagedArrayCollection>& InManagedArrayCollection, ENoInit)
+		: ManagedArrayCollection(InManagedArrayCollection)
+	{
+	}
+
+	bool FCollectionPropertyConstFacade::IsValid() const
+	{
+		return
+			ManagedArrayCollection->HasAttribute(PropertyFacadeNames::KeyName, PropertyFacadeNames::PropertyGroup) &&
+			ManagedArrayCollection->HasAttribute(PropertyFacadeNames::LowValueName, PropertyFacadeNames::PropertyGroup) &&
+			ManagedArrayCollection->HasAttribute(PropertyFacadeNames::HighValueName, PropertyFacadeNames::PropertyGroup) &&
+			ManagedArrayCollection->HasAttribute(PropertyFacadeNames::StringValueName, PropertyFacadeNames::PropertyGroup) &&
+			ManagedArrayCollection->HasAttribute(PropertyFacadeNames::FlagsName, PropertyFacadeNames::PropertyGroup);
+	}
+
+	void FCollectionPropertyConstFacade::UpdateArrays()
+	{
+		KeyArray = GetArray<FString>(PropertyFacadeNames::KeyName);
+		LowValueArray = GetArray<FVector3f>(PropertyFacadeNames::LowValueName);
+		HighValueArray = GetArray<FVector3f>(PropertyFacadeNames::HighValueName);
+		StringValueArray = GetArray<FString>(PropertyFacadeNames::StringValueName);
+		FlagsArray = GetArray<uint8>(PropertyFacadeNames::FlagsName);
+	}
+	
+	void FCollectionPropertyConstFacade::RebuildKeyIndices()
+	{
+		// Create a fast access search map (although it might only be faster for a large enough number of properties)
+		const int32 NumKeys = KeyArray.Num();
+		KeyIndices.Empty(NumKeys);
+		for (int32 Index = 0; Index < NumKeys; ++Index)
+		{
+			KeyIndices.Emplace(KeyArray[Index], Index);
+		}
+	}
+
+	template<typename T, typename ElementType>
+	T FCollectionPropertyConstFacade::GetValue(int32 KeyIndex, const TConstArrayView<ElementType>& ValueArray) const
+	{
+		return ValueArray[KeyIndex];
+	}
+	template CHAOS_API FVector3f FCollectionPropertyConstFacade::GetValue<FVector3f, FVector3f>(int32 KeyIndex, const TConstArrayView<FVector3f>& ValueArray) const;
+	template CHAOS_API const FVector3f& FCollectionPropertyConstFacade::GetValue<const FVector3f&, FVector3f>(int32 KeyIndex, const TConstArrayView<FVector3f>& ValueArray) const;
+	template CHAOS_API FString FCollectionPropertyConstFacade::GetValue<FString, FString>(int32 KeyIndex, const TConstArrayView<FString>& ValueArray) const;
+	template CHAOS_API const FString& FCollectionPropertyConstFacade::GetValue<const FString&, FString>(int32 KeyIndex, const TConstArrayView<FString>& ValueArray) const;
+	template CHAOS_API uint8 FCollectionPropertyConstFacade::GetValue<uint8, uint8>(int32 KeyIndex, const TConstArrayView<uint8>& ValueArray) const;
+
+	template<> CHAOS_API
+	bool FCollectionPropertyConstFacade::GetValue<bool, FVector3f>(int32 KeyIndex, const TConstArrayView<FVector3f>& ValueArray) const
+	{
+		return (bool)ValueArray[KeyIndex].X;
+	}
+
+	template<> CHAOS_API
+	int32 FCollectionPropertyConstFacade::GetValue<int32, FVector3f>(int32 KeyIndex, const TConstArrayView<FVector3f>& ValueArray) const
+	{
+		return (int32)ValueArray[KeyIndex].X;
+	}
+
+	template<> CHAOS_API
+	float FCollectionPropertyConstFacade::GetValue<float, FVector3f>(int32 KeyIndex, const TConstArrayView<FVector3f>& ValueArray) const
+	{
+		return ValueArray[KeyIndex].X;
+	}
+
+	FVector2f FCollectionPropertyConstFacade::GetWeightedFloatValue(const FString& Key, const float& Default, int32* OutKeyIndex) const
+	{
+		const int32 KeyIndex = GetKeyIndex(Key);
+		if (OutKeyIndex)
+		{
+			*OutKeyIndex = KeyIndex;
+		}
+		return KeyIndex != INDEX_NONE ? GetWeightedFloatValue(KeyIndex) : FVector2f(Default, Default);
+	}
+
+	template <typename T>
+	TConstArrayView<T> FCollectionPropertyConstFacade::GetArray(const FName& Name) const
+	{
+		const TManagedArray<T>* ManagedArray = ManagedArrayCollection->FindAttributeTyped<T>(Name, PropertyFacadeNames::PropertyGroup);
+		return ManagedArray ? TConstArrayView<T>(ManagedArray->GetConstArray()) : TConstArrayView<T>();
+	}
+
+	FCollectionPropertyFacade::FCollectionPropertyFacade(const TSharedPtr<FManagedArrayCollection>& InManagedArrayCollection)
+		: FCollectionPropertyConstFacade(InManagedArrayCollection, NoInit)
+	{
+		UpdateArrays();
+		RebuildKeyIndices();
+	}
+	
+	FCollectionPropertyFacade::FCollectionPropertyFacade(const TSharedPtr<FManagedArrayCollection>& InManagedArrayCollection, ENoInit)
+		: FCollectionPropertyConstFacade(InManagedArrayCollection, NoInit)
+	{
+	}
+
+	int32 FCollectionPropertyFacade::SetWeightedFloatValue(const FString& Key, const FVector2f& Value)
+	{
+		const int32 KeyIndex = GetKeyIndex(Key);
+		if (KeyIndex != INDEX_NONE)
+		{
+			SetWeightedFloatValue(KeyIndex, Value);
+		}
+		return KeyIndex;
+	}
+
+	void FCollectionPropertyFacade::EnableFlag(int32 KeyIndex, ECollectionPropertyFlag Flag, bool bEnable)
+	{
+		if (bEnable)
+		{
+			EnumAddFlags(GetFlagsArray()[KeyIndex], (uint8)Flag);
+		}
+		else
+		{
+			EnumRemoveFlags(GetFlagsArray()[KeyIndex], (uint8)Flag);
+		}
+	}
+
+	int32 FCollectionPropertyFacade::EnableFlag(const FString& Key, ECollectionPropertyFlag Flag, bool bEnable)
+	{
+		const int32 KeyIndex = GetKeyIndex(Key);
+		if (KeyIndex != INDEX_NONE)
+		{
+			EnableFlag(KeyIndex, Flag, bEnable);
+		}
+		return KeyIndex;
+	}
+
+	FCollectionPropertyMutableFacade::FCollectionPropertyMutableFacade(const TSharedPtr<FManagedArrayCollection>& InManagedArrayCollection)
+		: FCollectionPropertyFacade(InManagedArrayCollection, NoInit)
+	{
+		UpdateArrays();
+		RebuildKeyIndices();
+	}
+
+	void FCollectionPropertyMutableFacade::DefineSchema()
+	{
+		GetManagedArrayCollection()->AddAttribute<FString>(PropertyFacadeNames::KeyName, PropertyFacadeNames::PropertyGroup);
+		GetManagedArrayCollection()->AddAttribute<FVector3f>(PropertyFacadeNames::LowValueName, PropertyFacadeNames::PropertyGroup);
+		GetManagedArrayCollection()->AddAttribute<FVector3f>(PropertyFacadeNames::HighValueName, PropertyFacadeNames::PropertyGroup);
+		GetManagedArrayCollection()->AddAttribute<FString>(PropertyFacadeNames::StringValueName, PropertyFacadeNames::PropertyGroup);
+		GetManagedArrayCollection()->AddAttribute<uint8>(PropertyFacadeNames::FlagsName, PropertyFacadeNames::PropertyGroup);
+
+		UpdateArrays();
+		RebuildKeyIndices();
+	}
+
+	int32 FCollectionPropertyMutableFacade::AddProperty(const FString& Key, bool bEnabled, bool bAnimatable)
+	{
+		const int32 Index = GetManagedArrayCollection()->AddElements(1, PropertyFacadeNames::PropertyGroup);
+		const uint8 Flags = (bEnabled ? (uint8)ECollectionPropertyFlag::Enabled : 0) | (bAnimatable ? (uint8)ECollectionPropertyFlag::Animatable : 0);
+
+		// Update the arrayviews in case the new element triggered a reallocation 
+		UpdateArrays();
+
+		// Setup the new element's default value and enable the property by default
+		GetKeyArray()[Index] = Key;
+		GetLowValueArray()[Index] = GetHighValueArray()[Index] = FVector3f::ZeroVector;
+		GetFlagsArray()[Index] = Flags;
+
+		// Update search map
+		KeyIndices.Emplace(KeyArray[Index], Index);
+
+		return Index;
+	}
+
+	int32 FCollectionPropertyMutableFacade::AddProperties(const TArray<FString>& Keys, bool bEnabled, bool bAnimatable)
+	{
+		if (const int32 NumProperties = Keys.Num())
+		{
+			const int32 StartIndex = GetManagedArrayCollection()->AddElements(NumProperties, PropertyFacadeNames::PropertyGroup);
+			const uint8 Flags = (bEnabled ? (uint8)ECollectionPropertyFlag::Enabled : 0) | (bAnimatable ? (uint8)ECollectionPropertyFlag::Animatable : 0);
+
+			// Update the arrayviews in case the new elements triggered a reallocation 
+			UpdateArrays();
+
+			for (int32 Index = StartIndex; Index < NumProperties + StartIndex; ++Index)
+			{
+				// Setup the new elements' default value and enable the property by default
+				GetKeyArray()[Index] = Keys[Index - StartIndex];
+				GetLowValueArray()[Index] = GetHighValueArray()[Index] = FVector3f::ZeroVector;
+				GetFlagsArray()[Index] = Flags;
+
+				// Update search map
+				KeyIndices.Emplace(KeyArray[Index], Index);
+			}
+
+			return StartIndex;
+		}
+		return INDEX_NONE;
+	}
+
+	void FCollectionPropertyMutableFacade::Append(const FManagedArrayCollection& InManagedArrayCollection)
+	{
+		TArray<FName> GroupsToSkip = InManagedArrayCollection.GroupNames();
+		GroupsToSkip.RemoveSingleSwap(PropertyFacadeNames::PropertyGroup);
+
+		InManagedArrayCollection.CopyTo(GetManagedArrayCollection().Get(), GroupsToSkip);
+		UpdateArrays();
+		RebuildKeyIndices();
+	}
+
+	void FCollectionPropertyMutableFacade::Copy(const FManagedArrayCollection& InManagedArrayCollection)
+	{
+		GetManagedArrayCollection()->Reset();
+		DefineSchema();
+		Append(InManagedArrayCollection);
+	}
+
+	int32 FCollectionPropertyMutableFacade::AddWeightedFloatValue(const FString& Key, const FVector2f& Value, bool bEnabled, bool bAnimatable)
+	{
+		const int32 KeyIndex = AddProperty(Key, bEnabled, bAnimatable);
+		SetWeightedFloatValue(KeyIndex, Value);
+		return KeyIndex;
+	}
+
+	int32 FCollectionPropertyMutableFacade::AddStringValue(const FString& Key, const FString& Value, bool bEnabled, bool bAnimatable)
+	{
+		const int32 KeyIndex = AddProperty(Key, bEnabled, bAnimatable);
+		SetStringValue(Key, Value);
+		return KeyIndex;
+	}
+}  // End namespace Chaos
diff --git a/Engine/Source/Runtime/Experimental/Chaos/Private/Chaos/PropertyCollectionAdapter.cpp b/Engine/Source/Runtime/Experimental/Chaos/Private/Chaos/PropertyCollectionAdapter.cpp
deleted file mode 100644
diff --git a/Engine/Source/Runtime/Experimental/Chaos/Public/Chaos/CollectionPropertyFacade.h b/Engine/Source/Runtime/Experimental/Chaos/Public/Chaos/CollectionPropertyFacade.h
new file mode 100644
--- /dev/null
+++ b/Engine/Source/Runtime/Experimental/Chaos/Public/Chaos/CollectionPropertyFacade.h
@@ -0,0 +1,430 @@
+// Copyright Epic Games, Inc. All Rights Reserved.
+
+#pragma once
+
+#include "UObject/NameTypes.h"
+#include "Containers/Map.h"
+#include "Math/Vector.h"
+
+struct FManagedArrayCollection;
+
+namespace Chaos::Softs
+{
+	/** Property flags, whether properties are enabled, animatable, ...etc. */
+	enum class ECollectionPropertyFlag : uint8
+	{
+		None,
+		Enabled = 1 << 0,  /** Whether this property is enabled(so that it doesn't have to be removed from the collection when not needed). */
+		Animatable = 1 << 1,  /** Whether this property needs to be set at every frame. */
+		//~ Add new flags above this line
+		Dirty = 1 << 7  /** Whether this property has changed and needs to be updated at the next frame. */
+	};
+	ENUM_CLASS_FLAGS(ECollectionPropertyFlag)
+
+	/** Weighted types are all property types that can have a pair of low and high values to be associated with a weight map. */
+	template<typename T> struct TIsWeightedType { static constexpr bool Value = false; };
+	template<> struct TIsWeightedType<bool> { static constexpr bool Value = true; };
+	template<> struct TIsWeightedType<int32> { static constexpr bool Value = true; };
+	template<> struct TIsWeightedType<float> { static constexpr bool Value = true; };
+	template<> struct TIsWeightedType<FVector3f> { static constexpr bool Value = true; };
+
+	/**
+	 * Defines common API for reading simulation properties data and metadata.
+	 * This is mainly used for the cloth simulation properties to provide
+	 * weighted values that works in conjunction with weight maps.
+	 */
+	class CHAOS_API FCollectionPropertyConstFacade
+	{
+	public:
+		explicit FCollectionPropertyConstFacade(const TSharedPtr<const FManagedArrayCollection>& InManagedArrayCollection);
+		virtual ~FCollectionPropertyConstFacade() = default;
+
+		FCollectionPropertyConstFacade() = delete;
+		FCollectionPropertyConstFacade(const FCollectionPropertyConstFacade&) = delete;
+		FCollectionPropertyConstFacade& operator=(const FCollectionPropertyConstFacade&) = delete;
+
+		FCollectionPropertyConstFacade(FCollectionPropertyConstFacade&&) = default;
+		FCollectionPropertyConstFacade& operator=(FCollectionPropertyConstFacade&&) = default;
+
+		/** Return whether the facade is defined on the collection. */
+		bool IsValid() const;
+
+		/** Return the number of properties in this collection. */
+		int32 Num() const { return KeyIndices.Num(); }
+
+		/** Return the property index for the specified key if it exists, or INDEX_NONE otherwise. */
+		int32 GetKeyIndex(const FString& Key) const { const int32* const Index = KeyIndices.Find(Key); return Index ? *Index : INDEX_NONE; }
+
+		//~ Values access per index, fast, no check, index must be valid (0 <= KeyIndex < Num())
+		template<typename T, TEMPLATE_REQUIRES(TIsWeightedType<T>::Value)>
+		T GetLowValue(int32 KeyIndex) const { return GetValue<T>(KeyIndex, LowValueArray); }
+
+		template<typename T, TEMPLATE_REQUIRES(TIsWeightedType<T>::Value)>
+		T GetHighValue(int32 KeyIndex) const { return GetValue<T>(KeyIndex, HighValueArray); }
+
+		template<typename T, TEMPLATE_REQUIRES(TIsWeightedType<T>::Value)>
+		TPair<T, T> GetWeightedValue(int32 KeyIndex) const { return MakeTuple(GetLowValue<T>(KeyIndex), GetHighValue<T>(KeyIndex)); }
+
+		FVector2f GetWeightedFloatValue(int32 KeyIndex) const { return FVector2f(GetLowValue<float>(KeyIndex), GetHighValue<float>(KeyIndex)); }
+
+		template<typename T, TEMPLATE_REQUIRES(TIsWeightedType<T>::Value)>
+		T GetValue(int32 KeyIndex) const { return GetLowValue<T>(KeyIndex); }
+
+		const FString& GetStringValue(int32 KeyIndex, const FString& Default = "") const { return GetValue<const FString&>(KeyIndex, StringValueArray); }
+
+		uint8 GetFlags(int32 KeyIndex) const { return FlagsArray[KeyIndex]; }
+		bool IsEnabled(int32 KeyIndex) const { return EnumHasAnyFlags((ECollectionPropertyFlag)FlagsArray[KeyIndex], ECollectionPropertyFlag::Enabled); }
+		bool IsAnimatable(int32 KeyIndex) const { return EnumHasAnyFlags((ECollectionPropertyFlag)FlagsArray[KeyIndex], ECollectionPropertyFlag::Animatable); }
+		bool IsDirty(int32 KeyIndex) const { return EnumHasAnyFlags((ECollectionPropertyFlag)FlagsArray[KeyIndex], ECollectionPropertyFlag::Dirty); }
+
+		//~ Values access per key
+		template<typename T, TEMPLATE_REQUIRES(TIsWeightedType<T>::Value)>
+		T GetLowValue(const FString& Key, const T& Default = T(0), int32* OutKeyIndex = nullptr) const { return GetValue(Key, LowValueArray, Default, OutKeyIndex); }
+
+		template<typename T, TEMPLATE_REQUIRES(TIsWeightedType<T>::Value)>
+		T GetHighValue(const FString& Key, const T& Default = T(0), int32* OutKeyIndex = nullptr) const { return GetValue(Key, HighValueArray, Default, OutKeyIndex); }
+
+		template<typename T, TEMPLATE_REQUIRES(TIsWeightedType<T>::Value)>
+		inline TPair<T, T> GetWeightedValue(const FString& Key, const T& Default = T(0), int32* OutKeyIndex = nullptr) const;
+
+		FVector2f GetWeightedFloatValue(const FString& Key, const float& Default = 0.f, int32* OutKeyIndex = nullptr) const;
+
+		template<typename T, TEMPLATE_REQUIRES(TIsWeightedType<T>::Value)>
+		T GetValue(const FString& Key, const T& Default = T(0), int32* OutKeyIndex = nullptr) const { return GetValue(Key, LowValueArray, Default, OutKeyIndex); }
+
+		FString GetStringValue(const FString& Key, const FString& Default = "", int32* OutKeyIndex = nullptr) const { return GetValue(Key, StringValueArray, Default, OutKeyIndex); }
+
+		uint8 GetFlags(const FString& Key, uint8 Default = 0, int32* OutKeyIndex = nullptr) const { return GetValue(Key, FlagsArray, (uint8)Default, OutKeyIndex); }
+		bool IsEnabled(const FString& Key, bool bDefault = false, int32* OutKeyIndex = nullptr) const { return EnumHasAnyFlags((ECollectionPropertyFlag)GetValue(Key, FlagsArray, (uint8)bDefault, OutKeyIndex), ECollectionPropertyFlag::Enabled); }
+		bool IsAnimatable(const FString& Key, bool bDefault = false, int32* OutKeyIndex = nullptr) const { return EnumHasAnyFlags((ECollectionPropertyFlag)GetValue(Key, FlagsArray, (uint8)bDefault, OutKeyIndex), ECollectionPropertyFlag::Animatable); }
+		bool IsDirty(const FString& Key, bool bDefault = false, int32* OutKeyIndex = nullptr) const { return EnumHasAnyFlags((ECollectionPropertyFlag)GetValue(Key, FlagsArray, (uint8)bDefault, OutKeyIndex), ECollectionPropertyFlag::Dirty); }
+
+	protected:
+		// No init constructor for FCollectionPropertyFacade
+		FCollectionPropertyConstFacade(const TSharedPtr<const FManagedArrayCollection>& InManagedArrayCollection, ENoInit);
+
+		// Update the array views
+		void UpdateArrays();
+
+		// Rebuild the search map
+		void RebuildKeyIndices();
+
+		template<typename T, typename ElementType>
+		T GetValue(int32 KeyIndex, const TConstArrayView<ElementType>& ValueArray) const;
+
+		template<typename T, typename ElementType>
+		inline T GetValue(const FString& Key, const TConstArrayView<ElementType>& ValueArray, const T& Default, int32* OutKeyIndex) const;
+
+		template <typename T>
+		TConstArrayView<T> GetArray(const FName& Name) const;
+
+		// Attribute groups, predefined data member of the collection
+		static const FName PropertyGroup;
+		static const FName KeyName;  // Property key, name to look for
+		static const FName LowValueName;  // Boolean, 24 bit integer (max 16777215), float, or vector value, or value of the lowest weight on the weight map if any
+		static const FName HighValueName;  // Boolean, 24 bit integer (max 16777215), float, or vector value of the highest weight on the weight map if any
+		static const FName StringValueName;  // String value, or weight map name, ...etc.
+		static const FName FlagsName;  // Whether this property is enabled, animatable, ...etc.
+
+		// Property Group array views
+		TConstArrayView<FString> KeyArray;
+		TConstArrayView<FVector3f> LowValueArray;
+		TConstArrayView<FVector3f> HighValueArray;
+		TConstArrayView<FString> StringValueArray;
+		TConstArrayView<uint8> FlagsArray;
+
+		// Key to index search map
+		TMap<FString, int32> KeyIndices;
+
+		// Property collection
+		TSharedPtr<const FManagedArrayCollection> ManagedArrayCollection;
+	};
+
+	/**
+	 * Defines common API for reading and writing simulation properties data and metadata.
+	 * This is mainly used for the cloth simulation properties to provide
+	 * weighted values that works in conjunction with weight maps.
+	 * 
+	 * Note: Int property values are stored as float and therefore limited to 24 bits (-16777215 to 16777215).
+	 */
+	class CHAOS_API FCollectionPropertyFacade : public FCollectionPropertyConstFacade
+	{
+	public:
+		explicit FCollectionPropertyFacade(const TSharedPtr<FManagedArrayCollection>& InManagedArrayCollection);
+		virtual ~FCollectionPropertyFacade() = default;
+
+		FCollectionPropertyFacade() = delete;
+		FCollectionPropertyFacade(const FCollectionPropertyFacade&) = delete;
+		FCollectionPropertyFacade& operator=(const FCollectionPropertyFacade&) = delete;
+
+		FCollectionPropertyFacade(FCollectionPropertyFacade&&) = default;
+		FCollectionPropertyFacade& operator=(FCollectionPropertyFacade&&) = default;
+
+		//~ Values set per index
+		template<typename T, TEMPLATE_REQUIRES(TIsWeightedType<T>::Value)>
+		void SetLowValue(int32 KeyIndex, const T& Value) { GetLowValueArray()[KeyIndex] = FVector3f(Value); }
+
+		template<typename T, TEMPLATE_REQUIRES(TIsWeightedType<T>::Value)>
+		void SetHighValue(int32 KeyIndex, const T& Value) { GetHighValueArray()[KeyIndex] = FVector3f(Value); }
+
+		template<typename T, TEMPLATE_REQUIRES(TIsWeightedType<T>::Value)>
+		void SetWeightedValue(int32 KeyIndex, const T& LowValue, const T& HighValue) { SetLowValue(KeyIndex, LowValue); SetHighValue(KeyIndex, HighValue); }
+
+		void SetWeightedFloatValue(int32 KeyIndex, const FVector2f& Value) { SetLowValue<float>(KeyIndex, Value.X); SetHighValue<float>(KeyIndex, Value.Y); }
+
+		template<typename T, TEMPLATE_REQUIRES(TIsWeightedType<T>::Value)>
+		void SetValue(int32 KeyIndex, const T& Value) { SetWeightedValue(KeyIndex, Value, Value); }
+
+		void SetStringValue(int32 KeyIndex, const FString& Value) { GetStringValueArray()[KeyIndex] = Value; }
+
+		void SetFlags(int32 KeyIndex, uint8 Flags) { GetFlagsArray()[KeyIndex] = Flags; }
+		void SetEnabled(int32 KeyIndex, bool bEnabled) { EnableFlag(KeyIndex, ECollectionPropertyFlag::Enabled, bEnabled); }
+		void SetAnimatable(int32 KeyIndex, bool bAnimatable) { EnableFlag(KeyIndex, ECollectionPropertyFlag::Animatable, bAnimatable); }
+		void SetDirty(int32 KeyIndex, bool bDirty) { EnableFlag(KeyIndex, ECollectionPropertyFlag::Animatable, bDirty); }
+
+		//~ Values set per key
+		template<typename T, TEMPLATE_REQUIRES(TIsWeightedType<T>::Value)>
+		int32 SetLowValue(const FString& Key, const T& Value) { return SetValue(Key, GetLowValueArray(), FVector3f(Value)); }
+
+		template<typename T, TEMPLATE_REQUIRES(TIsWeightedType<T>::Value)>
+		int32 SetHighValue(const FString& Key, const T& Value) { return SetValue(Key, GetHighValueArray(), FVector3f(Value)); }
+
+		template<typename T, TEMPLATE_REQUIRES(TIsWeightedType<T>::Value)>
+		inline int32 SetWeightedValue(const FString& Key, const T& LowValue, const T& HighValue);
+
+		int32 SetWeightedFloatValue(const FString& Key, const FVector2f& Value);
+
+		template<typename T, TEMPLATE_REQUIRES(TIsWeightedType<T>::Value)>
+		int32 SetValue(const FString& Key, const T& Value) { return SetWeightedValue(Key, Value, Value); }
+
+		int32 SetStringValue(const FString& Key, const FString& Value) { return SetValue(Key, GetStringValueArray(), Value); }
+
+		int32 SetFlags(const FString& Key, uint8 Flags) { return SetValue(Key, GetFlagsArray(), Flags); }
+		int32 SetEnabled(const FString& Key, bool bEnabled) { return EnableFlag(Key, ECollectionPropertyFlag::Enabled, bEnabled); }
+		int32 SetAnimatable(const FString& Key, bool bAnimatable) { return EnableFlag(Key, ECollectionPropertyFlag::Animatable, bAnimatable); }
+		int32 SetDirty(const FString& Key, bool bDirty) { return EnableFlag(Key, ECollectionPropertyFlag::Animatable, bDirty); }
+
+	protected:
+		// No init constructor for FCollectionPropertyMutableFacade
+		FCollectionPropertyFacade(const TSharedPtr<FManagedArrayCollection>& InManagedArrayCollection, ENoInit);
+
+		// Access to a writeable ManagedArrayCollection is protected, use an FPropertyCollectionMutableAdapter if needed to get a non const pointer
+		TSharedPtr<FManagedArrayCollection> GetManagedArrayCollection() { return ConstCastSharedPtr<FManagedArrayCollection>(ManagedArrayCollection); }
+
+		const TArrayView<FString>& GetKeyArray() { return reinterpret_cast<TArrayView<FString>&>(KeyArray); }
+		const TArrayView<FVector3f>& GetLowValueArray() { return reinterpret_cast<TArrayView<FVector3f>&>(LowValueArray); }
+		const TArrayView<FVector3f>& GetHighValueArray() { return reinterpret_cast<TArrayView<FVector3f>&>(HighValueArray); }
+		const TArrayView<FString>& GetStringValueArray() { return reinterpret_cast<TArrayView<FString>&>(StringValueArray); }
+		const TArrayView<uint8>& GetFlagsArray() { return reinterpret_cast<const TArrayView<uint8>&>(FlagsArray); }
+
+	private:
+		template<typename T>
+		inline int32 SetValue(const FString& Key, const TArrayView<T>& ValueArray, const T& Value);
+
+		void EnableFlag(int32 KeyIndex, ECollectionPropertyFlag Flag, bool bEnable);
+		int32 EnableFlag(const FString& Key, ECollectionPropertyFlag Flag, bool bEnable);
+	};
+
+	/**
+	 * Defines common API for reading and writing, and adding/removing simulation properties data and metadata.
+	 * This is mainly used for the cloth simulation properties to provide
+	 * weighted values that works in conjunction with weight maps.
+	 */
+	class CHAOS_API FCollectionPropertyMutableFacade final : public FCollectionPropertyFacade
+	{
+	public:
+		explicit FCollectionPropertyMutableFacade(const TSharedPtr<FManagedArrayCollection>& InManagedArrayCollection);
+		virtual ~FCollectionPropertyMutableFacade() = default;
+
+		FCollectionPropertyMutableFacade() = delete;
+		FCollectionPropertyMutableFacade(const FCollectionPropertyMutableFacade&) = delete;
+		FCollectionPropertyMutableFacade& operator=(const FCollectionPropertyMutableFacade&) = delete;
+
+		FCollectionPropertyMutableFacade(FCollectionPropertyMutableFacade&&) = default;
+		FCollectionPropertyMutableFacade& operator=(FCollectionPropertyMutableFacade&&) = default;
+
+		/** Create this facade's groups and attributes. */
+		void DefineSchema();
+
+		/** Add a single property, and return its index. */
+		int32 AddProperty(const FString& Key, bool bEnabled = true, bool bAnimatable = false);
+
+		/** Add new properties, and return the index of the first added property. */
+		int32 AddProperties(const TArray<FString>& Keys, bool bEnabled = true, bool bAnimatable = false);
+
+		/**
+		 * Append all properties and values from an existing collection to this property collection.
+		 * This won't copy any other groups, only data from PropertyGroup.
+		 * Any pre-existing data will be preserved.
+		 */
+		void Append(const FManagedArrayCollection& InManagedArrayCollection);
+
+		/**
+		 * Copy all properties and values from an existing collection to this property collection.
+		 * This won't copy any other groups, only data from PropertyGroup.
+		 * Any pre-xisting data will be removed/replaced.
+		 */
+		void Copy(const FManagedArrayCollection& InManagedArrayCollection);
+
+		//~ Add values
+		template<typename T, TEMPLATE_REQUIRES(TIsWeightedType<T>::Value)>
+		inline int32 AddWeightedValue(const FString& Key, const T& LowValue, const T& HighValue, bool bEnabled = true, bool bAnimatable = false);
+
+		int32 AddWeightedFloatValue(const FString& Key, const FVector2f& Value, bool bEnabled, bool bAnimatable);
+
+		template<typename T, TEMPLATE_REQUIRES(TIsWeightedType<T>::Value)>
+		int32 AddValue(const FString& Key, const T& Value, bool bEnabled = true, bool bAnimatable = false) { return AddWeightedValue(Key, Value, Value, bEnabled, bAnimatable); }
+
+		int32 AddStringValue(const FString& Key, const FString& Value, bool bEnabled = true, bool bAnimatable = false);
+	};
+
+	template<typename T, typename TEnableIf<TIsWeightedType<T>::Value, int>::type>
+	inline TPair<T, T> FCollectionPropertyConstFacade::GetWeightedValue(const FString& Key, const T& Default, int32* OutKeyIndex) const
+	{
+		const int32 KeyIndex = GetKeyIndex(Key);
+		if (OutKeyIndex)
+		{
+			*OutKeyIndex = KeyIndex;
+		}
+		return KeyIndex != INDEX_NONE ? GetWeightedValue<T>(KeyIndex) : MakeTuple(Default, Default);
+	}
+
+	template<typename T, typename ElementType>
+	inline T FCollectionPropertyConstFacade::GetValue(const FString& Key, const TConstArrayView<ElementType>& ValueArray, const T& Default, int32* OutKeyIndex) const
+	{
+		const int32 KeyIndex = GetKeyIndex(Key);
+		if (OutKeyIndex)
+		{
+			*OutKeyIndex = KeyIndex;
+		}
+		return KeyIndex != INDEX_NONE ? GetValue<T>(KeyIndex, ValueArray) : Default;
+	}
+
+	template<typename T, typename TEnableIf<TIsWeightedType<T>::Value, int>::type>
+	inline int32 FCollectionPropertyFacade::SetWeightedValue(const FString& Key, const T& LowValue, const T& HighValue)
+	{
+		const int32 KeyIndex = GetKeyIndex(Key);
+		if (KeyIndex != INDEX_NONE)
+		{
+			SetLowValue(KeyIndex, LowValue);
+			SetHighValue(KeyIndex, HighValue);
+		}
+		return KeyIndex;
+	}
+
+	template<typename T>
+	inline int32 FCollectionPropertyFacade::SetValue(const FString& Key, const TArrayView<T>& ValueArray, const T& Value)
+	{
+		const int32 KeyIndex = GetKeyIndex(Key);
+		if (KeyIndex != INDEX_NONE)
+		{
+			ValueArray[KeyIndex] = Value;
+		}
+		return KeyIndex;
+	}
+
+	template<typename T, typename TEnableIf<TIsWeightedType<T>::Value, int>::type>
+	inline int32 FCollectionPropertyMutableFacade::AddWeightedValue(const FString& Key, const T& LowValue, const T& HighValue, bool bEnabled, bool bAnimatable)
+	{
+		const int32 KeyIndex = AddProperty(Key, bEnabled, bAnimatable);
+		SetWeightedValue(KeyIndex, LowValue, HighValue);
+		return KeyIndex;
+	}
+}  // End namespace Chaos
+
+// Use this macro to add shorthands for property getters and direct access through the declared key index
+#define UE_CHAOS_DECLARE_PROPERTYCOLLECTION_NAME(PropertyName, Type) \
+	inline static const FName PropertyName##Name = TEXT(#PropertyName); \
+	static FString PropertyName##String() { return PropertyName##Name.ToString(); } \
+	static bool Is##PropertyName##Enabled(const FCollectionPropertyConstFacade& PropertyCollection, bool bDefault) \
+	{ \
+		return PropertyCollection.IsEnabled(PropertyName##String(), bDefault); \
+	} \
+	static bool Is##PropertyName##Animatable(const FCollectionPropertyConstFacade& PropertyCollection, bool bDefault) \
+	{ \
+		return PropertyCollection.IsAnimatable(PropertyName##String(), bDefault); \
+	} \
+	Type GetLow##PropertyName(const FCollectionPropertyConstFacade& PropertyCollection, const Type& Default) \
+	{ \
+		return PropertyCollection.GetLowValue<Type>(PropertyName##String(), Default, &PropertyName##Index); \
+	} \
+	Type GetHigh##PropertyName(const FCollectionPropertyConstFacade& PropertyCollection, const Type& Default) \
+	{ \
+		return PropertyCollection.GetHighValue<Type>(PropertyName##String(), Default, &PropertyName##Index); \
+	} \
+	TPair<Type, Type> GetWeighted##PropertyName(const FCollectionPropertyConstFacade& PropertyCollection, const Type& Default) \
+	{ \
+		return PropertyCollection.GetWeightedValue<Type>(PropertyName##String(), Default, &PropertyName##Index); \
+	} \
+	FVector2f GetWeightedFloat##PropertyName(const FCollectionPropertyConstFacade& PropertyCollection, const float& Default) \
+	{ \
+		return PropertyCollection.GetWeightedFloatValue(PropertyName##String(), Default, &PropertyName##Index); \
+	} \
+	Type Get##PropertyName(const FCollectionPropertyConstFacade& PropertyCollection, const Type& Default) \
+	{ \
+		return PropertyCollection.GetValue<Type>(PropertyName##String(), Default, &PropertyName##Index); \
+	} \
+	uint8 Get##PropertyName##Flags(const FCollectionPropertyConstFacade& PropertyCollection, uint8 Default) \
+	{ \
+		return PropertyCollection.GetFlags(PropertyName##String(), Default, &PropertyName##Index); \
+	} \
+	Type GetLow##PropertyName(const FCollectionPropertyConstFacade& PropertyCollection) const \
+	{ \
+		checkSlow(PropertyName##Index == PropertyCollection.GetKeyIndex(PropertyName##String())); \
+		checkf(PropertyName##Index != INDEX_NONE, TEXT("The default value getter that sets the property index must be called once prior to calling this function.")); \
+		return PropertyCollection.GetLowValue<Type>(PropertyName##Index); \
+	} \
+	Type GetHigh##PropertyName(const FCollectionPropertyConstFacade& PropertyCollection) const \
+	{ \
+		checkSlow(PropertyName##Index == PropertyCollection.GetKeyIndex(PropertyName##String())); \
+		checkf(PropertyName##Index != INDEX_NONE, TEXT("The default value getter that sets the property index must be called once prior to calling this function.")); \
+		return PropertyCollection.GetHighValue<Type>(PropertyName##Index); \
+	} \
+	TPair<Type, Type> GetWeighted##PropertyName(const FCollectionPropertyConstFacade& PropertyCollection) const \
+	{ \
+		checkSlow(PropertyName##Index == PropertyCollection.GetKeyIndex(PropertyName##String())); \
+		checkf(PropertyName##Index != INDEX_NONE, TEXT("The default value getter that sets the property index must be called once prior to calling this function.")); \
+		return PropertyCollection.GetWeightedValue<Type>(PropertyName##Index); \
+	} \
+	FVector2f GetWeightedFloat##PropertyName(const FCollectionPropertyConstFacade& PropertyCollection) const \
+	{ \
+		checkSlow(PropertyName##Index == PropertyCollection.GetKeyIndex(PropertyName##String())); \
+		checkf(PropertyName##Index != INDEX_NONE, TEXT("The default value getter that sets the property index must be called once prior to calling this function.")); \
+		return PropertyCollection.GetWeightedFloatValue(PropertyName##Index); \
+	} \
+	Type Get##PropertyName(const FCollectionPropertyConstFacade& PropertyCollection) const \
+	{ \
+		checkSlow(PropertyName##Index == PropertyCollection.GetKeyIndex(PropertyName##String())); \
+		checkf(PropertyName##Index != INDEX_NONE, TEXT("The default value getter that sets the property index must be called once prior to calling this function.")); \
+		return PropertyCollection.GetValue<Type>(PropertyName##Index); \
+	} \
+	uint8 Get##PropertyName##Flags(const FCollectionPropertyConstFacade& PropertyCollection) const \
+	{ \
+		checkSlow(PropertyName##Index == PropertyCollection.GetKeyIndex(PropertyName##String())); \
+		checkf(PropertyName##Index != INDEX_NONE, TEXT("The default value getter that sets the property index must be called once prior to calling this function.")); \
+		return PropertyCollection.GetFlags(PropertyName##Index); \
+	} \
+	bool Is##PropertyName##Enabled(const FCollectionPropertyConstFacade& PropertyCollection) const \
+	{ \
+		checkSlow(PropertyName##Index == PropertyCollection.GetKeyIndex(PropertyName##String())); \
+		checkf(PropertyName##Index != INDEX_NONE, TEXT("The default value getter that sets the property index must be called once prior to calling this function.")); \
+		return PropertyCollection.IsEnabled(PropertyName##Index); \
+	} \
+	bool Is##PropertyName##Animatable(const FCollectionPropertyConstFacade& PropertyCollection) const\
+	{ \
+		checkSlow(PropertyName##Index == PropertyCollection.GetKeyIndex(PropertyName##String())); \
+		checkf(PropertyName##Index != INDEX_NONE, TEXT("The default value getter that sets the property index must be called once prior to calling this function.")); \
+		return PropertyCollection.IsAnimatable(PropertyName##Index); \
+	} \
+	bool Is##PropertyName##Dirty(const FCollectionPropertyConstFacade& PropertyCollection) const \
+	{ \
+		checkSlow(PropertyName##Index == PropertyCollection.GetKeyIndex(PropertyName##String())); \
+		checkf(PropertyName##Index != INDEX_NONE, TEXT("The default value getter that sets the property index must be called once prior to calling this function.")); \
+		return PropertyCollection.IsDirty(PropertyName##Index); \
+	} \
+	bool Is##PropertyName##Mutable(const FCollectionPropertyConstFacade& PropertyCollection) const \
+	{ \
+		return PropertyName##Index != INDEX_NONE && \
+			EnumHasAllFlags((ECollectionPropertyFlag)PropertyCollection.GetFlags(PropertyName##Index), ECollectionPropertyFlag::Enabled | ECollectionPropertyFlag::Animatable | ECollectionPropertyFlag::Dirty); \
+	} \
+	int32 PropertyName##Index = INDEX_NONE;
diff --git a/Engine/Source/Runtime/Experimental/Chaos/Public/Chaos/PBDAnimDriveConstraint.h b/Engine/Source/Runtime/Experimental/Chaos/Public/Chaos/PBDAnimDriveConstraint.h
--- a/Engine/Source/Runtime/Experimental/Chaos/Public/Chaos/PBDAnimDriveConstraint.h
+++ b/Engine/Source/Runtime/Experimental/Chaos/Public/Chaos/PBDAnimDriveConstraint.h
@@ -6,7 +6,7 @@
 #include "Chaos/Framework/Parallel.h"
 #include "Chaos/PBDSoftsEvolutionFwd.h"
 #include "Chaos/PBDStiffness.h"
-#include "Chaos/PropertyCollectionAdapter.h"
+#include "Chaos/CollectionPropertyFacade.h"
 #include "ChaosStats.h"
 
 DECLARE_CYCLE_STAT(TEXT("Chaos PBD Anim Drive Constraint"), STAT_PBD_AnimDriveConstraint, STATGROUP_Chaos);
@@ -17,7 +17,7 @@ namespace Chaos::Softs
 	class FPBDAnimDriveConstraint final
 	{
 	public:
-		static bool IsEnabled(const FPropertyCollectionConstAdapter& PropertyCollection)
+		static bool IsEnabled(const FCollectionPropertyConstFacade& PropertyCollection)
 		{
 			return IsAnimDriveStiffnessEnabled(PropertyCollection, false);
 		}
@@ -49,7 +49,7 @@ namespace Chaos::Softs
 			const TArray<FSolverVec3>& InAnimationVelocities,  // Use global indexation (will need adding ParticleOffset)
 			const TConstArrayView<FRealSingle>& StiffnessMultipliers,  // Use local indexation
 			const TConstArrayView<FRealSingle>& DampingMultipliers,  // Use local indexation
-			const FPropertyCollectionConstAdapter& PropertyCollection
+			const FCollectionPropertyConstFacade& PropertyCollection
 		)
 			: AnimationPositions(InAnimationPositions)
 			, OldAnimationPositions_deprecated(InAnimationVelocities) // Unused when not using apply
@@ -96,7 +96,7 @@ namespace Chaos::Softs
 		// Return the damping input values used by the constraint
 		FSolverVec2 GetDamping() const { return Damping.GetWeightedValue(); }
 
-		void SetProperties(const FPropertyCollectionConstAdapter& PropertyCollection)
+		void SetProperties(const FCollectionPropertyConstFacade& PropertyCollection)
 		{
 			if (IsAnimDriveStiffnessMutable(PropertyCollection))
 			{
diff --git a/Engine/Source/Runtime/Experimental/Chaos/Public/Chaos/PBDAxialSpringConstraints.h b/Engine/Source/Runtime/Experimental/Chaos/Public/Chaos/PBDAxialSpringConstraints.h
--- a/Engine/Source/Runtime/Experimental/Chaos/Public/Chaos/PBDAxialSpringConstraints.h
+++ b/Engine/Source/Runtime/Experimental/Chaos/Public/Chaos/PBDAxialSpringConstraints.h
@@ -4,7 +4,7 @@
 // HEADER_UNIT_SKIP - Internal
 
 #include "Chaos/PBDAxialSpringConstraintsBase.h"
-#include "Chaos/PropertyCollectionAdapter.h"
+#include "Chaos/CollectionPropertyFacade.h"
 
 namespace Chaos::Softs
 {
@@ -46,7 +46,7 @@ private:
 class CHAOS_API FPBDAreaSpringConstraints final : public FPBDAxialSpringConstraints
 {
 public:
-	static bool IsEnabled(const FPropertyCollectionConstAdapter& PropertyCollection)
+	static bool IsEnabled(const FCollectionPropertyConstFacade& PropertyCollection)
 	{
 		return IsAreaSpringStiffnessEnabled(PropertyCollection, false);
 	}
@@ -57,7 +57,7 @@ public:
 		int32 ParticleCount,
 		const TArray<TVec3<int32>>& InConstraints,
 		const TConstArrayView<FRealSingle>& StiffnessMultipliers,
-		const FPropertyCollectionConstAdapter& PropertyCollection,
+		const FCollectionPropertyConstFacade& PropertyCollection,
 		bool bTrimKinematicConstraints)
 		: FPBDAxialSpringConstraints(
 			Particles,
@@ -71,7 +71,7 @@ public:
 
 	virtual ~FPBDAreaSpringConstraints() override = default;
 
-	void SetProperties(const FPropertyCollectionConstAdapter& PropertyCollection)
+	void SetProperties(const FCollectionPropertyConstFacade& PropertyCollection)
 	{
 		if (IsAreaSpringStiffnessMutable(PropertyCollection))
 		{
diff --git a/Engine/Source/Runtime/Experimental/Chaos/Public/Chaos/PBDBendingConstraints.h b/Engine/Source/Runtime/Experimental/Chaos/Public/Chaos/PBDBendingConstraints.h
--- a/Engine/Source/Runtime/Experimental/Chaos/Public/Chaos/PBDBendingConstraints.h
+++ b/Engine/Source/Runtime/Experimental/Chaos/Public/Chaos/PBDBendingConstraints.h
@@ -2,7 +2,7 @@
 #pragma once
 
 #include "Chaos/PBDBendingConstraintsBase.h"
-#include "Chaos/PropertyCollectionAdapter.h"
+#include "Chaos/CollectionPropertyFacade.h"
 
 namespace Chaos::Softs
 {
@@ -12,7 +12,7 @@ class CHAOS_API FPBDBendingConstraints : public FPBDBendingConstraintsBase
 	typedef FPBDBendingConstraintsBase Base;
 
 public:
-	static bool IsEnabled(const FPropertyCollectionConstAdapter& PropertyCollection)
+	static bool IsEnabled(const FCollectionPropertyConstFacade& PropertyCollection)
 	{
 		return IsBendingElementStiffnessEnabled(PropertyCollection, false);
 	}
@@ -23,7 +23,7 @@ public:
 		TArray<TVec4<int32>>&& InConstraints,
 		const TConstArrayView<FRealSingle>& StiffnessMultipliers,
 		const TConstArrayView<FRealSingle>& BucklingStiffnessMultipliers,
-		const FPropertyCollectionConstAdapter& PropertyCollection,
+		const FCollectionPropertyConstFacade& PropertyCollection,
 		bool bTrimKinematicConstraints = false)
 		: Base(
 			InParticles,
@@ -78,7 +78,7 @@ PRAGMA_ENABLE_DEPRECATION_WARNINGS
 
 	using Base::SetProperties;
 
-	void SetProperties(const FPropertyCollectionConstAdapter& PropertyCollection)
+	void SetProperties(const FCollectionPropertyConstFacade& PropertyCollection)
 	{
 		if (IsBendingElementStiffnessMutable(PropertyCollection))
 		{
diff --git a/Engine/Source/Runtime/Experimental/Chaos/Public/Chaos/PBDCollisionSpringConstraints.h b/Engine/Source/Runtime/Experimental/Chaos/Public/Chaos/PBDCollisionSpringConstraints.h
--- a/Engine/Source/Runtime/Experimental/Chaos/Public/Chaos/PBDCollisionSpringConstraints.h
+++ b/Engine/Source/Runtime/Experimental/Chaos/Public/Chaos/PBDCollisionSpringConstraints.h
@@ -3,7 +3,7 @@
 
 #if !COMPILE_WITHOUT_UNREAL_SUPPORT
 #include "Chaos/PBDCollisionSpringConstraintsBase.h"
-#include "Chaos/PropertyCollectionAdapter.h"
+#include "Chaos/CollectionPropertyFacade.h"
 
 namespace Chaos::Softs
 {
@@ -18,7 +18,7 @@ public:
 	static constexpr FSolverReal MinFrictionCoefficient = (FSolverReal)0.;
 	static constexpr FSolverReal MaxFrictionCoefficient = (FSolverReal)10.;
 
-	static bool IsEnabled(const FPropertyCollectionConstAdapter& PropertyCollection)
+	static bool IsEnabled(const FCollectionPropertyConstFacade& PropertyCollection)
 	{
 		return IsSelfCollisionStiffnessEnabled(PropertyCollection, false);
 	}
@@ -29,7 +29,7 @@ public:
 		const FTriangleMesh& InTriangleMesh,
 		const TArray<FSolverVec3>* InRestPositions,
 		TSet<TVec2<int32>>&& InDisabledCollisionElements,
-		const FPropertyCollectionConstAdapter& PropertyCollection)
+		const FCollectionPropertyConstFacade& PropertyCollection)
 		: Base(
 			InOffset,
 			InNumParticles,
@@ -65,7 +65,7 @@ public:
 
 	using Base::Init;
 
-	void SetProperties(const FPropertyCollectionConstAdapter& PropertyCollection)
+	void SetProperties(const FCollectionPropertyConstFacade& PropertyCollection)
 	{
 		if (IsSelfCollisionThicknessMutable(PropertyCollection))
 		{
diff --git a/Engine/Source/Runtime/Experimental/Chaos/Public/Chaos/PBDLongRangeConstraints.h b/Engine/Source/Runtime/Experimental/Chaos/Public/Chaos/PBDLongRangeConstraints.h
--- a/Engine/Source/Runtime/Experimental/Chaos/Public/Chaos/PBDLongRangeConstraints.h
+++ b/Engine/Source/Runtime/Experimental/Chaos/Public/Chaos/PBDLongRangeConstraints.h
@@ -2,7 +2,7 @@
 #pragma once
 
 #include "Chaos/PBDLongRangeConstraintsBase.h"
-#include "Chaos/PropertyCollectionAdapter.h"
+#include "Chaos/CollectionPropertyFacade.h"
 #include "ChaosStats.h"
 
 DECLARE_CYCLE_STAT(TEXT("Chaos PBD Long Range Constraint"), STAT_PBD_LongRange, STATGROUP_Chaos);
@@ -16,7 +16,7 @@ public:
 	typedef FPBDLongRangeConstraintsBase Base;
 	typedef typename Base::FTether FTether;
 
-	static bool IsEnabled(const FPropertyCollectionConstAdapter& PropertyCollection)
+	static bool IsEnabled(const FCollectionPropertyConstFacade& PropertyCollection)
 	{
 		return IsTetherStiffnessEnabled(PropertyCollection, false);
 	}
@@ -28,7 +28,7 @@ public:
 		const TArray<TConstArrayView<TTuple<int32, int32, FRealSingle>>>& InTethers,
 		const TConstArrayView<FRealSingle>& StiffnessMultipliers,
 		const TConstArrayView<FRealSingle>& ScaleMultipliers,
-		const FPropertyCollectionConstAdapter& PropertyCollection,
+		const FCollectionPropertyConstFacade& PropertyCollection,
 		FSolverReal MeshScale)
 		: FPBDLongRangeConstraintsBase(
 			Particles,
@@ -69,7 +69,7 @@ public:
 
 	using Base::SetProperties;
 
-	void SetProperties(const FPropertyCollectionConstAdapter& PropertyCollection, FSolverReal MeshScale)
+	void SetProperties(const FCollectionPropertyConstFacade& PropertyCollection, FSolverReal MeshScale)
 	{
 		if (IsTetherStiffnessMutable(PropertyCollection))
 		{
diff --git a/Engine/Source/Runtime/Experimental/Chaos/Public/Chaos/PBDSphericalConstraint.h b/Engine/Source/Runtime/Experimental/Chaos/Public/Chaos/PBDSphericalConstraint.h
--- a/Engine/Source/Runtime/Experimental/Chaos/Public/Chaos/PBDSphericalConstraint.h
+++ b/Engine/Source/Runtime/Experimental/Chaos/Public/Chaos/PBDSphericalConstraint.h
@@ -3,7 +3,7 @@
 
 #include "Chaos/PBDSoftsEvolutionFwd.h"
 #include "Chaos/PBDSoftsSolverParticles.h"
-#include "Chaos/PropertyCollectionAdapter.h"
+#include "Chaos/CollectionPropertyFacade.h"
 #include "ChaosStats.h"
 #include "Containers/ArrayView.h"
 #include "Containers/ContainersFwd.h"
@@ -29,7 +29,7 @@ namespace Chaos::Softs
 class CHAOS_API FPBDSphericalConstraint final
 {
 public:
-	static bool IsEnabled(const FPropertyCollectionConstAdapter& PropertyCollection)
+	static bool IsEnabled(const FCollectionPropertyConstFacade& PropertyCollection)
 	{
 		return IsMaxDistanceEnabled(PropertyCollection, false);
 	}
@@ -39,7 +39,7 @@ public:
 		const uint32 InParticleCount,
 		const TArray<FSolverVec3>& InAnimationPositions,  // Use global indexation (will need adding ParticleOffset)
 		const TConstArrayView<FRealSingle>& InSphereRadii,  // Use local indexation
-		const FPropertyCollectionConstAdapter& PropertyCollection,
+		const FCollectionPropertyConstFacade& PropertyCollection,
 		FSolverReal MeshScale
 	)
 		: AnimationPositions(InAnimationPositions)
@@ -65,7 +65,7 @@ public:
 	}
 	~FPBDSphericalConstraint() {}
 
-	void SetProperties(const FPropertyCollectionConstAdapter& PropertyCollection, FSolverReal MeshScale)
+	void SetProperties(const FCollectionPropertyConstFacade& PropertyCollection, FSolverReal MeshScale)
 	{
 		SetScale((FSolverReal)1., MeshScale);
 		// TODO: MaxDistance
@@ -146,7 +146,7 @@ private:
 class CHAOS_API FPBDSphericalBackstopConstraint final
 {
 public:
-	static bool IsEnabled(const FPropertyCollectionConstAdapter& PropertyCollection)
+	static bool IsEnabled(const FCollectionPropertyConstFacade& PropertyCollection)
 	{
 		return IsBackstopDistanceEnabled(PropertyCollection, false) ||
 			IsBackstopDistanceAnimatable(PropertyCollection, false);  // Backstop can be re-enabled if animated
@@ -159,7 +159,7 @@ public:
 		const TArray<FSolverVec3>& InAnimationNormals,  // Use global indexation (will need adding ParticleOffset)
 		const TConstArrayView<FRealSingle>& InSphereRadii,  // Use local indexation
 		const TConstArrayView<FRealSingle>& InSphereOffsetDistances,  // Use local indexation
-		const FPropertyCollectionConstAdapter& PropertyCollection,
+		const FCollectionPropertyConstFacade& PropertyCollection,
 		FSolverReal MeshScale
 	)
 		: AnimationPositions(InAnimationPositions)
@@ -198,7 +198,7 @@ public:
 	}
 	~FPBDSphericalBackstopConstraint() {}
 
-	void SetProperties(const FPropertyCollectionConstAdapter& PropertyCollection, FSolverReal MeshScale)
+	void SetProperties(const FCollectionPropertyConstFacade& PropertyCollection, FSolverReal MeshScale)
 	{
 		SetScale((FSolverReal)1., MeshScale);
 		// TODO: BackstopDistance and BackstopRadius
diff --git a/Engine/Source/Runtime/Experimental/Chaos/Public/Chaos/PBDSpringConstraints.h b/Engine/Source/Runtime/Experimental/Chaos/Public/Chaos/PBDSpringConstraints.h
--- a/Engine/Source/Runtime/Experimental/Chaos/Public/Chaos/PBDSpringConstraints.h
+++ b/Engine/Source/Runtime/Experimental/Chaos/Public/Chaos/PBDSpringConstraints.h
@@ -2,7 +2,7 @@
 #pragma once
 
 #include "Chaos/PBDSpringConstraintsBase.h"
-#include "Chaos/PropertyCollectionAdapter.h"
+#include "Chaos/CollectionPropertyFacade.h"
 
 namespace Chaos::Softs
 {
@@ -47,7 +47,7 @@ private:
 class CHAOS_API FPBDEdgeSpringConstraints final : public FPBDSpringConstraints
 {
 public:
-	static bool IsEnabled(const FPropertyCollectionConstAdapter& PropertyCollection)
+	static bool IsEnabled(const FCollectionPropertyConstFacade& PropertyCollection)
 	{
 		return IsEdgeSpringStiffnessEnabled(PropertyCollection, false);
 	}
@@ -58,7 +58,7 @@ public:
 		int32 ParticleCount,
 		const TArray<TVec3<int32>>& InConstraints,
 		const TConstArrayView<FRealSingle>& StiffnessMultipliers,
-		const FPropertyCollectionConstAdapter& PropertyCollection,
+		const FCollectionPropertyConstFacade& PropertyCollection,
 		bool bTrimKinematicConstraints = false)
 		: FPBDSpringConstraints(
 			Particles,
@@ -72,7 +72,7 @@ public:
 
 	virtual ~FPBDEdgeSpringConstraints() override = default;
 
-	void SetProperties(const FPropertyCollectionConstAdapter& PropertyCollection)
+	void SetProperties(const FCollectionPropertyConstFacade& PropertyCollection)
 	{
 		if (IsEdgeSpringStiffnessMutable(PropertyCollection))
 		{
@@ -87,7 +87,7 @@ private:
 class CHAOS_API FPBDBendingSpringConstraints final : public FPBDSpringConstraints
 {
 public:
-	static bool IsEnabled(const FPropertyCollectionConstAdapter& PropertyCollection)
+	static bool IsEnabled(const FCollectionPropertyConstFacade& PropertyCollection)
 	{
 		return IsBendingSpringStiffnessEnabled(PropertyCollection, false);
 	}
@@ -98,7 +98,7 @@ public:
 		int32 ParticleCount,
 		const TArray<TVec2<int32>>& InConstraints,
 		const TConstArrayView<FRealSingle>& StiffnessMultipliers,
-		const FPropertyCollectionConstAdapter& PropertyCollection,
+		const FCollectionPropertyConstFacade& PropertyCollection,
 		bool bTrimKinematicConstraints = false)
 		: FPBDSpringConstraints(
 			Particles,
@@ -112,7 +112,7 @@ public:
 
 	virtual ~FPBDBendingSpringConstraints() override = default;
 
-	void SetProperties(const FPropertyCollectionConstAdapter& PropertyCollection)
+	void SetProperties(const FCollectionPropertyConstFacade& PropertyCollection)
 	{
 		if (IsBendingSpringStiffnessMutable(PropertyCollection))
 		{
diff --git a/Engine/Source/Runtime/Experimental/Chaos/Public/Chaos/PBDTriangleMeshCollisions.h b/Engine/Source/Runtime/Experimental/Chaos/Public/Chaos/PBDTriangleMeshCollisions.h
--- a/Engine/Source/Runtime/Experimental/Chaos/Public/Chaos/PBDTriangleMeshCollisions.h
+++ b/Engine/Source/Runtime/Experimental/Chaos/Public/Chaos/PBDTriangleMeshCollisions.h
@@ -5,7 +5,7 @@
 #include "Chaos/PBDSoftsEvolutionFwd.h"
 #include "Chaos/TriangleMesh.h"
 #include "Chaos/HierarchicalSpatialHash.h"
-#include "Chaos/PropertyCollectionAdapter.h"
+#include "Chaos/CollectionPropertyFacade.h"
 
 // This does initialization for PBDCollisionSpringConstraints and PBDTriangleMeshIntersections, 
 // including intersection detection and global intersection analysis
@@ -95,7 +95,7 @@ public:
 		Count
 	};
 
-	static bool IsEnabled(const FPropertyCollectionConstAdapter& PropertyCollection)
+	static bool IsEnabled(const FCollectionPropertyConstFacade& PropertyCollection)
 	{
 		return PropertyCollection.IsEnabled(FName(TEXT("SelfCollisionStiffness")).ToString(), false);  // Don't use UE_CHAOS_DECLARE_PROPERTYCOLLECTION_NAME here, SelfCollisionStiffness is only needed for activation
 	}
@@ -104,7 +104,7 @@ public:
 		const int32 InOffset,
 		const int32 InNumParticles,
 		const FTriangleMesh& InTriangleMesh,
-		const FPropertyCollectionConstAdapter& PropertyCollection
+		const FCollectionPropertyConstFacade& PropertyCollection
 	)
 		:TriangleMesh(InTriangleMesh)
 		, Offset(InOffset)
@@ -129,7 +129,7 @@ public:
 
 	virtual ~FPBDTriangleMeshCollisions() = default;
 
-	void SetProperties(const FPropertyCollectionConstAdapter& PropertyCollection)
+	void SetProperties(const FCollectionPropertyConstFacade& PropertyCollection)
 	{
 		bool bUseSelfIntersections = false;
 		if (IsUseSelfIntersectionsMutable(PropertyCollection))
diff --git a/Engine/Source/Runtime/Experimental/Chaos/Public/Chaos/PBDTriangleMeshIntersections.h b/Engine/Source/Runtime/Experimental/Chaos/Public/Chaos/PBDTriangleMeshIntersections.h
--- a/Engine/Source/Runtime/Experimental/Chaos/Public/Chaos/PBDTriangleMeshIntersections.h
+++ b/Engine/Source/Runtime/Experimental/Chaos/Public/Chaos/PBDTriangleMeshIntersections.h
@@ -4,7 +4,7 @@
 #include "Chaos/Core.h"
 #include "Chaos/PBDSoftsEvolutionFwd.h"
 #include "Chaos/PBDTriangleMeshCollisions.h"
-#include "Chaos/PropertyCollectionAdapter.h"
+#include "Chaos/CollectionPropertyFacade.h"
 
 // Resolve intersections via contour minimization. Contour Minimization data is calculated in PBDTriangleMeshCollisions.
 namespace Chaos::Softs
@@ -12,7 +12,7 @@ namespace Chaos::Softs
 class CHAOS_API FPBDTriangleMeshIntersections
 {
 public:
-	static bool IsEnabled(const FPropertyCollectionConstAdapter& PropertyCollection)
+	static bool IsEnabled(const FCollectionPropertyConstFacade& PropertyCollection)
 	{
 		return PropertyCollection.IsEnabled(FName(TEXT("SelfCollisionStiffness")).ToString(), false);  // Don't use UE_CHAOS_DECLARE_PROPERTYCOLLECTION_NAME here, SelfCollisionStiffness is only needed for activation
 	}
diff --git a/Engine/Source/Runtime/Experimental/Chaos/Public/Chaos/PropertyCollectionAdapter.h b/Engine/Source/Runtime/Experimental/Chaos/Public/Chaos/PropertyCollectionAdapter.h
deleted file mode 100644
diff --git a/Engine/Source/Runtime/Experimental/Chaos/Public/Chaos/VelocityField.h b/Engine/Source/Runtime/Experimental/Chaos/Public/Chaos/VelocityField.h
--- a/Engine/Source/Runtime/Experimental/Chaos/Public/Chaos/VelocityField.h
+++ b/Engine/Source/Runtime/Experimental/Chaos/Public/Chaos/VelocityField.h
@@ -4,7 +4,7 @@
 #include "Chaos/PBDSoftsEvolutionFwd.h"
 #include "Chaos/PBDSoftsSolverParticles.h"
 #include "Chaos/TriangleMesh.h"
-#include "Chaos/PropertyCollectionAdapter.h"
+#include "Chaos/CollectionPropertyFacade.h"
 
 namespace Chaos::Softs
 {
@@ -19,7 +19,7 @@ public:
 	static constexpr FSolverReal MinCoefficient = (FSolverReal)0.;   // Applies to both drag and lift
 	static constexpr FSolverReal MaxCoefficient = (FSolverReal)10.;  //
 
-	static bool IsEnabled(const FPropertyCollectionConstAdapter& PropertyCollection)
+	static bool IsEnabled(const FCollectionPropertyConstFacade& PropertyCollection)
 	{
 		return
 			IsDragEnabled(PropertyCollection, false) ||
@@ -27,7 +27,7 @@ public:
 			IsPressureEnabled(PropertyCollection, false);
 	}
 
-	explicit FVelocityAndPressureField(const FPropertyCollectionConstAdapter& PropertyCollection)
+	explicit FVelocityAndPressureField(const FCollectionPropertyConstFacade& PropertyCollection)
 		: Offset(INDEX_NONE)
 		, NumParticles(0)
 	{
@@ -61,7 +61,7 @@ public:
 		}
 	}
 
-	void SetProperties(const FPropertyCollectionConstAdapter& PropertyCollection, FSolverReal WorldScale)
+	void SetProperties(const FCollectionPropertyConstFacade& PropertyCollection, FSolverReal WorldScale)
 	{
 		if (IsFluidDensityMutable(PropertyCollection))
 		{
diff --git a/Engine/Source/Runtime/Experimental/Chaos/Public/Chaos/XPBDAxialSpringConstraints.h b/Engine/Source/Runtime/Experimental/Chaos/Public/Chaos/XPBDAxialSpringConstraints.h
--- a/Engine/Source/Runtime/Experimental/Chaos/Public/Chaos/XPBDAxialSpringConstraints.h
+++ b/Engine/Source/Runtime/Experimental/Chaos/Public/Chaos/XPBDAxialSpringConstraints.h
@@ -2,7 +2,7 @@
 #pragma once
 
 #include "Chaos/PBDAxialSpringConstraintsBase.h"
-#include "Chaos/PropertyCollectionAdapter.h"
+#include "Chaos/CollectionPropertyFacade.h"
 #include "ChaosStats.h"
 
 DECLARE_CYCLE_STAT(TEXT("Chaos XPBD Axial Spring Constraint"), STAT_XPBD_AxialSpring, STATGROUP_Chaos);
@@ -168,7 +168,7 @@ private:
 class FXPBDAreaSpringConstraints final : public FXPBDAxialSpringConstraints
 {
 public:
-	static bool IsEnabled(const FPropertyCollectionConstAdapter& PropertyCollection)
+	static bool IsEnabled(const FCollectionPropertyConstFacade& PropertyCollection)
 	{
 		return IsXPBDAreaSpringStiffnessEnabled(PropertyCollection, false);
 	}
@@ -179,7 +179,7 @@ public:
 		int32 ParticleCount,
 		const TArray<TVec3<int32>>& InConstraints,
 		const TConstArrayView<FRealSingle>& StiffnessMultipliers,
-		const FPropertyCollectionConstAdapter& PropertyCollection,
+		const FCollectionPropertyConstFacade& PropertyCollection,
 		bool bTrimKinematicConstraints)
 		: FXPBDAxialSpringConstraints(
 			Particles,
@@ -193,7 +193,7 @@ public:
 
 	virtual ~FXPBDAreaSpringConstraints() override = default;
 
-	void SetProperties(const FPropertyCollectionConstAdapter& PropertyCollection)
+	void SetProperties(const FCollectionPropertyConstFacade& PropertyCollection)
 	{
 		if (IsXPBDAreaSpringStiffnessMutable(PropertyCollection))
 		{
diff --git a/Engine/Source/Runtime/Experimental/Chaos/Public/Chaos/XPBDBendingConstraints.h b/Engine/Source/Runtime/Experimental/Chaos/Public/Chaos/XPBDBendingConstraints.h
--- a/Engine/Source/Runtime/Experimental/Chaos/Public/Chaos/XPBDBendingConstraints.h
+++ b/Engine/Source/Runtime/Experimental/Chaos/Public/Chaos/XPBDBendingConstraints.h
@@ -2,7 +2,7 @@
 #pragma once
 
 #include "Chaos/PBDBendingConstraintsBase.h"
-#include "Chaos/PropertyCollectionAdapter.h"
+#include "Chaos/CollectionPropertyFacade.h"
 #include "ChaosStats.h"
 
 namespace Chaos::Softs
@@ -25,7 +25,7 @@ public:
 	static constexpr FSolverReal MinDamping = (FSolverReal)0.;
 	static constexpr FSolverReal MaxDamping = (FSolverReal)1000.;
 
-	static bool IsEnabled(const FPropertyCollectionConstAdapter& PropertyCollection)
+	static bool IsEnabled(const FCollectionPropertyConstFacade& PropertyCollection)
 	{
 		return IsXPBDBendingElementStiffnessEnabled(PropertyCollection, false);
 	}
@@ -37,7 +37,7 @@ public:
 		const TConstArrayView<FRealSingle>& StiffnessMultipliers,
 		const TConstArrayView<FRealSingle>& BucklingStiffnessMultipliers,
 		const TConstArrayView<FRealSingle>& DampingMultipliers,
-		const FPropertyCollectionConstAdapter& PropertyCollection,
+		const FCollectionPropertyConstFacade& PropertyCollection,
 		bool bTrimKinematicConstraints = false)
 		: Base(
 			InParticles,
@@ -135,7 +135,7 @@ public:
 		FPBDBendingConstraintsBase::Init(InParticles);
 	}
 
-	void SetProperties(const FPropertyCollectionConstAdapter& PropertyCollection)
+	void SetProperties(const FCollectionPropertyConstFacade& PropertyCollection)
 	{
 		if (IsXPBDBendingElementStiffnessMutable(PropertyCollection))
 		{
diff --git a/Engine/Source/Runtime/Experimental/Chaos/Public/Chaos/XPBDLongRangeConstraints.h b/Engine/Source/Runtime/Experimental/Chaos/Public/Chaos/XPBDLongRangeConstraints.h
--- a/Engine/Source/Runtime/Experimental/Chaos/Public/Chaos/XPBDLongRangeConstraints.h
+++ b/Engine/Source/Runtime/Experimental/Chaos/Public/Chaos/XPBDLongRangeConstraints.h
@@ -3,7 +3,7 @@
 
 #include "Chaos/PBDLongRangeConstraintsBase.h"
 #include "Chaos/PBDParticles.h"
-#include "Chaos/PropertyCollectionAdapter.h"
+#include "Chaos/CollectionPropertyFacade.h"
 #include "ChaosStats.h"
 
 DECLARE_CYCLE_STAT(TEXT("Chaos XPBD Long Range Constraint"), STAT_XPBD_LongRange, STATGROUP_Chaos);
@@ -32,7 +32,7 @@ public:
 		const TArray<TConstArrayView<TTuple<int32, int32, FRealSingle>>>& InTethers,
 		const TConstArrayView<FRealSingle>& StiffnessMultipliers,
 		const TConstArrayView<FRealSingle>& ScaleMultipliers,
-		const FPropertyCollectionConstAdapter& PropertyCollection)
+		const FCollectionPropertyConstFacade& PropertyCollection)
 	    : FPBDLongRangeConstraintsBase(
 			Particles,
 			InParticleOffset,
@@ -82,7 +82,7 @@ public:
 
 	virtual ~FXPBDLongRangeConstraints() override {}
 
-	void SetProperties(const FPropertyCollectionConstAdapter& PropertyCollection, FSolverReal MeshScale)
+	void SetProperties(const FCollectionPropertyConstFacade& PropertyCollection, FSolverReal MeshScale)
 	{
 		if (IsXPBDTetherStiffnessMutable(PropertyCollection))
 		{
diff --git a/Engine/Source/Runtime/Experimental/Chaos/Public/Chaos/XPBDSpringConstraints.h b/Engine/Source/Runtime/Experimental/Chaos/Public/Chaos/XPBDSpringConstraints.h
--- a/Engine/Source/Runtime/Experimental/Chaos/Public/Chaos/XPBDSpringConstraints.h
+++ b/Engine/Source/Runtime/Experimental/Chaos/Public/Chaos/XPBDSpringConstraints.h
@@ -4,7 +4,7 @@
 #include "Chaos/PBDSpringConstraintsBase.h"
 #include "ChaosStats.h"
 #include "PBDWeightMap.h"
-#include "Chaos/PropertyCollectionAdapter.h"
+#include "Chaos/CollectionPropertyFacade.h"
 
 namespace Chaos::Softs
 {
@@ -158,7 +158,7 @@ private:
 class CHAOS_API FXPBDEdgeSpringConstraints final : public FXPBDSpringConstraints
 {
 public:
-	static bool IsEnabled(const FPropertyCollectionConstAdapter& PropertyCollection)
+	static bool IsEnabled(const FCollectionPropertyConstFacade& PropertyCollection)
 	{
 		return IsXPBDEdgeSpringStiffnessEnabled(PropertyCollection, false);
 	}
@@ -170,7 +170,7 @@ public:
 		const TArray<TVec3<int32>>& InConstraints,
 		const TConstArrayView<FRealSingle>& StiffnessMultipliers,
 		const TConstArrayView<FRealSingle>& DampingMultipliers,
-		const FPropertyCollectionConstAdapter& PropertyCollection,
+		const FCollectionPropertyConstFacade& PropertyCollection,
 		bool bTrimKinematicConstraints = false)
 		: FXPBDSpringConstraints(
 			Particles,
@@ -186,7 +186,7 @@ public:
 
 	virtual ~FXPBDEdgeSpringConstraints() override = default;
 
-	void SetProperties(const FPropertyCollectionConstAdapter& PropertyCollection)
+	void SetProperties(const FCollectionPropertyConstFacade& PropertyCollection)
 	{
 		if (IsXPBDEdgeSpringStiffnessMutable(PropertyCollection))
 		{
@@ -209,7 +209,7 @@ private:
 class CHAOS_API FXPBDBendingSpringConstraints : public FXPBDSpringConstraints
 {
 public:
-	static bool IsEnabled(const FPropertyCollectionConstAdapter& PropertyCollection)
+	static bool IsEnabled(const FCollectionPropertyConstFacade& PropertyCollection)
 	{
 		return IsXPBDBendingSpringStiffnessEnabled(PropertyCollection, false);
 	}
@@ -221,7 +221,7 @@ public:
 		const TArray<TVec2<int32>>& InConstraints,
 		const TConstArrayView<FRealSingle>& StiffnessMultipliers,
 		const TConstArrayView<FRealSingle>& DampingMultipliers,
-		const FPropertyCollectionConstAdapter& PropertyCollection,
+		const FCollectionPropertyConstFacade& PropertyCollection,
 		bool bTrimKinematicConstraints = false)
 		: FXPBDSpringConstraints(
 			Particles,
@@ -237,7 +237,7 @@ public:
 
 	virtual ~FXPBDBendingSpringConstraints() override = default;
 
-	void SetProperties(const FPropertyCollectionConstAdapter& PropertyCollection)
+	void SetProperties(const FCollectionPropertyConstFacade& PropertyCollection)
 	{
 		if (IsXPBDBendingSpringStiffnessMutable(PropertyCollection))
 		{
