Constraints: Added some missing nullptr checks that were causing a crash in Dev-Cine.

diff --git a/Engine/Source/Runtime/Experimental/Animation/Constraints/Private/TransformConstraint.cpp b/Engine/Source/Runtime/Experimental/Animation/Constraints/Private/TransformConstraint.cpp
--- a/Engine/Source/Runtime/Experimental/Animation/Constraints/Private/TransformConstraint.cpp
+++ b/Engine/Source/Runtime/Experimental/Animation/Constraints/Private/TransformConstraint.cpp
@@ -265,7 +265,7 @@ FString UTickableTransformConstraint::GetFullLabel() const
 		return DummyLabel;
 	}
 	
-	if (ParentTRSHandle->IsValid())
+	if (ParentTRSHandle && ParentTRSHandle->IsValid())
 	{
 		return FString::Printf(TEXT("%s.%s"), *ParentTRSHandle->GetFullLabel(), *ChildTRSHandle->GetFullLabel() );		
 	}
@@ -335,7 +335,7 @@ void UTickableTransformConstraint::RegisterDelegates()
 
 void UTickableTransformConstraint::Setup()
 {
-	if (!ParentTRSHandle->IsValid() || !ChildTRSHandle->IsValid())
+	if (!ParentTRSHandle || !ChildTRSHandle || !ParentTRSHandle->IsValid() || !ChildTRSHandle->IsValid())
 	{
 		// handle error
 		return;
