Geometry collection : static analysis fixes

diff --git a/Engine/Source/Runtime/Experimental/Chaos/Private/GeometryCollection/GeometryCollectionConvexUtility.cpp b/Engine/Source/Runtime/Experimental/Chaos/Private/GeometryCollection/GeometryCollectionConvexUtility.cpp
--- a/Engine/Source/Runtime/Experimental/Chaos/Private/GeometryCollection/GeometryCollectionConvexUtility.cpp
+++ b/Engine/Source/Runtime/Experimental/Chaos/Private/GeometryCollection/GeometryCollectionConvexUtility.cpp
@@ -2092,19 +2092,17 @@ void FGeometryCollectionConvexUtility::ConvertScaledImplicitToConvexArray(
 		const Chaos::FConvex* Convex = nullptr;
 		if (bInstanced)
 		{
-			const auto ScaledInstancedConvex = Implicit.template GetObject<Chaos::TImplicitObjectScaled<Chaos::FConvex, true>>();
-			ScaleTM.SetScale3D(ScaledInstancedConvex->GetScale());
-			if (ScaledInstancedConvex)
+			if (const auto ScaledInstancedConvex = Implicit.template GetObject<Chaos::TImplicitObjectScaled<Chaos::FConvex, true>>())
 			{
+				ScaleTM.SetScale3D(ScaledInstancedConvex->GetScale());
 				Convex = ScaledInstancedConvex->GetUnscaledObject();
 			}
 		}
 		else
 		{
-			const auto ScaledConvex = Implicit.template GetObject<Chaos::TImplicitObjectScaled<Chaos::FConvex, false>>();
-			ScaleTM.SetScale3D(ScaledConvex->GetScale());
-			if (ScaledConvex)
+			if (const auto ScaledConvex = Implicit.template GetObject<Chaos::TImplicitObjectScaled<Chaos::FConvex, false>>())
 			{
+				ScaleTM.SetScale3D(ScaledConvex->GetScale());
 				Convex = ScaledConvex->GetUnscaledObject();
 			}
 		}
