ChaosFlesh sim space.

diff --git a/Engine/Plugins/Experimental/ChaosFlesh/Source/ChaosFleshEngine/Private/ChaosFlesh/FleshComponent.cpp b/Engine/Plugins/Experimental/ChaosFlesh/Source/ChaosFleshEngine/Private/ChaosFlesh/FleshComponent.cpp
--- a/Engine/Plugins/Experimental/ChaosFlesh/Source/ChaosFleshEngine/Private/ChaosFlesh/FleshComponent.cpp
+++ b/Engine/Plugins/Experimental/ChaosFlesh/Source/ChaosFleshEngine/Private/ChaosFlesh/FleshComponent.cpp
@@ -19,6 +19,7 @@
 #include "GeometryCollection/Facades/CollectionTetrahedralSkeletalBindingsFacade.h"
 #include "GeometryCollection/TransformCollection.h"
 #include "ProceduralMeshComponent.h"
+#include "Rendering/SkeletalMeshRenderData.h"
 
 
 #include UE_INLINE_GENERATED_CPP_BY_NAME(FleshComponent)
@@ -124,6 +125,7 @@ UDeformablePhysicsComponent::FThreadingProxy* UFleshComponent::NewProxy()
 	SCOPE_CYCLE_COUNTER(STAT_ChaosDeformable_UFleshComponent_NewProxy);
 	TRACE_CPUPROFILER_EVENT_SCOPE(ChaosDeformable_UFleshComponent_NewProxy);
 
+	UpdateSimSpaceTransformIndex();
 	if (const UFleshAsset* RestAsset = GetRestCollection())
 	{
 		if (const FFleshCollection* Rest = RestAsset->GetCollection())
@@ -136,7 +138,17 @@ UDeformablePhysicsComponent::FThreadingProxy* UFleshComponent::NewProxy()
 			GetDynamicCollection()->Reset(Rest);
 			if (const FManagedArrayCollection* Dynamic = GetDynamicCollection()->GetCollection())
 			{
-				return new FFleshThreadingProxy(this, GetComponentTransform(), *Rest, *Dynamic);
+				// Mesh points are in component space, such that the exterior hull aligns with the
+                // surface of the skeletal mesh, which is subject to the transform hierarchy.
+				const FTransform& ComponentToWorldXf = GetComponentTransform();
+				const FTransform ComponentToSimXf = GetSimSpaceRestTransform();
+				return new FFleshThreadingProxy(
+					this,
+					ComponentToWorldXf, 
+					ComponentToSimXf, 
+					SimSpace,
+					*Rest, 
+					*Dynamic);
 			}
 		}
 	}
@@ -166,6 +178,7 @@ UDeformablePhysicsComponent::FDataMapValue UFleshComponent::NewDeformableData()
 						TArray<FTransform> AnimationTransforms = RestTransforms->GetConstArray();
 						TArray<FTransform> ComponentPose = RestTransforms->GetConstArray();
 
+						// Extract animated transforms from all skeletal meshes.
 						for (const USkeletalMeshComponent* SkeletalMeshComponent : SkeletalMeshComponents)
 						{
 							if (const USkeletalMesh* SkeletalMesh = SkeletalMeshComponent->GetSkeletalMeshAsset())
@@ -186,28 +199,133 @@ UDeformablePhysicsComponent::FDataMapValue UFleshComponent::NewDeformableData()
 												// @todo(flesh) : Can we just use one array?
 												AnimationTransforms[Adx] = ComponentTransforms[Cdx];
 												ComponentPose[Adx] = ComponentLocalPose[Cdx]; 
+
+												if (SimSpaceTransformGlobalIndex == INDEX_NONE &&
+													SimSpaceTransformIndex == Cdx &&
+													SkeletalMesh == SimSpaceSkeletalMesh)
+												{
+													SimSpaceTransformGlobalIndex = Adx;
+												}
 											}
+
+
 										}
 									}
 								}
 							}
 						}
-						return FDataMapValue(new Chaos::Softs::FFleshThreadingProxy::FFleshInputBuffer(this->GetComponentTransform(), AnimationTransforms, ComponentPose, bTempEnableGravity, StiffnessMultiplier, DampingMultiplier, MassMultiplier, this));
+
+						FTransform BoneSpaceXf;
+						if (AnimationTransforms.IsValidIndex(SimSpaceTransformGlobalIndex))
+						{
+							BoneSpaceXf = AnimationTransforms[SimSpaceTransformGlobalIndex];
+						}
+						else
+						{
+							BoneSpaceXf = FTransform::Identity;
+						}
+
+						return FDataMapValue(
+							new Chaos::Softs::FFleshThreadingProxy::FFleshInputBuffer(
+								this->GetComponentTransform(), 
+								BoneSpaceXf, 
+								SimSpaceTransformGlobalIndex,
+								AnimationTransforms, 
+								ComponentPose, 
+								bTempEnableGravity, 
+								StiffnessMultiplier, 
+								DampingMultiplier, 
+								MassMultiplier, 
+								this));
 					}
 				}
 			}
 		}
 	}
-	return FDataMapValue(new Chaos::Softs::FFleshThreadingProxy::FFleshInputBuffer(this->GetComponentTransform(), bTempEnableGravity, StiffnessMultiplier, DampingMultiplier, MassMultiplier, this));
+	return FDataMapValue(
+		new Chaos::Softs::FFleshThreadingProxy::FFleshInputBuffer(
+			this->GetComponentTransform(), 
+			GetSimSpaceRestTransform(),
+			SimSpaceTransformGlobalIndex,
+			bTempEnableGravity,
+			StiffnessMultiplier, 
+			DampingMultiplier, 
+			MassMultiplier, 
+			this));
+}
+
+TArray<FString> UFleshComponent::GetSimSpaceBoneNameOptions() const
+{
+	TArray<FString> Names;
+	if (RestCollection)
+	{
+		if (RestCollection->SkeletalMesh)
+		{
+			const FReferenceSkeleton& RefSkeleton = RestCollection->SkeletalMesh->GetSkeleton()->GetReferenceSkeleton();
+			Names.SetNum(RefSkeleton.GetNum());
+			for (int32 i = 0; i < RefSkeleton.GetNum(); i++)
+			{
+				Names[i] = RefSkeleton.GetBoneName(i).ToString();
+			}
+		}
+	}
+	return Names;
+}
+
+bool UFleshComponent::UpdateSimSpaceTransformIndex()
+{
+	SimSpaceTransformIndex = INDEX_NONE;
+	SimSpaceSkeletalMesh = nullptr;
+
+	if (SimSpace != ChaosDeformableSimSpace::Bone)
+	{
+		return false;
+	}
+
+	if (RestCollection && RestCollection->SkeletalMesh)
+	{
+		const FReferenceSkeleton& RefSkeleton = RestCollection->SkeletalMesh->GetSkeleton()->GetReferenceSkeleton();
+		for (int32 i = 0; i < RefSkeleton.GetNum(); i++)
+		{
+			if (RefSkeleton.GetBoneName(i).ToString() == SimSpaceBoneName.ToString())
+			{
+				SimSpaceSkeletalMesh = RestCollection->SkeletalMesh;
+				SimSpaceTransformIndex = i;
+				return true;
+			}
+		}
+	}
+	return false;
 }
 
+FTransform UFleshComponent::GetSimSpaceRestTransform() const
+{
+	if (SimSpaceSkeletalMesh == nullptr)
+	{
+		return FTransform::Identity;
+	}
+	
+	TArray<FTransform> ComponentTransforms;
+	ComponentTransforms.SetNum(SimSpaceSkeletalMesh->GetRefSkeleton().GetNum());
+
+	SimSpaceSkeletalMesh->FillComponentSpaceTransforms(
+		SimSpaceSkeletalMesh->GetRefSkeleton().GetRefBonePose(),
+		SimSpaceSkeletalMesh->GetResourceForRendering()->LODRenderData[0].RequiredBones,
+		ComponentTransforms);
+
+	if (!ComponentTransforms.IsValidIndex(SimSpaceTransformIndex))
+	{
+		return FTransform::Identity;
+	}
+	const FTransform& ComponentToBone = ComponentTransforms[SimSpaceTransformIndex];
+	return ComponentToBone;
+}
 
 void UFleshComponent::UpdateFromSimualtion(const FDataMapValue* SimualtionBuffer)
 {
 	SCOPE_CYCLE_COUNTER(STAT_ChaosDeformable_UFleshComponent_UpdateFromSimualtion);
 	TRACE_CPUPROFILER_EVENT_SCOPE(ChaosDeformable_UFleshComponent_UpdateFromSimualtion);
 
-	FTransform CurrTransform = GetComponentTransform();
 	if (const FFleshThreadingProxy::FFleshOutputBuffer* FleshBuffer = (*SimualtionBuffer)->As<FFleshThreadingProxy::FFleshOutputBuffer>())
 	{
 		if (GetDynamicCollection())
@@ -219,13 +337,14 @@ void UFleshComponent::UpdateFromSimualtion(const FDataMapValue* SimualtionBuffer
 			TManagedArray<FVector3f>& DynamicVertex = GetDynamicCollection()->GetPositions();
 			const TManagedArray<FVector3f>& SimulationVertex = FleshBuffer->Dynamic.GetAttribute<FVector3f>("Vertex", FGeometryCollection::VerticesGroup);
 
+			// Simulator produces results in component space.
 			for (int i = DynamicVertex.Num() - 1; i >= 0; i--)
 			{
-				DynamicVertex[i] = UEVertf(GetComponentTransform().InverseTransformPosition(UEVertd(SimulationVertex[i])));
+				DynamicVertex[i] = SimulationVertex[i];
 			}
-
-			 //p.Chaos.DebugDraw.Enabled 1
-			 //p.Chaos.DebugDraw.Deformable.SkeletalMeshBindingPositions 1
+			
+			//p.Chaos.DebugDraw.Enabled 1
+			//p.Chaos.DebugDraw.Deformable.SkeletalMeshBindingPositions 1
 			if (CVarParams.bDoDrawSkeletalMeshBindingPositions)
 			{
 				DebugDrawSkeletalMeshBindingPositions();
diff --git a/Engine/Plugins/Experimental/ChaosFlesh/Source/ChaosFleshEngine/Public/ChaosFlesh/FleshComponent.h b/Engine/Plugins/Experimental/ChaosFlesh/Source/ChaosFleshEngine/Public/ChaosFlesh/FleshComponent.h
--- a/Engine/Plugins/Experimental/ChaosFlesh/Source/ChaosFleshEngine/Public/ChaosFlesh/FleshComponent.h
+++ b/Engine/Plugins/Experimental/ChaosFlesh/Source/ChaosFleshEngine/Public/ChaosFlesh/FleshComponent.h
@@ -66,11 +66,41 @@ public:
 
 
 private:
+	// FleshAsset that describes the simulation rest state.
 	UPROPERTY(EditAnywhere, Category = "ChaosPhysics")
 	TObjectPtr<const UFleshAsset> RestCollection;
 
+	// Current simulation state.
 	UPROPERTY(EditAnywhere, Category = "ChaosPhysics")
-	TObjectPtr< UFleshDynamicAsset> DynamicCollection;
+	TObjectPtr<UFleshDynamicAsset> DynamicCollection;
+
+	//
+	// Sim Space
+	//
+
+	// Space the simulation will run in.
+	UPROPERTY(EditAnywhere, Category = "ChaosDeformable")
+	TEnumAsByte<ChaosDeformableSimSpace> SimSpace = ChaosDeformableSimSpace::World;
+
+	// Bone from the associated skeletal mesh (indicated by RestCollection.TargetSkeletalMesh) to use as the space the sim runs in.
+	UPROPERTY(EditAnywhere, Category = "ChaosDeformable", meta = (GetOptions = "GetSimSpaceBoneNameOptions", EditCondition = "SimSpace == ChaosDeformableSimSpace::Bone"))
+	FName SimSpaceBoneName;
+
+	UPROPERTY()
+	TObjectPtr<USkeletalMesh> SimSpaceSkeletalMesh;
+
+	// Returns a list of bone names from the currently selected skeletal mesh.
+	UFUNCTION(CallInEditor)
+	TArray<FString> GetSimSpaceBoneNameOptions() const;
+
+	//! Update \c SimSpaceSkeletalMesh and \c SimSpaceTransformIndex according to
+	//! \c RestCollection->TargetSkeletalMesh and SimSpaceBoneName.
+	//! \ret \c true if a valid sim space transform is found.
+	bool UpdateSimSpaceTransformIndex();
+
+	//! Return the rest transform to be used as the simulation space. 
+	//! \c UpdateSimSpaceTransformIndex() must be called prior to calling this function.
+	FTransform GetSimSpaceRestTransform() const;
 
 	//
 	// Render the Procedural Mesh
@@ -95,5 +125,8 @@ private:
 
 	TArray<FVector> GetSkeletalMeshBindingPositionsInternal(const USkeletalMesh* InSkeletalMesh, TArray<bool>* OutInfluence = nullptr) const;
 	void DebugDrawSkeletalMeshBindingPositions() const;
+
+	int32 SimSpaceTransformIndex = INDEX_NONE;
+	int32 SimSpaceTransformGlobalIndex = INDEX_NONE;
 };
 
diff --git a/Engine/Source/Runtime/Experimental/Chaos/Private/Chaos/Deformable/ChaosDeformableSolver.cpp b/Engine/Source/Runtime/Experimental/Chaos/Private/Chaos/Deformable/ChaosDeformableSolver.cpp
--- a/Engine/Source/Runtime/Experimental/Chaos/Private/Chaos/Deformable/ChaosDeformableSolver.cpp
+++ b/Engine/Source/Runtime/Experimental/Chaos/Private/Chaos/Deformable/ChaosDeformableSolver.cpp
@@ -163,10 +163,33 @@ namespace Chaos::Softs
 			RemoveSimulationObjects();
 			UpdateProxyInputPackages();
 			InitializeSimulationObjects();
+			InitializeSimulationSpace();
 			AdvanceDt(DeltaTime);
 		}
 	}
 
+	void FDeformableSolver::InitializeSimulationSpace()
+	{
+		for (int32 Index = 0; Index < this->MObjects.Num(); Index++)
+		{
+			if (const UObject* Owner = this->MObjects[Index])
+			{
+				if (FFleshThreadingProxy* Proxy = Proxies[Owner]->As<FFleshThreadingProxy>())
+				{
+					if (this->CurrentInputPackage->ObjectMap.Contains(Owner))
+					{
+						FFleshThreadingProxy::FFleshInputBuffer* FleshInputBuffer =
+							this->CurrentInputPackage->ObjectMap[Owner]->As<FFleshThreadingProxy::FFleshInputBuffer>();
+						if (FleshInputBuffer)
+						{
+							Proxy->UpdateSimSpace(FleshInputBuffer->WorldToComponentXf, FleshInputBuffer->ComponentToBoneXf);
+						}
+					}
+				}
+			}
+		}
+	}
+
 	void FDeformableSolver::InitializeSimulationObjects()
 	{
 		SCOPE_CYCLE_COUNTER(STAT_ChaosDeformableSolver_InitializeSimulationObjects);
@@ -292,23 +315,41 @@ namespace Chaos::Softs
 
 		Evolution->SetDamping(DampingMultiplier, GroupOffset - 1);
 
-		const FTransform& InitialTransform = Proxy.GetInitialTransform();
+		// Tet mesh points are in component space.  That means that if our sim space is:
+		//    World:     We need to multiply by the ComponentToWorldXf.
+		//	  Component: We do nothing.
+		//    Bone:      The points have the (component relative) bone transform baked in,
+		//               and we need to remove it.  Muliply by the BoneToComponentXf.
 
-
-		for (uint32 vdx = 0; vdx < NumParticles; ++vdx)
+		const FTransform& InitialPointsXf = Proxy.GetInitialPointsTransform();
+		if (!InitialPointsXf.Equals(FTransform::Identity))
 		{
-			int32 SolverParticleIndex = ParticleStart + vdx;
-			Evolution->Particles().X(SolverParticleIndex) = ChaosVert(InitialTransform.TransformPosition(DoubleVert(DynamicVertex[vdx])));
-			Evolution->Particles().V(SolverParticleIndex) = Chaos::FVec3(0.f, 0.f, 0.f);
-			Evolution->Particles().M(SolverParticleIndex) = MassWithMultiplier[vdx];
-			Evolution->Particles().InvM(SolverParticleIndex) = ChaosInvM(Evolution->Particles().M(SolverParticleIndex));
-			Evolution->Particles().PAndInvM(SolverParticleIndex).InvM = Evolution->Particles().InvM(SolverParticleIndex);
+			for (uint32 vdx = 0; vdx < NumParticles; ++vdx)
+			{
+				int32 SolverParticleIndex = ParticleStart + vdx;
+				Evolution->Particles().X(SolverParticleIndex) = ChaosVert(InitialPointsXf.TransformPosition(DoubleVert(DynamicVertex[vdx])));
+				Evolution->Particles().V(SolverParticleIndex) = Chaos::FVec3(0.f, 0.f, 0.f);
+				Evolution->Particles().M(SolverParticleIndex) = MassWithMultiplier[vdx];
+				Evolution->Particles().InvM(SolverParticleIndex) = ChaosInvM(Evolution->Particles().M(SolverParticleIndex));
+				Evolution->Particles().PAndInvM(SolverParticleIndex).InvM = Evolution->Particles().InvM(SolverParticleIndex);
+			}
+		}
+		else
+		{
+			for (uint32 vdx = 0; vdx < NumParticles; ++vdx)
+			{
+				int32 SolverParticleIndex = ParticleStart + vdx;
+				Evolution->Particles().X(SolverParticleIndex) = DynamicVertex[vdx];
+				Evolution->Particles().V(SolverParticleIndex) = Chaos::FVec3(0.f, 0.f, 0.f);
+				Evolution->Particles().M(SolverParticleIndex) = MassWithMultiplier[vdx];
+				Evolution->Particles().InvM(SolverParticleIndex) = ChaosInvM(Evolution->Particles().M(SolverParticleIndex));
+				Evolution->Particles().PAndInvM(SolverParticleIndex).InvM = Evolution->Particles().InvM(SolverParticleIndex);
+			}
 		}
 
 		bool ObjectEnableGravity = false;
-
-		//int32 SolverParticleIndex = ParticleStart;
-		if (const UObject* Owner = this->MObjects[ParticleStart]) {
+		if (const UObject* Owner = this->MObjects[ParticleStart]) 
+		{
 			FFleshThreadingProxy::FFleshInputBuffer* FleshInputBuffer = nullptr;
 			if (this->CurrentInputPackage->ObjectMap.Contains(Owner))
 			{
@@ -325,6 +366,14 @@ namespace Chaos::Softs
 			FSolverVec3 ZeroGravity(0.f);
 			Evolution->SetGravity(ZeroGravity, GroupOffset - 1);
 		}
+		else
+		{
+			// Gravity points "down" in world space, but we need to orient it to
+			// whatever our sim space is.
+			FSolverVec3 GravityDir = Evolution->GetGravity();
+			FSolverVec3 SimSpaceGravityDir = Proxy.RotateWorldSpaceVector(GravityDir);
+			Evolution->SetGravity(SimSpaceGravityDir, GroupOffset - 1);
+		}
 
 		Proxy.SetSolverParticleRange(ParticleStart, NumParticles);
 	}
@@ -649,7 +698,7 @@ namespace Chaos::Softs
 				{
 					if (const FFleshThreadingProxy* Proxy = Proxies[Owner]->As<FFleshThreadingProxy>())
 					{
-						FTransform GlobalTransform = Proxy->GetInitialTransform();
+						FTransform GlobalTransform = Proxy->GetCurrentPointsTransform();
 						const FIntVector2& Range = Proxy->GetSolverParticleRange();
 						const FManagedArrayCollection& Rest = Proxy->GetRestCollection();
 
@@ -665,10 +714,6 @@ namespace Chaos::Softs
 							if (this->CurrentInputPackage->ObjectMap.Contains(Owner))
 							{
 								FleshInputBuffer = this->CurrentInputPackage->ObjectMap[Owner]->As<FFleshThreadingProxy::FFleshInputBuffer>();
-								if (FleshInputBuffer)
-								{
-									GlobalTransform = FleshInputBuffer->GlobalTransform;
-								}
 							}
 
 							typedef GeometryCollection::Facades::FVertexBoneWeightsFacade FWeightsFacade;
@@ -704,19 +749,47 @@ namespace Chaos::Softs
 											{
 												if (BoneIndices[i] > -1 && BoneIndices[i] < RestNum && BoneIndices[i] < TransformNum)
 												{
+													// We're lerp'ing between the previous skinned position and the current skinned 
+													// position, based on the substep.  If our SimSpace is a bone, we need to transform
+													// those points by the inverse bone transform at each of those points in time.
+													const FVector3f& Pt = Vertex[Index - Range[0]];
 
 													// @todo(flesh) : Add the pre-cached component space rest transforms to the rest collection. 
 													// see  UFleshComponent::NewDeformableData for how its pulled from the SkeletalMesh
-													FVec3 LocalPoint = FleshInputBuffer->RestTransforms[BoneIndices[i]].InverseTransformPosition(ChaosVert(Vertex[Index - Range[0]]));
-													FVec3 ComponentPointAtT = FleshInputBuffer->Transforms[BoneIndices[i]].TransformPosition(LocalPoint);
+													const FTransform& BoneRestTransform = FleshInputBuffer->RestTransforms[BoneIndices[i]];
+													FVec3 LocalPoint = BoneRestTransform.InverseTransformPosition(ChaosVert(Pt));
+													const FTransform& BoneCurrentTransform = FleshInputBuffer->Transforms[BoneIndices[i]];
+													FVec3 ComponentPointAtT = BoneCurrentTransform.TransformPosition(LocalPoint);
 
 													if (PreviousFleshBuffer)
 													{
-														FTransform BonePreviousTransform = PreviousFleshBuffer->Transforms[BoneIndices[i]];
-														ComponentPointAtT = ComponentPointAtT * CurrentRatio + BonePreviousTransform.TransformPosition(LocalPoint) * ((FSolverReal)1. - CurrentRatio);
+														const FTransform& BonePreviousTransform = PreviousFleshBuffer->Transforms[BoneIndices[i]];
+														FVec3 PreviousComponentPointAtT = BonePreviousTransform.TransformPosition(LocalPoint);
+
+														if (Proxy->IsBoneSpace())
+														{
+															// Current position gets current bone xf
+															ComponentPointAtT = GlobalTransform.TransformPosition(ComponentPointAtT);
+															// Previous position gets previous bone xf
+															const FTransform& PrevGlobalTransform = Proxy->GetPreviousPointsTransform();
+															PreviousComponentPointAtT = PrevGlobalTransform.TransformPosition(PreviousComponentPointAtT);
+														}
+														ComponentPointAtT = 
+															ComponentPointAtT * CurrentRatio + PreviousComponentPointAtT * ((FSolverReal)1. - CurrentRatio);
+													}
+													else if (Proxy->IsBoneSpace())
+													{
+														ComponentPointAtT = GlobalTransform.TransformPosition(ComponentPointAtT);
 													}
 
-													MParticles.X(Index) = GlobalTransform.TransformPosition(ComponentPointAtT);
+													if (!Proxy->IsBoneSpace())
+													{
+														MParticles.X(Index) = GlobalTransform.TransformPosition(ComponentPointAtT);
+													}
+													else
+													{
+														MParticles.X(Index) = ComponentPointAtT;
+													}
 
 													bParticleTouched = true;
 
@@ -1101,7 +1174,6 @@ namespace Chaos::Softs
 		TRACE_CPUPROFILER_EVENT_SCOPE(ChaosDeformableSolver_UpdateOutputState);
 		if (FFleshThreadingProxy* Proxy = InProxy.As<FFleshThreadingProxy>())
 		{
-
 			const FIntVector2& Range = Proxy->GetSolverParticleRange();
 			if (0 <= Range[0])
 			{
@@ -1110,9 +1182,25 @@ namespace Chaos::Softs
 				auto UEVertf = [](FVector3d V) { return FVector3f((float)V.X, (float)V.Y, (float)V.Z); };
 
 				TManagedArray<FVector3f>& Position = Proxy->GetDynamicCollection().ModifyAttribute<FVector3f>("Vertex", FGeometryCollection::VerticesGroup);
-				for (int32 vdx = 0; vdx < Position.Num(); vdx++)
+
+				// The final transform gets us from whatever the simulation space is,
+				// to component space.
+				const FTransform FinalXf = Proxy->GetFinalTransform();
+				if (!FinalXf.Equals(FTransform::Identity))
 				{
-					Position[vdx] = UEVertf(UEVertd(Evolution->Particles().X(vdx + Range[0])));
+					for (int32 vdx = 0; vdx < Position.Num(); vdx++)
+					{
+						const Chaos::FVec3f& Pos = Evolution->Particles().X(vdx + Range[0]);
+						FVector PosD = UEVertd(Pos);
+						Position[vdx] = UEVertf(FinalXf.TransformPosition(PosD));
+					}
+				}
+				else
+				{
+					for (int32 vdx = 0; vdx < Position.Num(); vdx++)
+					{
+						Position[vdx] = UEVertf(UEVertd(Evolution->Particles().X(vdx + Range[0])));
+					}
 				}
 			}
 		}
@@ -1177,4 +1265,4 @@ namespace Chaos::Softs
 
 	}
 
-}; // Namespace Chaos::Softs
\ No newline at end of file
+}; // Namespace Chaos::Softs
diff --git a/Engine/Source/Runtime/Experimental/Chaos/Public/Chaos/Deformable/ChaosDeformableSolver.h b/Engine/Source/Runtime/Experimental/Chaos/Public/Chaos/Deformable/ChaosDeformableSolver.h
--- a/Engine/Source/Runtime/Experimental/Chaos/Public/Chaos/Deformable/ChaosDeformableSolver.h
+++ b/Engine/Source/Runtime/Experimental/Chaos/Public/Chaos/Deformable/ChaosDeformableSolver.h
@@ -116,6 +116,7 @@ namespace Chaos::Softs
 		TUniquePtr<FDeformablePackage> PullInputPackage( );
 
 		/* Iteration Advance */
+		void InitializeSimulationSpace();
 		void InitializeSimulationObjects();
 		void InitializeSimulationObject(FThreadingProxy&);
 		void InitializeDeformableParticles(FFleshThreadingProxy&);
diff --git a/Engine/Source/Runtime/Experimental/Chaos/Public/Chaos/Deformable/ChaosDeformableSolverProxy.h b/Engine/Source/Runtime/Experimental/Chaos/Public/Chaos/Deformable/ChaosDeformableSolverProxy.h
--- a/Engine/Source/Runtime/Experimental/Chaos/Public/Chaos/Deformable/ChaosDeformableSolverProxy.h
+++ b/Engine/Source/Runtime/Experimental/Chaos/Public/Chaos/Deformable/ChaosDeformableSolverProxy.h
@@ -6,6 +6,17 @@
 #include "GeometryCollection/GeometryCollection.h"
 #include "GeometryCollection/ManagedArrayCollection.h"
 
+/**
+* Supported simulation spaces for the ChaosDeformable solver.
+*/
+UENUM()
+enum ChaosDeformableSimSpace
+{
+	World		UMETA(DisplayName = "World"),
+	Component	UMETA(DisplayName = "Component"),
+	Bone		UMETA(DisplayName = "Bone"),
+};
+
 namespace Chaos::Softs
 {
 	class CHAOS_API FThreadingProxy {
@@ -60,29 +71,155 @@ namespace Chaos::Softs
 	class CHAOS_API FFleshThreadingProxy : public FThreadingProxy
 	{
 		typedef FThreadingProxy Super;
-		FTransform InitialTransform = FTransform::Identity;
+
+		// Component transform
+		FTransform WorldToComponentXf = FTransform::Identity; // component xf
+		FTransform ComponentToWorldXf = FTransform::Identity; // inv component xf
+
+		// Bone transform
+		FTransform ComponentToBoneXf = FTransform::Identity; // bone xf
+		FTransform BoneToComponentXf = FTransform::Identity; // inv bone xf
+		FTransform PrevComponentToBoneXf = FTransform::Identity; // bone xf
+		FTransform PrevBoneToComponentXf = FTransform::Identity; // inv bone xf
+
 		const FManagedArrayCollection& Rest;
 		FManagedArrayCollection Dynamic;
 		FIntVector2 SolverParticleRange = FIntVector2(0, 0);
+		ChaosDeformableSimSpace SimSpace = ChaosDeformableSimSpace::World;
 
 	public:
 
-		FFleshThreadingProxy(UObject* InOwner,
-			const FTransform& InInitialTransform,
+		FFleshThreadingProxy(
+			UObject* InOwner,
+			const FTransform& InWorldToComponentXf, // ComponentXf
+			const FTransform& InComponentToBoneXf,
+			const ChaosDeformableSimSpace InSimSpace,
 			const FManagedArrayCollection& InRest,
 			const FManagedArrayCollection& InDynamic)
 			: Super(InOwner, TypeName())
-			, InitialTransform(InInitialTransform)
+			, WorldToComponentXf(InWorldToComponentXf)
+			, ComponentToBoneXf(InComponentToBoneXf)
 			, Rest(InRest)
 			, Dynamic(InDynamic)
-		{}
+			, SimSpace(InSimSpace)
+		{
+			ComponentToWorldXf = FTransform(WorldToComponentXf.ToMatrixWithScale().Inverse());
+			if (SimSpace == ChaosDeformableSimSpace::Bone)
+			{
+				BoneToComponentXf = FTransform(ComponentToBoneXf.ToMatrixWithScale().Inverse());
+			}
+		}
 		virtual ~FFleshThreadingProxy() {}
 
 		static FName TypeName() { return FName("Flesh"); }
 
-		const FTransform& GetInitialTransform() const
+		bool IsBoneSpace() const
+		{
+			return SimSpace == ChaosDeformableSimSpace::Bone;
+		}
+
+		//! Update the component and bone transforms with current data from the scene.
+		void UpdateSimSpace(const FTransform& InWorldToComponentXf, const FTransform& InComponentToBoneXf)
 		{
-			return InitialTransform;
+			WorldToComponentXf = InWorldToComponentXf;
+			ComponentToWorldXf = FTransform(WorldToComponentXf.ToMatrixWithScale().Inverse());
+			if (SimSpace == ChaosDeformableSimSpace::Bone)
+			{
+				PrevComponentToBoneXf = ComponentToBoneXf;
+				PrevBoneToComponentXf = BoneToComponentXf;
+
+				ComponentToBoneXf = InComponentToBoneXf;
+				BoneToComponentXf = FTransform(ComponentToBoneXf.ToMatrixWithScale().Inverse());
+			}
+		}
+
+		//! Return the transform required to get points from component space to whatever
+		//! the specified simulation space is.
+		const FTransform& GetInitialPointsTransform() const
+		{
+			switch (SimSpace)
+			{
+			case ChaosDeformableSimSpace::World:
+				// Initial points are in component space, so apply the component xf.
+				return WorldToComponentXf;
+			case ChaosDeformableSimSpace::Component:
+				// We're already in the right space.
+				return FTransform::Identity;
+			case ChaosDeformableSimSpace::Bone:
+				// Points have ComponentToBoneXf baked in, so we invert it out to get
+				// points in bone space.
+				return BoneToComponentXf;
+			default:
+				return FTransform::Identity;
+			}
+		}
+
+		const FTransform& GetCurrentPointsTransform() const
+		{
+			switch (SimSpace)
+			{
+			case ChaosDeformableSimSpace::World:
+				return WorldToComponentXf;
+			case ChaosDeformableSimSpace::Component:
+				return FTransform::Identity;
+			case ChaosDeformableSimSpace::Bone:
+				return BoneToComponentXf;
+			default:
+				return FTransform::Identity;
+			}
+		}
+
+		const FTransform& GetPreviousPointsTransform() const
+		{
+			switch (SimSpace)
+			{
+			case ChaosDeformableSimSpace::World:
+				return WorldToComponentXf;
+			case ChaosDeformableSimSpace::Component:
+				return FTransform::Identity;
+			case ChaosDeformableSimSpace::Bone:
+				return PrevBoneToComponentXf;
+			default:
+				return FTransform::Identity;
+			}
+		}
+
+		Chaos::TVector<Chaos::FRealSingle,3> RotateWorldSpaceVector(const Chaos::TVector<Chaos::FRealSingle, 3>& Dir) const
+		{
+			switch (SimSpace)
+			{
+			case ChaosDeformableSimSpace::World:
+				return Dir;
+			case ChaosDeformableSimSpace::Component:
+			{
+				return Chaos::TVector<Chaos::FRealSingle, 3>(
+					ComponentToWorldXf.TransformVectorNoScale(FVector(Dir[0], Dir[1], Dir[2])));
+			}
+			case ChaosDeformableSimSpace::Bone:
+			{
+				FVector Tmp = ComponentToWorldXf.TransformVectorNoScale(FVector(Dir[0], Dir[1], Dir[2]));
+				return Chaos::TVector<Chaos::FRealSingle, 3>(ComponentToBoneXf.TransformVectorNoScale(Tmp));
+			}
+			default:
+				return Dir;
+			}
+		}
+
+		//! Return the transform required to get from whatever our simulation space is, to
+		//! component space.
+		FTransform GetFinalTransform() const
+		{
+			switch (SimSpace)
+			{
+			case ChaosDeformableSimSpace::World:
+				return ComponentToWorldXf;
+			case ChaosDeformableSimSpace::Component:
+				return FTransform::Identity;
+			case ChaosDeformableSimSpace::Bone:
+				return ComponentToBoneXf;
+			default:
+				return FTransform::Identity;
+			}
 		}
 
 		void SetSolverParticleRange(int32 InStart, int32 InRange)
@@ -96,16 +233,9 @@ namespace Chaos::Softs
 			return SolverParticleRange;
 		}
 
-		FManagedArrayCollection& GetDynamicCollection() {
-			return Dynamic;
-		}
-		const FManagedArrayCollection& GetDynamicCollection() const {
-			return Dynamic;
-		}
-
-		const FManagedArrayCollection& GetRestCollection() const {
-			return Rest;
-		}
+		FManagedArrayCollection& GetDynamicCollection()				{ return Dynamic; }
+		const FManagedArrayCollection& GetDynamicCollection() const { return Dynamic; }
+		const FManagedArrayCollection& GetRestCollection() const	{ return Rest; }
 
 		class FFleshInputBuffer : public FThreadingProxy::FBuffer
 		{
@@ -114,36 +244,61 @@ namespace Chaos::Softs
 		public:
 			typedef FFleshThreadingProxy Source;
 
-			FFleshInputBuffer(const FTransform& InGlobalTransform, const bool InbEnableGravity, const float InStiffnessMultiplier, const float InDampingMultiplier, const float InMassMultiplier, const UObject* InOwner = nullptr)
+			FFleshInputBuffer(
+				const FTransform& InWorldToComponentXf,
+				const FTransform& InComponentToBoneXf,
+				const int32 InSimSpaceBoneIndex,
+				const bool InbEnableGravity,
+				const float InStiffnessMultiplier, 
+				const float InDampingMultiplier, 
+				const float InMassMultiplier, 
+				const UObject* InOwner = nullptr)
 				: Super(InOwner, FFleshThreadingProxy::TypeName())
-				, GlobalTransform(InGlobalTransform)
+				, WorldToComponentXf(InWorldToComponentXf)
+				, ComponentToBoneXf(InComponentToBoneXf)
 				, Transforms(TArray<FTransform>())
 				, RestTransforms(TArray<FTransform>())
 				, bEnableGravity(InbEnableGravity)
 				, StiffnessMultiplier(InStiffnessMultiplier)
 				, DampingMultiplier(InDampingMultiplier)
 				, MassMultiplier(InMassMultiplier)
+				, SimSpaceBoneIndex(InSimSpaceBoneIndex)
 			{}
 
-			FFleshInputBuffer(const FTransform& InGlobalTransform, const TArray<FTransform>& InTransforms, const TArray<FTransform>& InRestTransforms, const bool InbEnableGravity, const float InStiffnessMultiplier, const float InDampingMultiplier, const float InMassMultiplier, const UObject* InOwner = nullptr)
+			FFleshInputBuffer(
+				const FTransform& InWorldToComponentXf,
+				const FTransform& InComponentToBoneXf, 
+				const int32 InSimSpaceBoneIndex,
+				const TArray<FTransform>& InTransforms, 
+				const TArray<FTransform>& InRestTransforms, 
+				const bool InbEnableGravity, 
+				const float InStiffnessMultiplier, 
+				const float InDampingMultiplier, 
+				const float InMassMultiplier, 
+				const UObject* InOwner = nullptr)
 				: Super(InOwner, FFleshThreadingProxy::TypeName())
-				, GlobalTransform(InGlobalTransform)
+				, WorldToComponentXf(InWorldToComponentXf)
+				, ComponentToBoneXf(InComponentToBoneXf)
 				, Transforms(InTransforms)
 				, RestTransforms(InRestTransforms)
 				, bEnableGravity(InbEnableGravity)
 				, StiffnessMultiplier(InStiffnessMultiplier)
 				, DampingMultiplier(InDampingMultiplier)
 				, MassMultiplier(InMassMultiplier)
+				, SimSpaceBoneIndex(InSimSpaceBoneIndex)
 			{}
 			virtual ~FFleshInputBuffer() {}
 
-			FTransform GlobalTransform;
+			FTransform WorldToComponentXf;
+			FTransform ComponentToBoneXf;
+
 			TArray<FTransform> Transforms;
 			TArray<FTransform> RestTransforms; 
 			bool bEnableGravity;
 			float StiffnessMultiplier;
 			float DampingMultiplier;
 			float MassMultiplier;
+			int32 SimSpaceBoneIndex = INDEX_NONE;
 		};
 
 		class FFleshOutputBuffer : public FThreadingProxy::FBuffer
@@ -164,6 +319,5 @@ namespace Chaos::Softs
 			FManagedArrayCollection Dynamic;
 		};
 
-
 	};
-}// namespace Chaos::Softs
\ No newline at end of file
+}// namespace Chaos::Softs
