Rename enum member "Component" to "ComponentXf" to fix Mac build issue, and specify the underlying type of the enum.

diff --git a/Engine/Source/Runtime/Experimental/Chaos/Public/Chaos/Deformable/ChaosDeformableSolverProxy.h b/Engine/Source/Runtime/Experimental/Chaos/Public/Chaos/Deformable/ChaosDeformableSolverProxy.h
--- a/Engine/Source/Runtime/Experimental/Chaos/Public/Chaos/Deformable/ChaosDeformableSolverProxy.h
+++ b/Engine/Source/Runtime/Experimental/Chaos/Public/Chaos/Deformable/ChaosDeformableSolverProxy.h
@@ -12,10 +12,10 @@
 * Supported simulation spaces for the ChaosDeformable solver.
 */
 UENUM()
-enum ChaosDeformableSimSpace
+enum ChaosDeformableSimSpace : uint8
 {
 	World		UMETA(DisplayName = "World"),
-	Component	UMETA(DisplayName = "Component"),
+	ComponentXf	UMETA(DisplayName = "Component"), // Component fails on Mac
 	Bone		UMETA(DisplayName = "Bone"),
 };
 
@@ -144,7 +144,7 @@ namespace Chaos::Softs
 			case ChaosDeformableSimSpace::World:
 				// Initial points are in component space, so apply the component xf.
 				return WorldToComponentXf;
-			case ChaosDeformableSimSpace::Component:
+			case ChaosDeformableSimSpace::ComponentXf:
 				// We're already in the right space.
 				return FTransform::Identity;
 			case ChaosDeformableSimSpace::Bone:
@@ -162,7 +162,7 @@ namespace Chaos::Softs
 			{
 			case ChaosDeformableSimSpace::World:
 				return WorldToComponentXf;
-			case ChaosDeformableSimSpace::Component:
+			case ChaosDeformableSimSpace::ComponentXf:
 				return FTransform::Identity;
 			case ChaosDeformableSimSpace::Bone:
 				return BoneToComponentXf;
@@ -177,7 +177,7 @@ namespace Chaos::Softs
 			{
 			case ChaosDeformableSimSpace::World:
 				return WorldToComponentXf;
-			case ChaosDeformableSimSpace::Component:
+			case ChaosDeformableSimSpace::ComponentXf:
 				return FTransform::Identity;
 			case ChaosDeformableSimSpace::Bone:
 				return PrevBoneToComponentXf;
@@ -192,7 +192,7 @@ namespace Chaos::Softs
 			{
 			case ChaosDeformableSimSpace::World:
 				return Dir;
-			case ChaosDeformableSimSpace::Component:
+			case ChaosDeformableSimSpace::ComponentXf:
 			{
 				return Chaos::TVector<Chaos::FRealSingle, 3>(
 					ComponentToWorldXf.TransformVectorNoScale(FVector(Dir[0], Dir[1], Dir[2])));
@@ -215,7 +215,7 @@ namespace Chaos::Softs
 			{
 			case ChaosDeformableSimSpace::World:
 				return ComponentToWorldXf;
-			case ChaosDeformableSimSpace::Component:
+			case ChaosDeformableSimSpace::ComponentXf:
 				return FTransform::Identity;
 			case ChaosDeformableSimSpace::Bone:
 				return ComponentToBoneXf;
