add support for custom world transforms in the axis overlap computation

diff --git a/Engine/Source/Runtime/Experimental/Chaos/Private/Chaos/PhysicsObjectInterface.cpp b/Engine/Source/Runtime/Experimental/Chaos/Private/Chaos/PhysicsObjectInterface.cpp
--- a/Engine/Source/Runtime/Experimental/Chaos/Private/Chaos/PhysicsObjectInterface.cpp
+++ b/Engine/Source/Runtime/Experimental/Chaos/Private/Chaos/PhysicsObjectInterface.cpp
@@ -201,14 +201,20 @@ namespace Chaos
 	template<EThreadContext Id>
 	bool FReadPhysicsObjectInterface<Id>::GetPhysicsObjectOverlap(FPhysicsObjectHandle ObjectA, FPhysicsObjectHandle ObjectB, bool bTraceComplex, Chaos::FOverlapInfo& OutOverlap)
 	{
-		TRACE_CPUPROFILER_EVENT_SCOPE(FReadPhysicsObjectInterface<Id>::GetPhysicsObjectOverlap);
+		return GetPhysicsObjectOverlapWithTransform(ObjectA, FTransform::Identity, ObjectB, FTransform::Identity, bTraceComplex, OutOverlap);
+	}
+
+	template<EThreadContext Id>
+	bool FReadPhysicsObjectInterface<Id>::GetPhysicsObjectOverlapWithTransform(FPhysicsObjectHandle ObjectA, const FTransform& InTransformA, FPhysicsObjectHandle ObjectB, const FTransform& InTransformB, bool bTraceComplex, Chaos::FOverlapInfo& OutOverlap)
+	{
+		TRACE_CPUPROFILER_EVENT_SCOPE(FReadPhysicsObjectInterface<Id>::GetPhysicsObjectOverlapWithTransform);
 		TArray<FPerShapeData*> ShapesA = GetAllShapes({ &ObjectA, 1 });
-		const FTransform TransformA(GetR(ObjectA), GetX(ObjectA));
-		const FBox BoxA = GetWorldBounds({&ObjectA, 1});
+		const FTransform TransformA = InTransformA * FTransform{ GetR(ObjectA), GetX(ObjectA) };
+		const FBox BoxA = GetWorldBounds({ &ObjectA, 1 }).TransformBy(InTransformA);
 
 		TArray<FPerShapeData*> ShapesB = GetAllShapes({ &ObjectB, 1 });
-		const FTransform TransformB(GetR(ObjectB), GetX(ObjectB));
-		const FBox BoxB = GetWorldBounds({ &ObjectB, 1 });
+		const FTransform TransformB = InTransformB * FTransform{ GetR(ObjectB), GetX(ObjectB) };
+		const FBox BoxB = GetWorldBounds({ &ObjectB, 1 }).TransformBy(InTransformB);
 
 		if (!BoxA.Intersect(BoxB))
 		{
@@ -243,7 +249,7 @@ namespace Chaos
 			const FAABB3 BoxShapeB = GeomB->CalculateTransformedBounds(TransformB);
 			// At this point on, this function should be mirror the Overlap_GeomInternal function in PhysInterface_Chaos.cpp.
 			// ShapeA is equivalent to InInstance and GeomB is equivalent to InGeom.
-			
+
 			for (FPerShapeData* A : ShapesA)
 			{
 				if (!A)
diff --git a/Engine/Source/Runtime/Experimental/Chaos/Public/Chaos/PhysicsObjectInterface.h b/Engine/Source/Runtime/Experimental/Chaos/Public/Chaos/PhysicsObjectInterface.h
--- a/Engine/Source/Runtime/Experimental/Chaos/Public/Chaos/PhysicsObjectInterface.h
+++ b/Engine/Source/Runtime/Experimental/Chaos/Public/Chaos/PhysicsObjectInterface.h
@@ -66,6 +66,7 @@ namespace Chaos
 
 		TArray<FPerShapeData*> GetAllShapes(TArrayView<FPhysicsObjectHandle> InObjects);
 		bool GetPhysicsObjectOverlap(FPhysicsObjectHandle ObjectA, FPhysicsObjectHandle ObjectB, bool bTraceComplex, Chaos::FOverlapInfo& OutOverlap);
+		bool GetPhysicsObjectOverlapWithTransform(FPhysicsObjectHandle ObjectA, const FTransform& InTransformA, FPhysicsObjectHandle ObjectB, const FTransform& InTransformB, bool bTraceComplex, Chaos::FOverlapInfo& OutOverlap);
 
 		bool AreAllValid(TArrayView<FPhysicsObjectHandle> InObjects);
 		bool AreAllKinematic(TArrayView<FPhysicsObjectHandle> InObjects);
