Geometry collection : add support for root proxy and leaf meshes to using ISM or HISM
- added cvar to control this
- updated ISMPool to take a pereferHISM option when registering a mesh

diff --git a/Engine/Source/Runtime/Experimental/GeometryCollectionEngine/Private/GeometryCollection/GeometryCollectionComponent.cpp b/Engine/Source/Runtime/Experimental/GeometryCollectionEngine/Private/GeometryCollection/GeometryCollectionComponent.cpp
--- a/Engine/Source/Runtime/Experimental/GeometryCollectionEngine/Private/GeometryCollection/GeometryCollectionComponent.cpp
+++ b/Engine/Source/Runtime/Experimental/GeometryCollectionEngine/Private/GeometryCollection/GeometryCollectionComponent.cpp
@@ -114,6 +114,12 @@ FAutoConsoleVariableRef CVarChaosGCUseISMPoolForNonFracturedParts(TEXT("p.Chaos.
 bool bChaos_GC_ForceAutoAssignISMPool = false;
 FAutoConsoleVariableRef CVarChaosGCForceAutoAssignISMPool(TEXT("p.Chaos.GC.ForceAutoAssignISMPool"), bChaos_GC_ForceAutoAssignISMPool, TEXT("When enabled, force assignement of ISMPool regardgless of the settings of the components"));
 
+bool bChaos_GC_UseHierarchicalISMForProxyMesh = true;
+FAutoConsoleVariableRef CVarChaosGCUseHierarchicalISMForProxyMesh(TEXT("p.Chaos.GC.UseHierarchicalISMForProxyMesh"), bChaos_GC_UseHierarchicalISMForProxyMesh, TEXT("When enabled along with ISM Pool, proxy mesh will prefer using HISM vs standard ISM"));
+
+bool bChaos_GC_UseHierarchicalISMForLeafMeshes = true;
+FAutoConsoleVariableRef CVarChaosGCUseHierarchicalISMForLeafMeshes(TEXT("p.Chaos.GC.UseHierarchicalISMForLeafMeshes"), bChaos_GC_UseHierarchicalISMForLeafMeshes, TEXT("When enabled along with ISM Pool, leaf meshes will prefer using HISM vs standard ISM"));
+
 bool bChaos_GC_InitConstantDataUseParallelFor = true;
 FAutoConsoleVariableRef CVarChaosGCInitConstantDataUseParallelFor(TEXT("p.Chaos.GC.InitConstantDataUseParallelFor"), bChaos_GC_InitConstantDataUseParallelFor, TEXT("When enabled, InitConstant data will use parallelFor for copying some of the data"));
 
@@ -4487,7 +4493,7 @@ void UGeometryCollectionComponent::RegisterToISMPool()
 								{
 									StaticMeshInstance.MaterialsOverrides = AutoInstanceMesh.Materials;
 								}
-								ISMPoolComp->AddMeshToGroup(ISMPoolMeshGroupIndex, StaticMeshInstance, InstanceCounts[MeshIndex]);
+								ISMPoolComp->AddMeshToGroup(ISMPoolMeshGroupIndex, StaticMeshInstance, InstanceCounts[MeshIndex], bChaos_GC_UseHierarchicalISMForLeafMeshes);
 							}
 						}
 					}
@@ -4504,7 +4510,7 @@ void UGeometryCollectionComponent::RegisterToISMPool()
 
 						FGeometryCollectionStaticMeshInstance StaticMeshInstance;
 						StaticMeshInstance.StaticMesh = RootMeshProxy;
-						ISMPoolRootProxyMeshId = ISMPoolComp->AddMeshToGroup(ISMPoolMeshGroupIndex, StaticMeshInstance, 1);
+						ISMPoolRootProxyMeshId = ISMPoolComp->AddMeshToGroup(ISMPoolMeshGroupIndex, StaticMeshInstance, 1, bChaos_GC_UseHierarchicalISMForProxyMesh);
 					}
 				}
 			}
diff --git a/Engine/Source/Runtime/Experimental/GeometryCollectionEngine/Private/GeometryCollection/GeometryCollectionISMPoolComponent.cpp b/Engine/Source/Runtime/Experimental/GeometryCollectionEngine/Private/GeometryCollection/GeometryCollectionISMPoolComponent.cpp
--- a/Engine/Source/Runtime/Experimental/GeometryCollectionEngine/Private/GeometryCollection/GeometryCollectionISMPoolComponent.cpp
+++ b/Engine/Source/Runtime/Experimental/GeometryCollectionEngine/Private/GeometryCollection/GeometryCollectionISMPoolComponent.cpp
@@ -43,13 +43,17 @@ void FGeometryCollectionMeshGroup::RemoveAllMeshes(FGeometryCollectionISMPool& I
 	Meshes.Empty();
 }
 
-FGeometryCollectionISM::FGeometryCollectionISM(AActor* OwmingActor, const FGeometryCollectionStaticMeshInstance& MeshInstance)
+FGeometryCollectionISM::FGeometryCollectionISM(AActor* OwmingActor, const FGeometryCollectionStaticMeshInstance& MeshInstance, bool bPreferHISM)
 {
 	check(MeshInstance.StaticMesh);
 	check(OwmingActor);
 
-	const FName ISMName = MakeUniqueObjectName(OwmingActor, UHierarchicalInstancedStaticMeshComponent::StaticClass(), MeshInstance.StaticMesh->GetFName());
-	if (UInstancedStaticMeshComponent* ISMC = NewObject<UHierarchicalInstancedStaticMeshComponent>(OwmingActor, ISMName, RF_Transient | RF_DuplicateTransient))
+	const FName ISMName = MakeUniqueObjectName(OwmingActor, bPreferHISM? UHierarchicalInstancedStaticMeshComponent::StaticClass(): UInstancedStaticMeshComponent::StaticClass(), MeshInstance.StaticMesh->GetFName());
+	UInstancedStaticMeshComponent* ISMC = bPreferHISM
+		? NewObject<UHierarchicalInstancedStaticMeshComponent>(OwmingActor, ISMName, RF_Transient | RF_DuplicateTransient)
+		: NewObject<UInstancedStaticMeshComponent>(OwmingActor, ISMName, RF_Transient | RF_DuplicateTransient)
+		;
+	if (ISMC)
 	{
 		ISMC->SetStaticMesh(MeshInstance.StaticMesh);
 		// material overrides
@@ -84,14 +88,14 @@ int32 FGeometryCollectionISM::AddInstanceGroup(int32 InstanceCount)
 	return InstanceGroupIndex;
 }
 
-FGeometryCollectionMeshInfo FGeometryCollectionISMPool::AddISM(UGeometryCollectionISMPoolComponent* OwningComponent, const FGeometryCollectionStaticMeshInstance& MeshInstance, int32 InstanceCount)
+FGeometryCollectionMeshInfo FGeometryCollectionISMPool::AddISM(UGeometryCollectionISMPoolComponent* OwningComponent, const FGeometryCollectionStaticMeshInstance& MeshInstance, int32 InstanceCount, bool bPreferHISM)
 {
 	FGeometryCollectionMeshInfo Info;
 
 	FISMIndex* ISMIndex = MeshToISMIndex.Find(MeshInstance);
 	if (!ISMIndex)
 	{
-		Info.ISMIndex = ISMs.Emplace(OwningComponent->GetOwner(), MeshInstance);
+		Info.ISMIndex = ISMs.Emplace(OwningComponent->GetOwner(), MeshInstance, bPreferHISM);
 		MeshToISMIndex.Add(MeshInstance, Info.ISMIndex);
 	}
 	else
@@ -177,11 +181,11 @@ void UGeometryCollectionISMPoolComponent::DestroyMeshGroup(FMeshGroupId MeshGrou
 	}
 }
 
-UGeometryCollectionISMPoolComponent::FMeshId UGeometryCollectionISMPoolComponent::AddMeshToGroup(FMeshGroupId MeshGroupId, const FGeometryCollectionStaticMeshInstance& MeshInstance, int32 InstanceCount)
+UGeometryCollectionISMPoolComponent::FMeshId UGeometryCollectionISMPoolComponent::AddMeshToGroup(FMeshGroupId MeshGroupId, const FGeometryCollectionStaticMeshInstance& MeshInstance, int32 InstanceCount, bool bPreferHISM)
 {
 	if (FGeometryCollectionMeshGroup* MeshGroup = MeshGroups.Find(MeshGroupId))
 	{
-		const FGeometryCollectionMeshInfo ISMInstanceInfo = Pool.AddISM(this, MeshInstance, InstanceCount);
+		const FGeometryCollectionMeshInfo ISMInstanceInfo = Pool.AddISM(this, MeshInstance, InstanceCount, bPreferHISM);
 		return MeshGroup->AddMesh(MeshInstance, InstanceCount, ISMInstanceInfo);
 	}
 	UE_LOG(LogChaos, Warning, TEXT("UGeometryCollectionISMPoolComponent : Trying to add a mesh to a mesh group (%d) that does not exists"), MeshGroupId);
diff --git a/Engine/Source/Runtime/Experimental/GeometryCollectionEngine/Public/GeometryCollection/GeometryCollectionISMPoolComponent.h b/Engine/Source/Runtime/Experimental/GeometryCollectionEngine/Public/GeometryCollection/GeometryCollectionISMPoolComponent.h
--- a/Engine/Source/Runtime/Experimental/GeometryCollectionEngine/Public/GeometryCollection/GeometryCollectionISMPoolComponent.h
+++ b/Engine/Source/Runtime/Experimental/GeometryCollectionEngine/Public/GeometryCollection/GeometryCollectionISMPoolComponent.h
@@ -165,7 +165,7 @@ struct FGeometryCollectionMeshGroup
 
 struct FGeometryCollectionISM
 {
-	FGeometryCollectionISM(AActor* OwmingActor, const FGeometryCollectionStaticMeshInstance& MeshInstance);
+	FGeometryCollectionISM(AActor* OwmingActor, const FGeometryCollectionStaticMeshInstance& MeshInstance, bool bPreferHISM);
 
 	int32 AddInstanceGroup(int32 InstanceCount);
 
@@ -178,7 +178,7 @@ struct FGeometryCollectionISMPool
 {
 	using FISMIndex = int32;
 
-	FGeometryCollectionMeshInfo AddISM(UGeometryCollectionISMPoolComponent* OwningComponent, const FGeometryCollectionStaticMeshInstance& MeshInstance, int32 InstanceCount);
+	FGeometryCollectionMeshInfo AddISM(UGeometryCollectionISMPoolComponent* OwningComponent, const FGeometryCollectionStaticMeshInstance& MeshInstance, int32 InstanceCount, bool bPreferHISM);
 	bool BatchUpdateInstancesTransforms(FGeometryCollectionMeshInfo& MeshInfo, int32 StartInstanceIndex, const TArray<FTransform>& NewInstancesTransforms, bool bWorldSpace, bool bMarkRenderStateDirty, bool bTeleport);
 	void RemoveISM(const FGeometryCollectionMeshInfo& MeshInfo);
 
@@ -215,7 +215,7 @@ public:
 	void DestroyMeshGroup(FMeshGroupId MeshGroupId);
 
 	/** Add a static mesh for a nmesh group */
-	FMeshId AddMeshToGroup(FMeshGroupId MeshGroupId, const FGeometryCollectionStaticMeshInstance& MeshInstance, int32 InstanceCount);
+	FMeshId AddMeshToGroup(FMeshGroupId MeshGroupId, const FGeometryCollectionStaticMeshInstance& MeshInstance, int32 InstanceCount, bool bPreferHISM = true);
 
 	/** Add a static mesh for a nmesh group */
 	bool BatchUpdateInstancesTransforms(FMeshGroupId MeshGroupId, FMeshId MeshId, int32 StartInstanceIndex, const TArray<FTransform>& NewInstancesTransforms, bool bWorldSpace = false, bool bMarkRenderStateDirty = false, bool bTeleport = false);
