ChaosFlesh : Ability to deactivate collisions on the manager

diff --git a/Engine/Plugins/Experimental/ChaosFlesh/Source/ChaosFleshEngine/Private/ChaosFlesh/ChaosDeformableCollisionsComponent.cpp b/Engine/Plugins/Experimental/ChaosFlesh/Source/ChaosFleshEngine/Private/ChaosFlesh/ChaosDeformableCollisionsComponent.cpp
--- a/Engine/Plugins/Experimental/ChaosFlesh/Source/ChaosFleshEngine/Private/ChaosFlesh/ChaosDeformableCollisionsComponent.cpp
+++ b/Engine/Plugins/Experimental/ChaosFlesh/Source/ChaosFleshEngine/Private/ChaosFlesh/ChaosDeformableCollisionsComponent.cpp
@@ -125,7 +125,13 @@ UDeformableCollisionsComponent::NewDeformableData()
 	}
 
 
-	for (auto& RemovedBody : RemovedBodies) { RemovedBodiesData.Add({RemovedBody}); }
+	for (auto& RemovedBody : RemovedBodies) 
+	{
+		if (RemovedBody)
+		{
+			RemovedBodiesData.Add({ RemovedBody });
+		}
+	}
 	AddedBodies.Empty();
 	RemovedBodies.Empty();
 
diff --git a/Engine/Source/Runtime/Experimental/Chaos/Private/Chaos/Deformable/ChaosDeformableSolver.cpp b/Engine/Source/Runtime/Experimental/Chaos/Private/Chaos/Deformable/ChaosDeformableSolver.cpp
--- a/Engine/Source/Runtime/Experimental/Chaos/Private/Chaos/Deformable/ChaosDeformableSolver.cpp
+++ b/Engine/Source/Runtime/Experimental/Chaos/Private/Chaos/Deformable/ChaosDeformableSolver.cpp
@@ -231,10 +231,13 @@ namespace Chaos::Softs
 		typedef TPair< FThreadingProxy::FKey, TUniquePtr<FThreadingProxy> > FType;
 		for (const FType& Entry : Proxies)
 		{
-			FThreadingProxy& Proxy = *Entry.Value.Get();
-			if (FCollisionManagerProxy* CollisionManagerProxy = Proxy.As< FCollisionManagerProxy>())
+			if (Entry.Value)
 			{
-				UpdateCollisionBodies(Entry.Key, DeltaTime);
+				FThreadingProxy& Proxy = *Entry.Value.Get();
+				if (FCollisionManagerProxy* CollisionManagerProxy = Proxy.As< FCollisionManagerProxy>())
+				{
+					UpdateCollisionBodies(*CollisionManagerProxy, Entry.Key, DeltaTime);
+				}
 			}
 		}
 	}
@@ -482,27 +485,42 @@ namespace Chaos::Softs
 		TRACE_CPUPROFILER_EVENT_SCOPE(ChaosDeformableSolver_InitializeCollisionBodies);
 	}
 
-	void FDeformableSolver::UpdateCollisionBodies(FThreadingProxy::FKey Owner, FSolverReal DeltaTime)
+	void FDeformableSolver::UpdateCollisionBodies(FCollisionManagerProxy& Proxy, FThreadingProxy::FKey Owner, FSolverReal DeltaTime)
 	{
 		SCOPE_CYCLE_COUNTER(STAT_ChaosDeformableSolver_InitializeCollisionBodies);
 		TRACE_CPUPROFILER_EVENT_SCOPE(ChaosDeformableSolver_InitializeCollisionBodies);
 
 		FCollisionManagerProxy::FCollisionsInputBuffer* CollisionsInputBuffer = nullptr;
-		if (this->CurrentInputPackage->ObjectMap.Contains(Owner))
+		if (this->CurrentInputPackage && this->CurrentInputPackage->ObjectMap.Contains(Owner))
 		{
-			CollisionsInputBuffer = this->CurrentInputPackage->ObjectMap[Owner]->As<FCollisionManagerProxy::FCollisionsInputBuffer>();
-			if (CollisionsInputBuffer)
+			if (this->CurrentInputPackage->ObjectMap[Owner] != nullptr)
 			{
-				for (auto& AddBody : CollisionsInputBuffer->Added)
+				CollisionsInputBuffer = this->CurrentInputPackage->ObjectMap[Owner]->As<FCollisionManagerProxy::FCollisionsInputBuffer>();
+				if (CollisionsInputBuffer)
 				{
-					if (AddBody.Shapes)
+					for (auto& AddBody : CollisionsInputBuffer->Added)
+					{
+						if (AddBody.Shapes)
+						{
+							typedef TPlane<Chaos::Softs::FSolverReal, 3> GeomType;
+							int32 Index = Evolution->AddCollisionParticle(INDEX_NONE, true);
+							int32 ViewIndex = Evolution->CollisionParticlesActiveView().GetRanges().Num()-1;
+							Evolution->CollisionParticles().X(Index) = AddBody.Transform.GetTranslation();
+							Evolution->CollisionParticles().R(Index) = AddBody.Transform.GetRotation();
+							TUniquePtr<FImplicitObject> UniquePtr(AddBody.Shapes); AddBody.Shapes = nullptr;
+							Evolution->CollisionParticles().SetDynamicGeometry(Index, MoveTemp(UniquePtr));
+							Proxy.CollisionBodies.Add(AddBody.BodyId, FCollisionObjectParticleHandel( Index,ViewIndex,AddBody.Transform ));
+						}
+					}
+					for (auto& RemovedBody : CollisionsInputBuffer->Removed)
 					{
-						typedef TPlane<Chaos::Softs::FSolverReal, 3> GeomType;
-						int32 Index = Evolution->AddCollisionParticleRange(1, INDEX_NONE, true);
-						Evolution->CollisionParticles().X(Index) = AddBody.Transform.GetTranslation();
-						Evolution->CollisionParticles().R(Index) = AddBody.Transform.GetRotation();
-						TUniquePtr<FImplicitObject> UniquePtr(AddBody.Shapes); AddBody.Shapes = nullptr;
-						Evolution->CollisionParticles().SetDynamicGeometry(Index, MoveTemp(UniquePtr));
+						if (Proxy.CollisionBodies.Contains(RemovedBody.BodyId))
+						{
+							int32 ParticleIndex = Proxy.CollisionBodies[RemovedBody.BodyId].ParticleIndex;
+							int32 ViewIndex = Proxy.CollisionBodies[RemovedBody.BodyId].ActiveViewIndex;
+							Evolution->RemoveCollisionParticle(ParticleIndex, ViewIndex);
+							Proxy.CollisionBodies.Remove(RemovedBody.BodyId);
+						}
 					}
 				}
 			}
diff --git a/Engine/Source/Runtime/Experimental/Chaos/Private/Chaos/PBDEvolution.cpp b/Engine/Source/Runtime/Experimental/Chaos/Private/Chaos/PBDEvolution.cpp
--- a/Engine/Source/Runtime/Experimental/Chaos/Private/Chaos/PBDEvolution.cpp
+++ b/Engine/Source/Runtime/Experimental/Chaos/Private/Chaos/PBDEvolution.cpp
@@ -186,6 +186,54 @@ int32 FPBDEvolution::AddCollisionParticleRange(int32 NumParticles, uint32 GroupI
 	return INDEX_NONE;
 }
 
+int32 FPBDEvolution::AddCollisionParticle(uint32 GroupId, bool bActivate)
+{
+	FParticleVievToken Token = { INDEX_NONE,INDEX_NONE };
+	if (!RemovedCollisionIndices.Num())
+	{
+		Token = { (int32)MCollisionParticles.Size(),INDEX_NONE };
+
+		MCollisionParticles.AddParticles(1);
+	}
+	else
+	{
+		Token = RemovedCollisionIndices.Pop(); 
+	}
+
+	if (0 <= Token.ParticleIndex && Token.ParticleIndex < (int32)MCollisionParticles.Size())
+	{
+		MCollisionParticleGroupIds[Token.ParticleIndex] = GroupId;
+
+		if (Token.ViewIndex == INDEX_NONE)
+		{
+			// Add range
+			MCollisionParticlesActiveView.AddRange(1, bActivate);
+		}
+		else
+		{
+			// Update Range
+			MCollisionParticlesActiveView.ActivateRange(Token.ViewIndex, bActivate);
+		}
+	}
+
+	return Token.ParticleIndex;
+}
+
+void FPBDEvolution::RemoveCollisionParticle(int32 CollisionParticleIndex, int32 CollisionParticleViewIndex)
+{
+	if (0 <= CollisionParticleIndex && CollisionParticleIndex < (int32)MCollisionParticles.Size())
+	{
+		CollisionParticleGroupIds()[CollisionParticleIndex] = CollisionParticleIndex;
+		MCollisionParticlesActiveView.ActivateRange(CollisionParticleViewIndex, false);
+		RemovedCollisionIndices.Add({ CollisionParticleIndex ,CollisionParticleViewIndex });
+	}
+	else
+	{
+		ensure(false);
+	}
+}
+
+
 int32 FPBDEvolution::AddConstraintInitRange(int32 NumConstraints, bool bActivate)
 {
 	// Add new constraint init functions
diff --git a/Engine/Source/Runtime/Experimental/Chaos/Public/Chaos/Deformable/ChaosDeformableCollisionsProxy.h b/Engine/Source/Runtime/Experimental/Chaos/Public/Chaos/Deformable/ChaosDeformableCollisionsProxy.h
--- a/Engine/Source/Runtime/Experimental/Chaos/Public/Chaos/Deformable/ChaosDeformableCollisionsProxy.h
+++ b/Engine/Source/Runtime/Experimental/Chaos/Public/Chaos/Deformable/ChaosDeformableCollisionsProxy.h
@@ -36,6 +36,20 @@ namespace Chaos::Softs
 		FTransform Transform = FTransform::Identity;
 	};
 
+	struct CHAOS_API FCollisionObjectParticleHandel
+	{
+		FCollisionObjectParticleHandel(int32 InParticleIndex = INDEX_NONE,
+									   int32 InActiveViewIndex = INDEX_NONE,
+									   FTransform InTransform = FTransform::Identity)
+			: ParticleIndex(InParticleIndex)
+			, ActiveViewIndex(InActiveViewIndex)
+			, Transform(InTransform) {}
+
+		int32 ParticleIndex = INDEX_NONE;
+		int32 ActiveViewIndex = INDEX_NONE;
+		FTransform Transform = FTransform::Identity;
+	};
+
 
 	class CHAOS_API FCollisionManagerProxy : public FThreadingProxy
 	{
@@ -73,7 +87,7 @@ namespace Chaos::Softs
 
 		TArray<FCollisionObjectAddedBodies> CollisionObjectsToAdd;
 		TArray< FCollisionObjectRemovedBodies > CollisionObjectsToRemove;
-		TMap< UObject*, FTransform > CollisionBodies;
+		TMap<const UObject*, FCollisionObjectParticleHandel > CollisionBodies;
 	};
 
 }// namespace Chaos::Softs
\ No newline at end of file
diff --git a/Engine/Source/Runtime/Experimental/Chaos/Public/Chaos/Deformable/ChaosDeformableSolver.h b/Engine/Source/Runtime/Experimental/Chaos/Public/Chaos/Deformable/ChaosDeformableSolver.h
--- a/Engine/Source/Runtime/Experimental/Chaos/Public/Chaos/Deformable/ChaosDeformableSolver.h
+++ b/Engine/Source/Runtime/Experimental/Chaos/Public/Chaos/Deformable/ChaosDeformableSolver.h
@@ -128,7 +128,7 @@ namespace Chaos::Softs
 		void InitializeCollisionBodies(FCollisionManagerProxy&);
 		void InitializeSelfCollisionVariables();
 		void InitializeGridBasedConstraintVariables();
-		void UpdateCollisionBodies(FThreadingProxy::FKey, FSolverReal DeltaTime);
+		void UpdateCollisionBodies(FCollisionManagerProxy&, FThreadingProxy::FKey, FSolverReal DeltaTime);
 		void RemoveSimulationObjects();
 		TArray<Chaos::TVec3<FSolverReal>> ComputeParticleTargets(const TArray<TArray<int32>>& ParticleIndices);
 
diff --git a/Engine/Source/Runtime/Experimental/Chaos/Public/Chaos/PBDEvolution.h b/Engine/Source/Runtime/Experimental/Chaos/Public/Chaos/PBDEvolution.h
--- a/Engine/Source/Runtime/Experimental/Chaos/Public/Chaos/PBDEvolution.h
+++ b/Engine/Source/Runtime/Experimental/Chaos/Public/Chaos/PBDEvolution.h
@@ -69,6 +69,12 @@ class CHAOS_API FPBDEvolution : public TArrayCollection
 	// Use INDEX_NONE as GroupId for collision particles that affect all particle groups.
 	int32 AddCollisionParticleRange(int32 NumParticles, uint32 GroupId, bool bActivate);
 
+	// Add a single collision body particle to the solver. 
+	int32 AddCollisionParticle(uint32 GroupId, bool bActivate);
+
+	// Remove a collision body from the active view list, and save its particle in the remvoed collision particle list
+	void RemoveCollisionParticle(int32 CollisionParticleIndex, int32 CollisionParticleViewIndex);
+
 	// Set a block of collision particles active or inactive, using the index of the first added particle to identify the block.
 	void ActivateCollisionParticleRange(int32 Offset, bool bActivate) { MCollisionParticlesActiveView.ActivateRange(Offset, bActivate); }
 
@@ -78,8 +84,10 @@ class CHAOS_API FPBDEvolution : public TArrayCollection
 	// Collision particles accessors
 	const FSolverRigidParticles& CollisionParticles() const { return MCollisionParticles; }
 	FSolverRigidParticles& CollisionParticles() { return MCollisionParticles; }
+	TArray<uint32>& CollisionParticleGroupIds() { return MCollisionParticleGroupIds; }
 	const TArray<uint32>& CollisionParticleGroupIds() const { return MCollisionParticleGroupIds; }
-	const TPBDActiveView<FSolverRigidParticles>& CollisionParticlesActiveView() { return MCollisionParticlesActiveView; }
+	TPBDActiveView<FSolverRigidParticles>& CollisionParticlesActiveView() { return MCollisionParticlesActiveView; }
+	const TPBDActiveView<FSolverRigidParticles>& CollisionParticlesActiveView() const { return MCollisionParticlesActiveView; }
 
 	// Reset all constraint init and rule functions.
 	void ResetConstraintRules() 
@@ -185,11 +193,14 @@ class CHAOS_API FPBDEvolution : public TArrayCollection
 	template<bool bForceRule, bool bVelocityField, bool bDampVelocityRule>
 	void PreIterationUpdate(const FSolverReal Dt, const int32 Offset, const int32 Range, const int32 MinParallelBatchSize);
 
+	struct FParticleVievToken { int32 ParticleIndex = INDEX_NONE; int32 ViewIndex = INDEX_NONE; };
+
 private:
 	FSolverParticles MParticles;
 	TPBDActiveView<FSolverParticles> MParticlesActiveView;
 	FSolverRigidParticles MCollisionParticles;
 	TPBDActiveView<FSolverRigidParticles> MCollisionParticlesActiveView;
+	TArray<FParticleVievToken> RemovedCollisionIndices;
 
 	TArrayCollectionArray<FSolverRigidTransform3> MCollisionTransforms;  // Used for CCD to store the initial state before the kinematic update
 	TArrayCollectionArray<bool> MCollided;
