SUV's back wheels falling through floor at the start of big angle test

diff --git a/Engine/Source/Runtime/Experimental/Chaos/Private/Chaos/TriangleMeshImplicitObject.cpp b/Engine/Source/Runtime/Experimental/Chaos/Private/Chaos/TriangleMeshImplicitObject.cpp
--- a/Engine/Source/Runtime/Experimental/Chaos/Private/Chaos/TriangleMeshImplicitObject.cpp
+++ b/Engine/Source/Runtime/Experimental/Chaos/Private/Chaos/TriangleMeshImplicitObject.cpp
@@ -1120,7 +1120,8 @@ bool FTriangleMeshImplicitObject::OverlapGeomImp(const QueryGeomType& QueryGeom,
 	}
 	else
 	{
-		if (IsSpherical || FMath::Abs(QueryTM.GetRotation().W) > 0.9)
+		// Disabled optimization due to a bug UE-176682, TODO Re-enabled it 
+		// if (IsSpherical || FMath::Abs(QueryTM.GetRotation().W) > 0.9)
 		{
 			const FVec3 InvTriMeshScale = FVec3(FReal(1) / TriMeshScale.X, FReal(1) / TriMeshScale.Y, FReal(1) / TriMeshScale.Z);
 			// IMPORTANT QueryTM comes with a invscaled translation so we need a version of the TM with world space translation to properly compute the bounds
@@ -1135,15 +1136,15 @@ bool FTriangleMeshImplicitObject::OverlapGeomImp(const QueryGeomType& QueryGeom,
 			ScaleTransformHelper(TriMeshScale, QueryTM, WorldScaleQueryTM);
 			return FastBVH.FindAllIntersectionsNoMTD(QueryBounds, WorldScaleQueryTM, WorldScaleQueryGeom, Thickness, TriMeshScale, this);
 		}
-		else
-		{
-			FAABB3 GeometryAABB = QueryGeom.BoundingBox();
-			GeometryAABB.ThickenSymmetrically(FVec3(Thickness));
-			Private::FOBBVectorized QueryObb(QueryTM, (GeometryAABB.Max() - GeometryAABB.Min()) * 0.5);
-			TRigidTransform<FReal, 3> WorldScaleQueryTM;
-			ScaleTransformHelper(TriMeshScale, QueryTM, WorldScaleQueryTM);
-			return FastBVH.FindAllIntersectionsNoMTD(QueryObb, WorldScaleQueryTM, WorldScaleQueryGeom, Thickness, TriMeshScale, this);
-		}
+		//else
+		//{
+		//	FAABB3 GeometryAABB = QueryGeom.BoundingBox();
+		//	GeometryAABB.ThickenSymmetrically(FVec3(Thickness));
+		//	Private::FOBBVectorized QueryObb(QueryTM, (GeometryAABB.Max() - GeometryAABB.Min()) * 0.5);
+		//	TRigidTransform<FReal, 3> WorldScaleQueryTM;
+		//	ScaleTransformHelper(TriMeshScale, QueryTM, WorldScaleQueryTM);
+		//	return FastBVH.FindAllIntersectionsNoMTD(QueryObb, WorldScaleQueryTM, WorldScaleQueryGeom, Thickness, TriMeshScale, this);
+		//}
 	}
 }
 
