stop setting the version number when saving a managed array collection, clean up how we set the version for FGeometryCollection, and add an ensure to make sure it's been set as expected

diff --git a/Engine/Source/Runtime/Experimental/Chaos/Private/GeometryCollection/GeometryCollection.cpp b/Engine/Source/Runtime/Experimental/Chaos/Private/GeometryCollection/GeometryCollection.cpp
--- a/Engine/Source/Runtime/Experimental/Chaos/Private/GeometryCollection/GeometryCollection.cpp
+++ b/Engine/Source/Runtime/Experimental/Chaos/Private/GeometryCollection/GeometryCollection.cpp
@@ -102,6 +102,8 @@ void FGeometryCollection::DefineGeometrySchema(FManagedArrayCollection& InCollec
 
 void FGeometryCollection::Construct()
 {
+	Version = GetLatestVersionNumber();
+
 	FManagedArrayCollection::FConstructionParameters TransformDependency(FTransformCollection::TransformGroup);
 	FManagedArrayCollection::FConstructionParameters VerticesDependency(FGeometryCollection::VerticesGroup);
 	FManagedArrayCollection::FConstructionParameters FacesDependency(FGeometryCollection::FacesGroup);
@@ -1296,41 +1298,43 @@ void FGeometryCollection::Serialize(Chaos::FChaosArchive& Ar)
 			HierarchyFacade.GenerateLevelAttribute();
 		}
 
-		if (Version < 10 && 
-			(!HasAttribute(GeometryCollection::UV::UVLayerNames[0], VerticesGroup) || HasAttribute("UVs", VerticesGroup)) )
+		if (Version < 10)
 		{
-			TManagedArray<TArray<FVector2f>>* OrigUVs = FindAttributeTyped<TArray<FVector2f>>("UVs", VerticesGroup);
-			if (OrigUVs)
+			if (!HasAttribute(GeometryCollection::UV::UVLayerNames[0], VerticesGroup) || HasAttribute("UVs", VerticesGroup))
 			{
-				// Note: We take the max of the num layers because in practice there have been some vertices with inconsistent layer counts
-				// and it seems better to transfer all the data (with missing data left as zeros) than to potentially lose data
-				int32 NumLayers = 1;
-				for (int32 Idx = 0; Idx < OrigUVs->Num(); ++Idx)
+				TManagedArray<TArray<FVector2f>>* OrigUVs = FindAttributeTyped<TArray<FVector2f>>("UVs", VerticesGroup);
+				if (OrigUVs)
 				{
-					NumLayers = FMath::Max((*OrigUVs)[Idx].Num(), NumLayers);
-				}
-				NumLayers = FMath::Min((int32)GeometryCollectionUV::MAX_NUM_UV_CHANNELS, NumLayers); // make sure we never exceed max layers
-				SetNumUVLayers(NumLayers);
-
-				for (int32 Idx = 0; Idx < OrigUVs->Num(); ++Idx)
-				{
-					TArray<FVector2f>& VertexLayers = (*OrigUVs)[Idx];
-					int32 NumVertexLayers = FMath::Min(VertexLayers.Num(), NumLayers);
-					for (int32 LayerIdx = 0; LayerIdx < NumVertexLayers; ++LayerIdx)
+					// Note: We take the max of the num layers because in practice there have been some vertices with inconsistent layer counts
+					// and it seems better to transfer all the data (with missing data left as zeros) than to potentially lose data
+					int32 NumLayers = 1;
+					for (int32 Idx = 0; Idx < OrigUVs->Num(); ++Idx)
 					{
-						ModifyUV(Idx, LayerIdx) = VertexLayers[LayerIdx];
+						NumLayers = FMath::Max((*OrigUVs)[Idx].Num(), NumLayers);
 					}
-					for (int32 LayerIdx = NumVertexLayers; LayerIdx < NumLayers; ++LayerIdx)
+					NumLayers = FMath::Min((int32)GeometryCollectionUV::MAX_NUM_UV_CHANNELS, NumLayers); // make sure we never exceed max layers
+					SetNumUVLayers(NumLayers);
+
+					for (int32 Idx = 0; Idx < OrigUVs->Num(); ++Idx)
 					{
-						ModifyUV(Idx, LayerIdx) = FVector2f(0, 0); // explicitly zero UVs of any missing layers
+						TArray<FVector2f>& VertexLayers = (*OrigUVs)[Idx];
+						int32 NumVertexLayers = FMath::Min(VertexLayers.Num(), NumLayers);
+						for (int32 LayerIdx = 0; LayerIdx < NumVertexLayers; ++LayerIdx)
+						{
+							ModifyUV(Idx, LayerIdx) = VertexLayers[LayerIdx];
+						}
+						for (int32 LayerIdx = NumVertexLayers; LayerIdx < NumLayers; ++LayerIdx)
+						{
+							ModifyUV(Idx, LayerIdx) = FVector2f(0, 0); // explicitly zero UVs of any missing layers
+						}
 					}
+
+					RemoveAttribute("UVs", VerticesGroup);
+				}
+				else
+				{
+					SetNumUVLayers(1);
 				}
-				
-				RemoveAttribute("UVs", VerticesGroup);
-			}
-			else
-			{
-				SetNumUVLayers(1);
 			}
 
 			Version = 10;
@@ -1352,6 +1356,8 @@ void FGeometryCollection::Serialize(Chaos::FChaosArchive& Ar)
 		// Finally, make sure expected interfaces are initialized
 		InitializeInterfaces();
 	}
+
+	ensure(Version == GetLatestVersionNumber());
 }
 
 bool FGeometryCollection::HasContiguousVertices( ) const
diff --git a/Engine/Source/Runtime/Experimental/Chaos/Private/GeometryCollection/ManagedArrayCollection.cpp b/Engine/Source/Runtime/Experimental/Chaos/Private/GeometryCollection/ManagedArrayCollection.cpp
--- a/Engine/Source/Runtime/Experimental/Chaos/Private/GeometryCollection/ManagedArrayCollection.cpp
+++ b/Engine/Source/Runtime/Experimental/Chaos/Private/GeometryCollection/ManagedArrayCollection.cpp
@@ -613,7 +613,6 @@ void FManagedArrayCollection::Serialize(Chaos::FChaosArchive& Ar)
 {
 	Ar.UsingCustomVersion(FUE5MainStreamObjectVersion::GUID);
 
-	if (Ar.IsSaving()) Version = 9;
 	Ar << Version;
 
 	if (Ar.IsLoading())
diff --git a/Engine/Source/Runtime/Experimental/Chaos/Public/GeometryCollection/GeometryCollection.h b/Engine/Source/Runtime/Experimental/Chaos/Public/GeometryCollection/GeometryCollection.h
--- a/Engine/Source/Runtime/Experimental/Chaos/Public/GeometryCollection/GeometryCollection.h
+++ b/Engine/Source/Runtime/Experimental/Chaos/Public/GeometryCollection/GeometryCollection.h
@@ -367,6 +367,14 @@ protected:
 	*/
 	void ReorderTransformElements(const TArray<int32>& NewOrder);
 
+	/**
+	 * @return the latest Version number used in serialization
+	 */
+	int32 GetLatestVersionNumber() const
+	{
+		return 10;
+	}
+
 
 public:
 	/* Backwards compatibility */
