make new FGeometryCollection latest version function be constexpr

diff --git a/Engine/Source/Runtime/Experimental/Chaos/Public/GeometryCollection/GeometryCollection.h b/Engine/Source/Runtime/Experimental/Chaos/Public/GeometryCollection/GeometryCollection.h
--- a/Engine/Source/Runtime/Experimental/Chaos/Public/GeometryCollection/GeometryCollection.h
+++ b/Engine/Source/Runtime/Experimental/Chaos/Public/GeometryCollection/GeometryCollection.h
@@ -370,7 +370,7 @@ protected:
 	/**
 	 * @return the latest Version number used in serialization
 	 */
-	int32 GetLatestVersionNumber() const
+	constexpr int32 GetLatestVersionNumber() const
 	{
 		return 10;
 	}
