rename "Split Components" in fracture mode/geometry collection settings to "Split Meshes" to reduce confusion with the UE "Component" concept

diff --git a/Engine/Plugins/Experimental/ChaosEditor/Source/FractureEditor/Private/FractureToolGenerators.cpp b/Engine/Plugins/Experimental/ChaosEditor/Source/FractureEditor/Private/FractureToolGenerators.cpp
--- a/Engine/Plugins/Experimental/ChaosEditor/Source/FractureEditor/Private/FractureToolGenerators.cpp
+++ b/Engine/Plugins/Experimental/ChaosEditor/Source/FractureEditor/Private/FractureToolGenerators.cpp
@@ -306,7 +306,7 @@ void SCreateGeometryCollectionFromObject::Construct(const FArguments& InArgs, TS
 					.VAlign(VAlign_Center)
 					[
 						SNew(STextBlock)
-						.Text(LOCTEXT("SplitComponentsLabel", "Split Components"))
+						.Text(LOCTEXT("SplitMeshesLabel", "Split Meshes"))
 						.ToolTipText(LOCTEXT("SplitComponentsToolTip", "If checked, triangles that are not topologically connected will be assigned separate bones on import. Does not apply if creating from a Geometry Collection source."))
 					]
 
diff --git a/Engine/Source/Runtime/Experimental/GeometryCollectionEngine/Public/GeometryCollection/GeometryCollectionObject.h b/Engine/Source/Runtime/Experimental/GeometryCollectionEngine/Public/GeometryCollection/GeometryCollectionObject.h
--- a/Engine/Source/Runtime/Experimental/GeometryCollectionEngine/Public/GeometryCollection/GeometryCollectionObject.h
+++ b/Engine/Source/Runtime/Experimental/GeometryCollectionEngine/Public/GeometryCollection/GeometryCollectionObject.h
@@ -48,7 +48,7 @@ struct GEOMETRYCOLLECTIONENGINE_API FGeometryCollectionSource
 	bool bAddInternalMaterials = true;
 
 	/** Whether individual source mesh components should be split into separate pieces of geometry based on mesh connectivity. If checked, triangles that are not topologically connected will be assigned separate bones. (For non-Geometry Collection inputs only.) */
-	UPROPERTY(BlueprintReadWrite, EditAnywhere, Category = "GeometrySource")
+	UPROPERTY(BlueprintReadWrite, EditAnywhere, Category = "GeometrySource", DisplayName = "Split Meshes")
 	bool bSplitComponents = false;
 
 	/** Whether to set the 'internal' flag for faces with odd-numbered materials slots. */
