Chaos - Reduce PerShapeData extended data size and don't create extended data on game thread

diff --git a/Engine/Source/Runtime/Engine/Private/PhysicsEngine/ImmediatePhysics/ImmediatePhysicsChaos/ImmediatePhysicsActorHandle_Chaos.cpp b/Engine/Source/Runtime/Engine/Private/PhysicsEngine/ImmediatePhysics/ImmediatePhysicsChaos/ImmediatePhysicsActorHandle_Chaos.cpp
--- a/Engine/Source/Runtime/Engine/Private/PhysicsEngine/ImmediatePhysics/ImmediatePhysicsChaos/ImmediatePhysicsActorHandle_Chaos.cpp
+++ b/Engine/Source/Runtime/Engine/Private/PhysicsEngine/ImmediatePhysics/ImmediatePhysicsChaos/ImmediatePhysicsActorHandle_Chaos.cpp
@@ -31,8 +31,8 @@ namespace ImmediatePhysics_Chaos
 		const FReal Mass = 1.0f;
 		const FReal Radius = 1.0f * Scale.GetMax();
 
-		auto ImplicitSphere = MakeUnique<Chaos::TSphere<FReal, 3>>(FVec3(0), Radius);
-		auto NewShape = Chaos::FPerShapeData::CreatePerShapeData(OutShapes.Num(), MakeSerializable(ImplicitSphere));
+		TUniquePtr<Chaos::FImplicitSphere3> ImplicitSphere = MakeUnique<Chaos::FImplicitSphere3>(FVec3(0), Radius);
+		TUniquePtr<FPerShapeData> NewShape = Chaos::FPerShapeData::CreatePerShapeDataInternal(OutShapes.Num(), MakeSerializable(ImplicitSphere));
 		NewShape->UpdateShapeBounds(FTransform::Identity);
 		NewShape->SetUserData(nullptr);
 		NewShape->SetQueryEnabled(false);
@@ -182,7 +182,7 @@ namespace ImmediatePhysics_Chaos
 			OutGeom = MakeUnique<FImplicitObjectUnion>(MoveTemp(Geoms));
 		}
 
-		for (auto& Shape : Shapes)
+		for (TUniquePtr<FPerShapeData>& Shape : Shapes)
 		{
 			OutShapes.Emplace(MoveTemp(Shape));
 		}
@@ -242,7 +242,7 @@ namespace ImmediatePhysics_Chaos
 				CollisionData.SimData.Word1 = 0xFFFFF;
 				CollisionData.SimData.Word3 = 0xFFFFF;
 				CollisionData.bSimCollision = 1;
-				for (const auto& Shape : ParticleHandle->ShapesArray())
+				for (const TUniquePtr<FPerShapeData>& Shape : ParticleHandle->ShapesArray())
 				{
 					Shape->SetCollisionData(CollisionData);
 				}
@@ -254,13 +254,13 @@ namespace ImmediatePhysics_Chaos
 					ParticleHandle->UpdateWorldSpaceState(FRigidTransform3(ParticleHandle->X(), ParticleHandle->R()), FVec3(0));
 				}
 
-				if (auto* Kinematic = ParticleHandle->CastToKinematicParticle())
+				if (FKinematicGeometryParticleHandle* Kinematic = ParticleHandle->CastToKinematicParticle())
 				{
 					Kinematic->SetV(FVector3f::ZeroVector);
 					Kinematic->SetW(FVector3f::ZeroVector);
 				}
 
-				auto* Dynamic = ParticleHandle->CastToRigidParticle();
+				FPBDRigidParticleHandle* Dynamic = ParticleHandle->CastToRigidParticle();
 				if (Dynamic && Dynamic->ObjectState() == EObjectStateType::Dynamic)
 				{
 					FReal MassInv = (Mass > 0.0f) ? 1.0f / Mass : 0.0f;
@@ -331,7 +331,7 @@ namespace ImmediatePhysics_Chaos
 
 		SetWorldTransform(WorldTM);
 
-		if (auto* Kinematic = ParticleHandle->CastToKinematicParticle())
+		if (FKinematicGeometryParticleHandle* Kinematic = ParticleHandle->CastToKinematicParticle())
 		{
 			Kinematic->V() = FVec3(0);
 			Kinematic->W() = FVec3(0);
@@ -346,7 +346,7 @@ namespace ImmediatePhysics_Chaos
 		ParticleHandle->X() = WorldTM.GetTranslation();
 		ParticleHandle->R() = WorldTM.GetRotation();
 
-		auto* Dynamic = ParticleHandle->CastToRigidParticle();
+		FPBDRigidParticleHandle* Dynamic = ParticleHandle->CastToRigidParticle();
 		if(Dynamic && Dynamic->ObjectState() == Chaos::EObjectStateType::Dynamic)
 		{
 			Dynamic->P() = Dynamic->X();
diff --git a/Engine/Source/Runtime/Experimental/Chaos/Private/Chaos/GeometryParticles.cpp b/Engine/Source/Runtime/Experimental/Chaos/Private/Chaos/GeometryParticles.cpp
--- a/Engine/Source/Runtime/Experimental/Chaos/Private/Chaos/GeometryParticles.cpp
+++ b/Engine/Source/Runtime/Experimental/Chaos/Private/Chaos/GeometryParticles.cpp
@@ -11,7 +11,41 @@
 
 namespace Chaos
 {
-	void UpdateShapesArrayFromGeometry(FShapesArray& ShapesArray, TSerializablePtr<FImplicitObject> Geometry, const FRigidTransform3& ActorTM, IPhysicsProxyBase* Proxy)
+	TUniquePtr<FPerShapeData> CreatePerShapeDataImpl(const int32 ShapeIndex, TSerializablePtr<FImplicitObject> InGeometry, const bool bAllowCachedLeafInfo)
+	{
+		if (bAllowCachedLeafInfo)
+		{
+			return FPerShapeData::CreatePerShapeDataInternal(ShapeIndex, InGeometry);
+		}
+		else
+		{
+			return FPerShapeData::CreatePerShapeData(ShapeIndex, InGeometry);
+		}
+	}
+
+	void UpdatePerShapeDataImpl(TUniquePtr<FPerShapeData>& ShapePtr, TSerializablePtr<FImplicitObject> InGeometry, const bool bAllowCachedLeafInfo)
+	{
+		if (bAllowCachedLeafInfo)
+		{
+			FPerShapeData::UpdateGeometryInternal(ShapePtr, InGeometry);
+		}
+		else
+		{
+			FPerShapeData::UpdateGeometry(ShapePtr, InGeometry);
+		}
+	}
+
+	// Creare or reuse the shapes in the shapes array and populate them with the Geometry.
+	// If we have a Union it will be unpacked into the ShapesArray.
+	// On the Physics Thread we set bAllowCachedLeafInfo which caches the shapes world space state to optimize collision detection,
+	// but this flag should not be used on the Game Thread because the extended data is not maintained so some PerShapeData funtions
+	// will return unitialized values (E.g., GetLeafWorldTransform).
+	void UpdateShapesArrayFromGeometryImpl(
+		FShapesArray& ShapesArray, 
+		TSerializablePtr<FImplicitObject> Geometry, 
+		const FRigidTransform3& ActorTM, 
+		IPhysicsProxyBase* Proxy, 
+		bool bAllowCachedLeafInfo)
 	{
 		if(Geometry)
 		{
@@ -27,12 +61,12 @@ namespace Chaos
 					if (ShapeIndex >= OldShapeNum)
 					{
 						// If newly allocated shape, initialize it.
-						ShapesArray[ShapeIndex] = FPerShapeData::CreatePerShapeData(ShapeIndex, ShapeGeometry);
+						ShapesArray[ShapeIndex] = CreatePerShapeDataImpl(ShapeIndex, ShapeGeometry, bAllowCachedLeafInfo);
 					}
 					else if (ShapeGeometry != ShapesArray[ShapeIndex]->GetGeometry())
 					{
 						// Update geometry pointer if it changed
-						FPerShapeData::UpdateGeometry(ShapesArray[ShapeIndex], ShapeGeometry);
+						UpdatePerShapeDataImpl(ShapesArray[ShapeIndex], ShapeGeometry, bAllowCachedLeafInfo);
 					}
 				}
 			}
@@ -41,11 +75,11 @@ namespace Chaos
 				ShapesArray.SetNum(1);
 				if (OldShapeNum == 0)
 				{
-					ShapesArray[0] = FPerShapeData::CreatePerShapeData(0, Geometry);
+					ShapesArray[0] = CreatePerShapeDataImpl(0, Geometry, bAllowCachedLeafInfo);
 				}
 				else
 				{
-					FPerShapeData::UpdateGeometry(ShapesArray[0], Geometry);
+					UpdatePerShapeDataImpl(ShapesArray[0], Geometry, bAllowCachedLeafInfo);
 				}
 			}
 
@@ -71,12 +105,21 @@ namespace Chaos
 		}
 	}
 
+	void UpdateShapesArrayFromGeometry(
+		FShapesArray& ShapesArray,
+		TSerializablePtr<FImplicitObject> Geometry,
+		const FRigidTransform3& ActorTM,
+		IPhysicsProxyBase* Proxy)
+	{
+		UpdateShapesArrayFromGeometryImpl(ShapesArray, Geometry, ActorTM, Proxy, false);
+	}
+
 	// Unwrap transformed shapes
 	// @todo(chaos): also unwrap Instanced and Scaled but that requires a lot of knock work because Convexes are usually Instanced or
 	// Scaled so the Scale and Margin is stored on the wrapper (the convex itself is shared).
 	// - support for Margin as per-shape data passed through the collision functions
 	// - support for Scale as per-shape data passed through the collision functions (or ideally in the Transforms)
-	const FImplicitObject* GetInnerGeometryInstanceData(const FImplicitObject* Implicit, FRigidTransformRealSingle3& OutTransform, FReal& OutMargin)
+	const FImplicitObject* GetInnerGeometryInstanceData(const FImplicitObject* Implicit, const FRigidTransform3** OutRelativeTransformPtr)
 	{
 		if (Implicit != nullptr)
 		{
@@ -85,21 +128,23 @@ namespace Chaos
 			{
 				// Transformed Implicit
 				const TImplicitObjectTransformed<FReal, 3>* TransformedImplicit = Implicit->template GetObject<const TImplicitObjectTransformed<FReal, 3>>();
-				OutTransform = (FRigidTransformRealSingle3)TransformedImplicit->GetTransform() * OutTransform;
-				//OutMargin += TransformedImplicit->GetMargin();
-				return GetInnerGeometryInstanceData(TransformedImplicit->GetTransformedObject(), OutTransform, OutMargin);
+				if (OutRelativeTransformPtr != nullptr)
+				{
+					*OutRelativeTransformPtr = &TransformedImplicit->GetTransform();
+				}
+				return GetInnerGeometryInstanceData(TransformedImplicit->GetTransformedObject(), OutRelativeTransformPtr);
 			}
 			else if ((uint32)ImplicitOuterType & ImplicitObjectType::IsInstanced)
 			{
 				// Instanced Implicit
-				// Currently we only unwrap instanced TriMesh and Heightfields. They don't have a margin, so they don't need the instance wrapper in collision detection.
-				// The only other type we wrap in instances right now are Convex, so we just check for that here...
+				// Currently we only unwrap instanced TriMesh and Heightfields. Convex is the only other type that will be Instanced,
+				// but we do not unwrap Convex because the convex margin is stored in the Instanced so we would need to extract
+				// it and pass it on to the collision detection.
 				const FImplicitObjectInstanced* Instanced = static_cast<const FImplicitObjectInstanced*>(Implicit);
 				EImplicitObjectType InnerType = Instanced->GetInnerObject()->GetType();
 				if (InnerType != FImplicitConvex3::StaticType())
 				{
-					OutMargin += Instanced->GetMargin();
-					return GetInnerGeometryInstanceData(Instanced->GetInnerObject().Get(), OutTransform, OutMargin);
+					return GetInnerGeometryInstanceData(Instanced->GetInnerObject().Get(), OutRelativeTransformPtr);
 				}
 			}
 			else if ((uint32)ImplicitOuterType & ImplicitObjectType::IsScaled)
@@ -107,13 +152,21 @@ namespace Chaos
 				// Scaled Implicit
 				//const FImplicitObjectScaled* Scaled = static_cast<const FImplicitObjectScaled*>(Implicit);
 				//OutTransform.Scale *= Scaled->GetScale();
-				//OutMargin += Scaled->GetMargin();
-				//return GetInnerGeometryInstanceData(Scaled->GetInnerObject().Get(), OutTransform, OutMargin);
+				//return GetInnerGeometryInstanceData(Scaled->GetInnerObject().Get(), OutRelativeTransformPtr);
 			}
 		}
 		return Implicit;
 	}
 
+	bool PerShapeDataRequiresCachedLeafInfo(const FImplicitObject* Geometry)
+	{
+		// Unwrap the geometry (if it's Transformed or Instanced)
+		const FRigidTransform3* LeafRelativeTransform = nullptr;
+		const FImplicitObject* LeafGeometry = GetInnerGeometryInstanceData(Geometry, &LeafRelativeTransform);
+		return (LeafGeometry != Geometry) && (LeafRelativeTransform != nullptr);
+	}
+
+
 
 	FPerShapeData::FPerShapeData(int32 InShapeIdx)
 		: bHasCachedLeafInfo(false)
@@ -124,8 +177,8 @@ namespace Chaos
 	{
 	}
 
-	FPerShapeData::FPerShapeData(int32 InShapeIdx, TSerializablePtr<FImplicitObject> InGeometry, bool bInHasCachedLeafInfo)
-		: bHasCachedLeafInfo(bInHasCachedLeafInfo)
+	FPerShapeData::FPerShapeData(int32 InShapeIdx, TSerializablePtr<FImplicitObject> InGeometry)
+		: bHasCachedLeafInfo(false)
 		, Proxy(nullptr)
 		, ShapeIdx(InShapeIdx)
 		, Geometry(InGeometry)
@@ -149,48 +202,65 @@ namespace Chaos
 	{
 	}
 
-	TUniquePtr<FPerShapeData> FPerShapeData::CreatePerShapeData(int32 ShapeIdx, TSerializablePtr<FImplicitObject> InGeometry)
+	TUniquePtr<FPerShapeData> FPerShapeData::CreatePerShapeData(int32 InShapeIdx, TSerializablePtr<FImplicitObject> InGeometry)
 	{
-		FReal LeafMargin = 0.0f;
-		FRigidTransformRealSingle3 LeafRelativeTransform = FRigidTransformRealSingle3::Identity;
-		const FImplicitObject* LeafGeometry = GetInnerGeometryInstanceData(InGeometry.Get(), LeafRelativeTransform, LeafMargin);
+		return TUniquePtr<FPerShapeData>(new FPerShapeData(InShapeIdx, InGeometry));
+	}
 
-		if (RequiresCachedLeafInfo(LeafRelativeTransform, LeafGeometry, InGeometry.Get()))
+	void FPerShapeData::UpdateGeometry(TUniquePtr<FPerShapeData>& ShapePtr, TSerializablePtr<FImplicitObject> InGeometry)
+	{
+		ShapePtr->Geometry = InGeometry;
+	}
+
+	TUniquePtr<FPerShapeData> FPerShapeData::CreatePerShapeDataInternal(int32 InShapeIdx, TSerializablePtr<FImplicitObject> InGeometry)
+	{
+		const bool bWantLeafGeometryCache = PerShapeDataRequiresCachedLeafInfo(InGeometry.Get());
+		if (bWantLeafGeometryCache)
 		{
-			return TUniquePtr<FPerShapeData>(new FPerShapeDataCachedLeafInfo(ShapeIdx, InGeometry, LeafGeometry, LeafRelativeTransform));
+			return TUniquePtr<FPerShapeData>(new FPerShapeDataCachedLeafInfo(InShapeIdx, InGeometry));
 		}
 		else
 		{
-			return TUniquePtr<FPerShapeData>(new FPerShapeData(ShapeIdx, InGeometry));
+			return TUniquePtr<FPerShapeData>(new FPerShapeData(InShapeIdx, InGeometry));
 		}
 	}
 
-	void FPerShapeData::UpdateGeometry(TUniquePtr<FPerShapeData>& ShapePtr, TSerializablePtr<FImplicitObject> InGeometry)
+	void FPerShapeData::UpdateGeometryInternal(TUniquePtr<FPerShapeData>& ShapePtr, TSerializablePtr<FImplicitObject> InGeometry)
 	{
-		ShapePtr->Geometry = InGeometry;
+		const bool bWantLeafGeometryCache = PerShapeDataRequiresCachedLeafInfo(InGeometry.Get());
 
-		FReal LeafMargin = 0.0f;
-		FRigidTransformRealSingle3 LeafRelativeTransform = FRigidTransformRealSingle3::Identity;
-		const FImplicitObject* LeafGeometry = GetInnerGeometryInstanceData(InGeometry.Get(), LeafRelativeTransform, LeafMargin);
-
-		if (ShapePtr->bHasCachedLeafInfo)
+		// Do we need to add or remove the cached leaf data? If so this requires we recreate the object
+		if (bWantLeafGeometryCache != ShapePtr->HasCachedLeafInfo())
 		{
-			// set relative info
-			ShapePtr->SetLeafRelativeTransform(LeafRelativeTransform);
-			ShapePtr->SetLeafGeometry(LeafGeometry);
-			return;
+			if (bWantLeafGeometryCache)
+			{
+				ShapePtr = TUniquePtr<FPerShapeData>(new FPerShapeDataCachedLeafInfo(MoveTemp(*ShapePtr.Get())));
+			}
+			else
+			{
+				ShapePtr = TUniquePtr<FPerShapeData>(new FPerShapeData(MoveTemp(*ShapePtr.Get())));
+			}
 		}
 
-		if (RequiresCachedLeafInfo(LeafRelativeTransform, LeafGeometry, InGeometry.Get()))
+		ShapePtr->Geometry = InGeometry;
+	}
+
+	FPerShapeDataCachedLeafInfo* FPerShapeData::AsCachedLeafInfo()
+	{
+		if (HasCachedLeafInfo())
 		{
-			// We need to move to FPerShapeDataCachedLeafInfo to cache this info.
-			ShapePtr = TUniquePtr<FPerShapeData>(new FPerShapeDataCachedLeafInfo(MoveTemp(*ShapePtr.Get()), LeafGeometry, LeafRelativeTransform));
+			return static_cast<FPerShapeDataCachedLeafInfo*>(this);
 		}
+		return nullptr;
 	}
 
-	bool FPerShapeData::RequiresCachedLeafInfo(const FRigidTransformRealSingle3& RelativeTransform, const FImplicitObject* LeafGeometry, const FImplicitObject* Geometry)
+	const FPerShapeDataCachedLeafInfo* FPerShapeData::AsCachedLeafInfo() const
 	{
-		return (LeafGeometry != Geometry || !RelativeTransform.Equals(FRigidTransformRealSingle3::Identity));
+		if (HasCachedLeafInfo())
+		{
+			return static_cast<const FPerShapeDataCachedLeafInfo*>(this);
+		}
+		return nullptr;
 	}
 
 	void FPerShapeData::UpdateShapeBounds(const FRigidTransform3& WorldTM, const FVec3& BoundsExpansion)
@@ -204,14 +274,22 @@ namespace Chaos
 	void FPerShapeData::UpdateWorldSpaceState(const FRigidTransform3& WorldTransform, const FVec3& BoundsExpansion)
 	{
 		FRigidTransform3 LeafWorldTransform = WorldTransform;
-		if (bHasCachedLeafInfo)
+		const FRigidTransform3* LeafRelativeTransform = nullptr;
+		const FImplicitObject* LeafGeometry = GetInnerGeometryInstanceData(Geometry.Get(), &LeafRelativeTransform);
+
+		// Calculate the leaf world transform if different from particle transform
+		if (LeafRelativeTransform != nullptr)
 		{
-			FPerShapeDataCachedLeafInfo& LeafInfo = *static_cast<FPerShapeDataCachedLeafInfo*>(this);
-			LeafWorldTransform = FRigidTransform3::MultiplyNoScale(FRigidTransform3(GetLeafRelativeTransform()), WorldTransform);
-			LeafInfo.LeafWorldTransform = LeafWorldTransform;
+			LeafWorldTransform = FRigidTransform3::MultiplyNoScale(*LeafRelativeTransform, WorldTransform);
 		}
 
-		const FImplicitObject* LeafGeometry = GetLeafGeometry();
+		// Store the leaf data if we have a cache
+		if (FPerShapeDataCachedLeafInfo* LeafCache = AsCachedLeafInfo())
+		{
+			LeafCache->SetWorldTransform(LeafWorldTransform);
+		}
+
+		// Update the bounds at the world transform
 		if ((LeafGeometry != nullptr) && LeafGeometry->HasBoundingBox())
 		{
 			WorldSpaceInflatedShapeBounds = LeafGeometry->CalculateTransformedBounds(LeafWorldTransform).ThickenSymmetrically(BoundsExpansion);;
@@ -220,62 +298,59 @@ namespace Chaos
 
 	const FImplicitObject* FPerShapeData::GetLeafGeometry() const
 	{
-		if (bHasCachedLeafInfo)
-		{
-			const FPerShapeDataCachedLeafInfo& LeafInfo = *static_cast<const FPerShapeDataCachedLeafInfo*>(this);
-			return LeafInfo.LeafGeometry;
-		}
-
-		return Geometry.Get();
+		return GetInnerGeometryInstanceData(Geometry.Get(), nullptr);
 	}
 
-	void FPerShapeData::SetLeafGeometry(const FImplicitObject* LeafGeometry)
+	FRigidTransform3 FPerShapeData::GetLeafRelativeTransform() const
 	{
-		if (ensure(bHasCachedLeafInfo))
+		const FRigidTransform3* LeafRelativeTransform = nullptr;
+		GetInnerGeometryInstanceData(Geometry.Get(), &LeafRelativeTransform);
+
+		if (LeafRelativeTransform != nullptr)
 		{
-			FPerShapeDataCachedLeafInfo& LeafInfo = *static_cast<FPerShapeDataCachedLeafInfo*>(this);
-			LeafInfo.LeafGeometry = LeafGeometry;
+			return *LeafRelativeTransform;
 		}
-	}
-
-	FRigidTransformRealSingle3 FPerShapeData::GetLeafRelativeTransform() const
-	{
-		if (bHasCachedLeafInfo)
+		else
 		{
-			const FPerShapeDataCachedLeafInfo& LeafInfo = *static_cast<const FPerShapeDataCachedLeafInfo*>(this);
-			return LeafInfo.LeafRelativeTransform;
+			return FRigidTransform3::Identity;
 		}
-
-		return FRigidTransformRealSingle3::Identity;
 	}
 
 	FRigidTransform3 FPerShapeData::GetLeafWorldTransform(const FGeometryParticleHandle* Particle) const
 	{
-		if (bHasCachedLeafInfo)
+		if (const FPerShapeDataCachedLeafInfo* LeafCache = AsCachedLeafInfo())
 		{
-			const FPerShapeDataCachedLeafInfo& LeafInfo = *static_cast<const FPerShapeDataCachedLeafInfo*>(this);
-			return LeafInfo.LeafWorldTransform;
+			return LeafCache->GetWorldTransform();
 		}
+		else
+		{
+			FRigidTransform3 LeafWorldTransform = FConstGenericParticleHandle(Particle)->GetTransformPQ();
 
-		return FConstGenericParticleHandle(Particle)->GetTransformPQ();
+			const FRigidTransform3* LeafRelativeTransform = nullptr;
+			GetInnerGeometryInstanceData(Geometry.Get(), &LeafRelativeTransform);
+			if (LeafRelativeTransform != nullptr)
+			{
+				LeafWorldTransform = FRigidTransform3::MultiplyNoScale(*LeafRelativeTransform, LeafWorldTransform);
+			}
+
+			return LeafWorldTransform;
+		}
 	}
 	
 	void FPerShapeData::UpdateLeafWorldTransform(FGeometryParticleHandle* Particle)
 	{
-		if (bHasCachedLeafInfo)
+		if (FPerShapeDataCachedLeafInfo* LeafCache = AsCachedLeafInfo())
 		{
-			FPerShapeDataCachedLeafInfo& LeafInfo = *static_cast<FPerShapeDataCachedLeafInfo*>(this);
-			const FRigidTransform3 ParticleTransform = FConstGenericParticleHandle(Particle)->GetTransformPQ();
-			LeafInfo.LeafWorldTransform = FRigidTransform3::MultiplyNoScale(FRigidTransform3(GetLeafRelativeTransform()), ParticleTransform);
-		}
-	}
+			FRigidTransform3 LeafWorldTransform = FConstGenericParticleHandle(Particle)->GetTransformPQ();
 
-	void FPerShapeData::SetLeafRelativeTransform(const FRigidTransformRealSingle3& RelativeTransform)
-	{
-		if (ensure(bHasCachedLeafInfo))
-		{
-			FPerShapeDataCachedLeafInfo& LeafInfo = *static_cast<FPerShapeDataCachedLeafInfo*>(this);
-			LeafInfo.LeafRelativeTransform = RelativeTransform;
+			const FRigidTransform3* LeafRelativeTransform = nullptr;
+			GetInnerGeometryInstanceData(Geometry.Get(), &LeafRelativeTransform);
+			if (LeafRelativeTransform != nullptr)
+			{
+				LeafWorldTransform = FRigidTransform3::MultiplyNoScale(*LeafRelativeTransform, LeafWorldTransform);
+			}
+
+			LeafCache->SetWorldTransform(LeafWorldTransform);
 		}
 	}
 
@@ -332,6 +407,11 @@ namespace Chaos
 		bIsSingleShape = true;
 	}
 
+	template <typename T, int d, EGeometryParticlesSimType SimType>
+	void TGeometryParticlesImp<T, d, SimType>::UpdateShapesArray(const int32 Index)
+	{
+		UpdateShapesArrayFromGeometryImpl(MShapesArray[Index], MGeometry[Index], FRigidTransform3(X(Index), R(Index)), nullptr, true);
+	}
 
 	template <typename T, int d, EGeometryParticlesSimType SimType>
 	void TGeometryParticlesImp<T, d, SimType>::SetHandle(int32 Index, FGeometryParticleHandle* Handle)
diff --git a/Engine/Source/Runtime/Experimental/Chaos/Private/Chaos/ParticleHandle.cpp b/Engine/Source/Runtime/Experimental/Chaos/Private/Chaos/ParticleHandle.cpp
--- a/Engine/Source/Runtime/Experimental/Chaos/Private/Chaos/ParticleHandle.cpp
+++ b/Engine/Source/Runtime/Experimental/Chaos/Private/Chaos/ParticleHandle.cpp
@@ -13,6 +13,8 @@
 
 namespace Chaos
 {
+	extern void UpdateShapesArrayFromGeometryImpl(FShapesArray& ShapesArray, TSerializablePtr<FImplicitObject> Geometry, const FRigidTransform3& ActorTM, IPhysicsProxyBase* Proxy, bool bAllowCachedLeafInfo);
+
 	void SetObjectStateHelper(IPhysicsProxyBase& Proxy, FPBDRigidParticleHandle& Rigid, EObjectStateType InState, bool bAllowEvents, bool bInvalidate)
 	{
 		if (auto PhysicsSolver = Proxy.GetSolver<Chaos::FPBDRigidsSolver>())
@@ -54,6 +56,12 @@ namespace Chaos
 		return nullptr;
 	}
 
+	template <typename T, int d>
+	void Chaos::TGeometryParticle<T, d>::UpdateShapesArray()
+	{
+		UpdateShapesArrayFromGeometryImpl(MShapesArray, MakeSerializable(MNonFrequentData.Read().Geometry()), FRigidTransform3(X(), R()), Proxy, /*bAllowLeafCache*/false);
+	}
+
 	template <typename T, int d>
 	void Chaos::TGeometryParticle<T, d>::MergeGeometry(TArray<TUniquePtr<FImplicitObject>>&& Objects)
 	{
diff --git a/Engine/Source/Runtime/Experimental/Chaos/Public/Chaos/GeometryParticles.h b/Engine/Source/Runtime/Experimental/Chaos/Public/Chaos/GeometryParticles.h
--- a/Engine/Source/Runtime/Experimental/Chaos/Public/Chaos/GeometryParticles.h
+++ b/Engine/Source/Runtime/Experimental/Chaos/Public/Chaos/GeometryParticles.h
@@ -27,6 +27,7 @@ namespace Chaos
 {
 	class FConstraintHandle;
 	class FParticleCollisions;
+	class FPerShapeDataCachedLeafInfo;
 
 	using FShapesArray = TArray<TUniquePtr<FPerShapeData>, TInlineAllocator<1>>;
 	using FConstraintHandleArray = TArray<FConstraintHandle*>;
@@ -37,15 +38,48 @@ namespace Chaos
 		CHAOS_API extern bool bCCDAxisThresholdUsesProbeShapes;
 	}
 
-	/** Data that is associated with geometry. If a union is used an entry is created per internal geometry */
+	/**
+	 * Per-particle data associated with a geometry used by collision detection,
+	 * and queries. This contains data like the collision/query filters, material
+	 * properties etc.
+	 *
+	 * Every particle holds one FPerShapeData objects for each geometry they use.
+	 * If the particle has a Union of geometries the shapes array will be a flattened
+	 * list of those geometries. (Except ClusteredUnions - these are not flattened
+	 * because they contain their own query acceleration structure.)
+	 * 
+	 * NOTE: Both the game thread and physics thread have FPerShapeData objects for
+	 * each shape. This class is used for both.
+	 * 
+	 * NOTE: It is extremely important to keep this class as small as possible
+	 * because there is one shape instance in the world. This can easily be in
+	 * the millions.
+	 * 
+	 * @todo(chaos): split FPerShapeData into GT and PT versions. The PT version does 
+	 * not need the interop data. The GT version does not need the extended data flag.
+	 * Also try to remove the GT Proxy pointer - this could easily be passed into the
+	 * relevant functions instead.
+	 * 
+	 * @todo(chaos): For the PT version (maybe the GT version) reduce the cost of
+	 * MaterialData for shapes with one material and no masks etc.
+	 * 
+	 */
 	class CHAOS_API FPerShapeData
 	{
 	public:
 
 		static constexpr bool AlwaysSerializable = true;
+
+		UE_DEPRECATED(5.3, "Not used")
+		static bool RequiresCachedLeafInfo(const FImplicitObject* Geometry) { return false; }
+
+		// Game Thread PerShapeData factory/update
 		static TUniquePtr<FPerShapeData> CreatePerShapeData(int32 InShapeIdx, TSerializablePtr<FImplicitObject> InGeometry);
 		static void UpdateGeometry(TUniquePtr<FPerShapeData>& ShapePtr, TSerializablePtr<FImplicitObject> InGeometry);
-		static bool RequiresCachedLeafInfo(const FRigidTransformRealSingle3& RelativeTransform, const FImplicitObject* LeafGeometry, const FImplicitObject* Geometry);
+
+		// Physics Thread PerShapeData factory/update - Do not use the Ext versions from the game thread because the extended data will not be maintained
+		static TUniquePtr<FPerShapeData> CreatePerShapeDataInternal(int32 InShapeIdx, TSerializablePtr<FImplicitObject> InGeometry);
+		static void UpdateGeometryInternal(TUniquePtr<FPerShapeData>& ShapePtr, TSerializablePtr<FImplicitObject> InGeometry);
 
 		virtual ~FPerShapeData();
 		
@@ -88,7 +122,7 @@ namespace Chaos
 		const FImplicitObject* GetLeafGeometry() const;
 
 		// The actor-relative transform of the leaf geometry.
-		FRigidTransformRealSingle3 GetLeafRelativeTransform() const;
+		FRigidTransform3 GetLeafRelativeTransform() const;
 
 		// The world-space transform of the leaf geometry.
 		// If we have non-identity leaf relative transform, is cached from the last call to UpdateWorldSpaceState.
@@ -254,17 +288,15 @@ namespace Chaos
 			ShapeIdx = NewShapeIndex;
 		}
 
-	private:
-
-		void SetLeafRelativeTransform(const FRigidTransformRealSingle3& RelativeTransform);
-		void SetLeafGeometry(const FImplicitObject* LeafGeometry);
-
 	protected:
 		// use CreatePerShapeData
-		FPerShapeData(int32 InShapeIdx, TSerializablePtr<FImplicitObject> InGeometry, bool bInHasCachedLeafInfo = false);
+		FPerShapeData(int32 InShapeIdx, TSerializablePtr<FImplicitObject> InGeometry);
 
 		explicit FPerShapeData(FPerShapeData&& Other);
 
+		FPerShapeDataCachedLeafInfo* AsCachedLeafInfo();
+		const FPerShapeDataCachedLeafInfo* AsCachedLeafInfo() const;
+
 		// Can we be downcasted to FPerShapeDataCachedLeafInfo?
 		uint8 bHasCachedLeafInfo : 1;
 
@@ -288,26 +320,32 @@ namespace Chaos
 		friend class FPerShapeData;
 
 	private:
-		FPerShapeDataCachedLeafInfo(int32 InShapeIdx, TSerializablePtr<FImplicitObject> InGeometry, const FImplicitObject* InLeafGeometry, const FRigidTransformRealSingle3& RelativeTransform)
-			: FPerShapeData(InShapeIdx, InGeometry, /*bHasCachedLeafInfo=*/true)
-			, LeafRelativeTransform(RelativeTransform)
-			, LeafWorldTransform()
-			, LeafGeometry(InLeafGeometry)
-		{}
+		FPerShapeDataCachedLeafInfo(int32 InShapeIdx, TSerializablePtr<FImplicitObject> InGeometry)
+			: FPerShapeData(InShapeIdx, InGeometry)
+		{
+			bHasCachedLeafInfo = true;
+		}
 
 		// Move old shape into shape with cached leaf info
-		FPerShapeDataCachedLeafInfo(FPerShapeData&& PerShapeData, const FImplicitObject* InLeafGeometry, const FRigidTransformRealSingle3& RelativeTransform)
+		FPerShapeDataCachedLeafInfo(FPerShapeData&& PerShapeData)
 			: FPerShapeData(MoveTemp(PerShapeData))
-			, LeafRelativeTransform(RelativeTransform)
-			, LeafWorldTransform()
-			, LeafGeometry(InLeafGeometry)
 		{
 			bHasCachedLeafInfo = true;
 		}
 
-		FRigidTransformRealSingle3 LeafRelativeTransform;
-		FRigidTransform3 LeafWorldTransform;
-		const FImplicitObject* LeafGeometry;
+		FRigidTransform3 GetWorldTransform() const
+		{
+			return FRigidTransform3(WorldPosition, WorldRotation);
+		}
+
+		void SetWorldTransform(const FRigidTransform3& LeafWorldTransform)
+		{
+			WorldPosition = LeafWorldTransform.GetTranslation();
+			WorldRotation = LeafWorldTransform.GetRotation();
+		}
+
+		FVec3 WorldPosition;
+		FRotation3 WorldRotation;
 	};
 
 	/**
@@ -370,6 +408,7 @@ namespace Chaos
 		return Ar;
 	}
 
+	UE_DEPRECATED(5.3, "Not for external use")
 	void CHAOS_API UpdateShapesArrayFromGeometry(FShapesArray& ShapesArray, TSerializablePtr<FImplicitObject> Geometry, const FRigidTransform3& ActorTM, IPhysicsProxyBase* Proxy);
 
 	FORCEINLINE uint32 GetTypeHash(const FParticleID& Unique)
@@ -969,10 +1008,7 @@ public:
 		TArrayCollectionArray<bool> MEnabledDuringResim;
 		TArrayCollectionArray<bool> MLightWeightDisabled;
 
-		void UpdateShapesArray(const int32 Index)
-		{
-			UpdateShapesArrayFromGeometry(MShapesArray[Index], MGeometry[Index], FRigidTransform3(X(Index), R(Index)), nullptr);
-		}
+		CHAOS_API void UpdateShapesArray(const int32 Index);
 
 		template <typename T2, int d2, EGeometryParticlesSimType SimType2>
 		friend class TGeometryParticlesImp;
diff --git a/Engine/Source/Runtime/Experimental/Chaos/Public/Chaos/ParticleHandle.h b/Engine/Source/Runtime/Experimental/Chaos/Public/Chaos/ParticleHandle.h
--- a/Engine/Source/Runtime/Experimental/Chaos/Public/Chaos/ParticleHandle.h
+++ b/Engine/Source/Runtime/Experimental/Chaos/Public/Chaos/ParticleHandle.h
@@ -2677,10 +2677,7 @@ protected:
 
 	void MarkDirty(const EChaosPropertyFlags DirtyBits, bool bInvalidate = true);
 
-	void UpdateShapesArray()
-	{
-		UpdateShapesArrayFromGeometry(MShapesArray, MakeSerializable(MNonFrequentData.Read().Geometry()), FRigidTransform3(X(), R()), Proxy);
-	}
+	CHAOS_API void UpdateShapesArray();
 
 	virtual void SyncRemoteDataImp(FDirtyPropertiesManager& Manager, int32 DataIdx, const FDirtyChaosProperties& RemoteData) const
 	{
