Updated the static mesh property("RootProxy") in UGeometryCollection to be a FGeometryCollectionProxyMeshData which currently contains an array of static meshes.
Deprecated the old RootProxy property in UGeometryCollection. Old data will still work.

diff --git a/Engine/Source/Runtime/Experimental/GeometryCollectionEngine/Private/GeometryCollection/GeometryCollectionComponent.cpp b/Engine/Source/Runtime/Experimental/GeometryCollectionEngine/Private/GeometryCollection/GeometryCollectionComponent.cpp
--- a/Engine/Source/Runtime/Experimental/GeometryCollectionEngine/Private/GeometryCollection/GeometryCollectionComponent.cpp
+++ b/Engine/Source/Runtime/Experimental/GeometryCollectionEngine/Private/GeometryCollection/GeometryCollectionComponent.cpp
@@ -1000,24 +1000,30 @@ bool UGeometryCollectionComponent::DoCustomNavigableGeometryExport(FNavigableGeo
 
 	if (bUseRootProxyForNavigation)
 	{
-		if (const UStaticMesh* ProxyMesh = Cast<UStaticMesh>(RestCollection->RootProxy.TryLoad()))
+		bool bHasData = false;
+		for (int32 MeshIndex = 0; MeshIndex < RestCollection->RootProxyData.ProxyMeshes.Num(); MeshIndex++)
 		{
-			const FTransform& CompToWorld = GetComponentToWorld();
-			const FVector Scale3D = CompToWorld.GetScale3D();
-			if (!Scale3D.IsZero())
+			const TObjectPtr<UStaticMesh>& ProxyMesh = RestCollection->RootProxyData.ProxyMeshes[MeshIndex];
+			if (ProxyMesh != nullptr)
 			{
-				if (const UNavCollisionBase* NavCollision = ProxyMesh->GetNavCollision())
+				const FTransform& CompToWorld = GetComponentToWorld();
+				const FVector Scale3D = CompToWorld.GetScale3D();
+				if (!Scale3D.IsZero())
 				{
-					const bool bHasData = NavCollision->ExportGeometry(CompToWorld, GeomExport);
-					if (bHasData)
+					if (const UNavCollisionBase* NavCollision = ProxyMesh->GetNavCollision())
 					{
-						// skip default export
-						return false;
+						bHasData = NavCollision->ExportGeometry(CompToWorld, GeomExport) || bHasData;
 					}
 				}
 			}
 		}
 
+		if (bHasData)
+		{
+			// skip default export
+			return false;
+		}
+
 		return true;
 	}
 
@@ -4501,16 +4507,17 @@ void UGeometryCollectionComponent::RegisterToISMPool()
 
 				// root proxy if available 
 				// TODO : if ISM pool is not available : uses a standard static mesh component
-				if (UObject* RootMeshProxyObject = RestCollection->RootProxy.TryLoad())
+				for (int32 MeshIndex = 0; MeshIndex < RestCollection->RootProxyData.ProxyMeshes.Num(); MeshIndex++)
 				{
-					if (UStaticMesh* RootMeshProxy = Cast<UStaticMesh>(RootMeshProxyObject))
+					const TObjectPtr<UStaticMesh>& Mesh = RestCollection->RootProxyData.ProxyMeshes[MeshIndex];
+					if (Mesh != nullptr)
 					{
 						// if we use a mesh proxy hide the component for rendering 
 						bCanRenderComponent = false;
 
 						FGeometryCollectionStaticMeshInstance StaticMeshInstance;
-						StaticMeshInstance.StaticMesh = RootMeshProxy;
-						ISMPoolRootProxyMeshId = ISMPoolComp->AddMeshToGroup(ISMPoolMeshGroupIndex, StaticMeshInstance, 1, bChaos_GC_UseHierarchicalISMForProxyMesh);
+						StaticMeshInstance.StaticMesh = Mesh;
+						ISMPoolRootProxyMeshIds.Add(ISMPoolComp->AddMeshToGroup(ISMPoolMeshGroupIndex, StaticMeshInstance, 1, bChaos_GC_UseHierarchicalISMForProxyMesh));
 					}
 				}
 			}
@@ -4530,7 +4537,7 @@ void UGeometryCollectionComponent::UnregisterFromISMPool()
 		{
 			ISMPoolComp->DestroyMeshGroup(ISMPoolMeshGroupIndex);
 			ISMPoolMeshGroupIndex = INDEX_NONE;
-			ISMPoolRootProxyMeshId = INDEX_NONE;
+			ISMPoolRootProxyMeshIds.Empty();
 		}
 		SetVisibility(true);
 	}
@@ -4568,14 +4575,17 @@ void UGeometryCollectionComponent::RefreshISMPoolInstances()
 							const int32 RootIndex = GetRootIndex();
 							//const bool bIsBroken = DynamicCollection ? (DynamicCollection->Children[RootIndex].Num() != Children[RootIndex].Num()) : false;
 							const bool bIsBroken = DynamicCollection ? !DynamicCollection->Active[RootIndex] : false;
-							const bool bHasRootProxyMesh = (ISMPoolRootProxyMeshId != INDEX_NONE);
+							const bool bHasRootProxyMesh = !ISMPoolRootProxyMeshIds.IsEmpty();
 
 							if (bHasRootProxyMesh && !bIsBroken)
 							{
 								if (GlobalMatrices.IsValidIndex(RootIndex))
 								{
 									FTransform RootTransform = FTransform(GlobalMatrices[RootIndex]) * ComponentTransform;
-									ISMPoolComp->BatchUpdateInstancesTransforms(ISMPoolMeshGroupIndex, ISMPoolRootProxyMeshId, 0, { RootTransform }, bWorlSpace, bMarkRenderStateDirty, bTeleport);
+									for (const int32 ISMPoolRootProxyMeshId : ISMPoolRootProxyMeshIds)
+									{
+										ISMPoolComp->BatchUpdateInstancesTransforms(ISMPoolMeshGroupIndex, ISMPoolRootProxyMeshId, 0, { RootTransform }, bWorlSpace, bMarkRenderStateDirty, bTeleport);
+									}
 								}
 							}
 							else if (bChaos_GC_UseISMPoolForNonFracturedParts)
@@ -4586,7 +4596,10 @@ void UGeometryCollectionComponent::RefreshISMPoolInstances()
 								{
 									FTransform RootTransformZeroScale;
 									RootTransformZeroScale.SetIdentityZeroScale();
-									ISMPoolComp->BatchUpdateInstancesTransforms(ISMPoolMeshGroupIndex, ISMPoolRootProxyMeshId, 0, { RootTransformZeroScale }, bWorlSpace, bMarkRenderStateDirty, bTeleport);
+									for (const int32 ISMPoolRootProxyMeshId : ISMPoolRootProxyMeshIds)
+									{
+										ISMPoolComp->BatchUpdateInstancesTransforms(ISMPoolMeshGroupIndex, ISMPoolRootProxyMeshId, 0, { RootTransformZeroScale }, bWorlSpace, bMarkRenderStateDirty, bTeleport);
+									}
 								}
 
 								TArray<FTransform> InstanceTransforms;
diff --git a/Engine/Source/Runtime/Experimental/GeometryCollectionEngine/Private/GeometryCollection/GeometryCollectionObject.cpp b/Engine/Source/Runtime/Experimental/GeometryCollectionEngine/Private/GeometryCollection/GeometryCollectionObject.cpp
--- a/Engine/Source/Runtime/Experimental/GeometryCollectionEngine/Private/GeometryCollection/GeometryCollectionObject.cpp
+++ b/Engine/Source/Runtime/Experimental/GeometryCollectionEngine/Private/GeometryCollection/GeometryCollectionObject.cpp
@@ -777,10 +777,8 @@ void UGeometryCollection::Serialize(FArchive& Ar)
 			// do we need to remove the simplicial attribute ? 
 			if (false == FGeometryCollection::AreCollisionParticlesEnabled())
 			{
-				ArchiveGeometryCollection = TSharedPtr<FGeometryCollection, ESPMode::ThreadSafe>(new FGeometryCollection);
-				const TArray<FName> NoGroupsToSkip;
-				const TArray<TTuple<FName, FName>> AttributesToSkip{ { FGeometryDynamicCollection::SimplicialsAttribute, FTransformCollection::TransformGroup } };
-				GeometryCollection->CopyTo(ArchiveGeometryCollection.Get(), NoGroupsToSkip, AttributesToSkip);
+				ArchiveGeometryCollection = TSharedPtr<FGeometryCollection, ESPMode::ThreadSafe>(GeometryCollection->NewCopy<FGeometryCollection>());
+				ArchiveGeometryCollection->RemoveAttribute(FGeometryDynamicCollection::SimplicialsAttribute, FTransformCollection::TransformGroup);
 			}
 		}
 
@@ -1525,6 +1523,17 @@ void UGeometryCollection::PostLoad()
 	{
 		InitResources();
 	}
+
+#if WITH_EDITORONLY_DATA
+	if (!RootProxy_DEPRECATED.IsNull())
+	{
+		if (UStaticMesh* ProxyMesh = Cast<UStaticMesh>(RootProxy_DEPRECATED.TryLoad()))
+		{
+			RootProxyData.ProxyMeshes.Add(TObjectPtr<UStaticMesh>(ProxyMesh));
+		}
+		RootProxy_DEPRECATED = nullptr;
+	}
+#endif
 }
 
 void UGeometryCollection::BeginDestroy()
diff --git a/Engine/Source/Runtime/Experimental/GeometryCollectionEngine/Public/GeometryCollection/GeometryCollectionObject.h b/Engine/Source/Runtime/Experimental/GeometryCollectionEngine/Public/GeometryCollection/GeometryCollectionObject.h
--- a/Engine/Source/Runtime/Experimental/GeometryCollectionEngine/Public/GeometryCollection/GeometryCollectionObject.h
+++ b/Engine/Source/Runtime/Experimental/GeometryCollectionEngine/Public/GeometryCollection/GeometryCollectionObject.h
@@ -330,6 +330,14 @@ private:
 	bool bIsInitialized = false;
 };
 
+USTRUCT(BlueprintType)
+struct GEOMETRYCOLLECTIONENGINE_API FGeometryCollectionProxyMeshData
+{
+	GENERATED_BODY()
+
+	UPROPERTY(EditAnywhere, BlueprintReadWrite, Category = "Rendering", meta = (AllowedClasses = "/Script/Engine.StaticMesh"))
+	TArray<TObjectPtr<UStaticMesh>> ProxyMeshes;
+};
 
 /**
 * UGeometryCollectionObject (UObject)
@@ -539,8 +547,8 @@ public:
 	TArray<FGeometryCollectionAutoInstanceMesh> AutoInstanceMeshes;
 
 	/** static mesh to use as a proxy for rendering until the geometry collection is broken */
-	UPROPERTY(BlueprintReadWrite, EditAnywhere, Category = "Rendering", meta = (AllowedClasses = "/Script/Engine.StaticMesh"))
-	FSoftObjectPath RootProxy;
+	UPROPERTY(BlueprintReadWrite, EditAnywhere, Category = "Rendering")
+	FGeometryCollectionProxyMeshData RootProxyData;
 
 	/**
 	 * Strip unnecessary data from the Geometry Collection to keep the memory footprint as small as possible.
@@ -599,6 +607,10 @@ public:
 	*/
 	UPROPERTY(meta = (DeprecatedProperty, DeprecationMessage = "This property is deprecated. Use the default SizeSpecificData instead."))
 	float CollisionObjectReductionPercentage_DEPRECATED;
+
+	/** static mesh to use as a proxy for rendering until the geometry collection is broken */
+	UPROPERTY(meta = (DeprecatedProperty, DeprecationMessage = "This property is deprecated. Use RootProxyData instead."))
+	FSoftObjectPath RootProxy_DEPRECATED;
 #endif
 	
 	/**
