Added missing file from GC change.

diff --git a/Engine/Source/Runtime/Experimental/GeometryCollectionEngine/Public/GeometryCollection/GeometryCollectionComponent.h b/Engine/Source/Runtime/Experimental/GeometryCollectionEngine/Public/GeometryCollection/GeometryCollectionComponent.h
--- a/Engine/Source/Runtime/Experimental/GeometryCollectionEngine/Public/GeometryCollection/GeometryCollectionComponent.h
+++ b/Engine/Source/Runtime/Experimental/GeometryCollectionEngine/Public/GeometryCollection/GeometryCollectionComponent.h
@@ -1063,7 +1063,7 @@ protected:
 
 	TObjectPtr<AGeometryCollectionISMPoolActor> AssignedISMPool = nullptr;
 	int32 ISMPoolMeshGroupIndex = INDEX_NONE;
-	int32 ISMPoolRootProxyMeshId = INDEX_NONE;
+	TArray<int32> ISMPoolRootProxyMeshIds;
 
 	/** Populate the static geometry structures for the render thread. */
 	void InitConstantData(TUniquePtr<FGeometryCollectionConstantData>& ConstantData) const;
