Geometry collection : make sure Resttransform are in sync with RestCollection to avoid scrambled assets in previews

diff --git a/Engine/Source/Runtime/Experimental/GeometryCollectionEngine/Private/GeometryCollection/GeometryCollectionComponent.cpp b/Engine/Source/Runtime/Experimental/GeometryCollectionEngine/Private/GeometryCollection/GeometryCollectionComponent.cpp
--- a/Engine/Source/Runtime/Experimental/GeometryCollectionEngine/Private/GeometryCollection/GeometryCollectionComponent.cpp
+++ b/Engine/Source/Runtime/Experimental/GeometryCollectionEngine/Private/GeometryCollection/GeometryCollectionComponent.cpp
@@ -2726,6 +2726,13 @@ void UGeometryCollectionComponent::ResetDynamicCollection()
 		SetRenderStateDirty();
 	}
 
+	// make sure we have the RestTransforms up to date, other wise, otherwise there may be case where they do not match the Restcollection ones
+	// can happen if the RestCollection asset has been changed without the component knowing about it 
+	if (RestCollection && RestCollection->GetGeometryCollection())
+	{
+		RestTransforms = RestCollection->GetGeometryCollection()->Transform.GetConstArray();
+	}
+
 	if (RestTransforms.Num() > 0)
 	{
 		SetInitialTransforms(RestTransforms);
@@ -3276,12 +3283,7 @@ void UGeometryCollectionComponent::SetRestCollection(const UGeometryCollection*
 	{
 		RestCollection = RestCollectionIn;
 
-		const int32 NumTransforms = RestCollection->GetGeometryCollection()->NumElements(FGeometryCollection::TransformGroup);
-		RestTransforms.SetNum(NumTransforms);
-		for (int32 Idx = 0; Idx < NumTransforms; ++Idx)
-		{
-			RestTransforms[Idx] = RestCollection->GetGeometryCollection()->Transform[Idx];
-		}
+		RestTransforms = RestCollection->GetGeometryCollection()->Transform.GetConstArray();
 
 		CalculateGlobalMatrices();
 		CalculateLocalBounds();
