[Backout] - CL24248607
Original CL Desc
-----------------------------------------------------------------
Updated the static mesh property("RootProxy") in UGeometryCollection to be a FGeometryCollectionProxyMeshData which currently contains an array of static meshes.
Deprecated the old RootProxy property in UGeometryCollection. Old data will still work.

diff --git a/Engine/Source/Runtime/Experimental/GeometryCollectionEngine/Private/GeometryCollection/GeometryCollectionComponent.cpp b/Engine/Source/Runtime/Experimental/GeometryCollectionEngine/Private/GeometryCollection/GeometryCollectionComponent.cpp
--- a/Engine/Source/Runtime/Experimental/GeometryCollectionEngine/Private/GeometryCollection/GeometryCollectionComponent.cpp
+++ b/Engine/Source/Runtime/Experimental/GeometryCollectionEngine/Private/GeometryCollection/GeometryCollectionComponent.cpp
@@ -1000,30 +1000,24 @@ bool UGeometryCollectionComponent::DoCustomNavigableGeometryExport(FNavigableGeo
 
 	if (bUseRootProxyForNavigation)
 	{
-		bool bHasData = false;
-		for (int32 MeshIndex = 0; MeshIndex < RestCollection->RootProxyData.ProxyMeshes.Num(); MeshIndex++)
+		if (const UStaticMesh* ProxyMesh = Cast<UStaticMesh>(RestCollection->RootProxy.TryLoad()))
 		{
-			const TObjectPtr<UStaticMesh>& ProxyMesh = RestCollection->RootProxyData.ProxyMeshes[MeshIndex];
-			if (ProxyMesh != nullptr)
+			const FTransform& CompToWorld = GetComponentToWorld();
+			const FVector Scale3D = CompToWorld.GetScale3D();
+			if (!Scale3D.IsZero())
 			{
-				const FTransform& CompToWorld = GetComponentToWorld();
-				const FVector Scale3D = CompToWorld.GetScale3D();
-				if (!Scale3D.IsZero())
+				if (const UNavCollisionBase* NavCollision = ProxyMesh->GetNavCollision())
 				{
-					if (const UNavCollisionBase* NavCollision = ProxyMesh->GetNavCollision())
+					const bool bHasData = NavCollision->ExportGeometry(CompToWorld, GeomExport);
+					if (bHasData)
 					{
-						bHasData = NavCollision->ExportGeometry(CompToWorld, GeomExport) || bHasData;
+						// skip default export
+						return false;
 					}
 				}
 			}
 		}
 
-		if (bHasData)
-		{
-			// skip default export
-			return false;
-		}
-
 		return true;
 	}
 
@@ -4509,17 +4503,16 @@ void UGeometryCollectionComponent::RegisterToISMPool()
 
 				// root proxy if available 
 				// TODO : if ISM pool is not available : uses a standard static mesh component
-				for (int32 MeshIndex = 0; MeshIndex < RestCollection->RootProxyData.ProxyMeshes.Num(); MeshIndex++)
+				if (UObject* RootMeshProxyObject = RestCollection->RootProxy.TryLoad())
 				{
-					const TObjectPtr<UStaticMesh>& Mesh = RestCollection->RootProxyData.ProxyMeshes[MeshIndex];
-					if (Mesh != nullptr)
+					if (UStaticMesh* RootMeshProxy = Cast<UStaticMesh>(RootMeshProxyObject))
 					{
 						// if we use a mesh proxy hide the component for rendering 
 						bCanRenderComponent = false;
 
 						FGeometryCollectionStaticMeshInstance StaticMeshInstance;
-						StaticMeshInstance.StaticMesh = Mesh;
-						ISMPoolRootProxyMeshIds.Add(ISMPoolComp->AddMeshToGroup(ISMPoolMeshGroupIndex, StaticMeshInstance, 1, bChaos_GC_UseHierarchicalISMForProxyMesh));
+						StaticMeshInstance.StaticMesh = RootMeshProxy;
+						ISMPoolRootProxyMeshId = ISMPoolComp->AddMeshToGroup(ISMPoolMeshGroupIndex, StaticMeshInstance, 1, bChaos_GC_UseHierarchicalISMForProxyMesh);
 					}
 				}
 			}
@@ -4539,7 +4532,7 @@ void UGeometryCollectionComponent::UnregisterFromISMPool()
 		{
 			ISMPoolComp->DestroyMeshGroup(ISMPoolMeshGroupIndex);
 			ISMPoolMeshGroupIndex = INDEX_NONE;
-			ISMPoolRootProxyMeshIds.Empty();
+			ISMPoolRootProxyMeshId = INDEX_NONE;
 		}
 		SetVisibility(true);
 	}
@@ -4577,17 +4570,14 @@ void UGeometryCollectionComponent::RefreshISMPoolInstances()
 							const int32 RootIndex = GetRootIndex();
 							//const bool bIsBroken = DynamicCollection ? (DynamicCollection->Children[RootIndex].Num() != Children[RootIndex].Num()) : false;
 							const bool bIsBroken = DynamicCollection ? !DynamicCollection->Active[RootIndex] : false;
-							const bool bHasRootProxyMesh = !ISMPoolRootProxyMeshIds.IsEmpty();
+							const bool bHasRootProxyMesh = (ISMPoolRootProxyMeshId != INDEX_NONE);
 
 							if (bHasRootProxyMesh && !bIsBroken)
 							{
 								if (GlobalMatrices.IsValidIndex(RootIndex))
 								{
 									FTransform RootTransform = FTransform(GlobalMatrices[RootIndex]) * ComponentTransform;
-									for (const int32 ISMPoolRootProxyMeshId : ISMPoolRootProxyMeshIds)
-									{
-										ISMPoolComp->BatchUpdateInstancesTransforms(ISMPoolMeshGroupIndex, ISMPoolRootProxyMeshId, 0, { RootTransform }, bWorlSpace, bMarkRenderStateDirty, bTeleport);
-									}
+									ISMPoolComp->BatchUpdateInstancesTransforms(ISMPoolMeshGroupIndex, ISMPoolRootProxyMeshId, 0, { RootTransform }, bWorlSpace, bMarkRenderStateDirty, bTeleport);
 								}
 							}
 							else if (bChaos_GC_UseISMPoolForNonFracturedParts)
@@ -4598,10 +4588,7 @@ void UGeometryCollectionComponent::RefreshISMPoolInstances()
 								{
 									FTransform RootTransformZeroScale;
 									RootTransformZeroScale.SetIdentityZeroScale();
-									for (const int32 ISMPoolRootProxyMeshId : ISMPoolRootProxyMeshIds)
-									{
-										ISMPoolComp->BatchUpdateInstancesTransforms(ISMPoolMeshGroupIndex, ISMPoolRootProxyMeshId, 0, { RootTransformZeroScale }, bWorlSpace, bMarkRenderStateDirty, bTeleport);
-									}
+									ISMPoolComp->BatchUpdateInstancesTransforms(ISMPoolMeshGroupIndex, ISMPoolRootProxyMeshId, 0, { RootTransformZeroScale }, bWorlSpace, bMarkRenderStateDirty, bTeleport);
 								}
 
 								TArray<FTransform> InstanceTransforms;
diff --git a/Engine/Source/Runtime/Experimental/GeometryCollectionEngine/Private/GeometryCollection/GeometryCollectionObject.cpp b/Engine/Source/Runtime/Experimental/GeometryCollectionEngine/Private/GeometryCollection/GeometryCollectionObject.cpp
--- a/Engine/Source/Runtime/Experimental/GeometryCollectionEngine/Private/GeometryCollection/GeometryCollectionObject.cpp
+++ b/Engine/Source/Runtime/Experimental/GeometryCollectionEngine/Private/GeometryCollection/GeometryCollectionObject.cpp
@@ -777,8 +777,10 @@ void UGeometryCollection::Serialize(FArchive& Ar)
 			// do we need to remove the simplicial attribute ? 
 			if (false == FGeometryCollection::AreCollisionParticlesEnabled())
 			{
-				ArchiveGeometryCollection = TSharedPtr<FGeometryCollection, ESPMode::ThreadSafe>(GeometryCollection->NewCopy<FGeometryCollection>());
-				ArchiveGeometryCollection->RemoveAttribute(FGeometryDynamicCollection::SimplicialsAttribute, FTransformCollection::TransformGroup);
+				ArchiveGeometryCollection = TSharedPtr<FGeometryCollection, ESPMode::ThreadSafe>(new FGeometryCollection);
+				const TArray<FName> NoGroupsToSkip;
+				const TArray<TTuple<FName, FName>> AttributesToSkip{ { FGeometryDynamicCollection::SimplicialsAttribute, FTransformCollection::TransformGroup } };
+				GeometryCollection->CopyTo(ArchiveGeometryCollection.Get(), NoGroupsToSkip, AttributesToSkip);
 			}
 		}
 
@@ -1523,17 +1525,6 @@ void UGeometryCollection::PostLoad()
 	{
 		InitResources();
 	}
-
-#if WITH_EDITORONLY_DATA
-	if (!RootProxy_DEPRECATED.IsNull())
-	{
-		if (UStaticMesh* ProxyMesh = Cast<UStaticMesh>(RootProxy_DEPRECATED.TryLoad()))
-		{
-			RootProxyData.ProxyMeshes.Add(TObjectPtr<UStaticMesh>(ProxyMesh));
-		}
-		RootProxy_DEPRECATED = nullptr;
-	}
-#endif
 }
 
 void UGeometryCollection::BeginDestroy()
diff --git a/Engine/Source/Runtime/Experimental/GeometryCollectionEngine/Public/GeometryCollection/GeometryCollectionComponent.h b/Engine/Source/Runtime/Experimental/GeometryCollectionEngine/Public/GeometryCollection/GeometryCollectionComponent.h
--- a/Engine/Source/Runtime/Experimental/GeometryCollectionEngine/Public/GeometryCollection/GeometryCollectionComponent.h
+++ b/Engine/Source/Runtime/Experimental/GeometryCollectionEngine/Public/GeometryCollection/GeometryCollectionComponent.h
@@ -1063,7 +1063,7 @@ protected:
 
 	TObjectPtr<AGeometryCollectionISMPoolActor> AssignedISMPool = nullptr;
 	int32 ISMPoolMeshGroupIndex = INDEX_NONE;
-	TArray<int32> ISMPoolRootProxyMeshIds;
+	int32 ISMPoolRootProxyMeshId = INDEX_NONE;
 
 	/** Populate the static geometry structures for the render thread. */
 	void InitConstantData(TUniquePtr<FGeometryCollectionConstantData>& ConstantData) const;
diff --git a/Engine/Source/Runtime/Experimental/GeometryCollectionEngine/Public/GeometryCollection/GeometryCollectionObject.h b/Engine/Source/Runtime/Experimental/GeometryCollectionEngine/Public/GeometryCollection/GeometryCollectionObject.h
--- a/Engine/Source/Runtime/Experimental/GeometryCollectionEngine/Public/GeometryCollection/GeometryCollectionObject.h
+++ b/Engine/Source/Runtime/Experimental/GeometryCollectionEngine/Public/GeometryCollection/GeometryCollectionObject.h
@@ -330,14 +330,6 @@ private:
 	bool bIsInitialized = false;
 };
 
-USTRUCT(BlueprintType)
-struct GEOMETRYCOLLECTIONENGINE_API FGeometryCollectionProxyMeshData
-{
-	GENERATED_BODY()
-
-	UPROPERTY(EditAnywhere, BlueprintReadWrite, Category = "Rendering", meta = (AllowedClasses = "/Script/Engine.StaticMesh"))
-	TArray<TObjectPtr<UStaticMesh>> ProxyMeshes;
-};
 
 /**
 * UGeometryCollectionObject (UObject)
@@ -547,8 +539,8 @@ public:
 	TArray<FGeometryCollectionAutoInstanceMesh> AutoInstanceMeshes;
 
 	/** static mesh to use as a proxy for rendering until the geometry collection is broken */
-	UPROPERTY(BlueprintReadWrite, EditAnywhere, Category = "Rendering")
-	FGeometryCollectionProxyMeshData RootProxyData;
+	UPROPERTY(BlueprintReadWrite, EditAnywhere, Category = "Rendering", meta = (AllowedClasses = "/Script/Engine.StaticMesh"))
+	FSoftObjectPath RootProxy;
 
 	/**
 	 * Strip unnecessary data from the Geometry Collection to keep the memory footprint as small as possible.
@@ -607,10 +599,6 @@ public:
 	*/
 	UPROPERTY(meta = (DeprecatedProperty, DeprecationMessage = "This property is deprecated. Use the default SizeSpecificData instead."))
 	float CollisionObjectReductionPercentage_DEPRECATED;
-
-	/** static mesh to use as a proxy for rendering until the geometry collection is broken */
-	UPROPERTY(meta = (DeprecatedProperty, DeprecationMessage = "This property is deprecated. Use RootProxyData instead."))
-	FSoftObjectPath RootProxy_DEPRECATED;
 #endif
 	
 	/**
