Move all immutable geometry collection rendering data into the geometry collection asset.
This is built as part of geometry collection DDC contents.
Reworked scene proxy to use this data. Some big changes there:
* Tidied up hit proxy code so that all hit proxy clients use the same path.
* Removing of internal faces are now handled without an additional index buffer by always storing internal faces at the end of a section and providing different section descriptions for including internal faces.
* Dynamic geometry hiding in editor is now done by zeroing transforms instead of rebuilding the index buffer.

diff --git a/Engine/Plugins/Experimental/ChaosEditor/Source/FractureEditor/Private/FractureEditorMode.cpp b/Engine/Plugins/Experimental/ChaosEditor/Source/FractureEditor/Private/FractureEditorMode.cpp
--- a/Engine/Plugins/Experimental/ChaosEditor/Source/FractureEditor/Private/FractureEditorMode.cpp
+++ b/Engine/Plugins/Experimental/ChaosEditor/Source/FractureEditor/Private/FractureEditorMode.cpp
@@ -250,9 +250,9 @@ bool UFractureEditorMode::HandleClick(FEditorViewportClient* InViewportClient, H
 
 bool UFractureEditorMode::SelectFromClick(HHitProxy* HitProxy, bool bCtrlDown, bool bShiftDown)
 {
-	if (HitProxy && HitProxy->IsA(HGeometryCollectionBone::StaticGetType()))
+	if (HitProxy && HitProxy->IsA(HGeometryCollection::StaticGetType()))
 	{
-		HGeometryCollectionBone* GeometryCollectionProxy = (HGeometryCollectionBone*)HitProxy;
+		HGeometryCollection* GeometryCollectionProxy = (HGeometryCollection*)HitProxy;
 
 		if (GeometryCollectionProxy->Component)
 		{
diff --git a/Engine/Plugins/Experimental/ChaosEditor/Source/FractureEditor/Private/FractureToolEditing.cpp b/Engine/Plugins/Experimental/ChaosEditor/Source/FractureEditor/Private/FractureToolEditing.cpp
--- a/Engine/Plugins/Experimental/ChaosEditor/Source/FractureEditor/Private/FractureToolEditing.cpp
+++ b/Engine/Plugins/Experimental/ChaosEditor/Source/FractureEditor/Private/FractureToolEditing.cpp
@@ -278,8 +278,6 @@ void UFractureToolHide::Execute(TWeakPtr<FFractureEditorModeToolkit> InToolkit)
 				}
 			}
 
-			Context.GetGeometryCollectionComponent()->MarkRenderStateDirty();
-			Context.GetGeometryCollectionComponent()->MarkRenderDynamicDataDirty();
 			Refresh(Context, Toolkit, true);
 		}
 
@@ -348,8 +346,6 @@ void UFractureToolUnhide::Execute(TWeakPtr<FFractureEditorModeToolkit> InToolkit
 				}
 			}
 
-			Context.GetGeometryCollectionComponent()->MarkRenderStateDirty();
-			Context.GetGeometryCollectionComponent()->MarkRenderDynamicDataDirty();
 			Refresh(Context, Toolkit, true);
 		}
 
@@ -479,25 +475,9 @@ void UFractureToolValidate::Execute(TWeakPtr<FFractureEditorModeToolkit> InToolk
 						FGeometryCollectionClusteringUtility::UpdateHierarchyLevelOfChildren(GeometryCollection, -1);
 						AddSingleRootNodeIfRequired(GeometryCollectionObject);
 
-						// Update Nanite resource data to correctly reflect modified geometry collection data
-						{
-							GeometryCollectionObject->ReleaseResources();
-
-							if (GeometryCollectionObject->EnableNanite)
-							{
-								GeometryCollectionObject->NaniteData = UGeometryCollection::CreateNaniteData(GeometryCollection);
-							}
-							else
-							{
-								GeometryCollectionObject->NaniteData = MakeUnique<FGeometryCollectionNaniteData>();
-							}
-							GeometryCollectionObject->InitResources();
-						}
-
-						GeometryCollectionComponent->MarkRenderStateDirty();
-						GeometryCollectionObject->MarkPackageDirty();
+						// Update resource data to correctly reflect modified geometry collection data
+						GeometryCollectionObject->RebuildRenderData();
 					}
-
 				}
 			}
 
diff --git a/Engine/Plugins/Experimental/ChaosEditor/Source/FractureEditor/Private/FractureToolGenerators.cpp b/Engine/Plugins/Experimental/ChaosEditor/Source/FractureEditor/Private/FractureToolGenerators.cpp
--- a/Engine/Plugins/Experimental/ChaosEditor/Source/FractureEditor/Private/FractureToolGenerators.cpp
+++ b/Engine/Plugins/Experimental/ChaosEditor/Source/FractureEditor/Private/FractureToolGenerators.cpp
@@ -680,13 +680,8 @@ AGeometryCollectionActor* UFractureToolGenerateAsset::ConvertActorsToGeometryCol
 
 	AddSingleRootNodeIfRequired(FracturedGeometryCollection);
 
-	if (FracturedGeometryCollection->EnableNanite)
-	{
-		FracturedGeometryCollection->InvalidateCollection();
-		FracturedGeometryCollection->EnsureDataIsCooked(true /* init resources */);
-	}
-
-	NewActor->GetGeometryCollectionComponent()->MarkRenderStateDirty();
+	FracturedGeometryCollection->InvalidateCollection();
+	FracturedGeometryCollection->RebuildRenderData();
 
 	// Add and initialize guids
 	::GeometryCollection::GenerateTemporaryGuids(FracturedGeometryCollection->GetGeometryCollection().Get(), 0 , true);
diff --git a/Engine/Plugins/Experimental/ChaosEditor/Source/FractureEditor/Public/FractureTool.cpp b/Engine/Plugins/Experimental/ChaosEditor/Source/FractureEditor/Public/FractureTool.cpp
--- a/Engine/Plugins/Experimental/ChaosEditor/Source/FractureEditor/Public/FractureTool.cpp
+++ b/Engine/Plugins/Experimental/ChaosEditor/Source/FractureEditor/Public/FractureTool.cpp
@@ -10,6 +10,7 @@
 #include "SLevelViewport.h"
 #include "EditorModeManager.h"
 #include "FractureEditorMode.h"
+#include "Materials/MaterialInterface.h"
 #include "Modules/ModuleManager.h"
 
 #include "FractureToolContext.h"
@@ -245,8 +246,7 @@ void UFractureActionTool::Refresh(FFractureToolContext& Context, FFractureEditor
 		FFractureSelectionTools::ToggleSelectedBones(GeometryCollectionComponent, Context.GetSelection(), true, false);
 	}
 
-	GeometryCollectionComponent->MarkRenderDynamicDataDirty();
-	GeometryCollectionComponent->MarkRenderStateDirty();
+	Context.GetFracturedGeometryCollection()->RebuildRenderData();
 }
 
 void UFractureActionTool::SetOutlinerComponents(TArray<FFractureToolContext>& InContexts, FFractureEditorModeToolkit* Toolkit)
@@ -397,21 +397,6 @@ void UFractureModalTool::Execute(TWeakPtr<FFractureEditorModeToolkit> InToolkit)
 
 			FGeometryCollectionClusteringUtility::UpdateHierarchyLevelOfChildren(FractureContext.GetGeometryCollection().Get(), -1);
 
-			// Update Nanite resource data to correctly reflect modified geometry collection data
-			{
-				FractureContext.GetFracturedGeometryCollection()->ReleaseResources();
-
-				if (FractureContext.GetFracturedGeometryCollection()->EnableNanite)
-				{
-					FractureContext.GetFracturedGeometryCollection()->NaniteData = UGeometryCollection::CreateNaniteData(FractureContext.GetGeometryCollection().Get());
-				}
-				else
-				{
-					FractureContext.GetFracturedGeometryCollection()->NaniteData = MakeUnique<FGeometryCollectionNaniteData>();
-				}
-				FractureContext.GetFracturedGeometryCollection()->InitResources();
-			}
-
 			Refresh(FractureContext, Toolkit);
 		}
 
diff --git a/Engine/Plugins/Experimental/GeometryCollectionPlugin/Source/GeometryCollectionEditor/Private/GeometryCollection/GeometryCollectionSelectRigidBodyEdMode.cpp b/Engine/Plugins/Experimental/GeometryCollectionPlugin/Source/GeometryCollectionEditor/Private/GeometryCollection/GeometryCollectionSelectRigidBodyEdMode.cpp
--- a/Engine/Plugins/Experimental/GeometryCollectionPlugin/Source/GeometryCollectionEditor/Private/GeometryCollection/GeometryCollectionSelectRigidBodyEdMode.cpp
+++ b/Engine/Plugins/Experimental/GeometryCollectionPlugin/Source/GeometryCollectionEditor/Private/GeometryCollection/GeometryCollectionSelectRigidBodyEdMode.cpp
@@ -131,7 +131,7 @@ bool FGeometryCollectionSelectRigidBodyEdMode::MouseMove(FEditorViewportClient*
 	// Hover status
 	const HHitProxy* const HitProxy = Viewport->GetHitProxy(x, y);
 	bIsHoveringGeometryCollection = HitProxy && HitProxy->IsA(HGeometryCollection::StaticGetType());
-	UE_CLOG(bIsHoveringGeometryCollection, LogGeometryCollectionSelectRigidBodyEdMode, VeryVerbose, TEXT("Hovering GeometryCollectionActor %s, Transform Index %d."), *static_cast<const HGeometryCollection*>(HitProxy)->Actor->GetName(), static_cast<const HGeometryCollection*>(HitProxy)->TransformIndex);
+	UE_CLOG(bIsHoveringGeometryCollection, LogGeometryCollectionSelectRigidBodyEdMode, VeryVerbose, TEXT("Hovering GeometryCollectionActor %s, Transform Index %d."), *static_cast<const HGeometryCollection*>(HitProxy)->Component->GetOwner()->GetName(), static_cast<const HGeometryCollection*>(HitProxy)->BoneIndex);
 #endif // #if GEOMETRYCOLLECTION_EDITOR_SELECTION
 	return true;
 }
@@ -142,45 +142,41 @@ bool FGeometryCollectionSelectRigidBodyEdMode::HandleClick(FEditorViewportClient
 	if (HitProxy && HitProxy->IsA(HGeometryCollection::StaticGetType()))
 	{
 		const HGeometryCollection* const HitGeometryCollection = static_cast<const HGeometryCollection*>(HitProxy);
-		const AGeometryCollectionActor* const GeometryCollectionActor = Cast<AGeometryCollectionActor>(HitGeometryCollection->Actor);
-		if (GeometryCollectionActor)
+		const UGeometryCollectionComponent* const GeometryCollectionComponent = HitGeometryCollection->Component;
+		if (GeometryCollectionComponent)
 		{
-			const UGeometryCollectionComponent* const GeometryCollectionComponent = GeometryCollectionActor->GeometryCollectionComponent;
-			if (GeometryCollectionComponent)
+			// Retrieve the transform index
+			int32 TransformIndex = HitGeometryCollection->BoneIndex;
+
+			// Get to the parent transform if the hit transform is still attached to its parent
+			const TManagedArray<int32>& ParentArray = GeometryCollectionComponent->GetParentArray();
+			while (ParentArray[TransformIndex] != FGeometryCollectionBoneNode::InvalidBone)
+			{
+				TransformIndex = ParentArray[TransformIndex];
+			}
+			/*
+			// Retrieve the rigid body id
+			//const TManagedArray<int32>& RigidBodyIdArray = GeometryCollectionComponent->GetRigidBodyIdArray();
+			const TManagedArray<FGuid>& RigidBodyIdArray = GeometryCollectionComponent->GetRigidBodyGuidArray();
+			//const int32 RigidBodyId = (TransformIndex != INDEX_NONE && ensure(TransformIndex < RigidBodyIdArray.Num())) ? RigidBodyIdArray[TransformIndex]: INDEX_NONE;
+			const FGuid& RigidBodyId = (TransformIndex != INDEX_NONE && ensure(TransformIndex < RigidBodyIdArray.Num())) ? RigidBodyIdArray[TransformIndex] : FGuid();
+
+			// Update the rigid body id property
+			//if (RigidBodyId != INDEX_NONE)
+			if (RigidBodyId.IsValid())
 			{
-				// Retrieve the transform index
-				int32 TransformIndex = HitGeometryCollection->TransformIndex;
+				UE_LOG(LogGeometryCollectionSelectRigidBodyEdMode, Verbose, TEXT("Hit GeometryCollectionActor %s at rigid body %s."), *GeometryCollectionComponent->GetOwner()->GetName(), *RigidBodyId.ToString());
 
-				// Get to the parent transform if the hit transform is still attached to its parent
-				const TManagedArray<int32>& ParentArray = GeometryCollectionComponent->GetParentArray();
-				while (ParentArray[TransformIndex] != FGeometryCollectionBoneNode::InvalidBone)
+				if (const TSharedPtr<IPropertyHandle> PropertyHandleIdPin = PropertyHandleId.Pin())
 				{
-					TransformIndex = ParentArray[TransformIndex];
+					PropertyHandleIdPin->SetValue(RigidBodyId.ToString());
 				}
-				/*
-				// Retrieve the rigid body id
-				//const TManagedArray<int32>& RigidBodyIdArray = GeometryCollectionComponent->GetRigidBodyIdArray();
-				const TManagedArray<FGuid>& RigidBodyIdArray = GeometryCollectionComponent->GetRigidBodyGuidArray();
-				//const int32 RigidBodyId = (TransformIndex != INDEX_NONE && ensure(TransformIndex < RigidBodyIdArray.Num())) ? RigidBodyIdArray[TransformIndex]: INDEX_NONE;
-				const FGuid& RigidBodyId = (TransformIndex != INDEX_NONE && ensure(TransformIndex < RigidBodyIdArray.Num())) ? RigidBodyIdArray[TransformIndex] : FGuid();
-
-				// Update the rigid body id property
-				//if (RigidBodyId != INDEX_NONE)
-				if (RigidBodyId.IsValid())
+				if (const TSharedPtr<IPropertyHandle> PropertyHandleSolverPin = PropertyHandleSolver.Pin())
 				{
-					UE_LOG(LogGeometryCollectionSelectRigidBodyEdMode, Verbose, TEXT("Hit GeometryCollectionActor %s at rigid body %s."), *GeometryCollectionActor->GetName(), *RigidBodyId.ToString());
-
-					if (const TSharedPtr<IPropertyHandle> PropertyHandleIdPin = PropertyHandleId.Pin())
-					{
-						PropertyHandleIdPin->SetValue(RigidBodyId.ToString());
-					}
-					if (const TSharedPtr<IPropertyHandle> PropertyHandleSolverPin = PropertyHandleSolver.Pin())
-					{
-						PropertyHandleSolverPin->SetValue(GeometryCollectionComponent->ChaosSolverActor);
-					}
+					PropertyHandleSolverPin->SetValue(GeometryCollectionComponent->ChaosSolverActor);
 				}
-				*/
 			}
+			*/
 		}
 	}
 #endif // #if GEOMETRYCOLLECTION_EDITOR_SELECTION
diff --git a/Engine/Source/Runtime/Engine/Private/GeometryCollectionRendering.cpp b/Engine/Source/Runtime/Engine/Private/GeometryCollectionRendering.cpp
--- a/Engine/Source/Runtime/Engine/Private/GeometryCollectionRendering.cpp
+++ b/Engine/Source/Runtime/Engine/Private/GeometryCollectionRendering.cpp
@@ -120,7 +120,8 @@ void FGeometryCollectionVertexFactory::InitRHI()
 
 	// VertexFactory needs to be able to support max possible shader platform and feature level
 	// in case if we switch feature level at runtime.
-	const bool bCanUseGPUScene = UseGPUScene(GMaxRHIShaderPlatform, GMaxRHIFeatureLevel);
+	const bool bCanUseGPUScene = UseGPUScene(GMaxRHIShaderPlatform, GetFeatureLevel());
+	const bool bUseManualVertexFetch = SupportsManualVertexFetch(GetFeatureLevel());
 
 	// If the vertex buffer containing position is not the same vertex buffer containing the rest of the data,
 	// then initialize PositionStream and PositionDeclaration.
@@ -222,7 +223,7 @@ void FGeometryCollectionVertexFactory::InitRHI()
 	const int32 DefaultBaseVertexIndex = 0;
 	const int32 DefaultPreSkinBaseVertexIndex = 0;
 
-	if (RHISupportsManualVertexFetch(GMaxRHIShaderPlatform) || bCanUseGPUScene)
+	if (bUseManualVertexFetch || bCanUseGPUScene)
 	{
 		SCOPED_LOADTIMER(FGeometryCollectionVertexFactory_InitRHI_CreateLocalVFUniformBuffer);
 
@@ -231,7 +232,7 @@ void FGeometryCollectionVertexFactory::InitRHI()
 		UniformParameters.LODLightmapDataIndex = Data.LODLightmapDataIndex;
 		int32 ColorIndexMask = 0;
 
-		if (RHISupportsManualVertexFetch(GMaxRHIShaderPlatform))
+		if (bUseManualVertexFetch)
 		{
 			UniformParameters.VertexFetch_PositionBuffer = GetPositionsSRV();
 			UniformParameters.VertexFetch_PackedTangentsBuffer = GetTangentsSRV();
diff --git a/Engine/Source/Runtime/Experimental/GeometryCollectionEngine/Private/GeometryCollection/DerivedDataGeometryCollectionCooker.cpp b/Engine/Source/Runtime/Experimental/GeometryCollectionEngine/Private/GeometryCollection/DerivedDataGeometryCollectionCooker.cpp
--- a/Engine/Source/Runtime/Experimental/GeometryCollectionEngine/Private/GeometryCollection/DerivedDataGeometryCollectionCooker.cpp
+++ b/Engine/Source/Runtime/Experimental/GeometryCollectionEngine/Private/GeometryCollection/DerivedDataGeometryCollectionCooker.cpp
@@ -2,8 +2,9 @@
 
 #include "GeometryCollection/DerivedDataGeometryCollectionCooker.h"
 #include "Chaos/Core.h"
-#include "GeometryCollection/GeometryCollectionObject.h"
 #include "GeometryCollection/GeometryCollection.h"
+#include "GeometryCollection/GeometryCollectionObject.h"
+#include "GeometryCollection/GeometryCollectionRenderData.h"
 #include "Serialization/MemoryWriter.h"
 #include "PhysicsProxy/GeometryCollectionPhysicsProxy.h"
 #include "Chaos/ChaosArchive.h"
@@ -40,18 +41,67 @@ bool FDerivedDataGeometryCollectionCooker::Build(TArray<uint8>& OutData)
 		BuildSimulationData(ErrorReporter, *Collection, SharedParams);
 		Collection->Serialize(ChaosAr);
 
-		if (GeometryCollection.EnableNanite)
-		{
-			TUniquePtr<FGeometryCollectionNaniteData> NaniteData = UGeometryCollection::CreateNaniteData(Collection);
-			NaniteData->Serialize(ChaosAr, &GeometryCollection);
-		}
-		else
+		if (false && ErrorReporter.EncounteredAnyErrors())
 		{
-			// No Nanite data, write out zero resources
-			int32 NumNaniteResources = 0;
-			Ar << NumNaniteResources;
+			bool bAllErrorsHandled = !ErrorReporter.ContainsUnhandledError();
+			ErrorReporter.ReportError(*FString::Printf(TEXT("Could not cook content for Collection:%s"), *GeometryCollection.GetPathName()));
+			if (bAllErrorsHandled)
+			{
+				ErrorReporter.HandleLatestError();
+			}
+			return false;	//Don't save into DDC if any errors found
 		}
 
+		return true;
+	}
+
+	return false;
+}
+
+const TCHAR* FDerivedDataGeometryCollectionCooker::GetVersionString() const
+{
+	const TCHAR* VersionString = TEXT("C3F0D73CE6094A1C84E57D75C9131F55");
+	return VersionString;
+}
+
+FString FDerivedDataGeometryCollectionCooker::GetPluginSpecificCacheKeySuffix() const
+{
+	FString KeySuffix = FString::Printf(
+		TEXT("%s_%s_%s_%d_%d"),
+		Chaos::ChaosVersionGUID,
+		*GeometryCollection.GetIdGuid().ToString(),
+		*GeometryCollection.GetStateGuid().ToString(),
+		FDestructionObjectVersion::Type::LatestVersion,
+		FUE5MainStreamObjectVersion::Type::LatestVersion
+	);
+
+#if PLATFORM_CPU_ARM_FAMILY
+	// Separate out arm keys as x64 and arm64 clang do not generate the same data for a given
+	// input. Add the arm specifically so that a) we avoid rebuilding the current DDC and
+	// b) we can remove it once we get arm64 to be consistent.
+	KeySuffix.Append(TEXT("_arm64"));
+#endif
+
+	return KeySuffix;
+}
+
+
+FString FDerivedDataGeometryCollectionRenderDataCooker::GetDebugContextString() const
+{
+	return GeometryCollection.GetFullName();
+}
+
+bool FDerivedDataGeometryCollectionRenderDataCooker::Build(TArray<uint8>& OutData)
+{
+	FMemoryWriter Ar(OutData, true);	// Must be persistent for BulkData to serialize
+	Chaos::FChaosArchive ChaosAr(Ar);
+	if (FGeometryCollection* Collection = GeometryCollection.GetGeometryCollection().Get())
+	{
+		Chaos::FErrorReporter ErrorReporter(GeometryCollection.GetName());
+
+		TUniquePtr<FGeometryCollectionRenderData> RenderData = FGeometryCollectionRenderData::Create(*Collection, GeometryCollection.EnableNanite, GeometryCollection.bUseFullPrecisionUVs);
+		RenderData->Serialize(ChaosAr, &GeometryCollection);
+
 		if (false && ErrorReporter.EncounteredAnyErrors())
 		{
 			bool bAllErrorsHandled = !ErrorReporter.ContainsUnhandledError();
@@ -69,21 +119,20 @@ bool FDerivedDataGeometryCollectionCooker::Build(TArray<uint8>& OutData)
 	return false;
 }
 
-const TCHAR* FDerivedDataGeometryCollectionCooker::GetVersionString() const
+const TCHAR* FDerivedDataGeometryCollectionRenderDataCooker::GetVersionString() const
 {
-	const TCHAR* VersionString = TEXT("007038F125ED47A4AF928FAC54FD9493");
+	const TCHAR* VersionString = TEXT("C3F0D73CE6094A1C84E57D75C9131F55");
 
 	static FString CachedNaniteVersionString;
 	if (CachedNaniteVersionString.IsEmpty())
 	{
-		const TCHAR* NaniteVersionString = TEXT("CE9B42F81A5XX388E51GGNC8C19D3E");
-		CachedNaniteVersionString = FString::Printf(TEXT("%s_%s_%s"), VersionString, NaniteVersionString, *Nanite::IBuilderModule::Get().GetVersionString());
+		CachedNaniteVersionString = FString::Printf(TEXT("%s_%s"), VersionString, *Nanite::IBuilderModule::Get().GetVersionString());
 	}
 
 	return GeometryCollection.EnableNanite ? *CachedNaniteVersionString : VersionString;
 }
 
-FString FDerivedDataGeometryCollectionCooker::GetPluginSpecificCacheKeySuffix() const
+FString FDerivedDataGeometryCollectionRenderDataCooker::GetPluginSpecificCacheKeySuffix() const
 {
 	FString KeySuffix = FString::Printf(
 		TEXT("%s_%s_%s_%d_%d"),
diff --git a/Engine/Source/Runtime/Experimental/GeometryCollectionEngine/Private/GeometryCollection/GeometryCollectionComponent.cpp b/Engine/Source/Runtime/Experimental/GeometryCollectionEngine/Private/GeometryCollection/GeometryCollectionComponent.cpp
--- a/Engine/Source/Runtime/Experimental/GeometryCollectionEngine/Private/GeometryCollection/GeometryCollectionComponent.cpp
+++ b/Engine/Source/Runtime/Experimental/GeometryCollectionEngine/Private/GeometryCollection/GeometryCollectionComponent.cpp
@@ -673,103 +673,28 @@ FPrimitiveSceneProxy* UGeometryCollectionComponent::CreateSceneProxy()
 	{
 		if (UseNanite(GetScene()->GetShaderPlatform()) &&
 			RestCollection->EnableNanite &&
-			RestCollection->NaniteData != nullptr &&
+			RestCollection->HasNaniteData() &&
 			GGeometryCollectionNanite != 0)
 		{
-			LocalSceneProxy = new FNaniteGeometryCollectionSceneProxy(this);
+			FNaniteGeometryCollectionSceneProxy* NaniteProxy = new FNaniteGeometryCollectionSceneProxy(this);
+			LocalSceneProxy = NaniteProxy;
 
 			// ForceMotionBlur means we maintain bIsMoving, regardless of actual state.
 			if (bForceMotionBlur)
 			{
 				bIsMoving = true;
-				if (LocalSceneProxy)
-				{
-					FNaniteGeometryCollectionSceneProxy* NaniteProxy = static_cast<FNaniteGeometryCollectionSceneProxy*>(LocalSceneProxy);
-					ENQUEUE_RENDER_COMMAND(NaniteProxyOnMotionEnd)(
-						[NaniteProxy](FRHICommandListImmediate& RHICmdList)
-						{
-							NaniteProxy->OnMotionBegin();
-						}
-					);
-				}
-			}
-		}
-		// If we didn't get a proxy, but Nanite was enabled on the asset when it was built, evaluate proxy creation
-		else if (RestCollection->EnableNanite && NaniteProxyRenderMode != 0)
-		{
-			// Do not render Nanite proxy
-			return nullptr;
-		}
-		else
-		{
-			LocalSceneProxy = new FGeometryCollectionSceneProxy(this);
-		}
-
-		if (RestCollection->HasVisibleGeometry())
-		{
-			TUniquePtr<FGeometryCollectionConstantData> ConstantData = MakeUnique<FGeometryCollectionConstantData>();
-			InitConstantData(ConstantData);
-
-			FGeometryCollectionDynamicData* const DynamicData = InitDynamicData(true /* initialization */);
-
-			if (LocalSceneProxy->IsNaniteMesh())
-			{
-				FNaniteGeometryCollectionSceneProxy* const GeometryCollectionSceneProxy = static_cast<FNaniteGeometryCollectionSceneProxy*>(LocalSceneProxy);
-
-				// ...
-
-			#if GEOMETRYCOLLECTION_EDITOR_SELECTION
-				if (bIsTransformSelectionModeEnabled)
-				{
-					// ...
-				}
-			#endif
-
-				ENQUEUE_RENDER_COMMAND(CreateRenderState)(
-					[GeometryCollectionSceneProxy, ConstantData = MoveTemp(ConstantData), DynamicData](FRHICommandListImmediate& RHICmdList)
-					{
-						GeometryCollectionSceneProxy->SetConstantData_RenderThread(ConstantData);
-
-						if (DynamicData)
-						{
-							GeometryCollectionSceneProxy->SetDynamicData_RenderThread(DynamicData);
-						}
-
-						bool bValidUpdate = false;
-						if (FPrimitiveSceneInfo* PrimitiveSceneInfo = GeometryCollectionSceneProxy->GetPrimitiveSceneInfo())
-						{
-							bValidUpdate = PrimitiveSceneInfo->RequestGPUSceneUpdate();
-						}
-
-						// Deferred the GPU Scene update if the primitive scene info is not yet initialized with a valid index.
-						GeometryCollectionSceneProxy->SetRequiresGPUSceneUpdate_RenderThread(!bValidUpdate);
-					}
-				);
-			}
-			else
-			{
-				FGeometryCollectionSceneProxy* const GeometryCollectionSceneProxy = static_cast<FGeometryCollectionSceneProxy*>(LocalSceneProxy);
-
-			#if GEOMETRYCOLLECTION_EDITOR_SELECTION
-				// Re-init subsections
-				if (bIsTransformSelectionModeEnabled)
-				{
-					GeometryCollectionSceneProxy->UseSubSections(true, false);  // Do not force reinit now, it'll be done in SetConstantData_RenderThread
-				}
-			#endif
-
-				ENQUEUE_RENDER_COMMAND(CreateRenderState)(
-					[GeometryCollectionSceneProxy, ConstantData = MoveTemp(ConstantData), DynamicData](FRHICommandListImmediate& RHICmdList)
+				ENQUEUE_RENDER_COMMAND(NaniteProxyOnMotionEnd)(
+					[NaniteProxy](FRHICommandListImmediate& RHICmdList)
 					{
-						GeometryCollectionSceneProxy->SetConstantData_RenderThread(ConstantData);
-						if (DynamicData)
-						{
-							GeometryCollectionSceneProxy->SetDynamicData_RenderThread(DynamicData);
-						}
+						NaniteProxy->OnMotionBegin();
 					}
 				);
 			}
 		}
+		else if (RestCollection->HasMeshData())
+		{
+			LocalSceneProxy = new FGeometryCollectionSceneProxy(this);
+		}
 	}
 
 	return LocalSceneProxy;
@@ -1344,6 +1269,39 @@ void UGeometryCollectionComponent::InitializeComponent()
 	}
 }
 
+void UGeometryCollectionComponent::GetResourceSizeEx(FResourceSizeEx& CumulativeResourceSize)
+{
+	Super::GetResourceSizeEx(CumulativeResourceSize);
+	
+	int32 SizeBytes =
+		InitializationFields.GetAllocatedSize()
+		+ DamageThreshold.GetAllocatedSize()
+		+ RestTransforms.GetAllocatedSize()
+		+ DisabledFlags.GetAllocatedSize()
+		+ CollisionProfilePerLevel.GetAllocatedSize()
+		+ GlobalMatrices.GetAllocatedSize()
+		+ EventsPlayed.GetAllocatedSize()
+		+ CopyOnWriteAttributeList.GetAllocatedSize()
+		+ EmbeddedGeometryComponents.GetAllocatedSize()
+		+ (ClustersToRep ? ClustersToRep->GetAllocatedSize() : 0);
+
+#if WITH_EDITORONLY_DATA
+	SizeBytes +=
+		+ SelectedBones.GetAllocatedSize()
+		+ HighlightedBones.GetAllocatedSize()
+		+ EmbeddedInstanceIndex.GetAllocatedSize()
+		+ EmbeddedBoneMaps.GetAllocatedSize();
+
+	for (TArray<int32>& BoneMap : EmbeddedBoneMaps)
+	{
+		SizeBytes += BoneMap.GetAllocatedSize();
+	}
+#endif
+
+	CumulativeResourceSize.AddDedicatedSystemMemoryBytes(SizeBytes);
+}
+
+
 #if WITH_EDITOR
 FDelegateHandle UGeometryCollectionComponent::RegisterOnGeometryCollectionPropertyChanged(const FOnGeometryCollectionPropertyChanged& Delegate)
 {
@@ -2058,237 +2016,67 @@ void UGeometryCollectionComponent::SetInitialClusterBreaks(const TArray<int32>&
 	}
 }
 
+#if WITH_EDITOR
 
-void UGeometryCollectionComponent::InitConstantData(TUniquePtr<FGeometryCollectionConstantData>& ConstantData) const
+void UGeometryCollectionComponent::GetBoneColors(TArray<FColor>& OutColors) const
 {
-	// Constant data should all be moved to the DDC as time permits.
-	// todo : this should be computed once per asset not per component or at least the part that does not depend on the component properties
-
-	check(ConstantData);
-	check(RestCollection);
 	const FGeometryCollection* Collection = RestCollection->GetGeometryCollection().Get();
-	check(Collection);
+	const int32 NumPoints = Collection->NumElements(FGeometryCollection::VerticesGroup);
+	const TManagedArray<int32>& BoneMap = Collection->BoneMap;
+	const TManagedArray<FLinearColor>& BoneColors = Collection->BoneColor;
 
-	if (!RestCollection->EnableNanite)
+	OutColors.SetNumUninitialized(NumPoints);
+	if (bChaos_GC_InitConstantDataUseParallelFor)
 	{
-		const int32 NumPoints = Collection->NumElements(FGeometryCollection::VerticesGroup);
-		const TManagedArray<FVector3f>& Vertex = Collection->Vertex;
-		const TManagedArray<int32>& BoneMap = Collection->BoneMap;
-		const TManagedArray<FVector3f>& TangentU = Collection->TangentU;
-		const TManagedArray<FVector3f>& TangentV = Collection->TangentV;
-		const TManagedArray<FVector3f>& Normal = Collection->Normal;
-		const TManagedArray<FLinearColor>& Color = Collection->Color;
-		const TManagedArray<FLinearColor>& BoneColors = Collection->BoneColor;
-		
-		const int32 NumGeom = Collection->NumElements(FGeometryCollection::GeometryGroup);
-		const TManagedArray<int32>& TransformIndex = Collection->TransformIndex;
-		const TManagedArray<int32>& FaceStart = Collection->FaceStart;
-		const TManagedArray<int32>& FaceCount = Collection->FaceCount;
-
-		ConstantData->Vertices = TArray<FVector3f>(Vertex.GetData(), Vertex.Num());
-		ConstantData->BoneMap = TArray<int32>(BoneMap.GetData(), BoneMap.Num());
-		ConstantData->TangentU = TArray<FVector3f>(TangentU.GetData(), TangentU.Num());
-		ConstantData->TangentV = TArray<FVector3f>(TangentV.GetData(), TangentV.Num());
-		ConstantData->Normals = TArray<FVector3f>(Normal.GetData(), Normal.Num());
-		ConstantData->Colors = TArray<FLinearColor>(Color.GetData(), Color.Num());
-
-		int32 NumUVChannels = Collection->NumUVLayers();
-		GeometryCollection::UV::FConstUVLayers CollectionUVs = GeometryCollection::UV::FindActiveUVLayers(*Collection);
-		ConstantData->UVChannels.SetNum(NumUVChannels);
-		for (int32 UVChannelIndex = 0; UVChannelIndex < NumUVChannels; UVChannelIndex++)
-		{
-			ConstantData->UVChannels[UVChannelIndex].AddUninitialized(NumPoints);
-		}
-
-		ConstantData->BoneColors.AddUninitialized(NumPoints);
-
-		if (bChaos_GC_InitConstantDataUseParallelFor)
-		{
-			ParallelFor(TEXT("GC:InitConstantData"), NumPoints, bChaos_GC_InitConstantDataParallelForBatchSize,
-				[&](const int32 InPointIndex)
-				{
-					const int32 BoneIndex = ConstantData->BoneMap[InPointIndex];
-					ConstantData->BoneColors[InPointIndex] = BoneColors[BoneIndex];
-
-					for (int32 UVChannelIndex = 0; UVChannelIndex < NumUVChannels; UVChannelIndex++)
-					{
-						ConstantData->UVChannels[UVChannelIndex][InPointIndex] = CollectionUVs[UVChannelIndex][InPointIndex];
-					}
-				});
-		}
-		else
-		{
-			for (int32 PointIndex = 0; PointIndex < NumPoints; PointIndex++)
-			{
-				const int32 BoneIndex = ConstantData->BoneMap[PointIndex];
-				ConstantData->BoneColors[PointIndex] = BoneColors[BoneIndex];
-			}
-
-			for (int32 UVChannelIndex = 0; UVChannelIndex < NumUVChannels; UVChannelIndex++)
+		ParallelFor(TEXT("GC:InitBoneColors"), NumPoints, bChaos_GC_InitConstantDataParallelForBatchSize,
+			[&](const int32 InPointIndex)
 			{
-				for (int32 PointIndex = 0; PointIndex < NumPoints; PointIndex++)
-				{
-					ConstantData->UVChannels[UVChannelIndex][PointIndex] = CollectionUVs[UVChannelIndex][PointIndex];
-				}
-			}
-		}
-
-		int32 NumIndices = 0;
-		const TManagedArray<FIntVector>& Indices = Collection->Indices;
-		const TManagedArray<int32>& MaterialID = Collection->MaterialID;
-		const TManagedArray<bool>& Internal = Collection->Internal;
-
-		const TManagedArray<bool>& Visible = GetVisibleArray();  // Use copy on write attribute. The rest collection visible array can be overriden for the convenience of debug drawing the collision volumes
-		
-#if WITH_EDITOR
-		// We will override visibility with the Hide array (if available).
-		TArray<bool> VisibleOverride;
-		VisibleOverride.SetNumUninitialized(Visible.Num());
-		for (int32 FaceIdx = 0; FaceIdx < Visible.Num(); FaceIdx++)
-		{
-			VisibleOverride[FaceIdx] = Visible[FaceIdx];
-		}
-		bool bUsingHideArray = false;
-
-		if (Collection->HasAttribute("Hide", FGeometryCollection::TransformGroup))
-		{
-			bUsingHideArray = true;
-
-			bool bAllHidden = true;
-
-			const TManagedArray<bool>& Hide = Collection->GetAttribute<bool>("Hide", FGeometryCollection::TransformGroup);
-			for (int32 GeomIdx = 0; GeomIdx < NumGeom; ++GeomIdx)
-			{
-				int32 TransformIdx = TransformIndex[GeomIdx];
-				if (Hide[TransformIdx])
-				{
-					// (Temporarily) hide faces of this hidden geometry
-					for (int32 FaceIdxOffset = 0; FaceIdxOffset < FaceCount[GeomIdx]; ++FaceIdxOffset)
-					{
-						VisibleOverride[FaceStart[GeomIdx]+FaceIdxOffset] = false;
-					}
-				}
-				else if (bAllHidden && Collection->IsVisible(TransformIdx))
-				{
-					bAllHidden = false;
-				}
-			}
-			if (!ensure(!bAllHidden)) // if they're all hidden, rendering would crash -- reset to default visibility instead
-			{
-				for (int32 FaceIdx = 0; FaceIdx < VisibleOverride.Num(); ++FaceIdx)
-				{
-					VisibleOverride[FaceIdx] = Visible[FaceIdx];
-				}
-			}
-		}
-#endif
-		
-		const TManagedArray<int32>& MaterialIndex = Collection->MaterialIndex;
-
-		const int32 NumFaceGroupEntries = Collection->NumElements(FGeometryCollection::FacesGroup);
-
-		for (int FaceIndex = 0; FaceIndex < NumFaceGroupEntries; ++FaceIndex)
-		{
-#if WITH_EDITOR
-			NumIndices += bUsingHideArray ? static_cast<int>(VisibleOverride[FaceIndex]) : static_cast<int>(Visible[FaceIndex]);
-#else
-			NumIndices += static_cast<int>(Visible[FaceIndex]);
-#endif
-		}
-
-		ConstantData->Indices.AddUninitialized(NumIndices);
-		for (int IndexIdx = 0, cdx = 0; IndexIdx < NumFaceGroupEntries; ++IndexIdx)
-		{
-#if WITH_EDITOR
-			const bool bUseVisible = bUsingHideArray ? VisibleOverride[MaterialIndex[IndexIdx]] : Visible[MaterialIndex[IndexIdx]];
-			if (bUseVisible)
-#else
-			if (Visible[MaterialIndex[IndexIdx]])
-#endif
-			{
-				ConstantData->Indices[cdx++] = Indices[MaterialIndex[IndexIdx]];
-			}
-		}
-
-		// We need to correct the section index start point & number of triangles since only the visible ones have been copied across in the code above
-		const int32 NumMaterialSections = Collection->NumElements(FGeometryCollection::MaterialGroup);
-		ConstantData->Sections.AddUninitialized(NumMaterialSections);
-		const TManagedArray<FGeometryCollectionSection>& Sections = Collection->Sections;
-		for (int SectionIndex = 0; SectionIndex < NumMaterialSections; ++SectionIndex)
+				const int32 BoneIndex = BoneMap[InPointIndex];
+				OutColors[InPointIndex] = BoneColors[BoneIndex].ToFColor(true);
+			});
+	}
+	else
+	{
+		for (int32 PointIndex = 0; PointIndex < NumPoints; PointIndex++)
 		{
-			FGeometryCollectionSection Section = Sections[SectionIndex]; // deliberate copy
-
-			for (int32 TriangleIndex = 0; TriangleIndex < Sections[SectionIndex].FirstIndex / 3; TriangleIndex++)
-			{
-#if WITH_EDITOR
-				const bool bUseVisible = bUsingHideArray ? VisibleOverride[MaterialIndex[TriangleIndex]] : Visible[MaterialIndex[TriangleIndex]];
-				if (!bUseVisible)
-#else
-				if (!Visible[MaterialIndex[TriangleIndex]])
-#endif
-				{
-					Section.FirstIndex -= 3;
-				}
-			}
-
-			for (int32 TriangleIndex = 0; TriangleIndex < Sections[SectionIndex].NumTriangles; TriangleIndex++)
-			{
-				int32 FaceIdx = MaterialIndex[Sections[SectionIndex].FirstIndex / 3 + TriangleIndex];
-#if WITH_EDITOR
-				const bool bUseVisible = bUsingHideArray ? VisibleOverride[FaceIdx] : Visible[FaceIdx];
-				if (!bUseVisible)
-#else
-				if (!Visible[FaceIdx])
-#endif
-				{
-					Section.NumTriangles--;
-				}
-			}
-
-			ConstantData->Sections[SectionIndex] = MoveTemp(Section);
+			const int32 BoneIndex = BoneMap[PointIndex];
+			OutColors[PointIndex] = BoneColors[BoneIndex].ToFColor(true);
 		}
-		
-		
-		ConstantData->NumTransforms = Collection->NumElements(FGeometryCollection::TransformGroup);
-		ConstantData->LocalBounds = LocalBounds;
+	}
+}
 
-		// store the index buffer and render sections for the base unfractured mesh
-		const TManagedArray<int32>& TransformToGeometryIndex = Collection->TransformToGeometryIndex;
-		
-		const int32 NumFaces = Collection->NumElements(FGeometryCollection::FacesGroup);
-		TArray<FIntVector> BaseMeshIndices;
-		TArray<int32> BaseMeshOriginalFaceIndices;
+void UGeometryCollectionComponent::GetHiddenTransforms(TArray<bool>& OutHiddenTransforms) const
+{
+	const FGeometryCollection* Collection = RestCollection->GetGeometryCollection().Get();
+	check(Collection);
 
-		BaseMeshIndices.Reserve(NumFaces);
-		BaseMeshOriginalFaceIndices.Reserve(NumFaces);
+	OutHiddenTransforms.Reset();
+	if (Collection->HasAttribute("Hide", FGeometryCollection::TransformGroup))
+	{
+		const TManagedArray<bool>& Hide = Collection->GetAttribute<bool>("Hide", FGeometryCollection::TransformGroup);
+		const TManagedArray<int32>& TransformIndices = Collection->TransformIndex;
+		const int32 NumGeom = Collection->NumElements(FGeometryCollection::GeometryGroup);
+		const int32 NumTransforms = Collection->Transform.Num();
 
-		// add all visible external faces to the original geometry index array
-		// #note:  This is a stopgap because the original geometry array is broken
-		for (int FaceIndex = 0; FaceIndex < NumFaces; ++FaceIndex)
+		OutHiddenTransforms.SetNumZeroed(NumTransforms);
+		for (int32 GeometryIndex = 0; GeometryIndex < NumGeom; ++GeometryIndex)
 		{
-			// only add visible external faces.
-#if WITH_EDITOR
-			const bool bUseVisible = bUsingHideArray ? VisibleOverride[FaceIndex] : Visible[FaceIndex];
-			if (bUseVisible && (Internal[FaceIndex] == false || bUsingHideArray))
-#else
-			if (Visible[FaceIndex] && Internal[FaceIndex] == false)
-#endif
+			const int32 TransformIndex = TransformIndices[GeometryIndex];
+			if (Hide[TransformIndex])
 			{
-				BaseMeshIndices.Add(Indices[FaceIndex]);
-				BaseMeshOriginalFaceIndices.Add(FaceIndex);
+				OutHiddenTransforms[TransformIndex] = 1;
 			}
 		}
+	}
+}
 
-		// We should always have external faces of a geometry collection
-		ensure(BaseMeshIndices.Num() > 0);
+#endif // WITH_EDITOR
 
-		ConstantData->OriginalMeshSections = Collection->BuildMeshSections(BaseMeshIndices, BaseMeshOriginalFaceIndices, ConstantData->OriginalMeshIndices);
-	}
-	
+void UGeometryCollectionComponent::GetRestTransforms(TArray<FMatrix44f>& OutRestTransforms) const
+{
 	TArray<FMatrix> RestMatrices;
 	GeometryCollectionAlgo::GlobalMatrices(RestCollection->GetGeometryCollection()->Transform, RestCollection->GetGeometryCollection()->Parent, RestMatrices);
-
-	ConstantData->SetRestTransforms(RestMatrices);
+	CopyTransformsWithConversionWhenNeeded(OutRestTransforms, RestMatrices);
 }
 
 FGeometryCollectionDynamicData* UGeometryCollectionComponent::InitDynamicData(bool bInitialization)
@@ -3294,7 +3082,7 @@ void UGeometryCollectionComponent::SetRestCollection(const UGeometryCollection*
 		// initialize the component damage progataion data from the asset defaults 
 		DamagePropagationData = RestCollection->DamagePropagationData;
 		
-		//ResetDynamicCollection();
+		ResetDynamicCollection();
 	}
 }
 
@@ -4333,11 +4121,8 @@ void UGeometryCollectionComponent::SwitchRenderModels(const AActor* Actor)
 void UGeometryCollectionComponent::EnableTransformSelectionMode(bool bEnable)
 {
 	// TODO: Support for Nanite?
-	if (SceneProxy && !SceneProxy->IsNaniteMesh() && RestCollection && RestCollection->HasVisibleGeometry())
-	{
-		static_cast<FGeometryCollectionSceneProxy*>(SceneProxy)->UseSubSections(bEnable, true);
-	}
 	bIsTransformSelectionModeEnabled = bEnable;
+	MarkRenderStateDirty();
 }
 #endif  // #if GEOMETRYCOLLECTION_EDITOR_SELECTION
 
diff --git a/Engine/Source/Runtime/Experimental/GeometryCollectionEngine/Private/GeometryCollection/GeometryCollectionHitProxy.cpp b/Engine/Source/Runtime/Experimental/GeometryCollectionEngine/Private/GeometryCollection/GeometryCollectionHitProxy.cpp
--- a/Engine/Source/Runtime/Experimental/GeometryCollectionEngine/Private/GeometryCollection/GeometryCollectionHitProxy.cpp
+++ b/Engine/Source/Runtime/Experimental/GeometryCollectionEngine/Private/GeometryCollection/GeometryCollectionHitProxy.cpp
@@ -2,7 +2,6 @@
 
 #include "GeometryCollection/GeometryCollectionHitProxy.h"
 
-#if GEOMETRYCOLLECTION_EDITOR_SELECTION
-IMPLEMENT_HIT_PROXY(HGeometryCollection, HActor);
-IMPLEMENT_HIT_PROXY(HGeometryCollectionBone, HHitProxy);
-#endif  // #if GEOMETRYCOLLECTION_EDITOR_SELECTION
+#if WITH_EDITOR
+IMPLEMENT_HIT_PROXY(HGeometryCollection, HHitProxy);
+#endif 
diff --git a/Engine/Source/Runtime/Experimental/GeometryCollectionEngine/Private/GeometryCollection/GeometryCollectionObject.cpp b/Engine/Source/Runtime/Experimental/GeometryCollectionEngine/Private/GeometryCollection/GeometryCollectionObject.cpp
--- a/Engine/Source/Runtime/Experimental/GeometryCollectionEngine/Private/GeometryCollection/GeometryCollectionObject.cpp
+++ b/Engine/Source/Runtime/Experimental/GeometryCollectionEngine/Private/GeometryCollection/GeometryCollectionObject.cpp
@@ -6,6 +6,7 @@
 #include "GeometryCollection/GeometryCollectionObject.h"
 #include "GeometryCollection/GeometryCollection.h"
 #include "GeometryCollection/GeometryCollectionCache.h"
+#include "GeometryCollection/GeometryCollectionRenderData.h"
 #include "Materials/Material.h"
 #include "UObject/DestructionObjectVersion.h"
 #include "UObject/UE5MainStreamObjectVersion.h"
@@ -19,6 +20,7 @@
 #include "Engine/StaticMesh.h"
 #include "PhysicsEngine/PhysicsSettings.h"
 #include "EditorFramework/AssetImportData.h"
+#include "Rendering/NaniteResources.h"
 
 #if WITH_EDITOR
 #include "GeometryCollection/DerivedDataGeometryCollectionCooker.h"
@@ -114,6 +116,7 @@ UGeometryCollection::UGeometryCollection(const FObjectInitializer& ObjectInitial
 	InvalidateCollection();
 #if WITH_EDITOR
 	SimulationDataGuid = StateGuid;
+	RenderDataGuid = StateGuid;
 	bStripOnCook = GeometryCollectionAssetForceStripOnCook;
 #endif
 }
@@ -643,7 +646,7 @@ bool UGeometryCollection::HasVisibleGeometry() const
 {
 	if(ensureMsgf(GeometryCollection.IsValid(), TEXT("Geometry Collection %s has an invalid internal collection")))
 	{
-		return ( (EnableNanite && NaniteData) || GeometryCollection->HasVisibleGeometry());
+		return ( (EnableNanite && RenderData && RenderData->bHasNaniteData) || GeometryCollection->HasVisibleGeometry());
 	}
 
 	return false;
@@ -761,7 +764,9 @@ void UGeometryCollection::Serialize(FArchive& Ar)
 
 		if (bStripOnCook)
 		{
-			if (EnableNanite && NaniteData)
+			// TODO: Since non-nanite path now stores mesh data in cooked build we may be able to unify 
+			// the simplification of the Geometry Collection for both nanite and non-nanite cases.
+			if (EnableNanite && HasNaniteData())
 			{
 				// If this is a cooked archive, we strip unnecessary data from the Geometry Collection to keep the memory footprint as small as possible.
 				ArchiveGeometryCollection = GenerateMinimalGeometryCollection();
@@ -859,7 +864,7 @@ void UGeometryCollection::Serialize(FArchive& Ar)
 #if WITH_EDITOR
 		if (Ar.IsSaving() && !Ar.IsTransacting())
 		{
-			CreateSimulationDataImp(/*bCopyFromDDC=*/false);	//make sure content is built before saving
+			EnsureDataIsCooked(false /*bInitResourcs*/, Ar.IsTransacting(), Ar.IsPersistent());
 		}
 #endif
 		if (Ar.IsLoading())
@@ -913,12 +918,12 @@ void UGeometryCollection::Serialize(FArchive& Ar)
 		Ar << bCooked;
 		if (bCooked)
 		{
-			if (NaniteData == nullptr)
+			if (RenderData == nullptr)
 			{
-				NaniteData = MakeUnique<FGeometryCollectionNaniteData>();
+				RenderData = MakeUnique<FGeometryCollectionRenderData>();
 			}
 
-			NaniteData->Serialize(ChaosAr, this);
+			RenderData->Serialize(ChaosAr, this);
 		}
 	}
 
@@ -977,10 +982,10 @@ void UGeometryCollection::Serialize(FArchive& Ar)
 		CreateSimulationData();
 	}
 
-	//for all versions loaded, make sure sim data is up to date
+	//for all versions loaded, make sure loaded content is built
  	if (Ar.IsLoading())
 	{
-		EnsureDataIsCooked(true, Ar.IsTransacting());	//make sure loaded content is built
+		EnsureDataIsCooked(true /*bInitResourcs*/, Ar.IsTransacting(), Ar.IsPersistent());
 	}
 #endif
 }
@@ -995,16 +1000,9 @@ void UGeometryCollection::SetEnableNanite(bool bValue)
 	if (EnableNanite != bValue)
 	{
 		EnableNanite = bValue;
-		NaniteData = MakeUnique<FGeometryCollectionNaniteData>();
 
 #if WITH_EDITOR
-		if (EnableNanite)
-		{
-			if (GeometryCollection)
-			{
-				NaniteData = UGeometryCollection::CreateNaniteData(GeometryCollection.Get());
-			}
-		}
+		RebuildRenderData();
 #endif
 	}
 }
@@ -1041,9 +1039,6 @@ void UGeometryCollection::CreateSimulationDataImp(bool bCopyFromDDC)
 			FMemoryReader Ar(DDCData, true);	// Must be persistent for BulkData to serialize
 			Chaos::FChaosArchive ChaosAr(Ar);
 			GeometryCollection->Serialize(ChaosAr);
-
-			NaniteData = MakeUnique<FGeometryCollectionNaniteData>();
-			NaniteData->Serialize(ChaosAr, this);
 		}
 	}
 }
@@ -1062,136 +1057,65 @@ void UGeometryCollection::CreateSimulationDataIfNeeded()
 	}
 }
 
-TUniquePtr<FGeometryCollectionNaniteData> UGeometryCollection::CreateNaniteData(FGeometryCollection* Collection)
-{
-	TUniquePtr<FGeometryCollectionNaniteData> NaniteData;
-
-	TRACE_CPUPROFILER_EVENT_SCOPE(UGeometryCollection::CreateNaniteData);
-
-	Nanite::IBuilderModule& NaniteBuilderModule = Nanite::IBuilderModule::Get();
-
-	NaniteData = MakeUnique<FGeometryCollectionNaniteData>();
-
-	// Transform Group
-	const TManagedArray<int32>& TransformToGeometryIndexArray = Collection->TransformToGeometryIndex;
-	const TManagedArray<int32>& SimulationTypeArray = Collection->SimulationType;
-	const TManagedArray<int32>& StatusFlagsArray = Collection->StatusFlags;
-
-	// Vertices Group
-	const TManagedArray<FVector3f>& VertexArray = Collection->Vertex;
-	GeometryCollection::UV::FUVLayers UVsLayers = GeometryCollection::UV::FindActiveUVLayers(*Collection);
-	const TManagedArray<FLinearColor>& ColorArray = Collection->Color;
-	const TManagedArray<FVector3f>& TangentUArray = Collection->TangentU;
-	const TManagedArray<FVector3f>& TangentVArray = Collection->TangentV;
-	const TManagedArray<FVector3f>& NormalArray = Collection->Normal;
-	const TManagedArray<int32>& BoneMapArray = Collection->BoneMap;
-
-	// Faces Group
-	const TManagedArray<FIntVector>& IndicesArray = Collection->Indices;
-	const TManagedArray<bool>& VisibleArray = Collection->Visible;
-	const TManagedArray<int32>& MaterialIndexArray = Collection->MaterialIndex;
-	const TManagedArray<int32>& MaterialIDArray = Collection->MaterialID;
-
-	// Geometry Group
-	const TManagedArray<int32>& TransformIndexArray = Collection->TransformIndex;
-	const TManagedArray<FBox>& BoundingBoxArray = Collection->BoundingBox;
-	const TManagedArray<float>& InnerRadiusArray = Collection->InnerRadius;
-	const TManagedArray<float>& OuterRadiusArray = Collection->OuterRadius;
-	const TManagedArray<int32>& VertexStartArray = Collection->VertexStart;
-	const TManagedArray<int32>& VertexCountArray = Collection->VertexCount;
-	const TManagedArray<int32>& FaceStartArray = Collection->FaceStart;
-	const TManagedArray<int32>& FaceCountArray = Collection->FaceCount;
-
-	// Material Group
-	const int32 NumGeometry = Collection->NumElements(FGeometryCollection::GeometryGroup);
 
-	const uint32 NumTexCoords = Collection->NumUVLayers();
-	const bool bHasColors = ColorArray.Num() > 0;
+void UGeometryCollection::CreateRenderDataImp(bool bCopyFromDDC)
+{
+	COOK_STAT(auto Timer = GeometryCollectionCookStats::UsageStats.TimeSyncWork());
 
-	TArray<FStaticMeshBuildVertex> BuildVertices;
-	TArray<uint32> BuildIndices;
-	TArray<int32> MaterialIndices;
+	// Skips the DDC fetch entirely for testing the builder without adding to the DDC
+	const static bool bSkipDDC = false;
 
-	TArray<uint32> MeshTriangleCounts;
-	MeshTriangleCounts.SetNum(NumGeometry);
+	//Use the DDC to build simulation data. If we are loading in the editor we then serialize this data into the geometry collection
+	TArray<uint8> DDCData;
+	FDerivedDataGeometryCollectionRenderDataCooker* GeometryCollectionCooker = new FDerivedDataGeometryCollectionRenderDataCooker(*this);
 
-	for (int32 GeometryGroupIndex = 0; GeometryGroupIndex < NumGeometry; GeometryGroupIndex++)
+	if (GeometryCollectionCooker->CanBuild())
 	{
-		const int32 VertexStart = VertexStartArray[GeometryGroupIndex];
-		const int32 VertexCount = VertexCountArray[GeometryGroupIndex];
-
-		uint32 DestVertexStart = BuildVertices.Num();
-		BuildVertices.Reserve(DestVertexStart + VertexCount);
-		for (int32 VertexIndex = 0; VertexIndex < VertexCount; ++VertexIndex)
+		if (bSkipDDC)
 		{
-			FStaticMeshBuildVertex& Vertex = BuildVertices.Emplace_GetRef();
-			Vertex.Position = VertexArray[VertexStart + VertexIndex];
-			Vertex.Color = bHasColors ? ColorArray[VertexStart + VertexIndex].ToFColor(false /* sRGB */) : FColor::White;
-			Vertex.TangentX = FVector3f::ZeroVector;
-			Vertex.TangentY = FVector3f::ZeroVector;
-			Vertex.TangentZ = NormalArray[VertexStart + VertexIndex];
-			for (int32 UVIdx = 0; UVIdx < UVsLayers.Num(); ++UVIdx)
-			{
-				Vertex.UVs[UVIdx] = UVsLayers[UVIdx][VertexStart + VertexIndex];
-				if (Vertex.UVs[UVIdx].ContainsNaN())
-				{
-					Vertex.UVs[UVIdx] = FVector2f::ZeroVector;
-				}
-			}
+			GeometryCollectionCooker->Build(DDCData);
+			COOK_STAT(Timer.AddMiss(DDCData.Num()));
 		}
-
-		const int32 FaceStart = FaceStartArray[GeometryGroupIndex];
-		const int32 FaceCount = FaceCountArray[GeometryGroupIndex];
-
-		// TODO: Respect multiple materials like in FGeometryCollectionConversion::AppendStaticMesh
-
-		int32 DestFaceStart = MaterialIndices.Num();
-		MaterialIndices.Reserve(DestFaceStart + FaceCount);
-		BuildIndices.Reserve((DestFaceStart + FaceCount) * 3);
-		for (int32 FaceIndex = 0; FaceIndex < FaceCount; ++FaceIndex)
+		else
 		{
-			if (!VisibleArray[FaceStart + FaceIndex]) // TODO: Always in range?
-			{
-				continue;
-			}
-
-			FIntVector FaceIndices = IndicesArray[FaceStart + FaceIndex];
-			FaceIndices = FaceIndices + FIntVector( DestVertexStart - VertexStart );
-
-			// Remove degenerates
-			if( BuildVertices[ FaceIndices[0] ].Position == BuildVertices[ FaceIndices[1] ].Position ||
-				BuildVertices[ FaceIndices[1] ].Position == BuildVertices[ FaceIndices[2] ].Position ||
-				BuildVertices[ FaceIndices[2] ].Position == BuildVertices[ FaceIndices[0] ].Position )
-			{
-				continue;
-			}
+			bool bBuilt = false;
+			const bool bSuccess = GetDerivedDataCacheRef().GetSynchronous(GeometryCollectionCooker, DDCData, &bBuilt);
+			COOK_STAT(Timer.AddHitOrMiss(!bSuccess || bBuilt ? FCookStats::CallStats::EHitOrMiss::Miss : FCookStats::CallStats::EHitOrMiss::Hit, DDCData.Num()));
+		}
 
-			BuildIndices.Add(FaceIndices.X);
-			BuildIndices.Add(FaceIndices.Y);
-			BuildIndices.Add(FaceIndices.Z);
+		if (bCopyFromDDC)
+		{
+			FMemoryReader Ar(DDCData, true);	// Must be persistent for BulkData to serialize
+			Chaos::FChaosArchive ChaosAr(Ar);
 
-			const int32 MaterialIndex = MaterialIDArray[FaceStart + FaceIndex];
-			MaterialIndices.Add(MaterialIndex);
+			RenderData = MakeUnique<FGeometryCollectionRenderData>();
+			RenderData->Serialize(ChaosAr, this);
 		}
-
-		MeshTriangleCounts[GeometryGroupIndex] = MaterialIndices.Num() - DestFaceStart;
 	}
+}
 
-	FMeshNaniteSettings NaniteSettings = {};
-	NaniteSettings.bEnabled = true;
-	NaniteSettings.TargetMinimumResidencyInKB = 0;	// Default to smallest possible, which is a single page
-	NaniteSettings.KeepPercentTriangles = 1.0f;
-	NaniteSettings.TrimRelativeError = 0.0f;
-	NaniteSettings.FallbackPercentTriangles = 1.0f; // 100% - no reduction
-	NaniteSettings.FallbackRelativeError = 0.0f;
-
-	NaniteData->NaniteResource = {};
-	if (!NaniteBuilderModule.Build(NaniteData->NaniteResource, BuildVertices, BuildIndices, MaterialIndices, MeshTriangleCounts, NumTexCoords, NaniteSettings))
+void UGeometryCollection::RebuildRenderData()
+{
+	if (RenderDataGuid != StateGuid)
 	{
-		UE_LOG(LogStaticMesh, Error, TEXT("Failed to build Nanite for geometry collection. See previous line(s) for details."));
+		ReleaseResources();
+		RenderData = FGeometryCollectionRenderData::Create(*GetGeometryCollection(), EnableNanite, bUseFullPrecisionUVs);
+		InitResources();
+		PropagateMarkDirtyToComponents();
+		RenderDataGuid = StateGuid;
 	}
+}
 
-	return NaniteData;
+void UGeometryCollection::PropagateMarkDirtyToComponents() const
+{
+	for (TObjectIterator<UGeometryCollectionComponent> It(RF_ClassDefaultObject, false, EInternalObjectFlags::Garbage); It; ++It)
+	{
+		if (It->RestCollection == this)
+		{
+			It->MarkRenderStateDirty();
+			It->MarkRenderDynamicDataDirty();
+		}
+	}
 }
 
 TSharedPtr<FGeometryCollection, ESPMode::ThreadSafe> UGeometryCollection::GenerateMinimalGeometryCollection() const
@@ -1318,17 +1242,17 @@ TSharedPtr<FGeometryCollection, ESPMode::ThreadSafe> UGeometryCollection::CopyCo
 
 void UGeometryCollection::InitResources()
 {
-	if (NaniteData)
+	if (RenderData)
 	{
-		NaniteData->InitResources(this);
+		RenderData->InitResources(this);
 	}
 }
 
 void UGeometryCollection::ReleaseResources()
 {
-	if (NaniteData)
+	if (RenderData)
 	{
-		NaniteData->ReleaseResources();
+		RenderData->ReleaseResources();
 	}
 }
 
@@ -1410,24 +1334,25 @@ FGuid UGeometryCollection::GetStateGuid() const
 
 void UGeometryCollection::PostEditChangeProperty(struct FPropertyChangedEvent& PropertyChangedEvent)
 {
+	bool bDoInvalidateCollection = false;
+	bool bValidateSizeSpecificDataDefaults = false;
+	bool bDoUpdateConvexGeometry = false;
+	bool bRebuildSimulationData = false;
+	bool bRebuildRenderData = false;
+
 	if (PropertyChangedEvent.Property)
 	{
 		FName PropertyName = PropertyChangedEvent.Property->GetFName();
 
-		bool bDoInvalidateCollection = false;
-		bool bDoEnsureDataIsCooked = false;
-		bool bValidateSizeSpecificDataDefaults = false;
-		bool bDoUpdateConvexGeometry = false;
-		bool bRebuildSimulationData = false;
-
 		if (PropertyChangedEvent.Property->GetFName() == GET_MEMBER_NAME_CHECKED(UGeometryCollection, EnableNanite))
 		{
 			bDoInvalidateCollection = true;
-			bDoEnsureDataIsCooked = true;
+			bRebuildRenderData = true;
 		}
 		else if (PropertyChangedEvent.Property->GetFName() == GET_MEMBER_NAME_CHECKED(UGeometryCollection, bUseFullPrecisionUVs))
 		{
 			bDoInvalidateCollection = true;
+			bRebuildRenderData = true;
 		}
 		else if (PropertyChangedEvent.Property->GetFName() == GET_MEMBER_NAME_CHECKED(UGeometryCollection, SizeSpecificData))
 		{
@@ -1449,35 +1374,40 @@ void UGeometryCollection::PostEditChangeProperty(struct FPropertyChangedEvent& P
 			bDoInvalidateCollection = true;
 			bRebuildSimulationData = true;
 		}
+	}
+	else if (PropertyChangedEvent.ChangeType == EPropertyChangeType::Unspecified)
+	{
+		// We get here on undo/redo operations.
+		// Make sure that render data rebuilds.
+		bRebuildRenderData = true;
+	}
 
+	if (bDoInvalidateCollection)
+	{
+		InvalidateCollection();
+	}
 
-		if (bDoInvalidateCollection)
-		{
-			InvalidateCollection();
-		}
-
-		if (bValidateSizeSpecificDataDefaults)
-		{
-			ValidateSizeSpecificDataDefaults();
-		}
+	if (bValidateSizeSpecificDataDefaults)
+	{
+		ValidateSizeSpecificDataDefaults();
+	}
 
-		if (bDoUpdateConvexGeometry)
-		{
-			UpdateConvexGeometry();
-		}
+	if (bDoUpdateConvexGeometry)
+	{
+		UpdateConvexGeometry();
+	}
 
-		if (bDoEnsureDataIsCooked)
+	if (bRebuildSimulationData)
+	{
+		if (!bManualDataCreate)
 		{
-			EnsureDataIsCooked();
+			CreateSimulationData();
 		}
+	}
 
-		if (bRebuildSimulationData)
-		{
-			if (!bManualDataCreate)
-			{
-				CreateSimulationData();
-			}
-		}
+	if (bRebuildRenderData)
+	{
+		RebuildRenderData();
 	}
 }
 
@@ -1494,24 +1424,34 @@ bool UGeometryCollection::Modify(bool bAlwaysMarkDirty /*= true*/)
 	return bSuperResult;
 }
 
-void UGeometryCollection::EnsureDataIsCooked(bool bInitResources, bool bIsTransacting)
+void UGeometryCollection::EnsureDataIsCooked(bool bInitResources, bool bIsTransacting, bool bIsPersistant)
 {
-	if (StateGuid != LastBuiltGuid)
+	if (StateGuid != LastBuiltSimulationDataGuid)
 	{
 		CreateSimulationDataImp(/*bCopyFromDDC=*/ !bIsTransacting);
+		LastBuiltSimulationDataGuid = StateGuid;
+	}
+
+	// Render data only goes through DDC when loading and saving (bIsPersistant).
+	// Using DDC during edits isn't worth it especially as we use a continually mutating guid instead of a state hash.
+	// That ensures that all edits are cache misses (slow) and unnecessarily fill up DDC disk space.
+	// TODO: SimulationData currently relies on these calls to update reliably, so we still need to use DDC for edits.
+	//       We could make CreateSimulationData() be reliably called for all edits and then only use DDC for loading and saving.
+	//       If we do that we can combine CreateSimulationDataImp() with CreateRenderDataImp() and FDerivedDataGeometryCollectionCooker
+	//       with FDerivedDataGeometryCollectionRenderDataCooker.
+	if (bIsPersistant && StateGuid != LastBuiltRenderDataGuid)
+	{
+		CreateRenderDataImp(/*bCopyFromDDC=*/ bInitResources);
 
 		if (FApp::CanEverRender() && bInitResources)
 		{
-			// If there is no geometry in the collection, we leave Nanite data alone.
-			if (GeometryCollection->NumElements(FGeometryCollection::GeometryGroup) > 0)
+			if (RenderData)
 			{
-				if (NaniteData)
-				{
-					NaniteData->InitResources(this);
-				}
+				RenderData->InitResources(this);
 			}
 		}
-		LastBuiltGuid = StateGuid;
+	
+		LastBuiltRenderDataGuid = StateGuid;
 	}
 }
 #endif
@@ -1533,69 +1473,36 @@ void UGeometryCollection::BeginDestroy()
 	ReleaseResources();
 }
 
-FGeometryCollectionNaniteData::FGeometryCollectionNaniteData()
+bool UGeometryCollection::HasMeshData() const
 {
+	return RenderData != nullptr && RenderData->bHasMeshData;
 }
 
-FGeometryCollectionNaniteData::~FGeometryCollectionNaniteData()
+bool UGeometryCollection::HasNaniteData() const
 {
-	ReleaseResources();
+	return RenderData != nullptr && RenderData->bHasNaniteData;
 }
 
-void FGeometryCollectionNaniteData::Serialize(FArchive& Ar, UGeometryCollection* Owner)
+uint32 UGeometryCollection::GetNaniteResourceID() const
 {
-	if (Ar.IsSaving())
-	{
-		if (Owner->EnableNanite)
-		{
-			// Nanite data is currently 1:1 with each geometry group in the collection.
-			const int32 NumGeometryGroups = Owner->NumElements(FGeometryCollection::GeometryGroup);
-			if (NumGeometryGroups != NaniteResource.HierarchyRootOffsets.Num())
-			{
-				Ar.SetError();
-			}
-		}
-
-		NaniteResource.Serialize(Ar, Owner, true);
-	}
-	else if (Ar.IsLoading())
-	{
-		NaniteResource.Serialize(Ar, Owner, true);
-	
-		if (!Owner->EnableNanite)
-		{
-			NaniteResource = {};
-		}
-	}
+	Nanite::FResources& Resource = RenderData->NaniteResource;
+	return Resource.RuntimeResourceID;
 }
 
-void FGeometryCollectionNaniteData::InitResources(UGeometryCollection* Owner)
+uint32 UGeometryCollection::GetNaniteHierarchyOffset() const
 {
-	if (bIsInitialized)
-	{
-		ReleaseResources();
-	}
-
-	NaniteResource.InitResources(Owner);
-
-	bIsInitialized = true;
+	Nanite::FResources& Resource = RenderData->NaniteResource;
+	return Resource.HierarchyOffset;
 }
 
-void FGeometryCollectionNaniteData::ReleaseResources()
+uint32 UGeometryCollection::GetNaniteHierarchyOffset(int32 GeometryIndex, bool bFlattened) const
 {
-	if (!bIsInitialized)
+	Nanite::FResources& Resource = RenderData->NaniteResource;
+	check(GeometryIndex >= 0 && GeometryIndex < Resource.HierarchyRootOffsets.Num());
+	uint32 HierarchyOffset = Resource.HierarchyRootOffsets[GeometryIndex];
+	if (bFlattened)
 	{
-		return;
+		HierarchyOffset += Resource.HierarchyOffset;
 	}
-
-	if (NaniteResource.ReleaseResources())
-	{
-		// HACK: Make sure the renderer is done processing the command, and done using NaniteResource, before we continue.
-		// This code could really use a refactor.
-		FRenderCommandFence Fence;
-		Fence.BeginFence();
-		Fence.Wait();
-	}
-
-	bIsInitialized = false;
+	return HierarchyOffset;
 }
diff --git a/Engine/Source/Runtime/Experimental/GeometryCollectionEngine/Private/GeometryCollection/GeometryCollectionRenderData.cpp b/Engine/Source/Runtime/Experimental/GeometryCollectionEngine/Private/GeometryCollection/GeometryCollectionRenderData.cpp
new file mode 100644
--- /dev/null
+++ b/Engine/Source/Runtime/Experimental/GeometryCollectionEngine/Private/GeometryCollection/GeometryCollectionRenderData.cpp
@@ -0,0 +1,739 @@
+// Copyright Epic Games, Inc. All Rights Reserved.
+
+#include "GeometryCollection/GeometryCollectionRenderData.h"
+
+#include "GeometryCollection/GeometryCollection.h"
+#include "GeometryCollection/GeometryCollectionAlgo.h"
+#include "GeometryCollection/GeometryCollectionObject.h"
+#include "GeometryCollection/ManagedArray.h"
+
+#if WITH_EDITOR
+#include "NaniteBuilder.h"
+#endif
+
+
+FBoneMapVertexBuffer::~FBoneMapVertexBuffer()
+{
+	CleanUp();
+}
+
+void FBoneMapVertexBuffer::CleanUp()
+{
+	delete BoneMapData;
+	BoneMapData = nullptr;
+}
+
+void FBoneMapVertexBuffer::Init(TArray<uint16> const& InBoneMap, bool bInNeedsCPUAccess)
+{
+	AllocateData(bInNeedsCPUAccess);
+	ResizeBuffer(InBoneMap.Num());
+	FMemory::Memcpy(Data, InBoneMap.GetData(), PixelFormatStride * NumVertices);
+}
+
+void FBoneMapVertexBuffer::ResizeBuffer(uint32 InNumVertices)
+{
+	NumVertices = InNumVertices;
+	BoneMapData->ResizeBuffer(NumVertices);
+	Data = NumVertices != 0 ? BoneMapData->GetDataPointer() : nullptr;
+}
+
+void FBoneMapVertexBuffer::AllocateData(bool bInNeedsCPUAccess)
+{
+	CleanUp();
+	BoneMapData = new TStaticMeshVertexData<uint16>(bInNeedsCPUAccess);
+}
+
+void FBoneMapVertexBuffer::InitRHI()
+{
+	const bool bHadData = BoneMapData != nullptr;
+	VertexBufferRHI = CreateRHIBuffer<true>(BoneMapData, NumVertices, BUF_Static | BUF_ShaderResource | BUF_SourceCopy, TEXT("FBoneMapVertexBuffer")); 
+	if (VertexBufferRHI != nullptr)
+	{
+		VertexBufferSRV = RHICreateShaderResourceView(FShaderResourceViewInitializer(bHadData ? VertexBufferRHI : nullptr, PixelFormat));
+	}
+}
+
+void FBoneMapVertexBuffer::ReleaseRHI()
+{
+	VertexBufferSRV.SafeRelease();
+	FVertexBuffer::ReleaseRHI();
+}
+
+void FBoneMapVertexBuffer::Serialize(FArchive& Ar, bool bInNeedsCPUAccess)
+{
+	bNeedsCPUAccess = bInNeedsCPUAccess;
+
+	Ar << NumVertices;
+
+	if (Ar.IsLoading())
+	{
+		// Allocate the vertex data storage type.
+		AllocateData(bInNeedsCPUAccess);
+	}
+
+	if (BoneMapData != nullptr)
+	{
+		// Serialize the vertex data.
+		BoneMapData->Serialize(Ar);
+
+		// Make a copy of the vertex data pointer.
+		Data = NumVertices != 0 ? BoneMapData->GetDataPointer() : nullptr;
+	}
+}
+
+
+void FGeometryCollectionMeshResources::InitResources(UGeometryCollection* Owner)
+{
+	FName OwnerName = Owner->GetFName();
+	IndexBuffer.SetOwnerName(OwnerName);
+	PositionVertexBuffer.SetOwnerName(OwnerName);
+	StaticMeshVertexBuffer.SetOwnerName(OwnerName);
+	ColorVertexBuffer.SetOwnerName(OwnerName);
+	BoneMapVertexBuffer.SetOwnerName(OwnerName);
+
+	BeginInitResource(&IndexBuffer);
+	BeginInitResource(&PositionVertexBuffer);
+	BeginInitResource(&StaticMeshVertexBuffer);
+	BeginInitResource(&ColorVertexBuffer);
+	BeginInitResource(&BoneMapVertexBuffer);
+}
+
+void FGeometryCollectionMeshResources::ReleaseResources()
+{
+	BeginReleaseResource(&IndexBuffer);
+	BeginReleaseResource(&PositionVertexBuffer);
+	BeginReleaseResource(&StaticMeshVertexBuffer);
+	BeginReleaseResource(&ColorVertexBuffer);
+	BeginReleaseResource(&BoneMapVertexBuffer);
+}
+
+void FGeometryCollectionMeshResources::Serialize(FArchive& Ar)
+{
+	// If platform doesn't have manual vertex fetch then we will need position and 
+	// bone map CPU access so that we can skin positions on CPU.
+	const bool bCPUAccessForBoneTransform = Ar.IsLoading() && GMaxRHIFeatureLevel <= ERHIFeatureLevel::ES3_1;
+	// Otherwise we don't need CPU access.
+	const bool bNeedsCPUAccess = false;
+
+	IndexBuffer.Serialize(Ar, bNeedsCPUAccess);
+	PositionVertexBuffer.Serialize(Ar, bCPUAccessForBoneTransform);
+	StaticMeshVertexBuffer.Serialize(Ar, bNeedsCPUAccess);
+	ColorVertexBuffer.Serialize(Ar, bNeedsCPUAccess);
+	BoneMapVertexBuffer.Serialize(Ar, bCPUAccessForBoneTransform);
+}
+
+
+void FGeometryCollectionMeshDescription::Serialize(FArchive& Ar)
+{
+	Ar << NumVertices << NumTriangles << PreSkinnedBounds << Sections << SectionsNoInternal;
+	
+	if (Ar.IsCooking())
+	{
+		// No need for SubSections in cooked build.
+		TArray<FGeometryCollectionMeshElement> Dummy;
+		Ar << Dummy;
+	}
+	else
+	{
+		Ar << SubSections;
+	}
+}
+
+
+#if WITH_EDITOR
+
+/** Data to return from BuildMeshDataFromGeometryCollection(). */
+struct FGeometryCollectionBuiltMeshData
+{
+	uint32 NumTexCoords = 0;
+	TArray<uint32> Indices;
+	TArray<FStaticMeshBuildVertex> Vertices;
+	TArray<uint16> BonesPerVertex;
+	TArray<int32> MaterialsPerTriangle;
+	TArray<bool> InternalPerTriangle;
+	TArray<uint32> MeshTriangleCounts;
+};
+
+/** Get the data needed to build render data from a Geometry Collection. */
+void BuildMeshDataFromGeometryCollection(FGeometryCollection& InCollection, FGeometryCollectionBuiltMeshData& OutMeshData)
+{
+	// Vertices Group
+	const TManagedArray<FVector3f>& VertexArray = InCollection.Vertex;
+	GeometryCollection::UV::FUVLayers UVsLayers = GeometryCollection::UV::FindActiveUVLayers(InCollection);
+	const TManagedArray<FLinearColor>& ColorArray = InCollection.Color;
+	const TManagedArray<FVector3f>& TangentUArray = InCollection.TangentU;
+	const TManagedArray<FVector3f>& TangentVArray = InCollection.TangentV;
+	const TManagedArray<FVector3f>& NormalArray = InCollection.Normal;
+	const TManagedArray<int32>& BoneMapArray = InCollection.BoneMap;
+
+	// Faces Group
+	const TManagedArray<FIntVector>& IndicesArray = InCollection.Indices;
+	const TManagedArray<bool>& VisibleArray = InCollection.Visible;
+	const TManagedArray<int32>& MaterialIDArray = InCollection.MaterialID;
+	const TManagedArray<bool>& Internal = InCollection.Internal;
+
+	// Geometry Group
+	const int32 NumGeometry = InCollection.NumElements(FGeometryCollection::GeometryGroup);
+	const TManagedArray<int32>& VertexStartArray = InCollection.VertexStart;
+	const TManagedArray<int32>& VertexCountArray = InCollection.VertexCount;
+	const TManagedArray<int32>& FaceStartArray = InCollection.FaceStart;
+	const TManagedArray<int32>& FaceCountArray = InCollection.FaceCount;
+
+	const uint32 NumTexCoords = InCollection.NumUVLayers();
+	const bool bHasColors = ColorArray.Num() > 0;
+
+	// Iterate geometry groups and collect vertices.
+	TArray<FStaticMeshBuildVertex> BuildVertices;
+	TArray<uint16> BonesPerVertex;
+	TArray<int32> DestVertexStarts;
+
+	{
+		TRACE_CPUPROFILER_EVENT_SCOPE(FGeometryCollectionRenderData::BuildVertices);
+
+		// Gather destination vertex locations so that we can fill them in parallel.
+		int32 NumBuildVertices = 0;
+		DestVertexStarts.SetNumUninitialized(NumGeometry);
+		for (int32 GeometryGroupIndex = 0; GeometryGroupIndex < NumGeometry; GeometryGroupIndex++)
+		{
+			DestVertexStarts[GeometryGroupIndex] = GeometryGroupIndex > 0 ? DestVertexStarts[GeometryGroupIndex - 1] + VertexCountArray[GeometryGroupIndex - 1] : 0;
+			NumBuildVertices += VertexCountArray[GeometryGroupIndex];
+		}
+
+		BuildVertices.AddUninitialized(NumBuildVertices);
+		BonesPerVertex.AddUninitialized(NumBuildVertices);
+
+		ParallelFor(NumGeometry, [&](int32 GeometryGroupIndex)
+		{
+			const int32 VertexStart = VertexStartArray[GeometryGroupIndex];
+			const int32 VertexCount = VertexCountArray[GeometryGroupIndex];
+			const int32 DestVertexStart = DestVertexStarts[GeometryGroupIndex];
+
+			ParallelFor(TEXT("GC:BuildVertices"), VertexCount, 500, [&](int32 VertexIndex)
+			{
+				FStaticMeshBuildVertex& Vertex = BuildVertices[DestVertexStart + VertexIndex];
+				Vertex.Position = VertexArray[VertexStart + VertexIndex];
+				Vertex.Color = bHasColors ? ColorArray[VertexStart + VertexIndex].ToFColor(false /* sRGB */) : FColor::White;
+				Vertex.TangentX = FVector3f::ZeroVector;
+				Vertex.TangentY = FVector3f::ZeroVector;
+				Vertex.TangentZ = NormalArray[VertexStart + VertexIndex];
+				for (int32 UVIdx = 0; UVIdx < UVsLayers.Num(); ++UVIdx)
+				{
+					Vertex.UVs[UVIdx] = UVsLayers[UVIdx][VertexStart + VertexIndex];
+					if (Vertex.UVs[UVIdx].ContainsNaN())
+					{
+						Vertex.UVs[UVIdx] = FVector2f::ZeroVector;
+					}
+				}
+
+				const int32 BoneIndex = BoneMapArray[VertexStart + VertexIndex];
+				BonesPerVertex[DestVertexStart + VertexIndex] = (uint16)BoneIndex;
+			});
+		});
+	}
+
+	// Iterate geometry groups and collect indices.
+	// Can't go parallel here because we remove hidden faces as we go.
+	TArray<uint32> BuildIndices;
+	TArray<int32> MaterialsPerFace;
+	TArray<bool> InternalPerFace;
+	TArray<uint32> MeshTriangleCounts;
+
+	{
+		TRACE_CPUPROFILER_EVENT_SCOPE(FGeometryCollectionRenderData::GatherBuildIndices);
+
+		BuildIndices.Reserve(IndicesArray.Num() * 3);
+		MaterialsPerFace.Reserve(IndicesArray.Num());
+		InternalPerFace.Reserve(IndicesArray.Num());
+		MeshTriangleCounts.Reserve(NumGeometry);
+
+		for (int32 GeometryGroupIndex = 0; GeometryGroupIndex < NumGeometry; ++GeometryGroupIndex)
+		{
+			const int32 FaceStart = FaceStartArray[GeometryGroupIndex];
+			const int32 FaceCount = FaceCountArray[GeometryGroupIndex];
+			const int32 DestFaceStart = MaterialsPerFace.Num();
+			const int32 DestVertexOffset = DestVertexStarts[GeometryGroupIndex] - VertexStartArray[GeometryGroupIndex];
+
+			for (int32 FaceIndex = 0; FaceIndex < FaceCount; ++FaceIndex)
+			{
+				// Remove hidden.
+				if (!VisibleArray[FaceStart + FaceIndex])
+				{
+					continue;
+				}
+
+				FIntVector FaceIndices = IndicesArray[FaceStart + FaceIndex];
+				FaceIndices = FaceIndices + FIntVector(DestVertexOffset);
+
+				// Remove degenerates.
+				if (BuildVertices[FaceIndices[0]].Position == BuildVertices[FaceIndices[1]].Position ||
+					BuildVertices[FaceIndices[1]].Position == BuildVertices[FaceIndices[2]].Position ||
+					BuildVertices[FaceIndices[2]].Position == BuildVertices[FaceIndices[0]].Position)
+				{
+					continue;
+				}
+
+				BuildIndices.Add(FaceIndices.X);
+				BuildIndices.Add(FaceIndices.Y);
+				BuildIndices.Add(FaceIndices.Z);
+
+				const int32 MaterialIndex = MaterialIDArray[FaceStart + FaceIndex];
+				MaterialsPerFace.Add(MaterialIndex);
+				const bool bInternal = Internal[FaceStart + FaceIndex];
+				InternalPerFace.Add(bInternal);
+			}
+
+			MeshTriangleCounts.Add(MaterialsPerFace.Num() - DestFaceStart);
+		}
+	}
+
+	OutMeshData.NumTexCoords = NumTexCoords;
+	OutMeshData.Indices = MoveTemp(BuildIndices);
+	OutMeshData.Vertices = MoveTemp(BuildVertices);
+	OutMeshData.BonesPerVertex = MoveTemp(BonesPerVertex);
+	OutMeshData.MaterialsPerTriangle = MoveTemp(MaterialsPerFace);
+	OutMeshData.InternalPerTriangle = MoveTemp(InternalPerFace);
+	OutMeshData.MeshTriangleCounts = MoveTemp(MeshTriangleCounts);
+}
+
+struct FGeometryCollectionBuiltMeshSectionData
+{
+	TArray<uint32> SortedIndices;
+	TArray<FGeometryCollectionMeshElement> Sections;
+	TArray<FGeometryCollectionMeshElement> SectionsNoInternal;
+	TArray<FGeometryCollectionMeshElement> SubSections;
+};
+
+void BuildSectionDataFromMeshData(FGeometryCollectionBuiltMeshData const& InMeshData, FGeometryCollectionBuiltMeshSectionData& OutSectionData)
+{
+	// Find spans within the mesh data.
+	TArray<FGeometryCollectionMeshElement> Spans;
+
+	{
+		TRACE_CPUPROFILER_EVENT_SCOPE(FGeometryCollectionRenderData::FindSpans);
+
+		int32 LastMaterialIndex = -1;
+		int32 LastIsInternal = -1;
+		int32 LastBoneIndex = -1;
+
+		for (int32 TriangleIndex = 0; TriangleIndex < InMeshData.MaterialsPerTriangle.Num(); ++TriangleIndex)
+		{
+			const int32 MaterialIndex = InMeshData.MaterialsPerTriangle[TriangleIndex];
+			const int32 IsInternal = InMeshData.InternalPerTriangle[TriangleIndex] ? 1 : 0;
+			const uint32 VertexIndex0 = InMeshData.Indices[TriangleIndex * 3];
+			const uint32 VertexIndex1 = InMeshData.Indices[TriangleIndex * 3 + 1];
+			const uint32 VertexIndex2 = InMeshData.Indices[TriangleIndex * 3 + 2];
+			const uint16 BoneIndex = InMeshData.BonesPerVertex[VertexIndex0];
+			const uint32 MinVertexIndex = FMath::Min(VertexIndex0, FMath::Min(VertexIndex1, VertexIndex2));
+			const uint32 MaxVertexIndex = FMath::Max(VertexIndex0, FMath::Max(VertexIndex1, VertexIndex2));
+
+			// Add a new span if the properties change.
+			if (MaterialIndex != LastMaterialIndex || IsInternal != LastIsInternal || BoneIndex != LastBoneIndex)
+			{
+				FGeometryCollectionMeshElement& Span = Spans.AddZeroed_GetRef();
+				Span.TransformIndex = BoneIndex;
+				Span.MaterialIndex = MaterialIndex;
+				Span.bIsInternal = IsInternal;
+				Span.TriangleStart = TriangleIndex;
+				Span.TriangleCount = 0;
+				Span.VertexStart = MinVertexIndex;
+				Span.VertexEnd = MaxVertexIndex;
+			}
+
+			// Update the current span.
+			FGeometryCollectionMeshElement& Span = Spans.Last();
+			Span.TriangleCount++;
+			Span.VertexStart = FMath::Min(Span.VertexStart, MinVertexIndex);
+			Span.VertexEnd = FMath::Max(Span.VertexEnd, MaxVertexIndex);
+
+			LastMaterialIndex = MaterialIndex;
+			LastIsInternal = IsInternal;
+			LastBoneIndex = BoneIndex;
+		}
+	}
+
+	// Sort spans by material and bone index.
+	TArray<int32> SortedSpanIndices;
+
+	{
+		TRACE_CPUPROFILER_EVENT_SCOPE(FGeometryCollectionRenderData::SortSpans);
+
+		TArray<uint64> SortKeys;
+		SortKeys.SetNum(Spans.Num());
+		SortedSpanIndices.SetNum(Spans.Num());
+		for (int32 SpanIndex = 0; SpanIndex < SortedSpanIndices.Num(); SpanIndex++)
+		{
+			FGeometryCollectionMeshElement const& Span = Spans[SpanIndex];
+			SortKeys[SpanIndex] = ((uint64)Span.MaterialIndex << 56) | ((uint64)Span.bIsInternal << 48) | ((uint64)Span.TransformIndex << 32) | (uint64)Span.TriangleStart;
+			SortedSpanIndices[SpanIndex] = SpanIndex;
+		}
+		SortedSpanIndices.StableSort([&SortKeys](int32 A, int32 B)
+		{
+			return SortKeys[A] < SortKeys[B];
+		});
+	}
+
+	// Now copy spans to final section arrays.
+	TArray<FGeometryCollectionMeshElement> Sections;
+	TArray<FGeometryCollectionMeshElement> SectionsNoInternal;
+	TArray<FGeometryCollectionMeshElement> SubSections;
+
+	{
+		TRACE_CPUPROFILER_EVENT_SCOPE(FGeometryCollectionRenderData::BuildSections);
+
+		Sections.Reserve(Spans.Num());
+		SectionsNoInternal.Reserve(Spans.Num());
+		SubSections.Reserve(Spans.Num());
+
+		int32 LastMaterialIndex = -1;
+		uint8 LastIsInternal = 0xff;
+		int32 LastBoneIndex = -1;
+		int32 TriangleIndex = 0;
+
+		for (int32 SortedSpanIndex = 0; SortedSpanIndex < SortedSpanIndices.Num(); SortedSpanIndex++)
+		{
+			const int32 SpanIndex = SortedSpanIndices[SortedSpanIndex];
+			FGeometryCollectionMeshElement const& Span = Spans[SpanIndex];
+
+			// Add new section if required.
+			if (Span.MaterialIndex != LastMaterialIndex)
+			{
+				FGeometryCollectionMeshElement& Section = Sections.AddZeroed_GetRef();
+				Section.TransformIndex = -1;
+				Section.MaterialIndex = Span.MaterialIndex;
+				Section.bIsInternal = 0;
+				Section.TriangleStart = TriangleIndex;
+				Section.TriangleCount = 0;
+				Section.VertexStart = Span.VertexStart;
+				Section.VertexEnd = Span.VertexEnd;
+			}
+			{
+				// Accumulate this triangle to the current section.
+				FGeometryCollectionMeshElement& Section = Sections.Last();
+				Section.TriangleCount += Span.TriangleCount;
+				Section.VertexStart = FMath::Min(Section.VertexStart, Span.VertexStart);
+				Section.VertexEnd = FMath::Max(Section.VertexEnd, Span.VertexEnd);
+			}
+
+			// Add new no-internal section if required.
+			// Note that internal face spans have been sorted to the end of each full section.
+			if (Span.MaterialIndex != LastMaterialIndex && !Span.bIsInternal)
+			{
+				FGeometryCollectionMeshElement& Section = SectionsNoInternal.AddZeroed_GetRef();
+				Section.TransformIndex = -1;
+				Section.MaterialIndex = Span.MaterialIndex;
+				Section.bIsInternal = 0;
+				Section.TriangleStart = TriangleIndex;
+				Section.TriangleCount = 0;
+				Section.VertexStart = Span.VertexStart;
+				Section.VertexEnd = Span.VertexEnd;
+			}
+			// Accumulate this triangle if this is non-internal section.
+			if (!Span.bIsInternal)
+			{
+				FGeometryCollectionMeshElement& Section = SectionsNoInternal.Last();
+				Section.TriangleCount += Span.TriangleCount;
+				Section.VertexStart = FMath::Min(Section.VertexStart, Span.VertexStart);
+				Section.VertexEnd = FMath::Max(Section.VertexEnd, Span.VertexEnd);
+			}
+
+			// Add new subsection if required.
+			// Note that subsections don't care about the internal flag, since they are only ever used
+			// in editor in modes where we want to render internal faces.
+			if (Span.MaterialIndex != LastMaterialIndex || Span.TransformIndex != LastBoneIndex)
+			{
+				FGeometryCollectionMeshElement& Section = SubSections.AddZeroed_GetRef();
+				Section.TransformIndex = Span.TransformIndex;
+				Section.MaterialIndex = Span.MaterialIndex;
+				Section.bIsInternal = 0;
+				Section.TriangleStart = TriangleIndex;
+				Section.TriangleCount = 0;
+				Section.VertexStart = Span.VertexStart;
+				Section.VertexEnd = Span.VertexEnd;
+			}
+			{
+				// Accumulate this triangle to the current subsection.
+				FGeometryCollectionMeshElement& Section = SubSections.Last();
+				Section.TriangleCount += Span.TriangleCount;
+				Section.VertexStart = FMath::Min(Section.VertexStart, Span.VertexStart);
+				Section.VertexEnd = FMath::Max(Section.VertexEnd, Span.VertexEnd);
+			}
+
+			TriangleIndex += Span.TriangleCount;
+			LastMaterialIndex = Span.MaterialIndex;
+			LastIsInternal = Span.bIsInternal;
+			LastBoneIndex = Span.TransformIndex;
+		}
+
+		//TODO: Clear SectionsNoInternal if it is just a duplicate of Sections.
+
+		Sections.Shrink();
+		SectionsNoInternal.Shrink();
+		SubSections.Shrink();
+	}
+
+	// Copy indices to sorted index buffer.
+	TArray<uint32> SortedIndices;
+	SortedIndices.SetNumUninitialized(InMeshData.Indices.Num());
+
+	{
+		TRACE_CPUPROFILER_EVENT_SCOPE(FGeometryCollectionRenderData::CopySortedIndices);
+
+		// Gather destination index locations so that we can fill them in parallel.
+		TArray<int32> DestIndexStarts;
+		DestIndexStarts.SetNumUninitialized(SortedSpanIndices.Num());
+		for (int32 SortedSpanIndex = 0; SortedSpanIndex < SortedSpanIndices.Num(); SortedSpanIndex++)
+		{
+			DestIndexStarts[SortedSpanIndex] = SortedSpanIndex > 0 ? DestIndexStarts[SortedSpanIndex - 1] + Spans[SortedSpanIndices[SortedSpanIndex - 1]].TriangleCount * 3 : 0;
+		}
+
+		ParallelFor(SortedSpanIndices.Num(), [&](int32 SortedSpanIndex)
+		{
+			const int32 SpanIndex = SortedSpanIndices[SortedSpanIndex];
+			FGeometryCollectionMeshElement const& Span = Spans[SpanIndex];
+
+			const int32 SourceIndexStart = Span.TriangleStart * 3;
+			const int32 DestIndexStart = DestIndexStarts[SortedSpanIndex];
+
+			ParallelFor(TEXT("GC:CopySortedIndices"), Span.TriangleCount, 500, [&](int32 TriangleIndex)
+			{
+				const int32 SourceIndex = SourceIndexStart + TriangleIndex * 3;
+				const int32 DestIndex = DestIndexStart + TriangleIndex * 3;
+				SortedIndices[DestIndex] = InMeshData.Indices[SourceIndex];
+				SortedIndices[DestIndex + 1] = InMeshData.Indices[SourceIndex + 1];
+				SortedIndices[DestIndex + 2] = InMeshData.Indices[SourceIndex + 2];
+			});
+		});
+	}
+
+	// Move to output.
+	OutSectionData.SortedIndices = MoveTemp(SortedIndices);
+	OutSectionData.Sections = MoveTemp(Sections);
+	OutSectionData.SectionsNoInternal = MoveTemp(SectionsNoInternal);
+	OutSectionData.SubSections = MoveTemp(SubSections);
+}
+
+FBoxSphereBounds BuildPreSkinnedBounds(FGeometryCollection& InCollection)
+{
+	TRACE_CPUPROFILER_EVENT_SCOPE(FGeometryCollectionRenderData::BuildBounds);
+
+	const TManagedArray<FBox>& BoundingBoxes = InCollection.BoundingBox;
+	const TManagedArray<FTransform>& Transform = InCollection.Transform;
+	const TManagedArray<int32>& Parent = InCollection.Parent;
+	const TManagedArray<int32>& TransformIndices = InCollection.TransformIndex;
+	const int32 NumBoxes = BoundingBoxes.Num();
+
+	TArray<FMatrix> GlobalMatrices;
+	GeometryCollectionAlgo::GlobalMatrices(Transform, Parent, GlobalMatrices);
+
+	bool bBoundsInit = false;
+	FBox PreSkinnedBounds(ForceInit);
+
+	for (int32 BoxIndex = 0; BoxIndex < NumBoxes; ++BoxIndex)
+	{
+		const int32 TransformIndex = TransformIndices[BoxIndex];
+		if (InCollection.IsGeometry(TransformIndex))
+		{
+			if (!bBoundsInit)
+			{
+				PreSkinnedBounds = BoundingBoxes[BoxIndex].TransformBy(GlobalMatrices[TransformIndex]);
+				bBoundsInit = true;
+			}
+			else
+			{
+				PreSkinnedBounds += BoundingBoxes[BoxIndex].TransformBy(GlobalMatrices[TransformIndex]);
+			}
+		}
+	}
+
+	return bBoundsInit ? FBoxSphereBounds(PreSkinnedBounds) : FBoxSphereBounds(ForceInit);
+}
+
+void CreateMeshData(FGeometryCollection& InCollection, FGeometryCollectionBuiltMeshData const& InMeshData, bool bInUseFullPrecisionUVs, FGeometryCollectionRenderData& OutRenderData)
+{
+	TRACE_CPUPROFILER_EVENT_SCOPE(FGeometryCollectionRenderData::CreateMeshData);
+
+	OutRenderData.bHasMeshData = true;
+
+	OutRenderData.MeshDescription.NumVertices = InMeshData.Vertices.Num();
+	OutRenderData.MeshDescription.NumTriangles = InMeshData.Indices.Num() / 3;
+
+	{
+		TRACE_CPUPROFILER_EVENT_SCOPE(FGeometryCollectionRenderData::BufferInit);
+		OutRenderData.MeshResource.PositionVertexBuffer.Init(InMeshData.Vertices);
+		OutRenderData.MeshResource.StaticMeshVertexBuffer.SetUseFullPrecisionUVs(bInUseFullPrecisionUVs);
+		OutRenderData.MeshResource.StaticMeshVertexBuffer.Init(InMeshData.Vertices, InMeshData.NumTexCoords);
+		OutRenderData.MeshResource.ColorVertexBuffer.Init(InMeshData.Vertices);
+		OutRenderData.MeshResource.BoneMapVertexBuffer.Init(InMeshData.BonesPerVertex);
+	}
+
+	FGeometryCollectionBuiltMeshSectionData SectionData;
+	BuildSectionDataFromMeshData(InMeshData, SectionData);
+
+	// Store the sorted indices that match the section data.
+	OutRenderData.MeshResource.IndexBuffer.SetIndices(SectionData.SortedIndices, EIndexBufferStride::AutoDetect);
+
+	OutRenderData.MeshDescription.Sections = MoveTemp(SectionData.Sections);
+	OutRenderData.MeshDescription.SectionsNoInternal = MoveTemp(SectionData.SectionsNoInternal);
+	OutRenderData.MeshDescription.SubSections = MoveTemp(SectionData.SubSections);
+
+	// Calculate bounds.
+	OutRenderData.MeshDescription.PreSkinnedBounds = BuildPreSkinnedBounds(InCollection);
+}
+
+void CreateNaniteData(FGeometryCollectionBuiltMeshData& InMeshData, FGeometryCollectionRenderData& OutRenderData)
+{
+	TRACE_CPUPROFILER_EVENT_SCOPE(FGeometryCollectionRenderData::CreateNaniteData);
+
+	OutRenderData.bHasNaniteData = true;
+
+	FMeshNaniteSettings NaniteSettings = {};
+	NaniteSettings.bEnabled = true;
+	NaniteSettings.TargetMinimumResidencyInKB = 0;	// Default to smallest possible, which is a single page
+	NaniteSettings.KeepPercentTriangles = 1.0f;
+	NaniteSettings.TrimRelativeError = 0.0f;
+	NaniteSettings.FallbackPercentTriangles = 1.0f; // 100% - no reduction
+	NaniteSettings.FallbackRelativeError = 0.0f;
+
+	OutRenderData.NaniteResource = {};
+
+	Nanite::IBuilderModule& NaniteBuilderModule = Nanite::IBuilderModule::Get();
+
+	if (!NaniteBuilderModule.Build(
+		OutRenderData.NaniteResource,
+		InMeshData.Vertices,
+		InMeshData.Indices,
+		InMeshData.MaterialsPerTriangle,
+		InMeshData.MeshTriangleCounts,
+		InMeshData.NumTexCoords,
+		NaniteSettings))
+	{
+		UE_LOG(LogStaticMesh, Error, TEXT("Failed to build Nanite for geometry collection. See previous line(s) for details."));
+	}
+}
+
+TUniquePtr<FGeometryCollectionRenderData> FGeometryCollectionRenderData::Create(FGeometryCollection& InCollection, bool bInEnableNanite, bool bInUseFullPrecisionUVs)
+{
+	TRACE_CPUPROFILER_EVENT_SCOPE(FGeometryCollectionRenderData::Create);
+
+	TUniquePtr<FGeometryCollectionRenderData> RenderData = MakeUnique<FGeometryCollectionRenderData>();
+
+	FGeometryCollectionBuiltMeshData MeshBuildData;
+	BuildMeshDataFromGeometryCollection(InCollection, MeshBuildData);
+
+	if (bInEnableNanite)
+	{
+		CreateNaniteData(MeshBuildData, *RenderData.Get());
+	}
+	else
+	{
+		// Could always create mesh data if we want to be able to enable/disable nanite at runtime in cooked build.
+		CreateMeshData(InCollection, MeshBuildData, bInUseFullPrecisionUVs, *RenderData.Get());
+	}
+
+	return RenderData;
+}
+
+#endif // WITH_EDITOR
+
+
+FGeometryCollectionRenderData::~FGeometryCollectionRenderData()
+{
+	ReleaseResources();
+
+	FRenderCommandFence Fence;
+	Fence.BeginFence();
+	Fence.Wait();
+}
+
+void FGeometryCollectionRenderData::Serialize(FArchive& Ar, UGeometryCollection* Owner)
+{
+	if (Owner->bStripOnCook && Ar.IsCooking())
+	{
+		// Don't cook rendering data.
+		// This is used if we expect to use a non-native GC render path such as the ISM pool.
+		bool bDummy = false;
+		Ar << bDummy; // bHasMeshData
+		Ar << bDummy; // bHasNaniteData
+		return;
+	}
+
+	Ar << bHasMeshData;
+	Ar << bHasNaniteData;
+
+	if (bHasMeshData)
+	{
+		MeshResource.Serialize(Ar);
+		MeshDescription.Serialize(Ar);
+	}
+	else if (Ar.IsLoading())
+	{
+		MeshResource = {};
+		MeshDescription = {};
+	}
+
+	if (bHasNaniteData)
+	{
+		if (Ar.IsSaving())
+		{
+			// Nanite data is 1:1 with each geometry group in the collection.
+			const int32 NumGeometryGroups = Owner->NumElements(FGeometryCollection::GeometryGroup);
+			if (!Owner->EnableNanite || NumGeometryGroups != NaniteResource.HierarchyRootOffsets.Num())
+			{
+				Ar.SetError();
+			}
+		}
+
+		NaniteResource.Serialize(Ar, Owner, true);
+	}
+	else if (Ar.IsLoading())
+	{
+		NaniteResource = {};
+	}
+}
+
+void FGeometryCollectionRenderData::InitResources(UGeometryCollection* Owner)
+{
+	if (bIsInitialized)
+	{
+		ReleaseResources();
+	}
+
+	if (bHasMeshData)
+	{
+		MeshResource.InitResources(Owner);
+	}
+
+	if (bHasNaniteData)
+	{
+		NaniteResource.InitResources(Owner);
+	}
+
+	bIsInitialized = true;
+}
+
+void FGeometryCollectionRenderData::ReleaseResources()
+{
+	if (!bIsInitialized)
+	{
+		return;
+	}
+
+	if (bHasMeshData)
+	{
+		MeshResource.ReleaseResources();
+	}
+
+	if (bHasNaniteData)
+	{
+		if (NaniteResource.ReleaseResources())
+		{
+			// HACK: Make sure the renderer is done processing the command, and done using NaniteResource, before we continue.
+			// This code could really use a refactor.
+			FRenderCommandFence Fence;
+			Fence.BeginFence();
+			Fence.Wait();
+		}
+	}
+
+	bIsInitialized = false;
+}
diff --git a/Engine/Source/Runtime/Experimental/GeometryCollectionEngine/Private/GeometryCollection/GeometryCollectionRenderData.h b/Engine/Source/Runtime/Experimental/GeometryCollectionEngine/Private/GeometryCollection/GeometryCollectionRenderData.h
new file mode 100644
--- /dev/null
+++ b/Engine/Source/Runtime/Experimental/GeometryCollectionEngine/Private/GeometryCollection/GeometryCollectionRenderData.h
@@ -0,0 +1,148 @@
+// Copyright Epic Games, Inc. All Rights Reserved.
+
+#pragma once
+
+#include "CoreMinimal.h"
+#include "RawIndexBuffer.h"
+#include "Rendering/ColorVertexBuffer.h"
+#include "Rendering/NaniteResources.h"
+#include "Rendering/PositionVertexBuffer.h"
+#include "Rendering/StaticMeshVertexBuffer.h"
+
+class FGeometryCollection;
+class UGeometryCollection;
+
+
+/** Vertex Buffer for Bone Map. */
+class FBoneMapVertexBuffer : public FVertexBuffer
+{
+public:
+	~FBoneMapVertexBuffer();
+
+	void Init(TArray<uint16> const& InBoneMap, bool bInNeedsCPUAccess = true);
+	void Serialize(FArchive& Ar, bool bInNeedsCPUAccess);
+	FRHIShaderResourceView* GetSRV() const { return VertexBufferSRV; }
+	
+	FORCEINLINE uint16& BoneIndex(uint32 VertexIndex) const
+	{
+		checkSlow(VertexIndex < GetNumVertices());
+		return *(uint16*)(Data + VertexIndex * PixelFormatStride);
+	}
+
+protected:
+	void CleanUp();
+	void AllocateData(bool bInNeedsCPUAccess = true);
+	void ResizeBuffer(uint32 InNumVertices);
+
+	void InitRHI() override;
+	void ReleaseRHI() override;
+
+	static const EPixelFormat PixelFormat = PF_R16_UINT;
+	static const uint32 PixelFormatStride = 2;
+
+private:
+	uint32 NumVertices = 0;
+	bool bNeedsCPUAccess = true;
+	FShaderResourceViewRHIRef VertexBufferSRV = nullptr;
+	FStaticMeshVertexDataInterface* BoneMapData = nullptr;
+	uint8* Data = nullptr;
+};
+
+/** Set of GPU vertex and index buffers required to render a geometry collection. */
+struct FGeometryCollectionMeshResources
+{
+	FRawStaticIndexBuffer IndexBuffer;
+	FPositionVertexBuffer PositionVertexBuffer;
+	FStaticMeshVertexBuffer StaticMeshVertexBuffer;
+	FColorVertexBuffer ColorVertexBuffer;
+	FBoneMapVertexBuffer BoneMapVertexBuffer;
+
+	/** Initialize the render resources. */
+	void InitResources(UGeometryCollection* Owner);
+	/** Release the render resources. */
+	void ReleaseResources();
+	/** Serialization. */
+	void Serialize(FArchive& Ar);
+};
+
+/** Description of a single mesh element within a set of vertex and vertex buffers. */
+struct FGeometryCollectionMeshElement
+{
+	int16 TransformIndex;
+	uint8 MaterialIndex;
+	uint8 bIsInternal;
+	uint32 TriangleStart;
+	uint32 TriangleCount;
+	uint32 VertexStart;
+	uint32 VertexEnd;
+
+	/** Serialization. */
+	friend FArchive& operator<<(FArchive& Ar, FGeometryCollectionMeshElement& Item)
+	{
+		Ar << Item.TransformIndex << Item.MaterialIndex << Item.bIsInternal << Item.TriangleStart << Item.TriangleCount << Item.VertexStart << Item.VertexEnd;
+		return Ar;
+	}
+};
+
+/** CPU side data required to render a geometry collection. */
+struct FGeometryCollectionMeshDescription
+{
+	uint32 NumVertices = 0;
+	uint32 NumTriangles = 0;
+
+	// Bounds before any movement or skinning ever happens.
+	FBoxSphereBounds PreSkinnedBounds;
+
+	// Info on grouped material sections within vertex buffers. 
+	// A section has a uniform material index and can contain multiple per-bone subsections.
+	TArray<FGeometryCollectionMeshElement> Sections;
+	// Info on grouped material sections within vertex buffers but with internal faces added by the fracture tool removed.
+	TArray<FGeometryCollectionMeshElement> SectionsNoInternal;
+	// Info on individual subsections within vertex buffers. 
+	// A subsection has a uniform bone and material index.
+	TArray<FGeometryCollectionMeshElement> SubSections;
+
+	/** Serialization. */
+	void Serialize(FArchive& Ar);
+};
+
+/** Built render data compiled and ready for use by the geometry collection scene proxy. */
+class FGeometryCollectionRenderData
+{
+public:
+	FGeometryCollectionRenderData() = default;
+	~FGeometryCollectionRenderData();
+
+#if WITH_EDITOR
+	static TUniquePtr<FGeometryCollectionRenderData> Create(FGeometryCollection& InCollection, bool bInEnableNanite, bool bInUseFullPrecisionUVs);
+#endif
+
+	bool IsInitialized()
+	{
+		return bIsInitialized;
+	}
+
+	/** Serialization. */
+	void Serialize(FArchive& Ar, UGeometryCollection* Owner);
+
+	/** Initialize the render resources. */
+	void InitResources(UGeometryCollection* Owner);
+
+	/** Releases the render resources. */
+	void ReleaseResources();
+
+	bool bHasMeshData = false;
+	bool bHasNaniteData = false;
+
+	/** Mesh GPU resources. */
+	FGeometryCollectionMeshResources MeshResource;
+
+	/** Mesh CPU description. */
+	FGeometryCollectionMeshDescription MeshDescription;
+
+	/** Nanite resources. */
+	Nanite::FResources NaniteResource;
+
+private:
+	bool bIsInitialized = false;
+};
diff --git a/Engine/Source/Runtime/Experimental/GeometryCollectionEngine/Private/GeometryCollection/GeometryCollectionSceneProxy.cpp b/Engine/Source/Runtime/Experimental/GeometryCollectionEngine/Private/GeometryCollection/GeometryCollectionSceneProxy.cpp
--- a/Engine/Source/Runtime/Experimental/GeometryCollectionEngine/Private/GeometryCollection/GeometryCollectionSceneProxy.cpp
+++ b/Engine/Source/Runtime/Experimental/GeometryCollectionEngine/Private/GeometryCollection/GeometryCollectionSceneProxy.cpp
@@ -7,7 +7,6 @@
 #include "GeometryCollection/GeometryCollection.h"
 #include "GeometryCollection/GeometryCollectionObject.h"
 #include "MaterialDomain.h"
-#include "MaterialShared.h"
 #include "MaterialShaderType.h"
 #include "Materials/Material.h"
 #include "Materials/MaterialRenderProxy.h"
@@ -16,9 +15,13 @@
 #include "PrimitiveSceneInfo.h"
 #include "GeometryCollection/GeometryCollectionComponent.h"
 #include "GeometryCollection/GeometryCollectionAlgo.h"
-#include "RHIDefinitions.h"
-#if GEOMETRYCOLLECTION_EDITOR_SELECTION
 #include "GeometryCollection/GeometryCollectionHitProxy.h"
+#include "RHIDefinitions.h"
+#include "ComponentReregisterContext.h"
+#include "ComponentRecreateRenderStateContext.h"
+
+#if RHI_RAYTRACING
+#include "RayTracingInstance.h"
 #endif
 
 #if INTEL_ISPC
@@ -46,12 +49,6 @@ static_assert(sizeof(ispc::FMatrix44f) == sizeof(FMatrix44f), "sizeof(ispc::FMat
 static_assert(sizeof(ispc::FVector3f) == sizeof(FVector3f), "sizeof(ispc::FVector3f) != sizeof(FVector3f)");
 #endif
 
-#include "ComponentReregisterContext.h"
-#include "ComponentRecreateRenderStateContext.h"
-
-#if RHI_RAYTRACING
-#include "RayTracingInstance.h"
-#endif
 
 static int32 GParallelGeometryCollectionBatchSize = 1024;
 static TAutoConsoleVariable<int32> CVarParallelGeometryCollectionBatchSize(
@@ -89,21 +86,6 @@ FAutoConsoleVariableRef CVarRayTracingGeometryCollectionProxyMeshes(
 	ECVF_RenderThreadSafe
 );
 
-static int32 GeometryCollectionBuildGeometryParallelVertexProcessingMinBatch = 5000;
-FAutoConsoleVariableRef CVarGeometryCollectionBuildGeometryParallelVertexProcessingMinBatch(
-	TEXT("r.GeometryCollectionBuildGeometry.ParallelVertexProcessingMinBatch"),
-	GeometryCollectionBuildGeometryParallelVertexProcessingMinBatch,
-	TEXT("When building geometry collection rendering geometry, minimum number of vertices to batch for parallel processing performance tuning"),
-	ECVF_RenderThreadSafe
-);
-
-static int32 GeometryCollectionBuildGeometryParallelIndexProcessingMinBatch = 5000;
-FAutoConsoleVariableRef CVarGeometryCollectionBuildGeometryParallelIndexProcessingMinBatch(
-	TEXT("r.GeometryCollectionBuildGeometry.ParallelIndexProcessingMinBatch"),
-	GeometryCollectionBuildGeometryParallelIndexProcessingMinBatch,
-	TEXT("When building geometry collection rendering geometry, minimum number of indices to batch for parallel processing performance tuning"),
-	ECVF_RenderThreadSafe
-);
 
 #if !defined(CHAOS_GEOMETRY_COLLECTION_SET_DYNAMIC_DATA_ISPC_ENABLED_DEFAULT)
 #define CHAOS_GEOMETRY_COLLECTION_SET_DYNAMIC_DATA_ISPC_ENABLED_DEFAULT 1
@@ -121,14 +103,17 @@ DEFINE_LOG_CATEGORY_STATIC(FGeometryCollectionSceneProxyLogging, Log, All);
 
 FGeometryCollectionDynamicDataPool GDynamicDataPool;
 
+
 class FGeometryCollectionMeshCollectorResources : public FOneFrameResource
 {
 public:
 	FGeometryCollectionVertexFactory VertexFactory;
 	
-	FGeometryCollectionMeshCollectorResources(ERHIFeatureLevel::Type InFeatureLevel):VertexFactory(InFeatureLevel,true)
+	FGeometryCollectionMeshCollectorResources(ERHIFeatureLevel::Type InFeatureLevel)
+		: VertexFactory(InFeatureLevel,true)
 	{
 	}
+
 	virtual ~FGeometryCollectionMeshCollectorResources()
 	{
 		VertexFactory.ReleaseResource();
@@ -137,27 +122,23 @@ public:
 	virtual FGeometryCollectionVertexFactory& GetVertexFactory() { return VertexFactory; }
 };
 
+
 FGeometryCollectionSceneProxy::FGeometryCollectionSceneProxy(UGeometryCollectionComponent* Component)
 	: FPrimitiveSceneProxy(Component)
 	, MaterialRelevance(Component->GetMaterialRelevance(GetScene().GetFeatureLevel()))
-	, NumVertices(0)
-	, NumIndices(0)
+	, MeshResource(Component->GetRestCollection()->RenderData->MeshResource)
+	, MeshDescription(Component->GetRestCollection()->RenderData->MeshDescription)
 	, VertexFactory(GetScene().GetFeatureLevel())
 	, bSupportsManualVertexFetch(VertexFactory.SupportsManualVertexFetch(GetScene().GetFeatureLevel()))
 	, bSupportsTripleBufferVertexUpload(GRHISupportsMapWriteNoOverwrite)
-#if GEOMETRYCOLLECTION_EDITOR_SELECTION
-	, SubSections()
-	, SubSectionHitProxies()
-	, SubSectionHitProxyIndexMap()
-	, bUsesSubSections(false)
-#endif
-	, DynamicData(nullptr)
+#if WITH_EDITOR
 	, bShowBoneColors(Component->GetShowBoneColors())
-	, bEnableBoneSelection(Component->GetEnableBoneSelection())
 	, bSuppressSelectionMaterial(Component->GetSuppressSelectionMaterial())
-	, bUseFullPrecisionUVs(Component->GetRestCollection()->bUseFullPrecisionUVs)
-	, TransformVertexBuffersContainsOriginalMesh(false)
+	, VertexFactoryDebugColor(GetScene().GetFeatureLevel())
+#endif
 {
+	EnableGPUSceneSupportFlags();
+
 	Materials.Empty();
 	const int32 NumMaterials = Component->GetNumMaterials();
 	for (int MaterialIndex = 0; MaterialIndex < NumMaterials; ++MaterialIndex)
@@ -169,146 +150,96 @@ FGeometryCollectionSceneProxy::FGeometryCollectionSceneProxy(UGeometryCollection
 			Materials[MaterialIndex] = UMaterial::GetDefaultMaterial(MD_Surface);
 		}
 	}
-	if (Component->GetRestCollection())
-	{
-		BoneSelectedMaterial = Component->GetRestCollection()->GetBoneSelectedMaterial();
-	}
-	if (!BoneSelectedMaterial)
-	{
-		int32 LegacyBoneMaterialID = Component->GetBoneSelectedMaterialID();
-		if (Materials.IsValidIndex(LegacyBoneMaterialID))
-		{
-			BoneSelectedMaterial = Materials[LegacyBoneMaterialID];
-		}
-	}
-	if (BoneSelectedMaterial && !BoneSelectedMaterial->CheckMaterialUsage_Concurrent(MATUSAGE_GeometryCollections))
-	{
-		// If we have an invalid BoneSelectedMaterial, switch it back to null to skip its usage in GetDynamicMeshElements below
-		BoneSelectedMaterial = nullptr;
-	}
 
-	// Make sure the vertex color material has the usage flag for rendering geometry collections
-	if (GEngine->VertexColorMaterial)
-	{
-		GEngine->VertexColorMaterial->CheckMaterialUsage_Concurrent(MATUSAGE_GeometryCollections);
-	}
-
-	// #todo(dmp): We create the sections before we set the constant data because we need to make sure these
-	// are set before the hit proxies are created via CreateHitProxies.  Ideally, all data is passed in
-	// here when we create proxies, and they are thrown away if underlying geometry changes.
-	TManagedArray<FGeometryCollectionSection>& InputSections = Component->GetRestCollection()->GetGeometryCollection()->Sections;
+	Component->GetRestTransforms(RestTransforms);
+	NumTransforms = RestTransforms.Num();
 
-	const int32 NumSections = InputSections.Num();
-	Sections.Reset(NumSections);
+#if GEOMETRYCOLLECTION_EDITOR_SELECTION
+	// Render by SubSection if we are in the rigid body picker.
+	bUsesSubSections = Component->GetIsTransformSelectionMode() && MeshDescription.SubSections.Num();
+	// Enable bone hit selection proxies if we are in the rigid body picker or in the fracture modes.
+	bEnableBoneSelection = Component->GetEnableBoneSelection();
 
-	for (int SectionIndex = 0; SectionIndex < NumSections; ++SectionIndex)
+	if (bEnableBoneSelection || bUsesSubSections)
 	{
-		FGeometryCollectionSection& Section = InputSections[SectionIndex];
-		
-		if (Section.NumTriangles > 0)
+		for (int32 TransformIndex = 0; TransformIndex < NumTransforms; ++TransformIndex)
 		{
-			Sections.Add(Section);
+			HGeometryCollection* HitProxy = new HGeometryCollection(Component, TransformIndex);
+			HitProxies.Add(HitProxy);
 		}
 	}
-
-	EnableGPUSceneSupportFlags();
-
-#if GEOMETRYCOLLECTION_EDITOR_SELECTION
-	// Init HitProxy array with the maximum number of subsections
-	SubSectionHitProxies.SetNumZeroed(Sections.Num() * Component->GetTransformArray().Num());
 #endif
 
-	// #todo(dmp): This flag means that when motion blur is turned on, it will always render geometry collections into the
-	// velocity buffer.  Note that the way around this is to loop through the global matrices and test whether they have
-	// changed from the prev to curr frame, but this is expensive.  We should revisit this if the draw calls for velocity
-	// rendering become a problem. One solution could be to use internal solver sleeping state to drive motion blur.
-	bAlwaysHasVelocity = true;
-
-	// Build pre-skinned bounds from the rest collection, never needs to change as this is the bounds before
-	// any movement, or skinning ever happens to the component so it is logically immutable.
-	const TSharedPtr<FGeometryCollection, ESPMode::ThreadSafe>            
-										Collection     = Component->RestCollection->GetGeometryCollection();
-	const TManagedArray<FBox>&			BoundingBoxes  = Collection->BoundingBox;
-	const TManagedArray<FTransform>&	Transform      = Collection->Transform;
-	const TManagedArray<int32>&			Parent         = Collection->Parent;
-	const TManagedArray<int32>&			TransformIndex = Collection->TransformIndex;
-
-	const int32 NumBoxes = BoundingBoxes.Num();
-	PreSkinnedBounds = Component->Bounds;
-
-	if(NumBoxes > 0)
+#if WITH_EDITOR
+	if (bShowBoneColors || bEnableBoneSelection)
 	{
-		TArray<FMatrix> TmpGlobalMatrices;
-		GeometryCollectionAlgo::GlobalMatrices(Transform, Parent, TmpGlobalMatrices);
+		Component->GetBoneColors(BoneColors);
+		ColorVertexBuffer.InitFromColorArray(BoneColors);
 
-		FBox PreSkinnedBoundsTemp(ForceInit);
-		bool bBoundsInit = false;
-		for(int32 BoxIdx = 0; BoxIdx < NumBoxes; ++BoxIdx)
+		if (Component->GetRestCollection())
+		{
+			BoneSelectedMaterial = Component->GetRestCollection()->GetBoneSelectedMaterial();
+		}
+		if (!BoneSelectedMaterial)
 		{
-			const int32 TIndex = TransformIndex[BoxIdx];
-			if(Collection->IsGeometry(TIndex))
+			int32 LegacyBoneMaterialID = Component->GetBoneSelectedMaterialID();
+			if (Materials.IsValidIndex(LegacyBoneMaterialID))
 			{
-				if(!bBoundsInit)
-				{
-					PreSkinnedBoundsTemp = BoundingBoxes[BoxIdx].TransformBy(TmpGlobalMatrices[TIndex]);
-					bBoundsInit = true;
-				}
-				else
-				{
-					PreSkinnedBoundsTemp += BoundingBoxes[BoxIdx].TransformBy(TmpGlobalMatrices[TIndex]);
-				}
+				BoneSelectedMaterial = Materials[LegacyBoneMaterialID];
 			}
 		}
+		if (BoneSelectedMaterial && !BoneSelectedMaterial->CheckMaterialUsage_Concurrent(MATUSAGE_GeometryCollections))
+		{
+			// If we have an invalid BoneSelectedMaterial, switch it back to null to skip its usage in GetDynamicMeshElements below
+			BoneSelectedMaterial = nullptr;
+		}
 
-		PreSkinnedBounds = FBoxSphereBounds(PreSkinnedBoundsTemp);
+		// Make sure the vertex color material has the usage flag for rendering geometry collections
+		if (GEngine->VertexColorMaterial)
+		{
+			GEngine->VertexColorMaterial->CheckMaterialUsage_Concurrent(MATUSAGE_GeometryCollections);
+		}
 	}
 
-#if RHI_RAYTRACING
-	if (IsRayTracingAllowed())
+	// Get hidden geometry and zero the associated transforms.
+	Component->GetHiddenTransforms(HiddenTransforms);
+	if (HiddenTransforms.Num())
 	{
-
-		ENQUEUE_RENDER_COMMAND(InitGeometryCollectionRayTracingGeometry)(
-			[this](FRHICommandListImmediate& RHICmdList)
+		check(HiddenTransforms.Num() == RestTransforms.Num());
+		for (int32 TransformIndex = 0; TransformIndex < RestTransforms.Num(); ++TransformIndex)
+		{
+			if (HiddenTransforms[TransformIndex])
 			{
-				FRayTracingGeometryInitializer Initializer;
-				Initializer.DebugName = FName(TEXT("GeometryCollection"));
-				Initializer.GeometryType = RTGT_Triangles;
-				Initializer.bFastBuild = true;
-				Initializer.bAllowUpdate = false;
-				Initializer.TotalPrimitiveCount = 0;
-
-				RayTracingGeometry.SetInitializer(Initializer);
-				RayTracingGeometry.InitResource();
-			});
+				RestTransforms[TransformIndex] = FMatrix44f(EForceInit::ForceInitToZero);
+			}
+		}
 	}
 #endif
-}
 
-FGeometryCollectionSceneProxy::~FGeometryCollectionSceneProxy()
-{
+	// #todo(dmp): This flag means that when motion blur is turned on, it will always render geometry collections into the
+	// velocity buffer.  Note that the way around this is to loop through the global matrices and test whether they have
+	// changed from the prev to curr frame, but this is expensive.  We should revisit this if the draw calls for velocity
+	// rendering become a problem. One solution could be to use internal solver sleeping state to drive motion blur.
+	bAlwaysHasVelocity = true;
 }
 
-void FGeometryCollectionSceneProxy::DestroyRenderThreadResources()
+FGeometryCollectionSceneProxy::~FGeometryCollectionSceneProxy()
 {
-	ReleaseResources();
-
 	if (DynamicData != nullptr)
 	{
 		GDynamicDataPool.Release(DynamicData);
 		DynamicData = nullptr;
 	}
+}
 
-
-#if RHI_RAYTRACING
-	if (IsRayTracingAllowed())
-	{
-		RayTracingGeometry.ReleaseResource();
-		RayTracingDynamicVertexBuffer.Release();
-	}
-#endif
+SIZE_T FGeometryCollectionSceneProxy::GetTypeHash() const
+{
+	static size_t UniquePointer;
+	return reinterpret_cast<size_t>(&UniquePointer);
 }
 
-void UpdateLooseParameter(FGeometryCollectionVertexFactory& VertexFactory,
+static void UpdateLooseParameter(
+	FGeometryCollectionVertexFactory& VertexFactory,
 	FRHIShaderResourceView* BoneTransformSRV,
 	FRHIShaderResourceView* BonePrevTransformSRV,
 	FRHIShaderResourceView* BoneMapSRV)
@@ -324,19 +255,33 @@ void UpdateLooseParameter(FGeometryCollectionVertexFactory& VertexFactory,
 	VertexFactory.LooseParameterUniformBuffer = FGCBoneLooseParametersRef::CreateUniformBufferImmediate(LooseParameters, UniformBufferUsage);
 }
 
-void FGeometryCollectionSceneProxy::SetupVertexFactory(FGeometryCollectionVertexFactory& GeometryCollectionVertexFactory)const
+void FGeometryCollectionSceneProxy::SetupVertexFactory(FGeometryCollectionVertexFactory& GeometryCollectionVertexFactory, FColorVertexBuffer* ColorOverride) const
 {
 	FGeometryCollectionVertexFactory::FDataType Data;
 	
-	VertexBuffers.PositionVertexBuffer.BindPositionVertexBuffer(&GeometryCollectionVertexFactory, Data);
-	VertexBuffers.StaticMeshVertexBuffer.BindTangentVertexBuffer(&GeometryCollectionVertexFactory, Data);
-	VertexBuffers.StaticMeshVertexBuffer.BindPackedTexCoordVertexBuffer(&GeometryCollectionVertexFactory, Data);
-	VertexBuffers.StaticMeshVertexBuffer.BindLightMapVertexBuffer(&GeometryCollectionVertexFactory, Data, 0);
-	VertexBuffers.ColorVertexBuffer.BindColorVertexBuffer(&GeometryCollectionVertexFactory, Data);
+	FPositionVertexBuffer const& PositionVB = bSupportsManualVertexFetch ? MeshResource.PositionVertexBuffer : SkinnedPositionVertexBuffer;
+	PositionVB.BindPositionVertexBuffer(&GeometryCollectionVertexFactory, Data);
+
+	MeshResource.StaticMeshVertexBuffer.BindTangentVertexBuffer(&GeometryCollectionVertexFactory, Data);
+	MeshResource.StaticMeshVertexBuffer.BindPackedTexCoordVertexBuffer(&GeometryCollectionVertexFactory, Data);
+	MeshResource.StaticMeshVertexBuffer.BindLightMapVertexBuffer(&GeometryCollectionVertexFactory, Data, 0);
 
-	Data.BoneMapSRV = BoneMapBuffer.VertexBufferSRV;
-	Data.BoneTransformSRV = TransformBuffers[CurrentTransformBufferIndex].VertexBufferSRV;
-	Data.BonePrevTransformSRV = PrevTransformBuffers[CurrentTransformBufferIndex].VertexBufferSRV;
+	FColorVertexBuffer const& ColorVB = ColorOverride ? *ColorOverride : MeshResource.ColorVertexBuffer;
+	ColorVB.BindColorVertexBuffer(&GeometryCollectionVertexFactory, Data);
+
+	if (bSupportsManualVertexFetch)
+	{
+		Data.BoneMapSRV = MeshResource.BoneMapVertexBuffer.GetSRV();
+		Data.BoneTransformSRV = TransformBuffers[CurrentTransformBufferIndex].VertexBufferSRV;
+		Data.BonePrevTransformSRV = PrevTransformBuffers[CurrentTransformBufferIndex].VertexBufferSRV;
+	}
+	else
+	{
+		// Make sure these are not null to pass UB validation
+		Data.BoneMapSRV = GNullColorVertexBuffer.VertexBufferSRV;
+		Data.BoneTransformSRV = GNullColorVertexBuffer.VertexBufferSRV;
+		Data.BonePrevTransformSRV = GNullColorVertexBuffer.VertexBufferSRV;
+	}
 
 	GeometryCollectionVertexFactory.SetData(Data);
 
@@ -350,111 +295,88 @@ void FGeometryCollectionSceneProxy::SetupVertexFactory(FGeometryCollectionVertex
 	}
 }
 
-void FGeometryCollectionSceneProxy::InitResources(const TUniquePtr<FGeometryCollectionConstantData>& ConstantData)
+void FGeometryCollectionSceneProxy::CreateRenderThreadResources()
 {
-	check(ConstantData);
-	check(IsInRenderingThread());
-	
-	NumVertices = ConstantData->Vertices.Num();
-	NumIndices = ConstantData->Indices.Num()*3;	
-
-	// taken from this, and expanded here to accommodate modifications for
-	// GeometryCollection vertex factory data (transform and bonemap)
-	// VertexBuffers.InitWithDummyData(&VertexFactory, GetRequiredVertexCount());
-
-	// get vertex factory data
-	FGeometryCollectionVertexFactory::FDataType Data;
-
-	VertexBuffers.StaticMeshVertexBuffer.SetUseFullPrecisionUVs(bUseFullPrecisionUVs);
-	
-	// Init buffers
-	VertexBuffers.PositionVertexBuffer.Init(NumVertices);
-	VertexBuffers.StaticMeshVertexBuffer.Init(NumVertices, GeometryCollectionUV::MAX_NUM_UV_CHANNELS);
-	VertexBuffers.ColorVertexBuffer.Init(NumVertices);
-
-	// Init resources
-	VertexBuffers.PositionVertexBuffer.InitResource();
-	VertexBuffers.StaticMeshVertexBuffer.InitResource();
-	VertexBuffers.ColorVertexBuffer.InitResource();
-
-	// Bind buffers
-	VertexBuffers.PositionVertexBuffer.BindPositionVertexBuffer(&VertexFactory, Data);
-	VertexBuffers.StaticMeshVertexBuffer.BindTangentVertexBuffer(&VertexFactory, Data);
-	VertexBuffers.StaticMeshVertexBuffer.BindPackedTexCoordVertexBuffer(&VertexFactory, Data);	
-	VertexBuffers.StaticMeshVertexBuffer.BindLightMapVertexBuffer(&VertexFactory, Data, 0);
-	VertexBuffers.ColorVertexBuffer.BindColorVertexBuffer(&VertexFactory, Data);
-
-#if GEOMETRYCOLLECTION_EDITOR_SELECTION
-	// Note: Could skip this when bEnableBoneSelection is false if the hitproxy shader was made to not require per-vertex hit proxy IDs in that case
-	{
-		HitProxyIdBuffer.Init(NumVertices);
-		HitProxyIdBuffer.InitResource();
-	}
-#endif
-
-	IndexBuffer.NumIndices = GetRequiredIndexCount();
-	IndexBuffer.InitResource();
-		
-	OriginalMeshIndexBuffer.NumIndices = GetRequiredIndexCount();
-	OriginalMeshIndexBuffer.InitResource();	
-	
-	// If using manual vertex fetch, then we will setup the GPU point transform implementation
 	if (bSupportsManualVertexFetch)
 	{
-		BoneMapBuffer.NumVertices = NumVertices;
-
+		// Initialize transform buffers and upload rest transforms.
 		TransformBuffers.AddDefaulted(1);
 		PrevTransformBuffers.AddDefaulted(1);
 
-		TransformBuffers[0].NumTransforms = ConstantData->NumTransforms;
-		PrevTransformBuffers[0].NumTransforms = ConstantData->NumTransforms;
+		TransformBuffers[0].NumTransforms = NumTransforms;
+		PrevTransformBuffers[0].NumTransforms = NumTransforms;
 		TransformBuffers[0].InitResource();
 		PrevTransformBuffers[0].InitResource();
 
-		BoneMapBuffer.InitResource();
+		const bool bLocalGeometryCollectionTripleBufferUploads = (GGeometryCollectionTripleBufferUploads != 0) && bSupportsTripleBufferVertexUpload;
+		const EResourceLockMode LockMode = bLocalGeometryCollectionTripleBufferUploads ? RLM_WriteOnly_NoOverwrite : RLM_WriteOnly;
 
-		Data.BoneMapSRV = BoneMapBuffer.VertexBufferSRV;
-		Data.BoneTransformSRV = TransformBuffers[0].VertexBufferSRV;
-		Data.BonePrevTransformSRV = PrevTransformBuffers[0].VertexBufferSRV;
+		FGeometryCollectionTransformBuffer& TransformBuffer = GetCurrentTransformBuffer();
+		TransformBuffer.UpdateDynamicData(RestTransforms, LockMode);
+		FGeometryCollectionTransformBuffer& PrevTransformBuffer = GetCurrentPrevTransformBuffer();
+		PrevTransformBuffer.UpdateDynamicData(RestTransforms, LockMode);
 	}
 	else
 	{
-		// make sure these are not null to pass UB validation
-		Data.BoneMapSRV = GNullColorVertexBuffer.VertexBufferSRV;
-		Data.BoneTransformSRV = GNullColorVertexBuffer.VertexBufferSRV;
-		Data.BonePrevTransformSRV = GNullColorVertexBuffer.VertexBufferSRV;
+		// Initialize CPU skinning buffer with rest transforms.
+		SkinnedPositionVertexBuffer.Init(MeshResource.PositionVertexBuffer.GetNumVertices(), false);
+		SkinnedPositionVertexBuffer.InitResource();
+		UpdateSkinnedPositions(RestTransforms);
 	}
 
-	// 
-	// from InitOrUpdateResource(VertexFactory);
-	//
+	SetupVertexFactory(VertexFactory);
 
-	// also make sure to do the binding to the vertex factory
-	VertexFactory.SetData(Data);
+#if WITH_EDITOR
+	if (bShowBoneColors || bEnableBoneSelection)
+	{
+		// Initialize debug color buffer and associated vertex factory.
+		ColorVertexBuffer.InitResource();
+		SetupVertexFactory(VertexFactoryDebugColor, &ColorVertexBuffer);
+	}
+#endif
 
-	if (!VertexFactory.IsInitialized())
+#if GEOMETRYCOLLECTION_EDITOR_SELECTION
+	if (MeshDescription.NumVertices && HitProxies.Num())
 	{
-		VertexFactory.InitResource();
+		// Create buffer containing per vertex hit proxy IDs.
+		HitProxyIdBuffer.Init(MeshDescription.NumVertices);
+		HitProxyIdBuffer.InitResource();
+
+		uint16 const* BoneMapData = &MeshResource.BoneMapVertexBuffer.BoneIndex(0);
+		ParallelFor(MeshDescription.NumVertices, [&](int32 i)
+		{
+			const uint16 ProxyIndex = BoneMapData[i];
+			HitProxyIdBuffer.VertexColor(i) = HitProxies[ProxyIndex]->Id.GetColor();
+		});
+
+		void* VertexBufferData = RHILockBuffer(HitProxyIdBuffer.VertexBufferRHI, 0, HitProxyIdBuffer.GetNumVertices() * HitProxyIdBuffer.GetStride(), RLM_WriteOnly);
+		FMemory::Memcpy(VertexBufferData, HitProxyIdBuffer.GetVertexData(), HitProxyIdBuffer.GetNumVertices() * HitProxyIdBuffer.GetStride());
+		RHIUnlockBuffer(HitProxyIdBuffer.VertexBufferRHI);
 	}
-	else
+#endif
+
+#if RHI_RAYTRACING
+	if (IsRayTracingEnabled())
 	{
-		VertexFactory.UpdateRHI();
+		FRayTracingGeometryInitializer Initializer;
+		Initializer.DebugName = FName(TEXT("GeometryCollection"));
+		Initializer.GeometryType = RTGT_Triangles;
+		Initializer.bFastBuild = true;
+		Initializer.bAllowUpdate = false;
+		Initializer.TotalPrimitiveCount = 0;
+
+		RayTracingGeometry.SetInitializer(Initializer);
+		RayTracingGeometry.InitResource();
+
+		bGeometryResourceUpdated = true;
 	}
+#endif
 }
 
-void FGeometryCollectionSceneProxy::ReleaseResources()
+void FGeometryCollectionSceneProxy::DestroyRenderThreadResources()
 {
-	VertexBuffers.PositionVertexBuffer.ReleaseResource();
-	VertexBuffers.StaticMeshVertexBuffer.ReleaseResource();
-	VertexBuffers.ColorVertexBuffer.ReleaseResource();
-	IndexBuffer.ReleaseResource();
-
-	OriginalMeshIndexBuffer.ReleaseResource();
-
 	if (bSupportsManualVertexFetch)
 	{
-		BoneMapBuffer.ReleaseResource();
-
 		for (int32 i = 0; i < TransformBuffers.Num(); i++)
 		{
 			TransformBuffers[i].ReleaseResource();
@@ -462,244 +384,125 @@ void FGeometryCollectionSceneProxy::ReleaseResources()
 		}
 		TransformBuffers.Reset();
 	}
+	else
+	{
+		SkinnedPositionVertexBuffer.ReleaseResource();
+	}
+
+	VertexFactory.ReleaseResource();
+
+#if WITH_EDITOR
+	VertexFactoryDebugColor.ReleaseResource();
+	ColorVertexBuffer.ReleaseResource();
+#endif
 
 #if GEOMETRYCOLLECTION_EDITOR_SELECTION
 	HitProxyIdBuffer.ReleaseResource();
 #endif
 
-	VertexFactory.ReleaseResource();
+#if RHI_RAYTRACING
+	if (IsRayTracingAllowed())
+	{
+		RayTracingGeometry.ReleaseResource();
+		RayTracingDynamicVertexBuffer.Release();
+	}
+#endif
 }
 
-inline static FColor GetOnlyVertexColors(const TUniquePtr<FGeometryCollectionConstantData>& ConstantDataIn, int32 PointIdx, bool bShowBoneColors)
-{
-	return ConstantDataIn->Colors[PointIdx].ToFColor(true);
-};
-
-inline static FColor GetBonesColors(const TUniquePtr<FGeometryCollectionConstantData>& ConstantDataIn, int32 PointIdx, bool bShowBoneColors)
-{
-	return bShowBoneColors
-		? ConstantDataIn->BoneColors[PointIdx].ToFColor(true)
-		: (ConstantDataIn->BoneColors[PointIdx] * ConstantDataIn->Colors[PointIdx]).ToFColor(true);
-};
-
-void FGeometryCollectionSceneProxy::BuildGeometry(const TUniquePtr<FGeometryCollectionConstantData>& ConstantDataIn, TArray<FDynamicMeshVertex>& OutVertices, TArray<int32>& OutIndices, TArray<int32> &OutOriginalMeshIndices)
+void FGeometryCollectionSceneProxy::SetDynamicData_RenderThread(FGeometryCollectionDynamicData* NewDynamicData)
 {
-	const bool bCapturedShowBoneColors = this->bShowBoneColors;
-	const auto SetOutVertex = [&OutVertices, &ConstantDataIn, bCapturedShowBoneColors](int32 PointIdx, FColor (*GetColorFunc)(const TUniquePtr<FGeometryCollectionConstantData>&, int32, bool) )
+	check(IsInRenderingThread());
+	if (DynamicData)
 	{
-		OutVertices[PointIdx] =
-			FDynamicMeshVertex(
-				ConstantDataIn->Vertices[PointIdx],
-				ConstantDataIn->UVChannels[0][PointIdx],
-				GetColorFunc(ConstantDataIn, PointIdx, bCapturedShowBoneColors)
-			);
-		OutVertices[PointIdx].SetTangents(ConstantDataIn->TangentU[PointIdx], ConstantDataIn->TangentV[PointIdx], ConstantDataIn->Normals[PointIdx]);
-
-		const int32 NumUvChannels = ConstantDataIn->UVChannels.Num();
-		if (ConstantDataIn->UVChannels.Num() > 1)
-		{
-			for (int32 UVLayerIdx = 1; UVLayerIdx < NumUvChannels; ++UVLayerIdx)
-			{
-				const FGeometryCollectionConstantData::FUVChannel& UVChannel = ConstantDataIn->UVChannels[UVLayerIdx];
-				OutVertices[PointIdx].TextureCoordinate[UVLayerIdx] = UVChannel[PointIdx];
-			}
-		}
-	};
+		GDynamicDataPool.Release(DynamicData);
+		DynamicData = nullptr;
+	}
+	DynamicData = NewDynamicData;
 
-	OutVertices.SetNumUninitialized(ConstantDataIn->Vertices.Num());
-	const int32 MinVertexBatchSize = GeometryCollectionBuildGeometryParallelVertexProcessingMinBatch;
-	if (bShowBoneColors || bEnableBoneSelection)
+	if (MeshDescription.NumVertices == 0)
 	{
-		ParallelFor(TEXT("GCSceneProxy::BuildGeometry(Vertices)"), ConstantDataIn->Vertices.Num(), MinVertexBatchSize,
-			[&SetOutVertex](int32 PointIdx)
-			{
-				SetOutVertex(PointIdx, GetBonesColors);
-			});
+		return;
 	}
-	else
+	// Early out if if we are applying (non-dynamic) rest transforms over multiple frames.
+	if (!DynamicData->IsDynamic && TransformVertexBuffersContainsRestTransforms)
 	{
-		ParallelFor(TEXT("GCSceneProxy::BuildGeometry(Vertices)"), ConstantDataIn->Vertices.Num(), MinVertexBatchSize,
-			[&SetOutVertex](int32 PointIdx)
-			{
-				SetOutVertex(PointIdx, GetOnlyVertexColors);
-			});
+		return;
 	}
-	check(ConstantDataIn->Indices.Num() * 3 == NumIndices);
-
-	const int32 MinIndexBatchSize = GeometryCollectionBuildGeometryParallelIndexProcessingMinBatch;
+	TransformVertexBuffersContainsRestTransforms = !DynamicData->IsDynamic;
+		
+	if (bSupportsManualVertexFetch)
+	{
+		const bool bLocalGeometryCollectionTripleBufferUploads = (GGeometryCollectionTripleBufferUploads != 0) && bSupportsTripleBufferVertexUpload;
 
-	OutIndices.SetNumUninitialized(NumIndices);
-	ParallelFor(TEXT("GCSceneProxy::BuildGeometry(Indices)"), ConstantDataIn->Indices.Num(), MinIndexBatchSize, 
-		[&OutIndices, &ConstantDataIn](int32 IndexIdx)
+		if (bLocalGeometryCollectionTripleBufferUploads && TransformBuffers.Num() == 1)
 		{
-			OutIndices[IndexIdx * 3 ]    = ConstantDataIn->Indices[IndexIdx].X;
-			OutIndices[IndexIdx * 3 + 1] = ConstantDataIn->Indices[IndexIdx].Y;
-			OutIndices[IndexIdx * 3 + 2] = ConstantDataIn->Indices[IndexIdx].Z;
-		});
-	
-	OutOriginalMeshIndices.SetNumUninitialized(ConstantDataIn->OriginalMeshIndices.Num() * 3);
-	ParallelFor(TEXT("GCSceneProxy::BuildGeometry(OriginalMeshIndices)"), ConstantDataIn->OriginalMeshIndices.Num(), MinIndexBatchSize, 
-		[&OutOriginalMeshIndices, &ConstantDataIn](int32 IndexIdx)
-		{
-			OutOriginalMeshIndices[IndexIdx * 3] = ConstantDataIn->OriginalMeshIndices[IndexIdx].X;
-			OutOriginalMeshIndices[IndexIdx * 3 + 1] = ConstantDataIn->OriginalMeshIndices[IndexIdx].Y;
-			OutOriginalMeshIndices[IndexIdx * 3 + 2] = ConstantDataIn->OriginalMeshIndices[IndexIdx].Z;
-		});
-}
+			TransformBuffers.AddDefaulted(2);
+			PrevTransformBuffers.AddDefaulted(2);
 
-void FGeometryCollectionSceneProxy::SetConstantData_RenderThread(const TUniquePtr<FGeometryCollectionConstantData>& ConstantData, bool ForceInit)
-{
-	check(IsInRenderingThread());
-	check(ConstantData.IsValid());
+			for (int32 i = 1; i < 3; i++)
+			{
+				TransformBuffers[i].NumTransforms = NumTransforms;
+				PrevTransformBuffers[i].NumTransforms = NumTransforms;
+				TransformBuffers[i].InitResource();
+				PrevTransformBuffers[i].InitResource();
+			}
+		}
 
+		// Copy the transform data over to the vertex buffer	
+		{
+			const EResourceLockMode LockMode = bLocalGeometryCollectionTripleBufferUploads ? RLM_WriteOnly_NoOverwrite : RLM_WriteOnly;
 
-	if (ConstantData->Vertices.Num() != VertexBuffers.PositionVertexBuffer.GetNumVertices() || ForceInit)
-	{
-		ReleaseResources();
-		InitResources(ConstantData);
-	}
+			CycleTransformBuffers(bLocalGeometryCollectionTripleBufferUploads);
 
-	TArray<int32> Indices;
-	TArray<int32> OriginalMeshIndices;
-	TArray<FDynamicMeshVertex> Vertices;
-	BuildGeometry(ConstantData, Vertices, Indices, OriginalMeshIndices);
-	check(Vertices.Num() == GetRequiredVertexCount());
-	check(Indices.Num() == GetRequiredIndexCount());
+			FGeometryCollectionTransformBuffer& TransformBuffer = GetCurrentTransformBuffer();
+			FGeometryCollectionTransformBuffer& PrevTransformBuffer = GetCurrentPrevTransformBuffer();
 
-	if (GetRequiredVertexCount())
-	{
-		ParallelFor(Vertices.Num(), [&](int32 i)
-		{
-			const FDynamicMeshVertex& Vertex = Vertices[i];
+			VertexFactory.SetBoneTransformSRV(TransformBuffer.VertexBufferSRV);
+			VertexFactory.SetBonePrevTransformSRV(PrevTransformBuffer.VertexBufferSRV);
 
-			VertexBuffers.PositionVertexBuffer.VertexPosition(i) = Vertex.Position;
-			VertexBuffers.StaticMeshVertexBuffer.SetVertexTangents(i, Vertex.TangentX.ToFVector3f(), Vertex.GetTangentY(), Vertex.TangentZ.ToFVector3f());
-			for (int UVChannelIndex = 0; UVChannelIndex < GeometryCollectionUV::MAX_NUM_UV_CHANNELS; UVChannelIndex++)
+#if WITH_EDITOR
+			// Implement hiding geometry in editor by zeroing the transform.
+			// Could move this to InitDynamicData?
+			if (HiddenTransforms.Num())
 			{
-				VertexBuffers.StaticMeshVertexBuffer.SetVertexUV(i, UVChannelIndex, Vertex.TextureCoordinate[UVChannelIndex]);
+				for (int32 TransformIndex = 0; TransformIndex < DynamicData->Transforms.Num(); ++TransformIndex)
+				{
+					if (HiddenTransforms[TransformIndex])
+					{
+						DynamicData->Transforms[TransformIndex] = FMatrix44f(EForceInit::ForceInitToZero);
+						DynamicData->PrevTransforms[TransformIndex] = FMatrix44f(EForceInit::ForceInitToZero);
+					}
+				}
 			}
-			VertexBuffers.ColorVertexBuffer.VertexColor(i) = Vertex.Color;
-#if GEOMETRYCOLLECTION_EDITOR_SELECTION
-			if (bEnableBoneSelection && PerBoneHitProxies.Num())
+#endif
+				
+			if (DynamicData->IsDynamic)
 			{
-				// One proxy per bone
-				const int32 ProxyIndex = ConstantData->BoneMap[i];
-				HitProxyIdBuffer.VertexColor(i) = PerBoneHitProxies[ProxyIndex]->Id.GetColor();
+				TransformBuffer.UpdateDynamicData(DynamicData->Transforms, LockMode);
+				PrevTransformBuffer.UpdateDynamicData(DynamicData->PrevTransforms, LockMode);
 			}
 			else
 			{
-				HitProxyIdBuffer.VertexColor(i) = WholeObjectHitProxyColor;
+				// If we are rendering the base mesh geometry then use RestTransforms for both current and previous transforms.
+				TransformBuffer.UpdateDynamicData(RestTransforms, LockMode);
+				PrevTransformBuffer.UpdateDynamicData(RestTransforms, LockMode);
 			}
-#endif
-		});
-
-		{
-			auto& VertexBuffer = VertexBuffers.PositionVertexBuffer;
-			void* VertexBufferData = RHILockBuffer(VertexBuffer.VertexBufferRHI, 0, VertexBuffer.GetNumVertices() * VertexBuffer.GetStride(), RLM_WriteOnly);
-			FMemory::Memcpy(VertexBufferData, VertexBuffer.GetVertexData(), VertexBuffer.GetNumVertices() * VertexBuffer.GetStride());
-			RHIUnlockBuffer(VertexBuffer.VertexBufferRHI);
-		}
-
-		{
-			auto& VertexBuffer = VertexBuffers.ColorVertexBuffer;
-			void* VertexBufferData = RHILockBuffer(VertexBuffer.VertexBufferRHI, 0, VertexBuffer.GetNumVertices() * VertexBuffer.GetStride(), RLM_WriteOnly);
-			FMemory::Memcpy(VertexBufferData, VertexBuffer.GetVertexData(), VertexBuffer.GetNumVertices() * VertexBuffer.GetStride());
-			RHIUnlockBuffer(VertexBuffer.VertexBufferRHI);
-		}
-
-#if GEOMETRYCOLLECTION_EDITOR_SELECTION
-		// Note: Could skip this when bEnableBoneSelection is false if the hitproxy shader was made to not require per-vertex hit proxy IDs in that case
-		{
-			auto& VertexBuffer = HitProxyIdBuffer;
-			void* VertexBufferData = RHILockBuffer(VertexBuffer.VertexBufferRHI, 0, VertexBuffer.GetNumVertices() * VertexBuffer.GetStride(), RLM_WriteOnly);
-			FMemory::Memcpy(VertexBufferData, VertexBuffer.GetVertexData(), VertexBuffer.GetNumVertices() * VertexBuffer.GetStride());
-			RHIUnlockBuffer(VertexBuffer.VertexBufferRHI);
-		}
-#endif
-
-		{
-			auto& VertexBuffer = VertexBuffers.StaticMeshVertexBuffer;
-			void* VertexBufferData = RHILockBuffer(VertexBuffer.TangentsVertexBuffer.VertexBufferRHI, 0, VertexBuffer.GetTangentSize(), RLM_WriteOnly);
-			FMemory::Memcpy(VertexBufferData, VertexBuffer.GetTangentData(), VertexBuffer.GetTangentSize());
-			RHIUnlockBuffer(VertexBuffer.TangentsVertexBuffer.VertexBufferRHI);
-		}
-
-		{
-			auto& VertexBuffer = VertexBuffers.StaticMeshVertexBuffer;
-			void* VertexBufferData = RHILockBuffer(VertexBuffer.TexCoordVertexBuffer.VertexBufferRHI, 0, VertexBuffer.GetTexCoordSize(), RLM_WriteOnly);
-			FMemory::Memcpy(VertexBufferData, VertexBuffer.GetTexCoordData(), VertexBuffer.GetTexCoordSize());
-			RHIUnlockBuffer(VertexBuffer.TexCoordVertexBuffer.VertexBufferRHI);
-		}
-
-		{
-			void* IndexBufferData = RHILockBuffer(IndexBuffer.IndexBufferRHI, 0, Indices.Num() * sizeof(int32), RLM_WriteOnly);
-			FMemory::Memcpy(IndexBufferData, &Indices[0], Indices.Num() * sizeof(int32));
-			RHIUnlockBuffer(IndexBuffer.IndexBufferRHI);
-		}
-
-		{							
-			void* OriginalMeshIndexBufferData = RHILockBuffer(OriginalMeshIndexBuffer.IndexBufferRHI, 0, OriginalMeshIndices.Num() * sizeof(int32), RLM_WriteOnly);
-			FMemory::Memcpy(OriginalMeshIndexBufferData, &OriginalMeshIndices[0], OriginalMeshIndices.Num() * sizeof(int32));
-			RHIUnlockBuffer(OriginalMeshIndexBuffer.IndexBufferRHI);
-		}
 
-		// If we are using the GeometryCollection vertex factory, populate the vertex buffer for bone map
-		if (bSupportsManualVertexFetch)
-		{
-			void* BoneMapBufferData = RHILockBuffer(BoneMapBuffer.VertexBufferRHI, 0, Vertices.Num() * sizeof(int32), RLM_WriteOnly);
-			FMemory::Memcpy(BoneMapBufferData, &ConstantData->BoneMap[0], ConstantData->BoneMap.Num() * sizeof(int32));
-			RHIUnlockBuffer(BoneMapBuffer.VertexBufferRHI);
-
-			// In order to use loose parameter to support raytracing, we need to initialize the transform/pretransform buffer
-			// before it's set up in the dynamic path. Otherwise, the transformation matrix will be all zero instead of identity. 
-			// Then, nothing will be drawn
-			const bool bLocalGeometryCollectionTripleBufferUploads = (GGeometryCollectionTripleBufferUploads != 0) && bSupportsTripleBufferVertexUpload;
-			const EResourceLockMode LockMode = bLocalGeometryCollectionTripleBufferUploads ? RLM_WriteOnly_NoOverwrite : RLM_WriteOnly;
-
-			FGeometryCollectionTransformBuffer& TransformBuffer = GetCurrentTransformBuffer();
-			FGeometryCollectionTransformBuffer& PrevTransformBuffer = GetCurrentPrevTransformBuffer();
-
-			
-			// if we are rendering the base mesh geometry, then use rest transforms rather than the simulated one for both current and previous transforms
-			TransformBuffer.UpdateDynamicData(ConstantData->RestTransforms, LockMode);
-			PrevTransformBuffer.UpdateDynamicData(ConstantData->RestTransforms, LockMode);
-
-		}
+			UpdateLooseParameter(VertexFactory, TransformBuffer.VertexBufferSRV, PrevTransformBuffer.VertexBufferSRV, MeshResource.BoneMapVertexBuffer.GetSRV());
 
-		// Update mesh sections
-		Sections.Reset(ConstantData->Sections.Num());
-		// #todo(dmp): We should restructure the component/SceneProxy usage to avoid this messy stuff.  We need to know the sections
-		// when we create the sceneproxy for the hit proxy to work, but then we are updating the sections here with potentially differing
-		// vertex counts due to hiding geometry.  Ideally, the SceneProxy is treated as const and recreated whenever the geometry
-		// changes rather than this.  SetConstantData_RenderThread should be done in the constructor for the sceneproxy, most likely
-		for (FGeometryCollectionSection Section : ConstantData->Sections)
-		{
-			if (Section.NumTriangles > 0)
+#if WITH_EDITOR
+			if (bShowBoneColors || bEnableBoneSelection)
 			{
-				FGeometryCollectionSection& NewSection = Sections.AddDefaulted_GetRef();
-				NewSection.MaterialID = Section.MaterialID;
-				NewSection.FirstIndex = Section.FirstIndex;
-				NewSection.NumTriangles = Section.NumTriangles;
-				NewSection.MinVertexIndex = Section.MinVertexIndex;
-				NewSection.MaxVertexIndex = Section.MaxVertexIndex;				
+				VertexFactoryDebugColor.SetBoneTransformSRV(TransformBuffer.VertexBufferSRV);
+				VertexFactoryDebugColor.SetBonePrevTransformSRV(PrevTransformBuffer.VertexBufferSRV);
+				UpdateLooseParameter(VertexFactoryDebugColor, TransformBuffer.VertexBufferSRV, PrevTransformBuffer.VertexBufferSRV, MeshResource.BoneMapVertexBuffer.GetSRV());
 			}
-		}
-
-#if GEOMETRYCOLLECTION_EDITOR_SELECTION
-		// Recreate or release subsections as needed
-		if (bUsesSubSections)
-		{
-			InitializeSubSections_RenderThread();
-		}
-		else
-		{
-			ReleaseSubSections_RenderThread();
+#endif
 		}
 	}
 	else
 	{
-		ReleaseSubSections_RenderThread();
-#endif
+		UpdateSkinnedPositions(DynamicData->IsDynamic ? DynamicData->Transforms : RestTransforms);
 	}
 
 #if RHI_RAYTRACING
@@ -708,182 +511,100 @@ void FGeometryCollectionSceneProxy::SetConstantData_RenderThread(const TUniquePt
 		bGeometryResourceUpdated = true;
 	}
 #endif
-
-	// Move don't copy it ! 
-	SubsetConstantData.BoneMap = MoveTemp(ConstantData->BoneMap);
-	SubsetConstantData.OriginalMeshIndices = MoveTemp(ConstantData->OriginalMeshIndices);
-	SubsetConstantData.OriginalMeshSections = MoveTemp(ConstantData->OriginalMeshSections);
-	SubsetConstantData.Indices = MoveTemp(ConstantData->Indices);
-	SubsetConstantData.NumTransforms = MoveTemp(ConstantData->NumTransforms);
-	SubsetConstantData.RestTransforms = MoveTemp(ConstantData->RestTransforms);
 }
 
-void FGeometryCollectionSceneProxy::SetDynamicData_RenderThread(FGeometryCollectionDynamicData* NewDynamicData)
+void FGeometryCollectionSceneProxy::UpdateSkinnedPositions(TArray<FMatrix44f> const& Transforms)
 {
-	check(IsInRenderingThread());
-	if (GetRequiredVertexCount())
-	{
-		if (DynamicData)
-		{
-			GDynamicDataPool.Release(DynamicData);
-			DynamicData = nullptr;
-		}
-		DynamicData = NewDynamicData;
-		
-
-		const int32 VerticesNum = VertexBuffers.PositionVertexBuffer.GetNumVertices();
+	const int32 VertexStride = SkinnedPositionVertexBuffer.GetStride();
+	const int32 VertexCount = SkinnedPositionVertexBuffer.GetNumVertices();
+	check (VertexCount == MeshDescription.NumVertices)
 
-		if (bSupportsManualVertexFetch)
-		{
-			const bool bLocalGeometryCollectionTripleBufferUploads = (GGeometryCollectionTripleBufferUploads != 0) && bSupportsTripleBufferVertexUpload;
-
-			if (bLocalGeometryCollectionTripleBufferUploads && TransformBuffers.Num() == 1)
-			{
-				TransformBuffers.AddDefaulted(2);
-				PrevTransformBuffers.AddDefaulted(2);
-
-				for (int32 i = 1; i < 3; i++)
-				{
-					TransformBuffers[i].NumTransforms = SubsetConstantData.NumTransforms;
-					PrevTransformBuffers[i].NumTransforms = SubsetConstantData.NumTransforms;
-					TransformBuffers[i].InitResource();
-					PrevTransformBuffers[i].InitResource();
-				}
-			}
-
-			// Copy the transform data over to the vertex buffer	
-			{
-				const EResourceLockMode LockMode = bLocalGeometryCollectionTripleBufferUploads ? RLM_WriteOnly_NoOverwrite : RLM_WriteOnly;
-
-				CycleTransformBuffers(bLocalGeometryCollectionTripleBufferUploads);
-
-				FGeometryCollectionTransformBuffer& TransformBuffer = GetCurrentTransformBuffer();
-				FGeometryCollectionTransformBuffer& PrevTransformBuffer = GetCurrentPrevTransformBuffer();
-
-				VertexFactory.SetBoneTransformSRV(TransformBuffer.VertexBufferSRV);
-				VertexFactory.SetBonePrevTransformSRV(PrevTransformBuffer.VertexBufferSRV);
-
-				if (DynamicData->IsDynamic)
-				{
-					TransformBuffer.UpdateDynamicData(DynamicData->Transforms, LockMode);
-					PrevTransformBuffer.UpdateDynamicData(DynamicData->PrevTransforms, LockMode);
-
-					TransformVertexBuffersContainsOriginalMesh = false;
-				}
-				else if (!TransformVertexBuffersContainsOriginalMesh)
-				{
-					// if we are rendering the base mesh geometry, then use rest transforms rather than the simulated one for both current and previous transforms
-					TransformBuffer.UpdateDynamicData(SubsetConstantData.RestTransforms, LockMode);
-					PrevTransformBuffer.UpdateDynamicData(SubsetConstantData.RestTransforms, LockMode);
-					TransformVertexBuffersContainsOriginalMesh = true;
-				}
+	void* VertexBufferData = RHILockBuffer(SkinnedPositionVertexBuffer.VertexBufferRHI, 0, VertexCount * VertexStride, RLM_WriteOnly);
+	check(VertexBufferData != nullptr);
 
-				UpdateLooseParameter(VertexFactory, TransformBuffer.VertexBufferSRV, PrevTransformBuffer.VertexBufferSRV, BoneMapBuffer.VertexBufferSRV);
-			}
-		}
-		else
-		{
-			auto& VertexBuffer = VertexBuffers.PositionVertexBuffer;
-			void* VertexBufferData = RHILockBuffer(VertexBuffer.VertexBufferRHI, 0, VertexBuffer.GetNumVertices() * VertexBuffer.GetStride(), RLM_WriteOnly);
+	FPositionVertexBuffer const& SourcePositionVertexBuffer = MeshResource.PositionVertexBuffer;
+	FBoneMapVertexBuffer const& SourceBoneMapVertexBuffer = MeshResource.BoneMapVertexBuffer;
 
-			bool bParallelGeometryCollection = true;
-			int32 TotalVertices = VerticesNum;
-			int32 ParallelGeometryCollectionBatchSize = CVarParallelGeometryCollectionBatchSize.GetValueOnRenderThread();
+	bool bParallelGeometryCollection = true;
+	int32 ParallelGeometryCollectionBatchSize = CVarParallelGeometryCollectionBatchSize.GetValueOnRenderThread();
 
-			int32 NumBatches = (TotalVertices / ParallelGeometryCollectionBatchSize);
+	int32 NumBatches = (VertexCount / ParallelGeometryCollectionBatchSize);
 
-			if (TotalVertices != ParallelGeometryCollectionBatchSize)
-			{
-				NumBatches++;
-			}
+	if (VertexCount != ParallelGeometryCollectionBatchSize)
+	{
+		NumBatches++;
+	}
 
-			// Batch too small, don't bother with parallel
-			if (ParallelGeometryCollectionBatchSize > TotalVertices)
-			{
-				bParallelGeometryCollection = false;
-				ParallelGeometryCollectionBatchSize = TotalVertices;
-			}
+	// Batch too small, don't bother with parallel
+	if (ParallelGeometryCollectionBatchSize > VertexCount)
+	{
+		bParallelGeometryCollection = false;
+		ParallelGeometryCollectionBatchSize = VertexCount;
+	}
 
-			auto GeometryCollectionBatch([&](int32 BatchNum)
-			{
-				int32 IndexOffset = ParallelGeometryCollectionBatchSize * BatchNum;
-				int32 ThisBatchSize = ParallelGeometryCollectionBatchSize;
+	auto GeometryCollectionBatch([&](int32 BatchNum)
+	{
+		uint32 IndexOffset = ParallelGeometryCollectionBatchSize * BatchNum;
+		uint32 ThisBatchSize = ParallelGeometryCollectionBatchSize;
 
-				// Check for final batch
-				if (IndexOffset + ParallelGeometryCollectionBatchSize > NumVertices)
-				{
-					ThisBatchSize = TotalVertices - IndexOffset;
-				}
+		// Check for final batch
+		if (IndexOffset + ParallelGeometryCollectionBatchSize > MeshDescription.NumVertices)
+		{
+			ThisBatchSize = VertexCount - IndexOffset;
+		}
 
-				if (ThisBatchSize > 0)
-				{
-					const FMatrix44f* RESTRICT BoneTransformsPtr = DynamicData->IsDynamic ? DynamicData->Transforms.GetData() : SubsetConstantData.RestTransforms.GetData();
+		if (ThisBatchSize > 0)
+		{
+			const FMatrix44f* RESTRICT BoneTransformsPtr = Transforms.GetData();
 
-					if (bGeometryCollection_SetDynamicData_ISPC_Enabled)
-					{
+			if (bGeometryCollection_SetDynamicData_ISPC_Enabled)
+			{
 #if INTEL_ISPC
-						uint8* VertexBufferOffset = (uint8*)VertexBufferData + (IndexOffset * VertexBuffer.GetStride());
-						ispc::SetDynamicData_RenderThread(
-							(ispc::FVector3f*)VertexBufferOffset,
-							ThisBatchSize,
-							VertexBuffer.GetStride(),
-							&SubsetConstantData.BoneMap[IndexOffset],
-							(ispc::FMatrix44f*)BoneTransformsPtr,
-							(ispc::FVector3f*)&VertexBuffers.PositionVertexBuffer.VertexPosition(IndexOffset));
+				uint8* VertexBufferOffset = (uint8*)VertexBufferData + (IndexOffset * VertexStride);
+				ispc::SetDynamicData_RenderThread(
+					(ispc::FVector3f*)VertexBufferOffset,
+					ThisBatchSize,
+					VertexStride,
+					&SourceBoneMapVertexBuffer.BoneIndex(IndexOffset),
+					(ispc::FMatrix44f*)BoneTransformsPtr,
+					(ispc::FVector3f*)&SourcePositionVertexBuffer.VertexPosition(IndexOffset));
 #endif
-					}
-					else
-					{
-						for (int32 i = IndexOffset; i < IndexOffset + ThisBatchSize; i++)
-						{
-							FVector3f Transformed = BoneTransformsPtr[SubsetConstantData.BoneMap[i]].TransformPosition(VertexBuffers.PositionVertexBuffer.VertexPosition(i));
-							FMemory::Memcpy((uint8*)VertexBufferData + (i * VertexBuffer.GetStride()), &Transformed, sizeof(FVector3f));
-						}
-					}
+			}
+			else
+			{
+				for (uint32 i = IndexOffset; i < IndexOffset + ThisBatchSize; i++)
+				{
+					FVector3f Transformed = BoneTransformsPtr[SourceBoneMapVertexBuffer.BoneIndex(i)].TransformPosition(SourcePositionVertexBuffer.VertexPosition(i));
+					FMemory::Memcpy((uint8*)VertexBufferData + (i * VertexStride), &Transformed, sizeof(FVector3f));
 				}
-			});
-
-			ParallelFor(NumBatches, GeometryCollectionBatch, !bParallelGeometryCollection);
-
-			RHIUnlockBuffer(VertexBuffer.VertexBufferRHI);
+			}
 		}
+	});
 
-#if RHI_RAYTRACING
-		if (IsRayTracingEnabled())
-		{
-			bGeometryResourceUpdated = true;
-		}
-#endif
-	}
+	ParallelFor(NumBatches, GeometryCollectionBatch, !bParallelGeometryCollection);
+
+	RHIUnlockBuffer(SkinnedPositionVertexBuffer.VertexBufferRHI);
 }
 
 FMaterialRenderProxy* FGeometryCollectionSceneProxy::GetMaterial(FMeshElementCollector& Collector, int32 MaterialIndex) const
 {
-	// material for wireframe
-	/*
-	never used
-	auto WireframeMaterialInstance = new FColoredMaterialRenderProxy(
-		GEngine->WireframeMaterial ? GEngine->WireframeMaterial->GetRenderProxy(IsSelected()) : nullptr,
-		FLinearColor(0, 0.5f, 1.f)
-	);
-	Collector.RegisterOneFrameMaterialProxy(WireframeMaterialInstance);
-	*/
-
-	// material for colored bones
-
 	FMaterialRenderProxy* MaterialProxy = nullptr;
 
+#if WITH_EDITOR
 	if (bShowBoneColors && GEngine->VertexColorMaterial)
 	{
+		// Material for colored bones
 		UMaterial* VertexColorVisualizationMaterial = GEngine->VertexColorMaterial;
-		auto VertexColorVisualizationMaterialInstance = new FColoredMaterialRenderProxy(
+		FMaterialRenderProxy* VertexColorVisualizationMaterialInstance = new FColoredMaterialRenderProxy(
 			VertexColorVisualizationMaterial->GetRenderProxy(),
 			GetSelectionColor(FLinearColor::White, false, false)
 		);
 		Collector.RegisterOneFrameMaterialProxy(VertexColorVisualizationMaterialInstance);
 		MaterialProxy = VertexColorVisualizationMaterialInstance;
 	}
-	else if(Materials.IsValidIndex(MaterialIndex))
+	else 
+#endif
+	if(Materials.IsValidIndex(MaterialIndex))
 	{
 		MaterialProxy = Materials[MaterialIndex]->GetRenderProxy();
 	}
@@ -896,391 +617,306 @@ FMaterialRenderProxy* FGeometryCollectionSceneProxy::GetMaterial(FMeshElementCol
 	return MaterialProxy;
 }
 
+FVertexFactory const* FGeometryCollectionSceneProxy::GetVertexFactory() const
+{
+#if WITH_EDITOR
+	return bShowBoneColors ? &VertexFactoryDebugColor : &VertexFactory;
+#else
+	return &VertexFactory;
+#endif
+}
+
 void FGeometryCollectionSceneProxy::GetDynamicMeshElements(const TArray<const FSceneView*>& Views, const FSceneViewFamily& ViewFamily, uint32 VisibilityMap, FMeshElementCollector& Collector) const
 {
 	QUICK_SCOPE_CYCLE_COUNTER(STAT_GeometryCollectionSceneProxy_GetDynamicMeshElements);
-	if (GetRequiredVertexCount())
+	if (MeshDescription.NumVertices == 0)
 	{
-		const bool bWireframe = AllowDebugViewmodes() && ViewFamily.EngineShowFlags.Wireframe;
-		const bool bProxyIsSelected = IsSelected();
-
-		const FEngineShowFlags& EngineShowFlags = ViewFamily.EngineShowFlags;
+		return;
+	}
+		
+	const FEngineShowFlags& EngineShowFlags = ViewFamily.EngineShowFlags;
+	const bool bWireframe = AllowDebugViewmodes() && EngineShowFlags.Wireframe;
+	const bool bProxyIsSelected = IsSelected();
 
-		auto SetDebugMaterial = [this, &Collector, &EngineShowFlags, bProxyIsSelected](FMeshBatch& Mesh) -> void
-		{
+	auto SetDebugMaterial = [this, &Collector, &EngineShowFlags, bProxyIsSelected](FMeshBatch& Mesh) -> void
+	{
 #if UE_ENABLE_DEBUG_DRAWING
 
-			// flag to indicate whether we've set a debug material yet
-			// Note: Will be used if we add more debug material options
-			// (compare to variable of same name in StaticMeshRender.cpp)
-			bool bDebugMaterialRenderProxySet = false;
+		// flag to indicate whether we've set a debug material yet
+		// Note: Will be used if we add more debug material options
+		// (compare to variable of same name in StaticMeshRender.cpp)
+		bool bDebugMaterialRenderProxySet = false;
 
-			if (!bDebugMaterialRenderProxySet && bProxyIsSelected && EngineShowFlags.VertexColors && AllowDebugViewmodes())
+		if (!bDebugMaterialRenderProxySet && bProxyIsSelected && EngineShowFlags.VertexColors && AllowDebugViewmodes())
+		{
+			// Override the mesh's material with our material that draws the vertex colors
+			UMaterial* VertexColorVisualizationMaterial = NULL;
+			switch (GVertexColorViewMode)
 			{
-				// Override the mesh's material with our material that draws the vertex colors
-				UMaterial* VertexColorVisualizationMaterial = NULL;
-				switch (GVertexColorViewMode)
-				{
-				case EVertexColorViewMode::Color:
-					VertexColorVisualizationMaterial = GEngine->VertexColorViewModeMaterial_ColorOnly;
-					break;
+			case EVertexColorViewMode::Color:
+				VertexColorVisualizationMaterial = GEngine->VertexColorViewModeMaterial_ColorOnly;
+				break;
 
-				case EVertexColorViewMode::Alpha:
-					VertexColorVisualizationMaterial = GEngine->VertexColorViewModeMaterial_AlphaAsColor;
-					break;
+			case EVertexColorViewMode::Alpha:
+				VertexColorVisualizationMaterial = GEngine->VertexColorViewModeMaterial_AlphaAsColor;
+				break;
 
-				case EVertexColorViewMode::Red:
-					VertexColorVisualizationMaterial = GEngine->VertexColorViewModeMaterial_RedOnly;
-					break;
+			case EVertexColorViewMode::Red:
+				VertexColorVisualizationMaterial = GEngine->VertexColorViewModeMaterial_RedOnly;
+				break;
 
-				case EVertexColorViewMode::Green:
-					VertexColorVisualizationMaterial = GEngine->VertexColorViewModeMaterial_GreenOnly;
-					break;
+			case EVertexColorViewMode::Green:
+				VertexColorVisualizationMaterial = GEngine->VertexColorViewModeMaterial_GreenOnly;
+				break;
 
-				case EVertexColorViewMode::Blue:
-					VertexColorVisualizationMaterial = GEngine->VertexColorViewModeMaterial_BlueOnly;
-					break;
-				}
-				check(VertexColorVisualizationMaterial != NULL);
+			case EVertexColorViewMode::Blue:
+				VertexColorVisualizationMaterial = GEngine->VertexColorViewModeMaterial_BlueOnly;
+				break;
+			}
+			check(VertexColorVisualizationMaterial != NULL);
 
-				// Note: static mesh renderer does something more complicated involving per-section selection,
-				// but whole component selection seems ok for now
-				bool bSectionIsSelected = bProxyIsSelected;
+			// Note: static mesh renderer does something more complicated involving per-section selection,
+			// but whole component selection seems ok for now
+			bool bSectionIsSelected = bProxyIsSelected;
 
-				auto VertexColorVisualizationMaterialInstance = new FColoredMaterialRenderProxy(
-					VertexColorVisualizationMaterial->GetRenderProxy(),
-					GetSelectionColor(FLinearColor::White, bSectionIsSelected, IsHovered())
-				);
+			auto VertexColorVisualizationMaterialInstance = new FColoredMaterialRenderProxy(
+				VertexColorVisualizationMaterial->GetRenderProxy(),
+				GetSelectionColor(FLinearColor::White, bSectionIsSelected, IsHovered())
+			);
 
-				Collector.RegisterOneFrameMaterialProxy(VertexColorVisualizationMaterialInstance);
-				Mesh.MaterialRenderProxy = VertexColorVisualizationMaterialInstance;
+			Collector.RegisterOneFrameMaterialProxy(VertexColorVisualizationMaterialInstance);
+			Mesh.MaterialRenderProxy = VertexColorVisualizationMaterialInstance;
 
-				bDebugMaterialRenderProxySet = true;
-			}
+			bDebugMaterialRenderProxySet = true;
+		}
 #endif
-		};
+	};
 
-		for (int32 ViewIndex = 0; ViewIndex < Views.Num(); ViewIndex++)
+	for (int32 ViewIndex = 0; ViewIndex < Views.Num(); ViewIndex++)
+	{
+		if ((VisibilityMap & (1 << ViewIndex)) == 0) 
 		{
-			if ((VisibilityMap & (1 << ViewIndex)) == 0) { continue; }
-
-			// Render Batches
-						
-			// render original mesh if it isn't dynamic and there is an unfractured mesh
-			// #todo(dmp): refactor this to share more code later
-			const bool bIsDynamic = DynamicData && DynamicData->IsDynamic;
-			if (!bIsDynamic)
-			{
-			#if GEOMETRYCOLLECTION_EDITOR_SELECTION
-				const TArray<FGeometryCollectionSection>& SectionArray = bUsesSubSections && SubSections.Num() ? SubSections: SubsetConstantData.OriginalMeshSections;
-				UE_LOG(FGeometryCollectionSceneProxyLogging, VeryVerbose, TEXT("GetDynamicMeshElements, bUseSubSections=%d, NumSections=%d for %p."), bUsesSubSections, SectionArray.Num(), this);
-			#else
-				const TArray<FGeometryCollectionSection>& SectionArray = SubsetConstantData.OriginalMeshSections;
-			#endif
-
-				// Grab the material proxies we'll be using for each section
-				TArray<FMaterialRenderProxy*, TInlineAllocator<32>> MaterialProxies;
-
-				for (int32 SectionIndex = 0; SectionIndex < SectionArray.Num(); SectionIndex++)
-				{
-					const FGeometryCollectionSection& Section = SectionArray[SectionIndex];
-					FMaterialRenderProxy* MaterialProxy = GetMaterial(Collector, Section.MaterialID);
-					MaterialProxies.Add(MaterialProxy);
-				}
-
-				for (int32 SectionIndex = 0; SectionIndex < SectionArray.Num(); SectionIndex++)
-				{
-					const FGeometryCollectionSection& Section = SectionArray[SectionIndex];
-
-					// Draw the mesh.
-					FMeshBatch& Mesh = Collector.AllocateMesh();
-					FMeshBatchElement& BatchElement = Mesh.Elements[0];
-					BatchElement.IndexBuffer = &OriginalMeshIndexBuffer;
-					Mesh.bWireframe = bWireframe;
-					Mesh.VertexFactory = &VertexFactory;
-					Mesh.MaterialRenderProxy = MaterialProxies[SectionIndex];
-
-					/*
-					bool bHasPrecomputedVolumetricLightmap;
-					FMatrix PreviousLocalToWorld;
-					int32 SingleCaptureIndex;
-					bool bOutputVelocity;
-					GetScene().GetPrimitiveUniformShaderParameters_RenderThread(GetPrimitiveSceneInfo(), bHasPrecomputedVolumetricLightmap, PreviousLocalToWorld, SingleCaptureIndex, bOutputVelocity);
-					bOutputVelocity |= AlwaysHasVelocity();
-
-					FDynamicPrimitiveUniformBuffer& DynamicPrimitiveUniformBuffer = Collector.AllocateOneFrameResource<FDynamicPrimitiveUniformBuffer>();
-					DynamicPrimitiveUniformBuffer.Set(GetLocalToWorld(), PreviousLocalToWorld, GetBounds(), GetLocalBounds(), true, bHasPrecomputedVolumetricLightmap, bOutputVelocity);
-					BatchElement.PrimitiveUniformBufferResource = &DynamicPrimitiveUniformBuffer.UniformBuffer;
-					*/
-
-					BatchElement.PrimitiveUniformBuffer = GetUniformBuffer();
-					BatchElement.FirstIndex = Section.FirstIndex;
-					BatchElement.NumPrimitives = Section.NumTriangles;
-					BatchElement.MinVertexIndex = Section.MinVertexIndex;
-					BatchElement.MaxVertexIndex = Section.MaxVertexIndex;
-					Mesh.ReverseCulling = IsLocalToWorldDeterminantNegative();
-					Mesh.Type = PT_TriangleList;
-					Mesh.DepthPriorityGroup = SDPG_World;
-					Mesh.bCanApplyViewModeOverrides = true;
-				#if WITH_EDITOR
-					if (GIsEditor)
-					{
-						Mesh.BatchHitProxyId = Section.HitProxy ? Section.HitProxy->Id : FHitProxyId();
-					}
-				#endif
-
-					SetDebugMaterial(Mesh);
-
-					Collector.AddMesh(ViewIndex, Mesh);
-				}
-			}
-			else
-			{
-			#if GEOMETRYCOLLECTION_EDITOR_SELECTION
-				const TArray<FGeometryCollectionSection>& SectionArray = bUsesSubSections && SubSections.Num() ? SubSections: Sections;
-				UE_LOG(FGeometryCollectionSceneProxyLogging, VeryVerbose, TEXT("GetDynamicMeshElements, bUseSubSections=%d, NumSections=%d for %p."), bUsesSubSections, SectionArray.Num(), this);
-			#else
-				const TArray<FGeometryCollectionSection>& SectionArray = Sections;
-			#endif
-
-				// Grab the material proxies we'll be using for each section
-				TArray<FMaterialRenderProxy*, TInlineAllocator<32>> MaterialProxies;
-
-				for (int32 SectionIndex = 0; SectionIndex < SectionArray.Num(); ++SectionIndex)
-				{
-					const FGeometryCollectionSection& Section = SectionArray[SectionIndex];
+			continue; 
+		}
 
-					FMaterialRenderProxy* MaterialProxy = GetMaterial(Collector, Section.MaterialID);
-					MaterialProxies.Add(MaterialProxy);
-				}
+		// If not dynamic then use the section array with interior fracture surfaces removed.
+		bool bRemoveInternalFaces = DynamicData != nullptr && !DynamicData->IsDynamic && MeshDescription.SectionsNoInternal.Num();
 
-				for (int32 SectionIndex = 0; SectionIndex < SectionArray.Num(); ++SectionIndex)
-				{
-					const FGeometryCollectionSection& Section = SectionArray[SectionIndex];
-
-					// Draw the mesh.
-					FMeshBatch& Mesh = Collector.AllocateMesh();
-					FMeshBatchElement& BatchElement = Mesh.Elements[0];
-					BatchElement.IndexBuffer = &IndexBuffer;
-					Mesh.bWireframe = bWireframe;
-					Mesh.VertexFactory = &VertexFactory;
-					Mesh.MaterialRenderProxy = MaterialProxies[SectionIndex];
-					BatchElement.PrimitiveUniformBuffer = GetUniformBuffer();
-					BatchElement.FirstIndex = Section.FirstIndex;
-					BatchElement.NumPrimitives = Section.NumTriangles;
-					BatchElement.MinVertexIndex = Section.MinVertexIndex;
-					BatchElement.MaxVertexIndex = Section.MaxVertexIndex;
-					Mesh.ReverseCulling = IsLocalToWorldDeterminantNegative();
-					Mesh.Type = PT_TriangleList;
-					Mesh.DepthPriorityGroup = SDPG_World;
-					Mesh.bCanApplyViewModeOverrides = true;
-				#if WITH_EDITOR
-					if (GIsEditor)
-					{
-						Mesh.BatchHitProxyId = Section.HitProxy ? Section.HitProxy->Id : FHitProxyId();
-					}
-				#endif
+#if WITH_EDITOR
+		// If hiding geometry in editor then we don't remove hidden faces.
+		bRemoveInternalFaces &= HiddenTransforms.Num() == 0;
+#endif
 
-					SetDebugMaterial(Mesh);
+#if GEOMETRYCOLLECTION_EDITOR_SELECTION
+		// If using subsections then use the subsection array. 
+		TArray<FGeometryCollectionMeshElement> const& SectionArray = bUsesSubSections 
+			? MeshDescription.SubSections 
+			: bRemoveInternalFaces ? MeshDescription.SectionsNoInternal : MeshDescription.Sections;
+#else
+		TArray<FGeometryCollectionMeshElement> const& SectionArray = bRemoveInternalFaces ? MeshDescription.SectionsNoInternal : MeshDescription.Sections;
+#endif
 
-					Collector.AddMesh(ViewIndex, Mesh);
-				}
-			}
+		// Grab the material proxies we'll be using for each section.
+		TArray<FMaterialRenderProxy*, TInlineAllocator<32>> MaterialProxies;
+		for (int32 SectionIndex = 0; SectionIndex < SectionArray.Num(); ++SectionIndex)
+		{
+			const FGeometryCollectionMeshElement& Section = SectionArray[SectionIndex];
+			FMaterialRenderProxy* MaterialProxy = GetMaterial(Collector, Section.MaterialIndex);
+			MaterialProxies.Add(MaterialProxy);
+		}
 
-			// Highlight selected bone using specialized material - when rendering bones as colors we don't need to run this code as the
-			// bone selection is already contained in the rendered colors
-			// #note: This renders the geometry again but with the bone selection material.  Ideally we'd have one render pass and one
-			// material.
-			if ((bShowBoneColors || bEnableBoneSelection) && !bSuppressSelectionMaterial && BoneSelectedMaterial)
-			{
-				FMaterialRenderProxy* MaterialRenderProxy = BoneSelectedMaterial->GetRenderProxy();
-
-				FMeshBatch& Mesh = Collector.AllocateMesh();
-				FMeshBatchElement& BatchElement = Mesh.Elements[0];
-				BatchElement.IndexBuffer = &IndexBuffer;
-				Mesh.bWireframe = bWireframe;
-				Mesh.VertexFactory = &VertexFactory;
-				Mesh.MaterialRenderProxy = MaterialRenderProxy;
-				BatchElement.PrimitiveUniformBuffer = GetUniformBuffer();
-				BatchElement.FirstIndex = 0;
-				BatchElement.NumPrimitives = GetRequiredIndexCount() / 3;
-				BatchElement.MinVertexIndex = 0;
-				BatchElement.MaxVertexIndex = GetRequiredVertexCount();
-				Mesh.ReverseCulling = IsLocalToWorldDeterminantNegative();
-				Mesh.Type = PT_TriangleList;
-				Mesh.DepthPriorityGroup = SDPG_World;
-				Mesh.bCanApplyViewModeOverrides = false;
-				Collector.AddMesh(ViewIndex, Mesh);
-			}
+		// Draw the meshes.
+		for (int32 SectionIndex = 0; SectionIndex < SectionArray.Num(); ++SectionIndex)
+		{
+			const FGeometryCollectionMeshElement& Section = SectionArray[SectionIndex];
+
+			FMeshBatch& Mesh = Collector.AllocateMesh();
+			Mesh.bWireframe = bWireframe;
+			Mesh.VertexFactory = GetVertexFactory();
+			Mesh.MaterialRenderProxy = MaterialProxies[SectionIndex];
+			Mesh.ReverseCulling = IsLocalToWorldDeterminantNegative();
+			Mesh.Type = PT_TriangleList;
+			Mesh.DepthPriorityGroup = SDPG_World;
+			Mesh.bCanApplyViewModeOverrides = true;
+			SetDebugMaterial(Mesh);
+
+			FMeshBatchElement& BatchElement = Mesh.Elements[0];
+			BatchElement.IndexBuffer = &MeshResource.IndexBuffer;
+			BatchElement.PrimitiveUniformBuffer = GetUniformBuffer();
+			BatchElement.FirstIndex = Section.TriangleStart * 3;
+			BatchElement.NumPrimitives = Section.TriangleCount;
+			BatchElement.MinVertexIndex = Section.VertexStart;
+			BatchElement.MaxVertexIndex = Section.VertexEnd;
+
+			Collector.AddMesh(ViewIndex, Mesh);
+		}
 
-		#if !(UE_BUILD_SHIPPING || UE_BUILD_TEST)
-			RenderBounds(Collector.GetPDI(ViewIndex), ViewFamily.EngineShowFlags, GetBounds(), IsSelected());
-		#endif
+#if GEOMETRYCOLLECTION_EDITOR_SELECTION
+		// Highlight selected bone using specialized material.
+		// #note: This renders the geometry again but with the bone selection material.  Ideally we'd have one render pass and one material.
+		if (bEnableBoneSelection && !bSuppressSelectionMaterial && BoneSelectedMaterial)
+		{
+			FMaterialRenderProxy* MaterialRenderProxy = BoneSelectedMaterial->GetRenderProxy();
+
+			FMeshBatch& Mesh = Collector.AllocateMesh();
+			Mesh.bWireframe = bWireframe;
+			Mesh.VertexFactory = &VertexFactoryDebugColor;
+			Mesh.MaterialRenderProxy = MaterialRenderProxy;
+			Mesh.ReverseCulling = IsLocalToWorldDeterminantNegative();
+			Mesh.Type = PT_TriangleList;
+			Mesh.DepthPriorityGroup = SDPG_World;
+			Mesh.bCanApplyViewModeOverrides = false;
+
+			FMeshBatchElement& BatchElement = Mesh.Elements[0];
+			BatchElement.IndexBuffer = &MeshResource.IndexBuffer;
+			BatchElement.PrimitiveUniformBuffer = GetUniformBuffer();
+			BatchElement.FirstIndex = 0;
+			BatchElement.NumPrimitives = MeshDescription.NumTriangles;
+			BatchElement.MinVertexIndex = 0;
+			BatchElement.MaxVertexIndex = MeshDescription.NumVertices;
+
+			Collector.AddMesh(ViewIndex, Mesh);
 		}
+#endif // GEOMETRYCOLLECTION_EDITOR_SELECTION
+
+#if !(UE_BUILD_SHIPPING || UE_BUILD_TEST)
+		RenderBounds(Collector.GetPDI(ViewIndex), ViewFamily.EngineShowFlags, GetBounds(), IsSelected());
+#endif
 	}
 }
 
 #if RHI_RAYTRACING
 void FGeometryCollectionSceneProxy::GetDynamicRayTracingInstances(FRayTracingMaterialGatheringContext& Context, TArray<struct FRayTracingInstance>& OutRayTracingInstances)
 {
-	if (GRayTracingGeometryCollectionProxyMeshes == 0)
+	if (GRayTracingGeometryCollectionProxyMeshes == 0 || MeshDescription.NumVertices == 0)
 	{
 		return;
 	}
 
 	QUICK_SCOPE_CYCLE_COUNTER(STAT_GeometryCollectionSceneProxy_GetDynamicRayTracingInstances);
 
-	if (GetRequiredVertexCount())
-	{
-		const uint32 LODIndex = 0;
-		const bool bWireframe = false; //AllowDebugViewmodes();//&& ViewFamily.EngineShowFlags.Wireframe;
+	const uint32 LODIndex = 0;
+	const bool bWireframe = false;
 		
-		// render original mesh if it isn't dynamic and there is an unfractured mesh
-		// #todo(dmp): refactor this to share more code later
-		const bool bIsDynamic = DynamicData && DynamicData->IsDynamic;
-
-		//Loose parameter needs to be updated every frame
-		FGeometryCollectionMeshCollectorResources* CollectorResources;
-		CollectorResources = &Context.RayTracingMeshResourceCollector.
-			AllocateOneFrameResource<FGeometryCollectionMeshCollectorResources>(GetScene().GetFeatureLevel());
-		FGeometryCollectionVertexFactory& GeometryCollectionVertexFactory = CollectorResources->GetVertexFactory();
+	//Loose parameter needs to be updated every frame
+	FGeometryCollectionMeshCollectorResources* CollectorResources;
+	CollectorResources = &Context.RayTracingMeshResourceCollector.
+		AllocateOneFrameResource<FGeometryCollectionMeshCollectorResources>(GetScene().GetFeatureLevel());
+	FGeometryCollectionVertexFactory& GeometryCollectionVertexFactory = CollectorResources->GetVertexFactory();
 		
-		// Render dynamic objects
-		if (!GeometryCollectionVertexFactory.GetType()->SupportsRayTracingDynamicGeometry())
-		{
-			return;
-		}
-
-		SetupVertexFactory(GeometryCollectionVertexFactory);
+	// Render dynamic objects
+	if (!GeometryCollectionVertexFactory.GetType()->SupportsRayTracingDynamicGeometry())
+	{
+		return;
+	}
 
-		FGeometryCollectionIndexBuffer* ActiveIndexBuffer = bIsDynamic ? &IndexBuffer : &OriginalMeshIndexBuffer;
-		UpdatingRayTracingGeometry_RenderingThread(ActiveIndexBuffer);
+	SetupVertexFactory(GeometryCollectionVertexFactory);
 
-#if GEOMETRYCOLLECTION_EDITOR_SELECTION
-		const TArray<FGeometryCollectionSection>& SectionArray = bUsesSubSections && SubSections.Num() ? SubSections : Sections;
-		UE_LOG(FGeometryCollectionSceneProxyLogging, VeryVerbose, TEXT("GetDynamicMeshElements, bUseSubSections=%d, NumSections=%d for %p."), bUsesSubSections, SectionArray.Num(), this);
-#else
-		const TArray<FGeometryCollectionSection>& SectionArray = Sections;
-#endif
-		const int32 InstanceCount = SectionArray.Num();
+	// If not dynamic then use the section array with interior fracture surfaces removed.
+	const bool bRemoveInternalFaces = DynamicData != nullptr && !DynamicData->IsDynamic && MeshDescription.SectionsNoInternal.Num();
+	TArray<FGeometryCollectionMeshElement> const& SectionArray = bRemoveInternalFaces ? MeshDescription.SectionsNoInternal : MeshDescription.Sections;
 
-		if (InstanceCount > 0 && RayTracingGeometry.RayTracingGeometryRHI.IsValid())
-		{
+	UpdatingRayTracingGeometry_RenderingThread(SectionArray);
 
-			FRayTracingInstance RayTracingInstance;
-			RayTracingInstance.Geometry = &RayTracingGeometry;
-			RayTracingInstance.InstanceTransforms.Emplace(GetLocalToWorld());
+	if (SectionArray.Num() && RayTracingGeometry.RayTracingGeometryRHI.IsValid())
+	{
+		FRayTracingInstance RayTracingInstance;
+		RayTracingInstance.Geometry = &RayTracingGeometry;
+		RayTracingInstance.InstanceTransforms.Emplace(GetLocalToWorld());
 
-			// Grab the material proxies we'll be using for each section
-			TArray<FMaterialRenderProxy*, TInlineAllocator<32>> MaterialProxies;
+		// Grab the material proxies we'll be using for each section
+		TArray<FMaterialRenderProxy*, TInlineAllocator<32>> MaterialProxies;
 
-			for (int32 SectionIndex = 0; SectionIndex < SectionArray.Num(); ++SectionIndex)
-			{
-				const FGeometryCollectionSection& Section = SectionArray[SectionIndex];
+		for (int32 SectionIndex = 0; SectionIndex < SectionArray.Num(); ++SectionIndex)
+		{
+			const FGeometryCollectionMeshElement& Section = SectionArray[SectionIndex];
 				
-				//TODO: Add BoneColor support in Path/Ray tracing?
-				FMaterialRenderProxy* MaterialProxy= Materials[Section.MaterialID]->GetRenderProxy();
-
-				if (MaterialProxy == nullptr)
-				{
-					MaterialProxy = UMaterial::GetDefaultMaterial(MD_Surface)->GetRenderProxy();
-				}
-
-				MaterialProxies.Add(MaterialProxy);
-			}
+			//TODO: Add BoneColor support in Path/Ray tracing?
+			FMaterialRenderProxy* MaterialProxy= Materials[Section.MaterialIndex]->GetRenderProxy();
 
-			int32 MaxVertexIndex = 0;
-			for (int32 SectionIndex = 0; SectionIndex < SectionArray.Num(); ++SectionIndex)
+			if (MaterialProxy == nullptr)
 			{
-				const FGeometryCollectionSection& Section = SectionArray[SectionIndex];
-
-				// Draw the mesh
-
-				FMeshBatch& Mesh = RayTracingInstance.Materials.AddDefaulted_GetRef();
-				Mesh.bWireframe = bWireframe;//bWireframe needs viewfamily access ?
-				Mesh.SegmentIndex = SectionIndex;
-				Mesh.VertexFactory = &GeometryCollectionVertexFactory;
-				Mesh.MaterialRenderProxy = MaterialProxies[SectionIndex];
-				Mesh.LODIndex = LODIndex;
-				Mesh.ReverseCulling = IsLocalToWorldDeterminantNegative();
-				Mesh.bDisableBackfaceCulling = true;
-				Mesh.Type = PT_TriangleList;
-				Mesh.DepthPriorityGroup = SDPG_World;
-				Mesh.bCanApplyViewModeOverrides = true;
-
-				FMeshBatchElement& BatchElement = Mesh.Elements[0];
-				BatchElement.IndexBuffer = ActiveIndexBuffer;
-				BatchElement.PrimitiveUniformBuffer = GetUniformBuffer();
-				BatchElement.FirstIndex = Section.FirstIndex;
-				BatchElement.NumPrimitives = Section.NumTriangles;
-				BatchElement.MinVertexIndex = Section.MinVertexIndex;
-				BatchElement.MaxVertexIndex = Section.MaxVertexIndex; 
-				BatchElement.NumInstances = 1;
-
-				MaxVertexIndex = std::max(Section.MaxVertexIndex, MaxVertexIndex);
-#if WITH_EDITOR
-				if (GIsEditor)
-				{
-					Mesh.BatchHitProxyId = Section.HitProxy ? Section.HitProxy->Id : FHitProxyId();
-				}
-#endif
-				//#TODO: bone color, bone selection and render bound?
+				MaterialProxy = UMaterial::GetDefaultMaterial(MD_Surface)->GetRenderProxy();
 			}
 
-			FRWBuffer* VertexBuffer = RayTracingDynamicVertexBuffer.NumBytes > 0 ? &RayTracingDynamicVertexBuffer : nullptr;
-
-			const uint32 VertexCount = MaxVertexIndex + 1;
-			Context.DynamicRayTracingGeometriesToUpdate.Add(
-				FRayTracingDynamicGeometryUpdateParams
-				{
-					RayTracingInstance.Materials,
-					false,
-					VertexCount,
-					VertexCount * (uint32)sizeof(FVector3f),
-					RayTracingGeometry.Initializer.TotalPrimitiveCount,
-					&RayTracingGeometry,
-					VertexBuffer,
-					true
-				}
-			);
+			MaterialProxies.Add(MaterialProxy);
+		}
 
-			OutRayTracingInstances.Emplace(RayTracingInstance);
+		uint32 MaxVertexIndex = 0;
+		for (int32 SectionIndex = 0; SectionIndex < SectionArray.Num(); ++SectionIndex)
+		{
+			const FGeometryCollectionMeshElement& Section = SectionArray[SectionIndex];
+
+			// Draw the mesh
+			FMeshBatch& Mesh = RayTracingInstance.Materials.AddDefaulted_GetRef();
+			Mesh.bWireframe = bWireframe;
+			Mesh.SegmentIndex = SectionIndex;
+			Mesh.VertexFactory = &GeometryCollectionVertexFactory;
+			Mesh.MaterialRenderProxy = MaterialProxies[SectionIndex];
+			Mesh.LODIndex = LODIndex;
+			Mesh.ReverseCulling = IsLocalToWorldDeterminantNegative();
+			Mesh.bDisableBackfaceCulling = true;
+			Mesh.Type = PT_TriangleList;
+			Mesh.DepthPriorityGroup = SDPG_World;
+			Mesh.bCanApplyViewModeOverrides = true;
+
+			FMeshBatchElement& BatchElement = Mesh.Elements[0];
+			BatchElement.IndexBuffer = &MeshResource.IndexBuffer;
+			BatchElement.PrimitiveUniformBuffer = GetUniformBuffer();
+			BatchElement.FirstIndex = Section.TriangleStart * 3;
+			BatchElement.NumPrimitives = Section.TriangleCount;
+			BatchElement.MinVertexIndex = Section.VertexStart;
+			BatchElement.MaxVertexIndex = Section.VertexEnd; 
+			BatchElement.NumInstances = 1;
+
+			MaxVertexIndex = std::max(Section.VertexEnd, MaxVertexIndex);
+
+			//#TODO: bone color, bone selection and render bound?
 		}
 
-	}
+		FRWBuffer* VertexBuffer = RayTracingDynamicVertexBuffer.NumBytes > 0 ? &RayTracingDynamicVertexBuffer : nullptr;
+
+		const uint32 VertexCount = MaxVertexIndex + 1;
+		Context.DynamicRayTracingGeometriesToUpdate.Add(
+			FRayTracingDynamicGeometryUpdateParams
+			{
+				RayTracingInstance.Materials,
+				false,
+				VertexCount,
+				VertexCount * (uint32)sizeof(FVector3f),
+				RayTracingGeometry.Initializer.TotalPrimitiveCount,
+				&RayTracingGeometry,
+				VertexBuffer,
+				true
+			}
+		);
 
+		OutRayTracingInstances.Emplace(RayTracingInstance);
+	}
 }
 
-void FGeometryCollectionSceneProxy::UpdatingRayTracingGeometry_RenderingThread(FGeometryCollectionIndexBuffer* InIndexBuffer)
+void FGeometryCollectionSceneProxy::UpdatingRayTracingGeometry_RenderingThread(TArray<FGeometryCollectionMeshElement> const& InSectionArray)
 {
-#if GEOMETRYCOLLECTION_EDITOR_SELECTION
-	const TArray<FGeometryCollectionSection>& SectionArray = bUsesSubSections && SubSections.Num() ? SubSections : Sections;
-	UE_LOG(FGeometryCollectionSceneProxyLogging, VeryVerbose, TEXT("GetDynamicMeshElements, bUseSubSections=%d, NumSections=%d for %p."), bUsesSubSections, SectionArray.Num(), this);
-#else
-	const TArray<FGeometryCollectionSection>& SectionArray = Sections;
-#endif
-	const int32 InstanceCount = SectionArray.Num();//InstanceSceneData.Num();
-
-	if (InIndexBuffer && bGeometryResourceUpdated)
+	if (bGeometryResourceUpdated)
 	{
 		RayTracingGeometry.Initializer.Segments.Empty();
 		RayTracingGeometry.Initializer.TotalPrimitiveCount = 0;
 
-		for (int SectionIndex = 0; SectionIndex < InstanceCount; ++SectionIndex)
+		for (int SectionIndex = 0; SectionIndex < InSectionArray.Num(); ++SectionIndex)
 		{
-			const FGeometryCollectionSection& Section = SectionArray[SectionIndex];
+			const FGeometryCollectionMeshElement& Section = InSectionArray[SectionIndex];
 			FRayTracingGeometrySegment Segment;
-			Segment.FirstPrimitive = Section.FirstIndex / 3;
-			Segment.VertexBuffer = VertexBuffers.PositionVertexBuffer.VertexBufferRHI;
-			Segment.NumPrimitives = Section.NumTriangles;
-			Segment.MaxVertices = VertexBuffers.PositionVertexBuffer.GetNumVertices();
+			Segment.FirstPrimitive = Section.TriangleStart;
+			Segment.VertexBuffer = MeshResource.PositionVertexBuffer.VertexBufferRHI;
+			Segment.NumPrimitives = Section.TriangleCount;
+			Segment.MaxVertices = Section.VertexEnd;
 			RayTracingGeometry.Initializer.Segments.Add(Segment);
-			RayTracingGeometry.Initializer.TotalPrimitiveCount += Section.NumTriangles;
+			RayTracingGeometry.Initializer.TotalPrimitiveCount += Section.TriangleCount;
 		}
 				
 		if (RayTracingGeometry.Initializer.TotalPrimitiveCount > 0)
 		{
-			RayTracingGeometry.Initializer.IndexBuffer = InIndexBuffer->IndexBufferRHI;
+			RayTracingGeometry.Initializer.IndexBuffer = MeshResource.IndexBuffer.IndexBufferRHI;
 			// Create the ray tracing geometry but delay the acceleration structure build.
 			RayTracingGeometry.CreateRayTracingGeometry(ERTAccelerationStructureBuildPriority::Skip);
 		}
@@ -1288,8 +924,7 @@ void FGeometryCollectionSceneProxy::UpdatingRayTracingGeometry_RenderingThread(F
 		bGeometryResourceUpdated = false;
 	}
 }
-
-#endif
+#endif // RHI_RAYTRACING
 
 FPrimitiveViewRelevance FGeometryCollectionSceneProxy::GetViewRelevance(const FSceneView* View) const
 {
@@ -1308,220 +943,43 @@ FPrimitiveViewRelevance FGeometryCollectionSceneProxy::GetViewRelevance(const FS
 	return Result;
 }
 
-#if WITH_EDITOR
+#if GEOMETRYCOLLECTION_EDITOR_SELECTION
 HHitProxy* FGeometryCollectionSceneProxy::CreateHitProxies(UPrimitiveComponent* Component, TArray<TRefCountPtr<HHitProxy> >& OutHitProxies)
 {
-	// In order to be able to click on static meshes when they're batched up, we need to have catch all default
-	// hit proxy to return.
 	HHitProxy* DefaultHitProxy = FPrimitiveSceneProxy::CreateHitProxies(Component, OutHitProxies);
-	WholeObjectHitProxyColor = DefaultHitProxy->Id.GetColor();
-
-	// @todo FractureTools - Reconcile with subsection hit proxies.  Subsection is a draw call per hit proxy but is not suitable per-vertex as written
-	if (bEnableBoneSelection)
-	{
-		UGeometryCollectionComponent* GeometryCollectionComp = CastChecked<UGeometryCollectionComponent>(Component);
-		int32 NumTransforms = GeometryCollectionComp->GetTransformArray().Num();
-		PerBoneHitProxies.Empty();
-		for (int32 TransformIndex = 0; TransformIndex < NumTransforms; ++TransformIndex)
-		{
-			HGeometryCollectionBone* HitProxy = new HGeometryCollectionBone(GeometryCollectionComp, TransformIndex);
-			PerBoneHitProxies.Add(HitProxy);
-		}
-
-		OutHitProxies.Append(PerBoneHitProxies);
-	}
-	else if (Component->GetOwner())
-	{
-#if GEOMETRYCOLLECTION_EDITOR_SELECTION
-		// Note the below-created subsection hitproxies are never drawn, because the per-vertex hitproxy 
-		// rendering path is always on for geometry collections (i.e., USE_PER_VERTEX_HITPROXY_ID is 1)
-		const int32 NumTransforms = (Sections.Num() > 0) ? SubSectionHitProxies.Num() / Sections.Num(): 0;
-		for (int32 SectionIndex = 0; SectionIndex < Sections.Num(); ++SectionIndex)
-		{
-			// Create HitProxy for regular material based sections, and update existing section
-			FGeometryCollectionSection& Section = Sections[SectionIndex];
-
-			const int32 MaterialID = Section.MaterialID;
-			HActor* const HitProxy = new HActor(Component->GetOwner(), Component, SectionIndex, MaterialID);
-
-			OutHitProxies.Add(HitProxy);
-			Section.HitProxy = HitProxy;
-
-			// Create HitProxy per transform index using the same material Id than the current sections
-			// All combinations of material id/transform index are populated,
-			// since it can't be assumed that any of them won't be needed.
-			const int32 SectionOffset = SectionIndex * NumTransforms;
-
-			for (int32 TransformIndex = 0; TransformIndex < NumTransforms; ++TransformIndex)
-			{
-				static const int32 SubSectionIndex = INDEX_NONE;  // The index will get updated later for existing subsections
-				
-				HGeometryCollection* const SubSectionHitProxy = new HGeometryCollection(Component->GetOwner(), Component, SubSectionIndex, MaterialID, TransformIndex);
-
-				OutHitProxies.Add(SubSectionHitProxy);
-				SubSectionHitProxies[SectionOffset + TransformIndex] = SubSectionHitProxy;
-			}
-		}
-	#else
-		for (int32 SectionIndex = 0; SectionIndex < Sections.Num(); ++SectionIndex)
-		{
-			// Create HitProxy for regular material based sections, and update existing section
-			FGeometryCollectionSection& Section = Sections[SectionIndex];
-
-			const int32 MaterialID = Section.MaterialID;
-			HActor* const HitProxy = new HActor(Component->GetOwner(), Component, SectionIndex, MaterialID);
-
-			OutHitProxies.Add(HitProxy);
-			Section.HitProxy = HitProxy;
-		}
-	#endif
-	}
-
+	OutHitProxies.Append(HitProxies);
 	return DefaultHitProxy;
 }
-#endif // WITH_EDITOR
-
-#if GEOMETRYCOLLECTION_EDITOR_SELECTION
-void FGeometryCollectionSceneProxy::UseSubSections(bool bInUsesSubSections, bool bForceInit)
-{
-	if (!bForceInit)
-	{
-		bUsesSubSections = bInUsesSubSections;
-	}
-	else if (bInUsesSubSections)
-	{
-		FGeometryCollectionSceneProxy* GeometryCollectionSceneProxy = this;
-		ENQUEUE_RENDER_COMMAND(InitializeSubSections)(
-			[GeometryCollectionSceneProxy](FRHICommandListImmediate& RHICmdList)
-			{
-				if (GeometryCollectionSceneProxy)
-				{
-					GeometryCollectionSceneProxy->InitializeSubSections_RenderThread();
-					GeometryCollectionSceneProxy->bUsesSubSections = true;
-					UE_LOG(FGeometryCollectionSceneProxyLogging, Verbose, TEXT("UseSubSections, %d SubSections initialized for %p."), GeometryCollectionSceneProxy->SubSections.Num(), GeometryCollectionSceneProxy);
-				}
-			});
-	}
-	else
-	{
-		FGeometryCollectionSceneProxy* GeometryCollectionSceneProxy = this;
-		ENQUEUE_RENDER_COMMAND(ReleaseSubSections)(
-			[GeometryCollectionSceneProxy](FRHICommandListImmediate& RHICmdList)
-			{
-				if (GeometryCollectionSceneProxy)
-				{
-					GeometryCollectionSceneProxy->ReleaseSubSections_RenderThread();
-					GeometryCollectionSceneProxy->bUsesSubSections = false;
-					UE_LOG(FGeometryCollectionSceneProxyLogging, Verbose, TEXT("UseSubSections, SubSections released for %p."), GeometryCollectionSceneProxy);
-				}
-			});
-	}
-}
+#endif
 
-void FGeometryCollectionSceneProxy::InitializeSubSections_RenderThread()
+void FGeometryCollectionSceneProxy::GetPreSkinnedLocalBounds(FBoxSphereBounds& OutBounds) const
 {
-	// Exit now if there isn't any data
-	if (SubsetConstantData.OriginalMeshSections.IsEmpty())
-	{
-		SubSections.Empty();
-		SubSectionHitProxyIndexMap.Empty();
-		return;
-	}
-
-	// Retrieve the correct arrays depending on the dynamic state
-	const bool bIsDynamic = DynamicData && DynamicData->IsDynamic;
-	const TArray<FGeometryCollectionSection>& SectionArray = bIsDynamic ? Sections: SubsetConstantData.OriginalMeshSections;
-	const TArray<FIntVector>& IndexArray = bIsDynamic ? SubsetConstantData.Indices: SubsetConstantData.OriginalMeshIndices;
-	const TArray<int32>& BoneMap = SubsetConstantData.BoneMap;
-
-	// Reserve sub sections array with a minimum of one transform per section
-	SubSections.Empty(SectionArray.Num());
-	SubSectionHitProxyIndexMap.Empty(SectionArray.Num());
-
-	// Lambda that adds a new subsection and update the HitProxy section index
-	auto AddSubSection = [this, IndexArray](int32 HitProxyIndex, const FGeometryCollectionSection& Section, int32 FirstFaceIndex, int32 EndFaceIndex)
-	{
-		// Find the matching HitProxy for this transform/section
-		HGeometryCollection* const SubSectionHitProxy = SubSectionHitProxies[HitProxyIndex];
-
-		// Add the subsection
-		FGeometryCollectionSection SubSection;
-		SubSection.MaterialID = Section.MaterialID;
-		SubSection.FirstIndex = FirstFaceIndex * 3;
-		SubSection.NumTriangles = EndFaceIndex - FirstFaceIndex;
-		{
-			// Find out new min/max vertex indices
-			check(SubSection.NumTriangles > 0);
-			SubSection.MinVertexIndex = TNumericLimits<int32>::Max();
-			SubSection.MaxVertexIndex = TNumericLimits<int32>::Min();
-			for (int32 FaceIndex = FirstFaceIndex; FaceIndex < EndFaceIndex; ++FaceIndex)
-			{
-				SubSection.MinVertexIndex = FMath::Min(SubSection.MinVertexIndex, IndexArray[FaceIndex].GetMin());
-				SubSection.MaxVertexIndex = FMath::Max(SubSection.MaxVertexIndex, IndexArray[FaceIndex].GetMax());
-			}
-			check(SubSection.MinVertexIndex >= Section.MinVertexIndex && SubSection.MinVertexIndex <= Section.MaxVertexIndex)
-			check(SubSection.MaxVertexIndex >= Section.MinVertexIndex && SubSection.MaxVertexIndex <= Section.MaxVertexIndex)
-		}
-		SubSection.HitProxy = SubSectionHitProxy;
-		const int32 SubSectionIndex = SubSections.Add(SubSection);
-
-		// Keep the HitProxy index in a map in case this section's HitProxy pointer ever needs to be updated (e.g. after CreateHitProxies is called)
-		SubSectionHitProxyIndexMap.Add(SubSectionIndex, HitProxyIndex);
-
-		// Update HitProxy with this subsection index
-		if (SubSectionHitProxy)
-		{
-			SubSectionHitProxy->SectionIndex = SubSectionIndex;
-		}
-	};
-
-	// Create subsections per transform
-	const int32 NumTransforms = (SectionArray.Num() > 0) ? SubSectionHitProxies.Num() / SectionArray.Num(): 0;
-
-	for (int32 SectionIndex = 0; SectionIndex < SectionArray.Num(); ++SectionIndex)
-	{
-		const int32 SectionOffset = SectionIndex * NumTransforms;
-
-		const FGeometryCollectionSection& Section = SectionArray[SectionIndex];
-		check(Section.NumTriangles > 0);  // Sections are not created with zero triangles
-
-		const int32 FirstFaceIndex = Section.FirstIndex / 3;
-		const int32 EndFaceIndex = FirstFaceIndex + Section.NumTriangles;
-
-		int32 TransformIndex = BoneMap[IndexArray[FirstFaceIndex][0]];  // Assumes one transform per triangle
-		int32 FaceIndex = FirstFaceIndex;
-
-		for (int32 NextFaceIndex = FaceIndex + 1; NextFaceIndex < EndFaceIndex; ++NextFaceIndex)
-		{
-			const int32 NextTransformIndex = BoneMap[IndexArray[NextFaceIndex][0]];  // Assumes one transform per triangle
-			if (TransformIndex != NextTransformIndex)
-			{
-				// Add the current subsection
-				AddSubSection(SectionOffset + TransformIndex, Section, FaceIndex, NextFaceIndex);
-
-				// Update variables for the next subsection
-				TransformIndex = NextTransformIndex;
-				FaceIndex = NextFaceIndex;
-			}
-		}
-
-		// Add the last remaining subsection
-		AddSubSection(SectionOffset + TransformIndex, Section, FaceIndex, EndFaceIndex);
-	}
+	OutBounds = MeshDescription.PreSkinnedBounds; 
 }
 
-void FGeometryCollectionSceneProxy::ReleaseSubSections_RenderThread()
+uint32 FGeometryCollectionSceneProxy::GetAllocatedSize() const
 {
-	SubSections.Reset();
-	SubSectionHitProxyIndexMap.Reset();
+	return FPrimitiveSceneProxy::GetAllocatedSize()
+		+ Materials.GetAllocatedSize()
+		+ MeshDescription.Sections.GetAllocatedSize()
+		+ MeshDescription.SubSections.GetAllocatedSize()
+		+ RestTransforms.GetAllocatedSize()
+		+ (SkinnedPositionVertexBuffer.GetAllowCPUAccess() ? SkinnedPositionVertexBuffer.GetStride() * SkinnedPositionVertexBuffer.GetNumVertices() : 0)
+#if WITH_EDITOR
+		+ BoneColors.GetAllocatedSize()
+		+ (ColorVertexBuffer.GetAllowCPUAccess() ? ColorVertexBuffer.GetStride() * ColorVertexBuffer.GetNumVertices() : 0)
+		+ HiddenTransforms.GetAllocatedSize()
+#endif
+#if GEOMETRYCOLLECTION_EDITOR_SELECTION
+		+ HitProxies.GetAllocatedSize()
+		+ (HitProxyIdBuffer.GetAllowCPUAccess() ? HitProxyIdBuffer.GetStride() * HitProxyIdBuffer.GetNumVertices() : 0)
+#endif
+#if RHI_RAYTRACING
+		+ RayTracingGeometry.RawData.GetAllocatedSize()
+#endif
+		;
 }
 
-#endif  // #if GEOMETRYCOLLECTION_EDITOR_SELECTION
-
-void FGeometryCollectionSceneProxy::GetPreSkinnedLocalBounds(FBoxSphereBounds& OutBounds) const
-{
-	OutBounds = PreSkinnedBounds;
-}
 
 FNaniteGeometryCollectionSceneProxy::FNaniteGeometryCollectionSceneProxy(UGeometryCollectionComponent* Component)
 : Nanite::FSceneProxyBase(Component)
@@ -1549,9 +1007,7 @@ FNaniteGeometryCollectionSceneProxy::FNaniteGeometryCollectionSceneProxy(UGeomet
 	bShouldUpdateGPUSceneTransforms = (GGeometryCollectionOptimizedTransforms == 0);
 
 	bSupportsDistanceFieldRepresentation = false;
-
-	// Dynamic draw path without Nanite isn't supported by Lumen
-	bVisibleInLumenScene = false;
+	bSupportsMeshCardRepresentation = false;
 
 	// Use fast path that does not update static draw lists.
 	bStaticElementsAlwaysUseProxyPrimitiveUniformBuffer = true;
@@ -1576,6 +1032,8 @@ FNaniteGeometryCollectionSceneProxy::FNaniteGeometryCollectionSceneProxy(UGeomet
 	}
 
 	const TSharedPtr<FGeometryCollection, ESPMode::ThreadSafe> Collection = GeometryCollection->GetGeometryCollection();
+	const TManagedArray<int32>& TransformToGeometryIndices = Collection->TransformToGeometryIndex;
+	const TManagedArray<int32>& SimulationType = Collection->SimulationType;
 	const TManagedArray<FGeometryCollectionSection>& SectionsArray = Component->GetSectionsArray();
 
 	MaterialSections.SetNumZeroed(SectionsArray.Num());
@@ -1591,7 +1049,7 @@ FNaniteGeometryCollectionSceneProxy::FNaniteGeometryCollectionSceneProxy(UGeomet
 		UMaterialInterface* MaterialInterface = bValidMeshSection ? Component->GetMaterial(MeshSection.MaterialID) : nullptr;
 
 		// TODO: PROG_RASTER (Implement programmable raster support)
-		const bool bInvalidMaterial = !MaterialInterface || !IsOpaqueBlendMode(*MaterialInterface);
+		const bool bInvalidMaterial = !MaterialInterface || MaterialInterface->GetBlendMode() != BLEND_Opaque;
 		if (bInvalidMaterial)
 		{
 			bHasMaterialErrors = true;
@@ -1618,7 +1076,7 @@ FNaniteGeometryCollectionSceneProxy::FNaniteGeometryCollectionSceneProxy(UGeomet
 		check(MaterialInterface != nullptr);
 
 		// Should always be opaque blend mode here.
-		check(IsOpaqueBlendMode(*MaterialInterface));
+		check(MaterialInterface->GetBlendMode() == BLEND_Opaque);
 
 		MaterialSections[SectionIndex].ShadingMaterialProxy = MaterialInterface->GetRenderProxy();
 		MaterialSections[SectionIndex].RasterMaterialProxy  = MaterialInterface->GetRenderProxy(); // TODO: PROG_RASTER (Implement programmable raster support)
@@ -1650,7 +1108,7 @@ FNaniteGeometryCollectionSceneProxy::FNaniteGeometryCollectionSceneProxy(UGeomet
 	}
 	else if (bHasTransformBoundingBoxes)
 	{
-		Nanite::FResources& Resource = GeometryCollection->NaniteData->NaniteResource;
+		Nanite::FResources& Resource = GeometryCollection->RenderData->NaniteResource;
 		NumGeometry = Resource.HierarchyRootOffsets.Num();
 		GeometryNaniteData.SetNumUninitialized(NumGeometry);
 		
@@ -1664,29 +1122,20 @@ FNaniteGeometryCollectionSceneProxy::FNaniteGeometryCollectionSceneProxy(UGeomet
 			{
 				FGeometryNaniteData& Instance = GeometryNaniteData[GeometryIndex];
 				Instance.HierarchyOffset = GeometryCollection->GetNaniteHierarchyOffset(GeometryIndex);
-				Instance.LocalBounds = BoundingBoxes[TransformIndex];
+				SetInstanceLocalBounds(GeometryIndex, FRenderBounds(), false);
 			}
 		}
 	}
 
-	// Need to specify initial instance list, even with just identity transforms, so that the
-	// GPUScene instance data allocator reserves space for the instances early on. The instance
-	// transforms will be corrected during the first frame before any rendering occurs.
-	InstanceSceneData.SetNumUninitialized(NumGeometry);
-	InstanceDynamicData.SetNumUninitialized(NumGeometry);
-	InstanceLocalBounds.SetNumUninitialized(NumGeometry);
-	InstanceHierarchyOffset.SetNumZeroed(NumGeometry);
-
-	for (int32 GeometryIndex = 0; GeometryIndex < NumGeometry; ++GeometryIndex)
-	{
-		FInstanceSceneData& SceneData = InstanceSceneData[GeometryIndex];
-		SceneData.LocalToPrimitive.SetIdentity();
-
-		FInstanceDynamicData& DynamicData = InstanceDynamicData[GeometryIndex];
-		DynamicData.PrevLocalToPrimitive.SetIdentity();
+	// Initialize to rest transforms.
+	TArray<FMatrix44f> RestTransforms;
+	Component->GetRestTransforms(RestTransforms);
 
-		SetInstanceLocalBounds(GeometryIndex, FRenderBounds(), false);
-	}
+	FGeometryCollectionDynamicData* DynamicData = GDynamicDataPool.Allocate();
+	DynamicData->IsDynamic = true;
+	DynamicData->Transforms = RestTransforms;
+	DynamicData->PrevTransforms = RestTransforms;
+	SetDynamicData_RenderThread(DynamicData);
 }
 
 SIZE_T FNaniteGeometryCollectionSceneProxy::GetTypeHash() const
@@ -1726,7 +1175,7 @@ FPrimitiveViewRelevance FNaniteGeometryCollectionSceneProxy::GetViewRelevance(co
 	return Result;
 }
 
-#if WITH_EDITOR
+#if GEOMETRYCOLLECTION_EDITOR_SELECTION
 HHitProxy* FNaniteGeometryCollectionSceneProxy::CreateHitProxies(UPrimitiveComponent* Component, TArray<TRefCountPtr<HHitProxy>>& OutHitProxies)
 {
 	LLM_SCOPE_BYTAG(Nanite);
@@ -1778,7 +1227,7 @@ void FNaniteGeometryCollectionSceneProxy::GetNaniteMaterialMask(FUint32Vector2&
 
 Nanite::FResourceMeshInfo FNaniteGeometryCollectionSceneProxy::GetResourceMeshInfo() const
 {
-	Nanite::FResources& Resource = GeometryCollection->NaniteData->NaniteResource;
+	Nanite::FResources& Resource = GeometryCollection->RenderData->NaniteResource;
 
 	Nanite::FResourceMeshInfo OutInfo;
 
@@ -1795,42 +1244,6 @@ Nanite::FResourceMeshInfo FNaniteGeometryCollectionSceneProxy::GetResourceMeshIn
 	return MoveTemp(OutInfo);
 }
 
-void FNaniteGeometryCollectionSceneProxy::SetConstantData_RenderThread(const TUniquePtr<FGeometryCollectionConstantData>& NewConstantData, bool ForceInit)
-{
-	const TSharedPtr<FGeometryCollection, ESPMode::ThreadSafe> Collection = GeometryCollection->GetGeometryCollection();
-	const TManagedArray<int32>& TransformToGeometryIndices = Collection->TransformToGeometryIndex;
-
-	const int32 TransformCount = NewConstantData->RestTransforms.Num();
-	check(TransformCount == TransformToGeometryIndices.Num());
-	InstanceSceneData.Reset(TransformCount);
-	InstanceDynamicData.Reset(TransformCount);
-	InstanceLocalBounds.Reset(TransformCount);
-	InstanceHierarchyOffset.Reset(TransformCount);
-
-	for (int32 TransformIndex = 0; TransformIndex < TransformCount; ++TransformIndex)
-	{
-		const int32 TransformToGeometryIndex = TransformToGeometryIndices[TransformIndex];
-		if (!Collection->IsGeometry(TransformIndex))
-		{
-			continue;
-		}
-
-		const FGeometryNaniteData& NaniteData = GeometryNaniteData[TransformToGeometryIndex];
-
-		FInstanceSceneData& Instance	= InstanceSceneData.Emplace_GetRef();
-		Instance.LocalToPrimitive		= NewConstantData->RestTransforms[TransformIndex];
-
-		FInstanceDynamicData& DynamicData = InstanceDynamicData.Emplace_GetRef();
-		DynamicData.PrevLocalToPrimitive = Instance.LocalToPrimitive;
-
-		int32 InstanceIndex = InstanceLocalBounds.Num();
-		InstanceLocalBounds.SetNumUninitialized(InstanceIndex + 1);
-		SetInstanceLocalBounds(InstanceIndex, NaniteData.LocalBounds);
-		
-		InstanceHierarchyOffset.Emplace(NaniteData.HierarchyOffset);
-	}
-}
-
 void FNaniteGeometryCollectionSceneProxy::SetDynamicData_RenderThread(FGeometryCollectionDynamicData* NewDynamicData)
 {
 	// Are we currently simulating?
@@ -1929,6 +1342,7 @@ void FNaniteGeometryCollectionSceneProxy::OnMotionEnd()
 	ResetPreviousTransforms_RenderThread();
 }
 
+
 FGeometryCollectionDynamicDataPool::FGeometryCollectionDynamicDataPool()
 {
 	FreeList.SetNum(32);
diff --git a/Engine/Source/Runtime/Experimental/GeometryCollectionEngine/Private/GeometryCollection/GeometryCollectionSceneProxy.h b/Engine/Source/Runtime/Experimental/GeometryCollectionEngine/Private/GeometryCollection/GeometryCollectionSceneProxy.h
--- a/Engine/Source/Runtime/Experimental/GeometryCollectionEngine/Private/GeometryCollection/GeometryCollectionSceneProxy.h
+++ b/Engine/Source/Runtime/Experimental/GeometryCollectionEngine/Private/GeometryCollection/GeometryCollectionSceneProxy.h
@@ -13,51 +13,19 @@
 #include "HitProxies.h"
 #include "EngineUtils.h"
 #include "NaniteSceneProxy.h"
-
-#if GEOMETRYCOLLECTION_EDITOR_SELECTION
+#include "GeometryCollection/GeometryCollectionObject.h"
+#include "GeometryCollection/GeometryCollectionRenderData.h"
 #include "GeometryCollection/GeometryCollectionHitProxy.h"
-#endif
 
 class UGeometryCollection;
 class UGeometryCollectionComponent;
 struct FGeometryCollectionSection;
-struct HGeometryCollection;
 
 namespace Nanite
 {
 	struct FResources;
 }
 
-/** Index Buffer */
-class FGeometryCollectionIndexBuffer : public FIndexBuffer
-{
-public:
-	virtual void InitRHI() override
-	{
-		FRHIResourceCreateInfo CreateInfo(TEXT("FGeometryCollectionIndexBuffer"));
-		IndexBufferRHI = RHICreateIndexBuffer(sizeof(int32), NumIndices * sizeof(int32), BUF_Dynamic, CreateInfo);
-	}
-
-	int32 NumIndices;
-};
-
-/** Vertex Buffer for Bone Map*/
-class FGeometryCollectionBoneMapBuffer : public FVertexBuffer
-{
-public:
-	virtual void InitRHI() override
-	{
-		FRHIResourceCreateInfo CreateInfo(TEXT("FGeometryCollectionBoneMapBuffer"));
-
-		// #note: Bone Map is stored in uint16, but shaders only support uint32
-		VertexBufferRHI = RHICreateVertexBuffer(NumVertices * sizeof(uint32), BUF_Static | BUF_ShaderResource, CreateInfo);		
-		VertexBufferSRV = RHICreateShaderResourceView(VertexBufferRHI, sizeof(uint32), PF_R32_UINT);		
-	}
-
-	int32 NumVertices;
-
-	FShaderResourceViewRHIRef VertexBufferSRV;
-};
 
 /** Vertex Buffer for transform data */
 class FGeometryCollectionTransformBuffer : public FVertexBuffer
@@ -91,37 +59,6 @@ inline void CopyTransformsWithConversionWhenNeeded(TArray<FMatrix44f>& DstTransf
 	}
 }
 
-/** Immutable rendering data (kind of) */
-struct FGeometryCollectionConstantData
-{
-	using FUVChannel = TArray<FVector2f>;
-
-	TArray<FVector3f> Vertices;
-	TArray<FIntVector> Indices;
-	TArray<FVector3f> Normals;
-	TArray<FVector3f> TangentU;
-	TArray<FVector3f> TangentV;
-	TArray<FUVChannel> UVChannels;
-	TArray<FLinearColor> Colors;
-	TArray<int32> BoneMap;
-	TArray<FLinearColor> BoneColors;
-	TArray<FGeometryCollectionSection> Sections;
-
-	uint32 NumTransforms;
-
-	FBox LocalBounds;
-	
-	TArray<FIntVector> OriginalMeshIndices;
-	TArray<FGeometryCollectionSection> OriginalMeshSections;
-
-	TArray<FMatrix44f> RestTransforms;
-
-	void SetRestTransforms(const TArray<FMatrix>& InTransforms)
-	{
-		// use for LWC as FMatrix and FMatrix44f are different when LWC is on 
-		CopyTransformsWithConversionWhenNeeded(RestTransforms, InTransforms);
-	}
-};
 
 /** Mutable rendering data */
 struct FGeometryCollectionDynamicData
@@ -188,6 +125,7 @@ struct FGeometryCollectionDynamicData
 	}
 };
 
+
 class FGeometryCollectionDynamicDataPool
 {
 public:
@@ -205,7 +143,6 @@ private:
 };
 
 
-
 /***
 *   FGeometryCollectionSceneProxy
 *    
@@ -216,153 +153,103 @@ private:
 *   thoughts are appreciated though. The remaining items to address involve:
 *   - @todo double buffer - The double buffering of the FGeometryCollectionDynamicData.
 *   - @todo previous state - Saving the previous FGeometryCollectionDynamicData for rendering motion blur.
-*   - @todo shared memory model - The Asset(or Actor?) should hold the Vertex buffer, and pass the reference to the SceneProxy
 *   - @todo GPU skin : Make the skinning use the GpuVertexShader
 */
 class FGeometryCollectionSceneProxy final : public FPrimitiveSceneProxy
 {
 	TArray<UMaterialInterface*> Materials;
-	UMaterialInterface* BoneSelectedMaterial = nullptr;
 
 	FMaterialRelevance MaterialRelevance;
 
-	int32 NumVertices;
-	int32 NumIndices;
+	FGeometryCollectionMeshResources const& MeshResource;
+	FGeometryCollectionMeshDescription MeshDescription;
+
+	int32 NumTransforms = 0;
+	TArray<FMatrix44f> RestTransforms;
+
+	FBoxSphereBounds PreSkinnedBounds;
 
 	FGeometryCollectionVertexFactory VertexFactory;
 	
 	bool bSupportsManualVertexFetch;
-	const bool bSupportsTripleBufferVertexUpload;
-	
-	FStaticMeshVertexBuffers VertexBuffers;
-	FGeometryCollectionIndexBuffer IndexBuffer;
-	FGeometryCollectionIndexBuffer OriginalMeshIndexBuffer;
-	FGeometryCollectionBoneMapBuffer BoneMapBuffer;
+	FPositionVertexBuffer SkinnedPositionVertexBuffer;
+
+	int32 CurrentTransformBufferIndex = 0;
+	bool TransformVertexBuffersContainsRestTransforms = true;
+	bool bSupportsTripleBufferVertexUpload = false;
 	TArray<FGeometryCollectionTransformBuffer, TInlineAllocator<3>> TransformBuffers;
 	TArray<FGeometryCollectionTransformBuffer, TInlineAllocator<3>> PrevTransformBuffers;
 
-	int32 CurrentTransformBufferIndex = 0;
-	FBoxSphereBounds PreSkinnedBounds;
+	FGeometryCollectionDynamicData* DynamicData = nullptr;
+
+#if WITH_EDITOR
+	bool bShowBoneColors = false;
+	bool bSuppressSelectionMaterial = false;
+	TArray<FColor> BoneColors;
+	FColorVertexBuffer ColorVertexBuffer;
+	FGeometryCollectionVertexFactory VertexFactoryDebugColor;
+	UMaterialInterface* BoneSelectedMaterial = nullptr;
+	TArray<bool> HiddenTransforms;
+#endif
 
-	TArray<FGeometryCollectionSection> Sections;
 #if GEOMETRYCOLLECTION_EDITOR_SELECTION
+	bool bUsesSubSections = false;
+	bool bEnableBoneSelection = false;
+	TArray<TRefCountPtr<HHitProxy>> HitProxies;
 	FColorVertexBuffer HitProxyIdBuffer;
-	TArray<FGeometryCollectionSection> SubSections;
-	TArray<TRefCountPtr<HGeometryCollection>> SubSectionHitProxies;
-	TMap<int32, int32> SubSectionHitProxyIndexMap;
-	// @todo FractureTools - Reconcile with SubSectionHitProxies.  Currently subsection hit proxies dont work for per-vertex submission
-	TArray<TRefCountPtr<HGeometryCollectionBone>> PerBoneHitProxies;
-	FColor WholeObjectHitProxyColor = FColor(EForceInit::ForceInit);
-	bool bUsesSubSections;
 #endif
 
-	FGeometryCollectionDynamicData* DynamicData;
-
-	struct SubsetConstantData
-	{
-		TArray<FIntVector> Indices;
-		TArray<int32> BoneMap;
-		uint32 NumTransforms;
-		TArray<FIntVector> OriginalMeshIndices;
-		TArray<FGeometryCollectionSection> OriginalMeshSections;
-		TArray<FMatrix44f> RestTransforms;
-	} SubsetConstantData;
-
-	bool bShowBoneColors;
-	bool bEnableBoneSelection;
-	bool bSuppressSelectionMaterial;
-
-	bool bUseFullPrecisionUVs = false;
-
-	bool TransformVertexBuffersContainsOriginalMesh;
+#if RHI_RAYTRACING
+	bool bGeometryResourceUpdated = false;
+	FRayTracingGeometry RayTracingGeometry;
+	FRWBuffer RayTracingDynamicVertexBuffer;
+#endif
 
 public:
-	SIZE_T GetTypeHash() const override
-	{
-		static size_t UniquePointer;
-		return reinterpret_cast<size_t>(&UniquePointer);
-	}
-
 	FGeometryCollectionSceneProxy(UGeometryCollectionComponent* Component);
-
-	/** virtual destructor */
 	virtual ~FGeometryCollectionSceneProxy();
 
-	void DestroyRenderThreadResources()override;
-
-	/** Current number of vertices to render */
-	int32 GetRequiredVertexCount() const { return NumVertices; }
-
-	/** Current number of indices to connect */
-	int32 GetRequiredIndexCount() const { return NumIndices; }
-
-	/** Called on render thread to setup static geometry for rendering */
-	void SetConstantData_RenderThread(const TUniquePtr<FGeometryCollectionConstantData>& NewConstantData, bool ForceInit = false);
-
 	/** Called on render thread to setup dynamic geometry for rendering */
 	void SetDynamicData_RenderThread(FGeometryCollectionDynamicData* NewDynamicData);
 
-	/** Called on render thread to construct the vertex definitions */
-	void BuildGeometry(const TUniquePtr<FGeometryCollectionConstantData>& ConstantDataIn, TArray<FDynamicMeshVertex>& OutVertices, TArray<int32>& OutIndices, TArray<int32> &OutOriginalMeshIndices);
-
-	/** Called on render thread to setup dynamic geometry for rendering */
-	virtual void GetDynamicMeshElements(const TArray<const FSceneView*>& Views, const FSceneViewFamily& ViewFamily, uint32 VisibilityMap, FMeshElementCollector& Collector) const override;
-
-#if RHI_RAYTRACING
-	virtual bool					IsRayTracingRelevant() const { return true; }
-	virtual bool					IsRayTracingStaticRelevant() const { return false; }
-	virtual void					GetDynamicRayTracingInstances(FRayTracingMaterialGatheringContext& Context, TArray<struct FRayTracingInstance>& OutRayTracingInstances) override;
-
-	void UpdatingRayTracingGeometry_RenderingThread(FGeometryCollectionIndexBuffer* IndexBuffer);
-#endif
-
-	/** Manage the view assignment */
-	virtual FPrimitiveViewRelevance GetViewRelevance(const FSceneView* View) const override;
-
-	// @todo allocated size : make this reflect internally allocated memory. 
-	virtual uint32 GetMemoryFootprint(void) const override { return(sizeof(*this) + GetAllocatedSize()); }
+	uint32 GetMemoryFootprint() const override { return sizeof(*this) + GetAllocatedSize(); }
+	uint32 GetAllocatedSize() const;
 
-	/** Size of the base class */
-	uint32 GetAllocatedSize(void) const { return(FPrimitiveSceneProxy::GetAllocatedSize()); }
-
-	// FPrimitiveSceneProxy interface.
-#if WITH_EDITOR
-	virtual HHitProxy* CreateHitProxies(UPrimitiveComponent* Component, TArray<TRefCountPtr<HHitProxy> >& OutHitProxies) override;
-	virtual const FColorVertexBuffer* GetCustomHitProxyIdBuffer() const override
-	{
-		// Note: Could return nullptr when bEnableBoneSelection is false if the hitproxy shader was made to not require per-vertex hit proxy IDs in that case
-		return &HitProxyIdBuffer;
-	}
-#endif // WITH_EDITOR
+	SIZE_T GetTypeHash() const override;
+	void CreateRenderThreadResources() override;
+	void DestroyRenderThreadResources() override;
+	void GetPreSkinnedLocalBounds(FBoxSphereBounds& OutBounds) const override;
+	FPrimitiveViewRelevance GetViewRelevance(const FSceneView* View) const override;
+	void GetDynamicMeshElements(const TArray<const FSceneView*>& Views, const FSceneViewFamily& ViewFamily, uint32 VisibilityMap, FMeshElementCollector& Collector) const override;
 
 #if GEOMETRYCOLLECTION_EDITOR_SELECTION
-	/** Enable/disable the per transform selection mode. 
-	 *  This forces more sections/mesh batches to be sent to the renderer while also allowing the editor
-	 *  to return a special HitProxy containing the transform index of the section that has been clicked on.
-	 */
-	void UseSubSections(bool bInUsesSubSections, bool bForceInit);
+	virtual HHitProxy* CreateHitProxies(UPrimitiveComponent* Component, TArray<TRefCountPtr<HHitProxy> >& OutHitProxies) override;
+ 	virtual const FColorVertexBuffer* GetCustomHitProxyIdBuffer() const override
+ 	{
+		return (bEnableBoneSelection || bUsesSubSections) ? &HitProxyIdBuffer : nullptr;
+ 	}
 #endif
 
-	void GetPreSkinnedLocalBounds(FBoxSphereBounds& OutBounds) const override;
-
-	void SetupVertexFactory(FGeometryCollectionVertexFactory& GeometryCollectionVertexFactory) const;
+#if RHI_RAYTRACING
+	bool IsRayTracingRelevant() const override { return true; }
+	bool IsRayTracingStaticRelevant() const override { return false; }
+	void GetDynamicRayTracingInstances(FRayTracingMaterialGatheringContext& Context, TArray<struct FRayTracingInstance>& OutRayTracingInstances) override;
+#endif
 
 protected:
-
-	/** Create the rendering buffer resources */
-	void InitResources(const TUniquePtr<FGeometryCollectionConstantData>& ConstantData);
-
-	/** Return the rendering buffer resources */
-	void ReleaseResources();
-
+	/** Setup a geometry collection vertex factory. */
+	void SetupVertexFactory(FGeometryCollectionVertexFactory& GeometryCollectionVertexFactory, FColorVertexBuffer* ColorOverride = nullptr) const;
+	/** Update skinned position buffer used by mobile CPU skinning path. */
+	void UpdateSkinnedPositions(TArray<FMatrix44f> const& Transforms);
 	/** Get material proxy from material ID */
 	FMaterialRenderProxy* GetMaterial(FMeshElementCollector& Collector, int32 MaterialIndex) const;
+	/** Get the standard or debug vertex factory dependent on current state. */
+	FVertexFactory const* GetVertexFactory() const;
 
 	FGeometryCollectionTransformBuffer& GetCurrentTransformBuffer()
 	{
 		return TransformBuffers[CurrentTransformBufferIndex];
 	}
-
 	FGeometryCollectionTransformBuffer& GetCurrentPrevTransformBuffer()
 	{
 		return PrevTransformBuffers[CurrentTransformBufferIndex];
@@ -376,29 +263,18 @@ protected:
 		}
 	}
 
-private:
-#if GEOMETRYCOLLECTION_EDITOR_SELECTION
-	/** Create transform index based subsections for all current sections. */
-	void InitializeSubSections_RenderThread();
-
-	/** Release subsections by emptying the associated arrays. */
-	void ReleaseSubSections_RenderThread();
-#endif
-
 #if RHI_RAYTRACING
-	bool bGeometryResourceUpdated = false;
-	FRayTracingGeometry RayTracingGeometry;
-	FRWBuffer RayTracingDynamicVertexBuffer;
+	void UpdatingRayTracingGeometry_RenderingThread(TArray<FGeometryCollectionMeshElement> const& InSectionArray);
 #endif
 };
 
+
 class FNaniteGeometryCollectionSceneProxy : public Nanite::FSceneProxyBase
 {
 public:
 	using Super = Nanite::FSceneProxyBase;
 	
 	FNaniteGeometryCollectionSceneProxy(UGeometryCollectionComponent* Component);
-
 	virtual ~FNaniteGeometryCollectionSceneProxy() = default;
 
 public:
@@ -420,9 +296,6 @@ public:
 
 	virtual Nanite::FResourceMeshInfo GetResourceMeshInfo() const override;
 
-	/** Called on render thread to setup static geometry for rendering */
-	void SetConstantData_RenderThread(const TUniquePtr<FGeometryCollectionConstantData>& NewConstantData, bool ForceInit = false);
-
 	/** Called on render thread to setup dynamic geometry for rendering */
 	void SetDynamicData_RenderThread(FGeometryCollectionDynamicData* NewDynamicData);
 
@@ -444,6 +317,7 @@ public:
 	void OnMotionEnd();
 
 protected:
+	// TODO : Copy required data from UObject instead of using unsafe object pointer.
 	const UGeometryCollection* GeometryCollection = nullptr;
 
 	struct FGeometryNaniteData
diff --git a/Engine/Source/Runtime/Experimental/GeometryCollectionEngine/Private/GeometryCollection/GeometryCollectionSceneProxy.ispc b/Engine/Source/Runtime/Experimental/GeometryCollectionEngine/Private/GeometryCollection/GeometryCollectionSceneProxy.ispc
--- a/Engine/Source/Runtime/Experimental/GeometryCollectionEngine/Private/GeometryCollection/GeometryCollectionSceneProxy.ispc
+++ b/Engine/Source/Runtime/Experimental/GeometryCollectionEngine/Private/GeometryCollection/GeometryCollectionSceneProxy.ispc
@@ -6,7 +6,7 @@
 export void SetDynamicData_RenderThread(uniform FVector3f PositionBuffer[], 
 										const uniform unsigned int NumVertices, 
 										const uniform unsigned int Stride, 
-										const uniform int32 BoneMap[], 
+										const uniform uint16 BoneMap[], 
 										const uniform FMatrix44f Transforms[], 
 										const uniform FVector3f Vertices[])
 {
diff --git a/Engine/Source/Runtime/Experimental/GeometryCollectionEngine/Public/GeometryCollection/DerivedDataGeometryCollectionCooker.h b/Engine/Source/Runtime/Experimental/GeometryCollectionEngine/Public/GeometryCollection/DerivedDataGeometryCollectionCooker.h
--- a/Engine/Source/Runtime/Experimental/GeometryCollectionEngine/Public/GeometryCollection/DerivedDataGeometryCollectionCooker.h
+++ b/Engine/Source/Runtime/Experimental/GeometryCollectionEngine/Public/GeometryCollection/DerivedDataGeometryCollectionCooker.h
@@ -42,7 +42,7 @@ public:
 	virtual bool Build(TArray<uint8>& OutData) override;
 
 	/** Return true if we can build **/
-	bool CanBuild()
+	bool CanBuild() const
 	{
 		return true;
 	}
@@ -51,4 +51,24 @@ private:
 	UGeometryCollection& GeometryCollection;
 };
 
+
+class FDerivedDataGeometryCollectionRenderDataCooker : public FDerivedDataPluginInterface
+{
+public:
+	FDerivedDataGeometryCollectionRenderDataCooker(UGeometryCollection& InGeometryCollection) : GeometryCollection(InGeometryCollection) {}
+
+	virtual const TCHAR* GetPluginName() const override { return TEXT("GeometryCollectionRenderData"); }
+	virtual const TCHAR* GetVersionString() const override;
+	virtual FString GetPluginSpecificCacheKeySuffix() const override;
+	virtual bool IsBuildThreadsafe() const override { return false;	}
+	virtual bool IsDeterministic() const override {	return true; }
+	virtual FString GetDebugContextString() const override;
+	virtual bool Build(TArray<uint8>& OutData) override;
+
+	bool CanBuild() const { return true; }
+
+private:
+	UGeometryCollection& GeometryCollection;
+};
+
 #endif	// WITH_EDITOR
diff --git a/Engine/Source/Runtime/Experimental/GeometryCollectionEngine/Public/GeometryCollection/GeometryCollectionComponent.h b/Engine/Source/Runtime/Experimental/GeometryCollectionEngine/Public/GeometryCollection/GeometryCollectionComponent.h
--- a/Engine/Source/Runtime/Experimental/GeometryCollectionEngine/Public/GeometryCollection/GeometryCollectionComponent.h
+++ b/Engine/Source/Runtime/Experimental/GeometryCollectionEngine/Public/GeometryCollection/GeometryCollectionComponent.h
@@ -453,6 +453,7 @@ public:
 	virtual void EndPlay(const EEndPlayReason::Type ReasonEnd) override;
 	virtual void GetLifetimeReplicatedProps(TArray<FLifetimeProperty>& OutLifetimeProps) const override;
 	virtual void InitializeComponent() override;
+	virtual void GetResourceSizeEx(FResourceSizeEx& CumulativeResourceSize) override;
 
 #if WITH_EDITOR
 	virtual void PostEditChangeChainProperty(FPropertyChangedChainEvent& PropertyChangedEvent) override;
@@ -893,6 +894,7 @@ public:
 #if GEOMETRYCOLLECTION_EDITOR_SELECTION
 	/** Enable/disable the scene proxy per transform selection mode. When disabled the per material id default selection is used instead. */
 	void EnableTransformSelectionMode(bool bEnable);
+	bool GetIsTransformSelectionMode() const { return bIsTransformSelectionModeEnabled; }
 #endif  // #if GEOMETRYCOLLECTION_EDITOR_SELECTION
 
 #if UE_ENABLE_DEBUG_DRAWING
@@ -973,8 +975,12 @@ public:
 #if WITH_EDITOR
 	void SetEmbeddedGeometrySelectable(bool bSelectableIn);
 	int32 EmbeddedIndexToTransformIndex(const UInstancedStaticMeshComponent* ISMComponent, int32 InstanceIndex) const;
+	void GetBoneColors(TArray<FColor>& OutColors) const;
+	void GetHiddenTransforms(TArray<bool>& OutHiddenTransforms) const;
 #endif
 
+	void GetRestTransforms(TArray<FMatrix44f>& OutRestTransforms) const;
+
 #if WITH_EDITORONLY_DATA
 	const FDamageCollector* GetRunTimeDataCollector() const;
 #endif
@@ -1065,9 +1071,6 @@ protected:
 	int32 ISMPoolMeshGroupIndex = INDEX_NONE;
 	int32 ISMPoolRootProxyMeshId = INDEX_NONE;
 
-	/** Populate the static geometry structures for the render thread. */
-	void InitConstantData(TUniquePtr<FGeometryCollectionConstantData>& ConstantData) const;
-
 	/** Populate the dynamic particle data for the render thread. */
 	FGeometryCollectionDynamicData* InitDynamicData(bool bInitialization);
 
diff --git a/Engine/Source/Runtime/Experimental/GeometryCollectionEngine/Public/GeometryCollection/GeometryCollectionHitProxy.h b/Engine/Source/Runtime/Experimental/GeometryCollectionEngine/Public/GeometryCollection/GeometryCollectionHitProxy.h
--- a/Engine/Source/Runtime/Experimental/GeometryCollectionEngine/Public/GeometryCollection/GeometryCollectionHitProxy.h
+++ b/Engine/Source/Runtime/Experimental/GeometryCollectionEngine/Public/GeometryCollection/GeometryCollectionHitProxy.h
@@ -1,38 +1,21 @@
 // Copyright Epic Games, Inc. All Rights Reserved.
 #pragma once
 
-#include "GeometryCollectionEditorSelection.h"
+#if WITH_EDITOR
 
-#if GEOMETRYCOLLECTION_EDITOR_SELECTION
 #include "EngineUtils.h"
 
 class UGeometryCollectionComponent;
 
-/** HitProxy with transform index information.
- * @todo: These are never actually drawn because the shader is always in per-vertex mode, where it is passed HGeometryCollectionBone (or the whole-object HActor)
- */
-struct HGeometryCollection : public HActor
-{
-	DECLARE_HIT_PROXY(GEOMETRYCOLLECTIONENGINE_API)
-	int32 TransformIndex;
-
-	HGeometryCollection(AActor* InActor, const UPrimitiveComponent* InPrimitiveComponent, int32 InSectionIndex, int32 InMaterialIndex, int32 InTransformIndex)
-		: HActor(InActor, InPrimitiveComponent, InSectionIndex, InMaterialIndex)
-		, TransformIndex(InTransformIndex)
-		{}
-};
-
-
 /** HitProxy with transform index information. */
-// @todo FractureTools - This is temp, we should merge it with HGeometryCollection
-struct HGeometryCollectionBone : public HHitProxy
+struct HGeometryCollection : public HHitProxy
 {
 	DECLARE_HIT_PROXY(GEOMETRYCOLLECTIONENGINE_API)
 
 	UGeometryCollectionComponent* Component;
 	int32 BoneIndex;
 
-	HGeometryCollectionBone(UGeometryCollectionComponent* InGeometryCollectionComponent, int32 InBoneIndex)
+	HGeometryCollection(UGeometryCollectionComponent* InGeometryCollectionComponent, int32 InBoneIndex)
 		: Component(InGeometryCollectionComponent)
 		, BoneIndex(InBoneIndex)
 	{}
@@ -42,4 +25,5 @@ struct HGeometryCollectionBone : public HHitProxy
 		Collector.AddReferencedObject(Component);
 	}
 };
-#endif  // #if GEOMETRYCOLLECTION_EDITOR_SELECTION
+
+#endif // WITH_EDITOR
diff --git a/Engine/Source/Runtime/Experimental/GeometryCollectionEngine/Public/GeometryCollection/GeometryCollectionObject.h b/Engine/Source/Runtime/Experimental/GeometryCollectionEngine/Public/GeometryCollection/GeometryCollectionObject.h
--- a/Engine/Source/Runtime/Experimental/GeometryCollectionEngine/Public/GeometryCollection/GeometryCollectionObject.h
+++ b/Engine/Source/Runtime/Experimental/GeometryCollectionEngine/Public/GeometryCollection/GeometryCollectionObject.h
@@ -3,23 +3,24 @@
 #pragma once
 
 #include "CoreMinimal.h"
-#include "Misc/Crc.h"
-#include "Rendering/NaniteResources.h"
-#include "InstanceUniformShaderParameters.h"
-#include "GeometryCollection/ManagedArray.h"
+#include "Chaos/ChaosSolverActor.h"
 #include "GeometryCollection/GeometryCollectionDamagePropagationData.h"
 #include "GeometryCollection/GeometryCollectionSimulationTypes.h"
-#include "Chaos/ChaosSolverActor.h"
+#include "GeometryCollection/ManagedArray.h"
+#include "InstanceUniformShaderParameters.h"
+#include "Misc/Crc.h"
 
 #include "GeometryCollectionObject.generated.h"
 
-class UMaterialInterface;
-class UGeometryCollectionCache;
 class FGeometryCollection;
-struct FManagedArrayCollection;
+class FGeometryCollectionRenderData;
 struct FGeometryCollectionSection;
+struct FManagedArrayCollection;
 struct FSharedSimulationParameters;
 class UDataflow;
+class UGeometryCollectionCache;
+class UMaterial;
+class UMaterialInterface;
 
 USTRUCT(BlueprintType)
 struct GEOMETRYCOLLECTIONENGINE_API FGeometryCollectionSource
@@ -304,32 +305,6 @@ struct TStructOpsTypeTraits<FGeometryCollectionSizeSpecificData> : public TStruc
 	};
 };
 
-class FGeometryCollectionNaniteData
-{
-public:
-	GEOMETRYCOLLECTIONENGINE_API FGeometryCollectionNaniteData();
-	GEOMETRYCOLLECTIONENGINE_API ~FGeometryCollectionNaniteData();
-
-	FORCEINLINE bool IsInitialized()
-	{
-		return bIsInitialized;
-	}
-
-	/** Serialization. */
-	void Serialize(FArchive& Ar, UGeometryCollection* Owner);
-
-	/** Initialize the render resources. */
-	void InitResources(UGeometryCollection* Owner);
-
-	/** Releases the render resources. */
-	GEOMETRYCOLLECTIONENGINE_API void ReleaseResources();
-
-	Nanite::FResources NaniteResource;
-
-private:
-	bool bIsInitialized = false;
-};
-
 
 /**
 * UGeometryCollectionObject (UObject)
@@ -361,7 +336,7 @@ public:
 #endif
 
 #if WITH_EDITOR
-	void EnsureDataIsCooked(bool bInitResources = true, bool bIsTransacting = false);
+	void EnsureDataIsCooked(bool bInitResources, bool bIsTransacting, bool bIsPersistant);
 #endif
 
 	/** Accessors for internal geometry collection */
@@ -379,34 +354,14 @@ public:
 	int32 NumElements(const FName& Group) const;
 	void RemoveElements(const FName& Group, const TArray<int32>& SortedDeletionList);
 
-	FORCEINLINE bool HasNaniteData() const
-	{
-		return NaniteData != nullptr;
-	}
-
-	FORCEINLINE uint32 GetNaniteResourceID() const
-	{
-		Nanite::FResources& Resource = NaniteData->NaniteResource;
-		return Resource.RuntimeResourceID;
-	}
-
-	FORCEINLINE uint32 GetNaniteHierarchyOffset() const
-	{
-		Nanite::FResources& Resource = NaniteData->NaniteResource;
-		return Resource.HierarchyOffset;
-	}
+	/** Has data for static mesh rendering. */
+	bool HasMeshData() const;
+	/** Has data for nanite rendering. */
+	bool HasNaniteData() const;
 
-	FORCEINLINE uint32 GetNaniteHierarchyOffset(int32 GeometryIndex, bool bFlattened = false) const
-	{
-		Nanite::FResources& Resource = NaniteData->NaniteResource;
-		check(GeometryIndex >= 0 && GeometryIndex < Resource.HierarchyRootOffsets.Num());
-		uint32 HierarchyOffset = Resource.HierarchyRootOffsets[GeometryIndex];
-		if (bFlattened)
-		{
-			HierarchyOffset += Resource.HierarchyOffset;
-		}
-		return HierarchyOffset;
-	}
+	uint32 GetNaniteResourceID() const;
+	uint32 GetNaniteHierarchyOffset() const;
+	uint32 GetNaniteHierarchyOffset(int32 GeometryIndex, bool bFlattened = false) const;
 
 	/** ReindexMaterialSections */
 	void ReindexMaterialSections();
@@ -465,8 +420,11 @@ public:
 	/** Create the simulation data ( calls CreateSimulationData) only if the simulation data is dirty */
 	void CreateSimulationDataIfNeeded();
 
-	/** Create the Nanite rendering data. */
-	static TUniquePtr<FGeometryCollectionNaniteData> CreateNaniteData(FGeometryCollection* Collection);
+	/** Rebuild the render data. */
+	void RebuildRenderData();
+
+	/** Propogate render state dirty to components */
+	void PropagateMarkDirtyToComponents() const;
 #endif
 
 	void InitResources();
@@ -479,8 +437,8 @@ public:
 	FGuid GetIdGuid() const;
 	FGuid GetStateGuid() const;
 
-	/** Pointer to the data used to render this geometry collection with Nanite. */
-	TUniquePtr<class FGeometryCollectionNaniteData> NaniteData;
+	/** Pointer to the data used to render this geometry collection. */
+	TUniquePtr<FGeometryCollectionRenderData> RenderData;
 
 	UPROPERTY(EditAnywhere, Category = "Clustering")
 	bool EnableClustering;
@@ -729,6 +687,7 @@ public:
 private:
 #if WITH_EDITOR
 	void CreateSimulationDataImp(bool bCopyFromDDC);
+	void CreateRenderDataImp(bool bCopyFromDDC);
 #endif
 
 	/*
@@ -752,11 +711,15 @@ private:
 	FGuid StateGuid;
 
 #if WITH_EDITOR
-	//Used to determine whether we need to cook content
-	FGuid LastBuiltGuid;
-
+	//Used to determine whether we need to cook simulation data
+	FGuid LastBuiltSimulationDataGuid;
 	//Used to determine whether we need to regenerate simulation data
 	FGuid SimulationDataGuid;
+
+	//Used to determine whether we need to cook render data
+	FGuid LastBuiltRenderDataGuid;
+	//Used to determine whether we need to regenerate render data
+	FGuid RenderDataGuid;
 #endif
 
 	// #todo(dmp): rename to be consistent BoneSelectedMaterialID?
