Fix merge issues

diff --git a/Engine/Source/Runtime/Experimental/GeometryCollectionEngine/Private/GeometryCollection/GeometryCollectionSceneProxy.cpp b/Engine/Source/Runtime/Experimental/GeometryCollectionEngine/Private/GeometryCollection/GeometryCollectionSceneProxy.cpp
--- a/Engine/Source/Runtime/Experimental/GeometryCollectionEngine/Private/GeometryCollection/GeometryCollectionSceneProxy.cpp
+++ b/Engine/Source/Runtime/Experimental/GeometryCollectionEngine/Private/GeometryCollection/GeometryCollectionSceneProxy.cpp
@@ -1007,7 +1007,9 @@ FNaniteGeometryCollectionSceneProxy::FNaniteGeometryCollectionSceneProxy(UGeomet
 	bShouldUpdateGPUSceneTransforms = (GGeometryCollectionOptimizedTransforms == 0);
 
 	bSupportsDistanceFieldRepresentation = false;
-	bSupportsMeshCardRepresentation = false;
+
+	// Dynamic draw path without Nanite isn't supported by Lumen
+	bVisibleInLumenScene = false;
 
 	// Use fast path that does not update static draw lists.
 	bStaticElementsAlwaysUseProxyPrimitiveUniformBuffer = true;
@@ -1049,7 +1051,7 @@ FNaniteGeometryCollectionSceneProxy::FNaniteGeometryCollectionSceneProxy(UGeomet
 		UMaterialInterface* MaterialInterface = bValidMeshSection ? Component->GetMaterial(MeshSection.MaterialID) : nullptr;
 
 		// TODO: PROG_RASTER (Implement programmable raster support)
-		const bool bInvalidMaterial = !MaterialInterface || MaterialInterface->GetBlendMode() != BLEND_Opaque;
+		const bool bInvalidMaterial = !MaterialInterface || !IsOpaqueBlendMode(*MaterialInterface);
 		if (bInvalidMaterial)
 		{
 			bHasMaterialErrors = true;
@@ -1076,7 +1078,7 @@ FNaniteGeometryCollectionSceneProxy::FNaniteGeometryCollectionSceneProxy(UGeomet
 		check(MaterialInterface != nullptr);
 
 		// Should always be opaque blend mode here.
-		check(MaterialInterface->GetBlendMode() == BLEND_Opaque);
+		check(IsOpaqueBlendMode(*MaterialInterface));
 
 		MaterialSections[SectionIndex].ShadingMaterialProxy = MaterialInterface->GetRenderProxy();
 		MaterialSections[SectionIndex].RasterMaterialProxy  = MaterialInterface->GetRenderProxy(); // TODO: PROG_RASTER (Implement programmable raster support)
