Fix for debug build.

diff --git a/Engine/Source/Runtime/Experimental/GeometryCollectionEngine/Private/GeometryCollection/GeometryCollectionRenderData.h b/Engine/Source/Runtime/Experimental/GeometryCollectionEngine/Private/GeometryCollection/GeometryCollectionRenderData.h
--- a/Engine/Source/Runtime/Experimental/GeometryCollectionEngine/Private/GeometryCollection/GeometryCollectionRenderData.h
+++ b/Engine/Source/Runtime/Experimental/GeometryCollectionEngine/Private/GeometryCollection/GeometryCollectionRenderData.h
@@ -25,7 +25,7 @@ public:
 	
 	FORCEINLINE uint16& BoneIndex(uint32 VertexIndex) const
 	{
-		checkSlow(VertexIndex < GetNumVertices());
+		checkSlow(VertexIndex < NumVertices);
 		return *(uint16*)(Data + VertexIndex * PixelFormatStride);
 	}
 
