replace the cluster union map with a cluster union manager that allows us to dynamically create cluster unions and add particles to them

- Removed the check and ensure for zero child partiesl in UpdateClusterMassProperties since this condition is handled properly.
- Removed the ensure on non null geometry/dynamic geometry in UpdateGeometry since this would get triggered even if you're only setting the shared geometry.

diff --git a/Engine/Source/Runtime/Experimental/Chaos/Private/Chaos/ClusterUnionManager.cpp b/Engine/Source/Runtime/Experimental/Chaos/Private/Chaos/ClusterUnionManager.cpp
new file mode 100644
--- /dev/null
+++ b/Engine/Source/Runtime/Experimental/Chaos/Private/Chaos/ClusterUnionManager.cpp
@@ -0,0 +1,284 @@
+// Copyright Epic Games, Inc. All Rights Reserved.
+
+#include "Chaos/ClusterUnionManager.h"
+#include "Chaos/ImplicitObjectUnion.h"
+#include "Chaos/PBDRigidClustering.h"
+#include "Chaos/PBDRigidClusteringAlgo.h"
+#include "Chaos/PBDRigidsEvolution.h"
+#include "Chaos/PBDRigidsEvolutionGBF.h"
+#include "ChaosStats.h"
+
+namespace Chaos
+{
+	FClusterUnionManager::FClusterUnionManager(FRigidClustering& InClustering, FPBDRigidsEvolutionGBF& InEvolution)
+		: MClustering(InClustering)
+		, MEvolution(InEvolution)
+	{
+	}
+
+	DECLARE_CYCLE_STAT(TEXT("FClusterUnionManager::CreateNewClusterUnion"), STAT_CreateNewClusterUnion, STATGROUP_Chaos);
+	FClusterUnionIndex FClusterUnionManager::CreateNewClusterUnion(const FClusterCreationParameters& Parameters, FClusterUnionExplicitIndex ExplicitIndex)
+	{
+		SCOPE_CYCLE_COUNTER(STAT_CreateNewClusterUnion);
+		FClusterUnionIndex NewIndex = ClaimNextUnionIndex();
+		check(NewIndex > 0);
+
+		FClusterUnion NewUnion;
+		NewUnion.ExplicitIndex = ExplicitIndex;
+		NewUnion.SharedGeometry = ForceRecreateClusterUnionSharedGeometry(NewUnion);
+		NewUnion.InternalCluster = MClustering.CreateClusterParticle(-NewIndex, {}, Parameters, NewUnion.SharedGeometry);
+		NewUnion.Parameters = Parameters;
+		if (ensure(NewUnion.InternalCluster != nullptr))
+		{
+			NewUnion.InternalCluster->SetInternalCluster(true);
+		}
+
+		ClusterUnions.Add(NewIndex, NewUnion);
+		return NewIndex;
+	}
+
+	DECLARE_CYCLE_STAT(TEXT("FClusterUnionManager::ForceRecreateClusterUnionSharedGeometry"), STAT_ForceRecreateClusterUnionSharedGeometry, STATGROUP_Chaos);
+	TSharedPtr<FImplicitObject, ESPMode::ThreadSafe> FClusterUnionManager::ForceRecreateClusterUnionSharedGeometry(const FClusterUnion& Union)
+	{
+		if (Union.ChildParticles.IsEmpty() || !Union.InternalCluster)
+		{
+			return MakeShared<FImplicitObjectUnionClustered>();
+		}
+
+		// TODO: Can we do something better than a union?
+
+		const FRigidTransform3 ClusterWorldTM(Union.InternalCluster->X(), Union.InternalCluster->R());
+		TArray<TUniquePtr<FImplicitObject>> Objects;
+		TArray<FPBDRigidParticleHandle*> ChildParticleHandles;
+		Objects.Reserve(Union.ChildParticles.Num());
+		ChildParticleHandles.Reserve(Union.ChildParticles.Num());
+
+		for (FPBDRigidParticleHandle* Child : Union.ChildParticles)
+		{
+			const FRigidTransform3 ChildWorldTM(Child->X(), Child->R());
+			FRigidTransform3 Frame = ChildWorldTM.GetRelativeTransform(ClusterWorldTM);
+			if (Child->Geometry())
+			{
+				Objects.Add(TUniquePtr<FImplicitObject>(new TImplicitObjectTransformed<FReal, 3>(Child->Geometry(), Frame)));
+				ChildParticleHandles.Add(Child);
+			}
+		}
+
+		return MakeShared<FImplicitObjectUnionClustered>(MoveTemp(Objects), ChildParticleHandles);
+	}
+
+	DECLARE_CYCLE_STAT(TEXT("FClusterUnionManager::ClaimNextUnionIndex"), STAT_ClaimNextUnionIndex, STATGROUP_Chaos);
+	FClusterUnionIndex FClusterUnionManager::ClaimNextUnionIndex()
+	{
+		SCOPE_CYCLE_COUNTER(STAT_ClaimNextUnionIndex);
+		if (ReusableIndices.IsEmpty())
+		{
+			return NextAvailableUnionIndex++;
+		}
+		else
+		{
+			return ReusableIndices.Pop(false);
+		}
+	}
+
+	DECLARE_CYCLE_STAT(TEXT("FClusterUnionManager::AddPendingExplicitIndexOperation"), STAT_AddPendingExplicitIndexOperation, STATGROUP_Chaos);
+	void FClusterUnionManager::AddPendingExplicitIndexOperation(FClusterUnionExplicitIndex Index, EClusterUnionOperation Op, const TArray<FPBDRigidParticleHandle*>& Particles)
+	{
+		SCOPE_CYCLE_COUNTER(STAT_AddPendingExplicitIndexOperation);
+		AddPendingOperation(PendingExplicitIndexOperations, Index, Op, Particles);
+	}
+
+	DECLARE_CYCLE_STAT(TEXT("FClusterUnionManager::AddPendingClusterIndexOperation"), STAT_AddPendingClusterIndexOperation, STATGROUP_Chaos);
+	void FClusterUnionManager::AddPendingClusterIndexOperation(FClusterUnionIndex Index, EClusterUnionOperation Op, const TArray<FPBDRigidParticleHandle*>& Particles)
+	{
+		SCOPE_CYCLE_COUNTER(STAT_AddPendingClusterIndexOperation);
+		AddPendingOperation(PendingClusterIndexOperations, Index, Op, Particles);
+	}
+
+	DECLARE_CYCLE_STAT(TEXT("FClusterUnionManager::FlushPendingOperations"), STAT_FlushPendingOperations, STATGROUP_Chaos);
+	void FClusterUnionManager::FlushPendingOperations()
+	{
+		SCOPE_CYCLE_COUNTER(STAT_FlushPendingOperations);
+		if (PendingExplicitIndexOperations.IsEmpty() && PendingClusterIndexOperations.IsEmpty())
+		{
+			return;
+		}
+
+		// Go through every explicit index operation and convert them into a normal cluster index operation.
+		// This could be made more efficient but shouldn't happen enough for it to really matter.
+		for (const TPair<FClusterUnionExplicitIndex, FClusterOpMap>& OpMap : PendingExplicitIndexOperations)
+		{
+			const FClusterUnionIndex UnionIndex = GetOrCreateClusterUnionIndexFromExplicitIndex(OpMap.Key);
+			for (const TPair<EClusterUnionOperation, TArray<FPBDRigidParticleHandle*>>& Op : OpMap.Value)
+			{
+				AddPendingClusterIndexOperation(UnionIndex, Op.Key, Op.Value);
+			}
+		}
+		PendingExplicitIndexOperations.Empty();
+
+		for (const TPair<FClusterUnionIndex, FClusterOpMap>& OpMap : PendingClusterIndexOperations)
+		{
+			for (const TPair<EClusterUnionOperation, TArray<FPBDRigidParticleHandle*>>& Op : OpMap.Value)
+			{
+				switch (Op.Key)
+				{
+				case EClusterUnionOperation::Add:
+				case EClusterUnionOperation::AddReleased:
+					HandleAddOperation(OpMap.Key, Op.Value, Op.Key == EClusterUnionOperation::AddReleased);
+					break;
+				}
+			}
+		}
+		PendingClusterIndexOperations.Empty();
+	}
+
+	FClusterUnion* FClusterUnionManager::FindClusterUnionFromExplicitIndex(FClusterUnionExplicitIndex Index)
+	{
+		FClusterUnionIndex* ClusterIndex = ExplicitIndexMap.Find(Index);
+		if (ClusterIndex == nullptr)
+		{
+			return nullptr;
+		}
+
+		return FindClusterUnion(*ClusterIndex);
+	}
+
+	FClusterUnion* FClusterUnionManager::FindClusterUnion(FClusterUnionIndex Index)
+	{
+		return ClusterUnions.Find(Index);
+	}
+
+	DECLARE_CYCLE_STAT(TEXT("FClusterUnionManager::HandleAddOperation"), STAT_HandleAddOperation, STATGROUP_Chaos);
+	void FClusterUnionManager::HandleAddOperation(FClusterUnionIndex ClusterIndex, const TArray<FPBDRigidParticleHandle*>& Particles, bool bReleaseClustersFirst)
+	{
+		SCOPE_CYCLE_COUNTER(STAT_HandleAddOperation);
+		FClusterUnion* Cluster = ClusterUnions.Find(ClusterIndex);
+		if (!Cluster)
+		{
+			return;
+		}
+
+		TGuardValue_Bitfield_Cleanup<TFunction<void()>> Cleanup(
+			[this, OldGenerateClusterBreaking=MClustering.GetDoGenerateBreakingData()]() {
+				MClustering.SetGenerateClusterBreaking(OldGenerateClusterBreaking);
+			}
+		);
+		MClustering.SetGenerateClusterBreaking(false);
+
+		// If we're a new cluster, we need to determine whether to start the cluster in a sleeping or dynamic state.
+		// Only stay sleeping if all the particles we add are also sleeping.
+		const bool bIsNewCluster = Cluster->ChildParticles.IsEmpty();
+		bool bIsSleeping = true;
+		bool bIsAnchored = false;
+
+		TArray<FPBDRigidParticleHandle*> FinalParticlesToAdd;
+		FinalParticlesToAdd.Reserve(Particles.Num());
+
+		// This is only relevant when bReleaseClustersFirst=true. This is used to be able to
+		// properly notify the parent cluster about its child proxies.
+		TMap<FPBDRigidParticleHandle*, FPBDRigidParticleHandle*> ChildToParentMap;
+
+		for (FPBDRigidParticleHandle* Handle : Particles)
+		{
+			if (!Handle || Handle->Disabled())
+			{
+				continue;
+			}
+
+			bIsSleeping &= Handle->ObjectState() == EObjectStateType::Sleeping;
+
+			if (FPBDRigidClusteredParticleHandle* ClusterHandle = Handle->CastToClustered(); ClusterHandle && bReleaseClustersFirst)
+			{
+				TSet<FPBDRigidParticleHandle*> Children = MClustering.ReleaseClusterParticles(ClusterHandle, true);
+				FinalParticlesToAdd.Append(Children.Array());
+
+				for (FPBDRigidParticleHandle* Child : Children)
+				{
+					ChildToParentMap.Add(Child, ClusterHandle);
+				}
+
+				bIsAnchored |= ClusterHandle->IsAnchored();
+			}
+			else
+			{
+				FinalParticlesToAdd.Add(Handle);
+			}
+		}
+
+		if (FinalParticlesToAdd.IsEmpty())
+		{
+			return;
+		}
+
+		Cluster->ChildParticles.Append(FinalParticlesToAdd);
+
+		// If a physics proxy was set already on the cluster we want to make sure that doesn't change.
+		// This is needed to eventually be able to introduce a new physics proxy that gets attached to the
+		// cluster union particle so that it can communicate with the game thread.
+		IPhysicsProxyBase* OldProxy = Cluster->InternalCluster->PhysicsProxy();
+
+		MClustering.AddParticlesToCluster(Cluster->InternalCluster, FinalParticlesToAdd, ChildToParentMap);
+
+		// Update cluster properties.
+		FMatrix33 ClusterInertia(0);
+
+		// TODO: These functions are generally just re-building the cluster from scratch. Need to figure out a way
+		// to get these functions to update the already existing cluster instead.
+		TSet<FPBDRigidParticleHandle*> FullChildrenSet(FinalParticlesToAdd);
+
+		const FRigidTransform3 ForceMassOrientation{ Cluster->InternalCluster->X(), Cluster->InternalCluster->R() };
+		UpdateClusterMassProperties(Cluster->InternalCluster, FullChildrenSet, ClusterInertia, bIsNewCluster ? nullptr : &ForceMassOrientation);
+
+		if (bIsNewCluster && bIsAnchored)
+		{
+			// The anchored flag is taken care of in UpdateKinematicProperties so it must be set before that is called.
+			Cluster->InternalCluster->SetIsAnchored(true);
+		}
+		UpdateKinematicProperties(Cluster->InternalCluster, MClustering.GetChildrenMap(), MEvolution);
+
+		// The recreation of the geometry must happen after the call to UpdateClusterMassProperties.
+		// Creating the geometry requires knowing the relative frame between the parent cluster and the child clusters. The
+		// parent transform is not set properly for a new empty cluster until UpdateClusterMassProperties is called for the first time.
+		Cluster->SharedGeometry = ForceRecreateClusterUnionSharedGeometry(*Cluster);
+
+		UpdateGeometry(Cluster->InternalCluster, FullChildrenSet, MClustering.GetChildrenMap(), Cluster->SharedGeometry, Cluster->Parameters);
+		MClustering.GenerateConnectionGraph(Cluster->InternalCluster, Cluster->Parameters);
+
+		if (OldProxy)
+		{
+			Cluster->InternalCluster->SetPhysicsProxy(OldProxy);
+		}
+
+		if (bIsNewCluster)
+		{
+			if (bIsSleeping)
+			{
+				MEvolution.SetParticleObjectState(Cluster->InternalCluster, Chaos::EObjectStateType::Sleeping);
+			}
+
+			MEvolution.SetPhysicsMaterial(Cluster->InternalCluster, MEvolution.GetPhysicsMaterial(FinalParticlesToAdd[0]));
+		}
+	}
+
+	DECLARE_CYCLE_STAT(TEXT("FClusterUnionManager::GetOrCreateClusterUnionIndexFromExplicitIndex"), STAT_GetOrCreateClusterUnionIndexFromExplicitIndex, STATGROUP_Chaos);
+	FClusterUnionIndex FClusterUnionManager::GetOrCreateClusterUnionIndexFromExplicitIndex(FClusterUnionExplicitIndex InIndex)
+	{
+		SCOPE_CYCLE_COUNTER(STAT_GetOrCreateClusterUnionIndexFromExplicitIndex);
+		FClusterUnionIndex* OutIndex = ExplicitIndexMap.Find(InIndex);
+		if (OutIndex != nullptr)
+		{
+			return *OutIndex;
+		}
+
+		FClusterUnionIndex NewIndex = CreateNewClusterUnion(DefaultClusterCreationParameters(), InIndex);
+		ExplicitIndexMap.Add(InIndex, NewIndex);
+		return NewIndex;
+	}
+
+	FClusterCreationParameters FClusterUnionManager::DefaultClusterCreationParameters() const
+	{
+		 FClusterCreationParameters Parameters{ 0.3f, 100, false, FRigidClustering::ShouldUnionsHaveCollisionParticles() };
+		 Parameters.ConnectionMethod = MClustering.GetClusterUnionConnectionType();
+		 return Parameters;
+	}
+}
\ No newline at end of file
diff --git a/Engine/Source/Runtime/Experimental/Chaos/Private/Chaos/PBDRigidClustering.cpp b/Engine/Source/Runtime/Experimental/Chaos/Private/Chaos/PBDRigidClustering.cpp
--- a/Engine/Source/Runtime/Experimental/Chaos/Private/Chaos/PBDRigidClustering.cpp
+++ b/Engine/Source/Runtime/Experimental/Chaos/Private/Chaos/PBDRigidClustering.cpp
@@ -90,6 +90,7 @@ namespace Chaos
 	FRigidClustering::FRigidClustering(FPBDRigidsEvolution& InEvolution, FPBDRigidClusteredParticles& InParticles)
 		: MEvolution(InEvolution)
 		, MParticles(InParticles)
+		, ClusterUnionManager(*this, InEvolution)
 		, MCollisionImpulseArrayDirty(true)
 		, DoGenerateBreakingData(false)
 		, MClusterConnectionFactor(1.0)
@@ -174,8 +175,6 @@ namespace Chaos
 			UpdateTopLevelParticle(NewParticle);
 		}
 
-		ensureMsgf(!ProxyGeometry || ForceMassOrientation, TEXT("If ProxyGeometry is passed, we must override the mass orientation as they are tied"));
-
 		// TODO: This needs to be rotated to diagonal, used to update I()/InvI() from diagonal, and update transform with rotation.
 		FMatrix33 ClusterInertia(0);
 		UpdateClusterMassProperties(NewParticle, ChildrenSet, ClusterInertia, ForceMassOrientation);
@@ -185,19 +184,14 @@ namespace Chaos
 
 		NewParticle->SetSleeping(bClusterIsAsleep);
 
-		auto AddToClusterUnion = [&](int32 ClusterID, FPBDRigidClusteredParticleHandle* Handle)
+		auto AddToClusterUnion = [this](int32 ClusterID, FPBDRigidClusteredParticleHandle* Handle)
 		{
 			if (ClusterID <= 0)
 			{
 				return;
 			}
 
-			if (!ClusterUnionMap.Contains(ClusterID))
-			{
-				ClusterUnionMap.Add(ClusterID, TArray<FPBDRigidClusteredParticleHandle*>());
-			}
-
-			ClusterUnionMap[ClusterID].Add(Handle);
+			ClusterUnionManager.AddPendingExplicitIndexOperation(ClusterID, EClusterUnionOperation::AddReleased, { Handle });
 		};
 
 		if(ClusterGroupIndex)
@@ -208,9 +202,74 @@ namespace Chaos
 		return NewParticle;
 	}
 
+	DECLARE_CYCLE_STAT(TEXT("TPBDRigidClustering<>::AddParticlesToCluster"), STAT_AddParticlesToCluster, STATGROUP_Chaos);
+	void
+	FRigidClustering::AddParticlesToCluster(
+		FPBDRigidClusteredParticleHandle* Cluster,
+		const TArray<FPBDRigidParticleHandle*>& InChildren,
+		const TMap<FPBDRigidParticleHandle*, FPBDRigidParticleHandle*>& ChildToParentMap)
+	{
+		SCOPE_CYCLE_COUNTER(STAT_AddParticlesToCluster);
+		if (!Cluster || InChildren.IsEmpty())
+		{
+			return;
+		}
+
+		FRigidHandleArray& Children = MChildren.FindOrAdd(Cluster);
+		const int32 OldNumChildren = Children.Num();
+		Children.Append(InChildren);
+
+		// Disable all the input children since they no longer need to be simulated.
+		TSet<FPBDRigidParticleHandle*> InChildrenSet(InChildren);
+		MEvolution.DisableParticles(reinterpret_cast<TSet<FGeometryParticleHandle*>&>(InChildrenSet));
+
+		// Note that we want to compute the internal strain on the cluster the same if we build it up incrementally as well as if we
+		// build it all at the same time. The parent cluster's internal strain should be the average of all the child strains.
+		// The easy way to compute the new average is the multiply the old average by the number of old elements, add in the new strains,
+		// and then divide by the new total number of elements.
+		Cluster->SetInternalStrains(Cluster->GetInternalStrains() * OldNumChildren);
+		Cluster->ClusterIds().NumChildren = Children.Num();
+
+		// An initial pass through the children to transfer some of their cluster properties to their new parent.
+		for (FPBDRigidParticleHandle* Child : InChildrenSet)
+		{
+			if (FPBDRigidClusteredParticleHandle* ClusteredChild = Child->CastToClustered())
+			{
+				TopLevelClusterParents.Remove(ClusteredChild);
+				TopLevelClusterParentsStrained.Remove(ClusteredChild);
+
+				// Cluster group id 0 means "don't union with other things"
+				// TODO: Use INDEX_NONE instead of 0?
+				ClusteredChild->SetClusterGroupIndex(0);
+				ClusteredChild->ClusterIds().Id = Cluster;
+				Cluster->SetInternalStrains(Cluster->GetInternalStrains() + ClusteredChild->GetInternalStrains());
+				Cluster->SetCollisionImpulses(FMath::Max(Cluster->CollisionImpulses(), ClusteredChild->CollisionImpulses()));
+
+				const int32 NewCG = Cluster->CollisionGroup();
+				const int32 ChildCG = ClusteredChild->CollisionGroup();
+				Cluster->SetCollisionGroup(NewCG < ChildCG ? NewCG : ChildCG);
+			}
+
+			FPBDRigidParticleHandle* const* OldParent = ChildToParentMap.Find(Child);
+			FPBDRigidParticleHandle* ProxyParticle = (OldParent != nullptr) ? *OldParent : Child;
+			MEvolution.DoInternalParticleInitilization(ProxyParticle, Cluster);
+		}
+
+		if (!Children.IsEmpty())
+		{
+			Cluster->SetInternalStrains(Cluster->GetInternalStrains() / static_cast<FReal>(Children.Num()));
+		}
+	}
+
 	int32 UnionsHaveCollisionParticles = 0;
 	FAutoConsoleVariableRef CVarUnionsHaveCollisionParticles(TEXT("p.UnionsHaveCollisionParticles"), UnionsHaveCollisionParticles, TEXT(""));
 
+	bool
+	FRigidClustering::ShouldUnionsHaveCollisionParticles()
+	{
+		return !!UnionsHaveCollisionParticles;
+	}
+
 	DECLARE_CYCLE_STAT(TEXT("TPBDRigidClustering<>::CreateClusterParticleFromClusterChildren"), STAT_CreateClusterParticleFromClusterChildren, STATGROUP_Chaos);
 	Chaos::FPBDRigidClusteredParticleHandle* 
 	FRigidClustering::CreateClusterParticleFromClusterChildren(
@@ -298,94 +357,9 @@ namespace Chaos
 	FRigidClustering::UnionClusterGroups()
 	{
 		SCOPE_CYCLE_COUNTER(STAT_UnionClusterGroups);
-
-		if(ClusterUnionMap.Num())
-		{
-			struct FClusterGroup {
-				TArray<FPBDRigidParticleHandle*> Bodies = TArray < FPBDRigidParticleHandle*>();
-				bool bIsSleeping = true;
-			};
-
-
-			TMap<FPBDRigidParticleHandle*, FPBDRigidParticleHandle*> ChildToParentMap;
-			TMap<int32, FClusterGroup> NewClusterGroups;
-
-			// Walk the list of registered cluster groups
-			for(TTuple<int32, TArray<FPBDRigidClusteredParticleHandle* >>& Group : ClusterUnionMap)
-			{
-				int32 ClusterGroupID = Group.Key;
-				TArray<FPBDRigidClusteredParticleHandle*> Handles = Group.Value;
-
-				if(Handles.Num() > 1)
-				{
-					// First see if this is a new group
-					if(!NewClusterGroups.Contains(ClusterGroupID))
-					{
-						NewClusterGroups.Add(ClusterGroupID, FClusterGroup());
-					}
-
-					bool bIsSleeping = true;
-					TArray<FPBDRigidParticleHandle*> ClusterBodies;
-					for(FPBDRigidClusteredParticleHandle* ActiveCluster : Handles)
-					{
-						if(ActiveCluster && !ActiveCluster->Disabled())
-						{
-							// If this is an external cluster (from the rest collection) we release its children and append them to the current group
-							TSet<FPBDRigidParticleHandle*> Children;
-							
-							// let sleeping clusters stay asleep
-							bIsSleeping &= ActiveCluster->ObjectState() == EObjectStateType::Sleeping;
-
-							{
-								// First disable breaking data generation - this is not a break we're just reclustering under a dynamic parent.
-								TGuardValue<bool> BreakFlagGuard(DoGenerateBreakingData, false);
-								Children = ReleaseClusterParticles(ActiveCluster, true);
-							}
-
-							NewClusterGroups[ClusterGroupID].Bodies.Append(Children.Array());
-							
-							for(FPBDRigidParticleHandle* Child : Children)
-							{
-								ChildToParentMap.Add(Child, ActiveCluster);
-							}
-						}
-					}
-					NewClusterGroups[ClusterGroupID].bIsSleeping = bIsSleeping;
-				}
-			}
-
-			// For new cluster groups, create an internal cluster parent.
-			for(TTuple<int32, FClusterGroup>& Group : NewClusterGroups)
-			{
-				int32 ClusterGroupID = FMath::Abs(Group.Key);
-
-				TArray<FPBDRigidParticleHandle*> ActiveCluster = Group.Value.Bodies;
-
-				FClusterCreationParameters Parameters(0.3f, 100, false, !!UnionsHaveCollisionParticles);
-				Parameters.ConnectionMethod = MClusterUnionConnectionType;
-				TPBDRigidClusteredParticleHandleImp<FReal, 3, true>* Handle = 
-					CreateClusterParticle(-ClusterGroupID, MoveTemp(Group.Value.Bodies), Parameters, 
-						TSharedPtr<FImplicitObject, ESPMode::ThreadSafe>());
-				Handle->SetInternalCluster(true);
-
-				if (Group.Value.bIsSleeping)
-				{
-					MEvolution.SetParticleObjectState(Handle, Chaos::EObjectStateType::Sleeping);
-				}
-
-				MEvolution.SetPhysicsMaterial(Handle, MEvolution.GetPhysicsMaterial(ActiveCluster[0]));
-
-				for(FPBDRigidParticleHandle* Constituent : ActiveCluster)
-				{
-					MEvolution.DoInternalParticleInitilization(ChildToParentMap[Constituent], Handle);
-				}
-			}
-
-			ClusterUnionMap.Empty();
-		}
+		ClusterUnionManager.FlushPendingOperations();
 	}
 
-
 	DECLARE_CYCLE_STAT(TEXT("TPBDRigidClustering<>::DeactivateClusterParticle"), STAT_DeactivateClusterParticle, STATGROUP_Chaos);
 	TSet<FPBDRigidParticleHandle*> 
 	FRigidClustering::DeactivateClusterParticle(
diff --git a/Engine/Source/Runtime/Experimental/Chaos/Private/Chaos/PBDRigidClusteringAlgo.cpp b/Engine/Source/Runtime/Experimental/Chaos/Private/Chaos/PBDRigidClusteringAlgo.cpp
--- a/Engine/Source/Runtime/Experimental/Chaos/Private/Chaos/PBDRigidClusteringAlgo.cpp
+++ b/Engine/Source/Runtime/Experimental/Chaos/Private/Chaos/PBDRigidClusteringAlgo.cpp
@@ -43,7 +43,6 @@ namespace Chaos
 		const FRigidTransform3* ForceMassOrientation)
 	{
 		SCOPE_CYCLE_COUNTER(STAT_UpdateClusterMassProperties);
-		check(Children.Num());
 
 		Parent->SetX(FVec3(0));
 		Parent->SetR(FRotation3(FQuat::MakeFromEuler(FVec3(0))));
@@ -81,7 +80,7 @@ namespace Chaos
 			Parent->V() += OriginalChild->V() * ChildMass; // Use orig child for vel because we don't sim the proxy
 			Parent->W() += OriginalChild->W() * ChildMass;
 		}
-		if (!ensure(bHasProxyChild))
+		if (!bHasProxyChild)
 		{
 			for (FPBDRigidParticleHandle* OriginalChild : Children)
 			{
@@ -113,7 +112,7 @@ namespace Chaos
 			}
 		}
 
-		if (!ensure(bHasChild) || !ensure(Parent->M() > UE_SMALL_NUMBER))
+		if (!bHasChild || !(Parent->M() > UE_SMALL_NUMBER))
 		{
 			Parent->M() = 1.0;
 			Parent->X() = FVec3(0);
@@ -310,10 +309,6 @@ namespace Chaos
 		Objects.Reserve(Children.Num());
 		Objects2.Reserve(Children.Num());
 
-		//we should never update existing geometry since this is used by SQ threads.
-		ensure(!Parent->Geometry());
-		ensure(!Parent->DynamicGeometry());
-
 		const FRigidTransform3 ClusterWorldTM(Parent->X(), Parent->R());
 
 		TArray<FVec3> OriginalPoints;
diff --git a/Engine/Source/Runtime/Experimental/Chaos/Public/Chaos/ClusterUnionManager.h b/Engine/Source/Runtime/Experimental/Chaos/Public/Chaos/ClusterUnionManager.h
new file mode 100644
--- /dev/null
+++ b/Engine/Source/Runtime/Experimental/Chaos/Public/Chaos/ClusterUnionManager.h
@@ -0,0 +1,119 @@
+// Copyright Epic Games, Inc. All Rights Reserved.
+#pragma once
+
+#include "Chaos/ClusterCreationParameters.h"
+#include "Chaos/ParticleHandleFwd.h"
+#include "Containers/Array.h"
+#include "Containers/Map.h"
+#include "Containers/Set.h"
+
+namespace Chaos
+{
+	using FClusterUnionIndex = int32;
+	using FClusterUnionExplicitIndex = int32;
+
+	class FRigidClustering;
+	class FPBDRigidsEvolutionGBF;
+	struct FClusterCreationParameters;
+
+	enum class EClusterUnionOperation
+	{
+		Add,
+		// AddReleased is the original behavior where if the particle to be added is a cluster, we will release the cluster first
+		// and add its children instead.
+		AddReleased
+	};
+
+	struct CHAOS_API FClusterUnion
+	{
+		// The root cluster particle that we created internally to represent the cluster.
+		FPBDRigidClusteredParticleHandle* InternalCluster;
+
+		// The thread-safe collision geometry that can be shared between the GT and PT.
+		TSharedPtr<FImplicitObject, ESPMode::ThreadSafe> SharedGeometry;
+
+		// All the particles that belong to this cluster.
+		TSet<FPBDRigidParticleHandle*> ChildParticles;
+
+		// An explicit index set by the user if any.
+		FClusterUnionExplicitIndex ExplicitIndex;
+
+		// Need to remember the parameters used to create the cluster so we can update it later.
+		FClusterCreationParameters Parameters;
+	};
+
+	/**
+	 * This class is used by Chaos to create internal clusters that will
+	 * cause one or more clusters to simulate together as a single rigid
+	 * particle.
+	 */
+	class CHAOS_API FClusterUnionManager
+	{
+	public:
+		FClusterUnionManager(FRigidClustering& InClustering, FPBDRigidsEvolutionGBF& InEvolution);
+
+		// Creates a new cluster union with an automatically assigned cluster union index.
+		FClusterUnionIndex CreateNewClusterUnion(const FClusterCreationParameters& Parameters, FClusterUnionExplicitIndex ExplicitIndex = INDEX_NONE);
+
+		// Add a new operation to the queue.
+		void AddPendingExplicitIndexOperation(FClusterUnionExplicitIndex Index, EClusterUnionOperation Op, const TArray<FPBDRigidParticleHandle*>& Particles);
+		void AddPendingClusterIndexOperation(FClusterUnionIndex Index, EClusterUnionOperation Op, const TArray<FPBDRigidParticleHandle*>& Particles);
+
+		// Will be called at the beginning of every time step to ensure that all the expected cluster unions have been modified.
+		void FlushPendingOperations();
+
+		// Access the cluster union externally.
+		FClusterUnion* FindClusterUnionFromExplicitIndex(FClusterUnionExplicitIndex Index);
+		FClusterUnion* FindClusterUnion(FClusterUnionIndex Index);
+
+	private:
+		FRigidClustering& MClustering;
+		FPBDRigidsEvolutionGBF& MEvolution;
+
+		using FClusterOpMap = TMap<EClusterUnionOperation, TArray<FPBDRigidParticleHandle*>>;
+		template<typename TIndex>
+		using TClusterIndexOpMap = TMap<TIndex, FClusterOpMap>;
+
+		TClusterIndexOpMap<FClusterUnionIndex> PendingClusterIndexOperations;
+		TClusterIndexOpMap<FClusterUnionExplicitIndex> PendingExplicitIndexOperations;
+
+		template<typename TIndex>
+		void AddPendingOperation(TClusterIndexOpMap<TIndex>& OpMap, TIndex Index, EClusterUnionOperation Op, const TArray<FPBDRigidParticleHandle*>& Particles)
+		{
+			FClusterOpMap& Ops = OpMap.FindOrAdd(Index);
+			TArray<FPBDRigidParticleHandle*>& OpData = Ops.FindOrAdd(Op);
+			OpData.Append(Particles);
+		}
+
+		// All of our actively managed cluster unions. We need to keep track of these
+		// so a user could use the index to request modifications to a specific cluster union.
+		TMap<FClusterUnionIndex, FClusterUnion> ClusterUnions;
+		
+		//
+		// There are two ways we can pick a new union index:
+		// - If a cluster union gets released/destroyed, that index can be reused.
+		// - Otherwise, we use the NextAvailableUnionIndex which is just the max index we've seen + 1.
+		//
+		FClusterUnionIndex ClaimNextUnionIndex();
+		TArray<FClusterUnionIndex> ReusableIndices;
+		FClusterUnionIndex NextAvailableUnionIndex = 1;
+
+		// Using the user's passed in FClusterUnionIndex may result in strange unexpected behavior if the 
+		// user creates a cluster with a specified index. Thus we will map all explicitly requested indices
+		// (i.e. an index that comes in via FClusterUnionOperationData for the first time) to an automatically
+		// generated index (i.e one that would returned via CreateNewClusterUnion).
+		TMap<FClusterUnionExplicitIndex, FClusterUnionIndex> ExplicitIndexMap;
+		FClusterCreationParameters DefaultClusterCreationParameters() const;
+
+		// If no cluster index is set but an explicit index is set, map the explicit index to a regular index.
+		FClusterUnionIndex GetOrCreateClusterUnionIndexFromExplicitIndex(FClusterUnionExplicitIndex InIndex);
+
+		// Forcefully recreate the shared geometry on a cluster. Potentially expensive so ideally should be used rarely.
+		TSharedPtr<FImplicitObject, ESPMode::ThreadSafe> ForceRecreateClusterUnionSharedGeometry(const FClusterUnion& Union);
+
+		// Actually performs the change specified in the FClusterUnionOperationData structure.
+		void HandleAddOperation(FClusterUnionIndex ClusterIndex, const TArray<FPBDRigidParticleHandle*>& Particles, bool bReleaseClustersFirst);
+
+	};
+
+}
\ No newline at end of file
diff --git a/Engine/Source/Runtime/Experimental/Chaos/Public/Chaos/PBDRigidClustering.h b/Engine/Source/Runtime/Experimental/Chaos/Public/Chaos/PBDRigidClustering.h
--- a/Engine/Source/Runtime/Experimental/Chaos/Public/Chaos/PBDRigidClustering.h
+++ b/Engine/Source/Runtime/Experimental/Chaos/Public/Chaos/PBDRigidClustering.h
@@ -1,6 +1,7 @@
 // Copyright Epic Games, Inc. All Rights Reserved.
 #pragma once
 
+#include "Chaos/ClusterUnionManager.h"
 #include "Chaos/PBDRigidClusteredParticles.h"
 #include "Chaos/PBDCollisionConstraints.h"
 #include "Chaos/Transform.h"
@@ -95,6 +96,16 @@ public:
 		const FRigidTransform3& ClusterWorldTM, 
 		const FClusterCreationParameters& Parameters/* = FClusterCreationParameters()*/);
 
+	/**
+	 * Manually add a set of particles to a cluster after the cluster has already been created.
+	 *	ChildToParentMap: A map that may contain a pointer to one of the child particles as a key, and a pointer to its old parent particle (prior to be being released for example).
+	 *					  If a child particle exists in this map, its parent (or whatever is specified as the value) will be used to determine the correct proxy to add to the parent cluster.
+	 */
+	void AddParticlesToCluster(
+		FPBDRigidClusteredParticleHandle* Cluster,
+		const TArray<FPBDRigidParticleHandle*>& InChildren,
+		const TMap<FPBDRigidParticleHandle*, FPBDRigidParticleHandle*>& ChildToParentMap);
+
 	/**
 	 *  UnionClusterGroups
 	 *    Clusters that share a group index should be unioned into a single cluster prior to simulation.
@@ -266,6 +277,7 @@ public:
 	*/
 	const TArray<FBreakingData>& GetAllClusterBreakings() const { return MAllClusterBreakings; }
 	void SetGenerateClusterBreaking(bool DoGenerate) { DoGenerateBreakingData = DoGenerate; }
+	bool GetDoGenerateBreakingData() const { return DoGenerateBreakingData; }
 	void ResetAllClusterBreakings() { MAllClusterBreakings.Reset(); }
 
 	/*
@@ -334,6 +346,7 @@ public:
 	*/
 	void SetClusterConnectionFactor(FReal ClusterConnectionFactorIn) { MClusterConnectionFactor = ClusterConnectionFactorIn; }
 	void SetClusterUnionConnectionType(FClusterCreationParameters::EConnectionMethod ClusterConnectionType) { MClusterUnionConnectionType = ClusterConnectionType; }
+	FClusterCreationParameters::EConnectionMethod GetClusterUnionConnectionType() const { return MClusterUnionConnectionType; }
 
 	void GenerateConnectionGraph(
 		Chaos::FPBDRigidClusteredParticleHandle* Parent,
@@ -348,6 +361,10 @@ public:
 	void SetInternalStrain(FPBDRigidClusteredParticleHandle* Particle, FReal Strain);
 	void SetExternalStrain(FPBDRigidClusteredParticleHandle* Particle, FReal Strain);
 
+	static bool ShouldUnionsHaveCollisionParticles();
+
+	FClusterUnionManager& GetClusterUnionManager() { return ClusterUnionManager; }
+	const FClusterUnionManager& GetClusterUnionManager() const { return ClusterUnionManager; }
  protected:
 
 	void ComputeStrainFromCollision(const FPBDCollisionConstraints& CollisionRule);
@@ -397,8 +414,13 @@ private:
 
 	// Cluster data
 	FClusterMap MChildren;
-	TMap<int32, TArray<FPBDRigidClusteredParticleHandle*> > ClusterUnionMap;
 
+	/**
+	 * The old cluster union map has been replaced by the cluster union manager to allow for more
+	 * dynamic behavior of adding and removing particles from a cluster instead of being just restricted
+	 * to unioning particles together at construction.
+	 */
+	FClusterUnionManager ClusterUnionManager;
 
 	// Collision Impulses
 	bool MCollisionImpulseArrayDirty;
