Fix for non unity build.

diff --git a/Engine/Source/Runtime/Experimental/GeometryCollectionEngine/Private/GeometryCollection/GeometryCollectionSceneProxy.cpp b/Engine/Source/Runtime/Experimental/GeometryCollectionEngine/Private/GeometryCollection/GeometryCollectionSceneProxy.cpp
--- a/Engine/Source/Runtime/Experimental/GeometryCollectionEngine/Private/GeometryCollection/GeometryCollectionSceneProxy.cpp
+++ b/Engine/Source/Runtime/Experimental/GeometryCollectionEngine/Private/GeometryCollection/GeometryCollectionSceneProxy.cpp
@@ -8,6 +8,7 @@
 #include "GeometryCollection/GeometryCollectionObject.h"
 #include "MaterialDomain.h"
 #include "MaterialShaderType.h"
+#include "MaterialShared.h"
 #include "Materials/Material.h"
 #include "Materials/MaterialRenderProxy.h"
 #include "CommonRenderResources.h"
