Chaos - PerShapeData split into FShapeInstanceProxy and FShapeInstance for GT/PT respectively (Phase 1)

diff --git a/Engine/Source/Runtime/Engine/Private/PhysicsEngine/BodyUtils.h b/Engine/Source/Runtime/Engine/Private/PhysicsEngine/BodyUtils.h
--- a/Engine/Source/Runtime/Engine/Private/PhysicsEngine/BodyUtils.h
+++ b/Engine/Source/Runtime/Engine/Private/PhysicsEngine/BodyUtils.h
@@ -2,7 +2,7 @@
 
 #pragma once
 
-#include "CoreMinimal.h"
+#include "Chaos/ShapeInstanceFwd.h"
 #include "Containers/Array.h"
 #include "Containers/ContainerAllocationPolicies.h"
 #include "PhysicsInterfaceDeclaresCore.h"
@@ -12,7 +12,6 @@ struct FBodyInstance;
 namespace Chaos
 {
 	struct FMassProperties;
-	using FShapesArray = TArray<TUniquePtr<FPerShapeData>, TInlineAllocator<1>>;
 }
 
 /**
diff --git a/Engine/Source/Runtime/Engine/Private/PhysicsEngine/Experimental/ChaosInterfaceUtils.cpp b/Engine/Source/Runtime/Engine/Private/PhysicsEngine/Experimental/ChaosInterfaceUtils.cpp
--- a/Engine/Source/Runtime/Engine/Private/PhysicsEngine/Experimental/ChaosInterfaceUtils.cpp
+++ b/Engine/Source/Runtime/Engine/Private/PhysicsEngine/Experimental/ChaosInterfaceUtils.cpp
@@ -89,7 +89,7 @@ namespace ChaosInterface
 
 		auto NewShapeHelper = [&InParams, &CollisionTraceType](Chaos::TSerializablePtr<Chaos::FImplicitObject> InGeom, int32 ShapeIdx, void* UserData, ECollisionEnabled::Type ShapeCollisionEnabled, bool bComplexShape = false)
 		{
-			TUniquePtr<Chaos::FPerShapeData> NewShape = Chaos::FPerShapeData::CreatePerShapeData(ShapeIdx, InGeom);
+			TUniquePtr<Chaos::FPerShapeData> NewShape = Chaos::FShapeInstanceProxy::Make(ShapeIdx, InGeom);
 			NewShape->SetQueryData(bComplexShape ? InParams.CollisionData.CollisionFilterData.QueryComplexFilter : InParams.CollisionData.CollisionFilterData.QuerySimpleFilter);
 			NewShape->SetSimData(InParams.CollisionData.CollisionFilterData.SimFilter);
 			NewShape->SetCollisionTraceType(ConvertCollisionTraceFlag(CollisionTraceType));
diff --git a/Engine/Source/Runtime/Engine/Private/PhysicsEngine/ImmediatePhysics/ImmediatePhysicsChaos/ImmediatePhysicsActorHandle_Chaos.cpp b/Engine/Source/Runtime/Engine/Private/PhysicsEngine/ImmediatePhysics/ImmediatePhysicsChaos/ImmediatePhysicsActorHandle_Chaos.cpp
--- a/Engine/Source/Runtime/Engine/Private/PhysicsEngine/ImmediatePhysics/ImmediatePhysicsChaos/ImmediatePhysicsActorHandle_Chaos.cpp
+++ b/Engine/Source/Runtime/Engine/Private/PhysicsEngine/ImmediatePhysics/ImmediatePhysicsChaos/ImmediatePhysicsActorHandle_Chaos.cpp
@@ -32,7 +32,7 @@ namespace ImmediatePhysics_Chaos
 		const FReal Radius = 1.0f * Scale.GetMax();
 
 		TUniquePtr<Chaos::FImplicitSphere3> ImplicitSphere = MakeUnique<Chaos::FImplicitSphere3>(FVec3(0), Radius);
-		TUniquePtr<FPerShapeData> NewShape = Chaos::FPerShapeData::CreatePerShapeDataInternal(OutShapes.Num(), MakeSerializable(ImplicitSphere));
+		TUniquePtr<FPerShapeData> NewShape = Chaos::FShapeInstance::Make(OutShapes.Num(), MakeSerializable(ImplicitSphere));
 		NewShape->UpdateShapeBounds(FTransform::Identity);
 		NewShape->SetUserData(nullptr);
 		NewShape->SetQueryEnabled(false);
diff --git a/Engine/Source/Runtime/Experimental/Chaos/Private/Chaos/GeometryParticles.cpp b/Engine/Source/Runtime/Experimental/Chaos/Private/Chaos/GeometryParticles.cpp
--- a/Engine/Source/Runtime/Experimental/Chaos/Private/Chaos/GeometryParticles.cpp
+++ b/Engine/Source/Runtime/Experimental/Chaos/Private/Chaos/GeometryParticles.cpp
@@ -11,375 +11,7 @@
 
 namespace Chaos
 {
-	TUniquePtr<FPerShapeData> CreatePerShapeDataImpl(const int32 ShapeIndex, TSerializablePtr<FImplicitObject> InGeometry, const bool bAllowCachedLeafInfo)
-	{
-		if (bAllowCachedLeafInfo)
-		{
-			return FPerShapeData::CreatePerShapeDataInternal(ShapeIndex, InGeometry);
-		}
-		else
-		{
-			return FPerShapeData::CreatePerShapeData(ShapeIndex, InGeometry);
-		}
-	}
-
-	void UpdatePerShapeDataImpl(TUniquePtr<FPerShapeData>& ShapePtr, TSerializablePtr<FImplicitObject> InGeometry, const bool bAllowCachedLeafInfo)
-	{
-		if (bAllowCachedLeafInfo)
-		{
-			FPerShapeData::UpdateGeometryInternal(ShapePtr, InGeometry);
-		}
-		else
-		{
-			FPerShapeData::UpdateGeometry(ShapePtr, InGeometry);
-		}
-	}
-
-	// Creare or reuse the shapes in the shapes array and populate them with the Geometry.
-	// If we have a Union it will be unpacked into the ShapesArray.
-	// On the Physics Thread we set bAllowCachedLeafInfo which caches the shapes world space state to optimize collision detection,
-	// but this flag should not be used on the Game Thread because the extended data is not maintained so some PerShapeData funtions
-	// will return unitialized values (E.g., GetLeafWorldTransform).
-	void UpdateShapesArrayFromGeometryImpl(
-		FShapesArray& ShapesArray, 
-		TSerializablePtr<FImplicitObject> Geometry, 
-		const FRigidTransform3& ActorTM, 
-		IPhysicsProxyBase* Proxy, 
-		bool bAllowCachedLeafInfo)
-	{
-		if(Geometry)
-		{
-			const int32 OldShapeNum = ShapesArray.Num();
-			if(const auto* Union = Geometry->template GetObject<FImplicitObjectUnion>())
-			{
-				ShapesArray.SetNum(Union->GetObjects().Num());
-
-				for (int32 ShapeIndex = 0; ShapeIndex < ShapesArray.Num(); ++ShapeIndex)
-				{
-					TSerializablePtr<FImplicitObject> ShapeGeometry = MakeSerializable(Union->GetObjects()[ShapeIndex]);
-
-					if (ShapeIndex >= OldShapeNum)
-					{
-						// If newly allocated shape, initialize it.
-						ShapesArray[ShapeIndex] = CreatePerShapeDataImpl(ShapeIndex, ShapeGeometry, bAllowCachedLeafInfo);
-					}
-					else if (ShapeGeometry != ShapesArray[ShapeIndex]->GetGeometry())
-					{
-						// Update geometry pointer if it changed
-						UpdatePerShapeDataImpl(ShapesArray[ShapeIndex], ShapeGeometry, bAllowCachedLeafInfo);
-					}
-				}
-			}
-			else
-			{
-				ShapesArray.SetNum(1);
-				if (OldShapeNum == 0)
-				{
-					ShapesArray[0] = CreatePerShapeDataImpl(0, Geometry, bAllowCachedLeafInfo);
-				}
-				else
-				{
-					UpdatePerShapeDataImpl(ShapesArray[0], Geometry, bAllowCachedLeafInfo);
-				}
-			}
-
-			if (Geometry->HasBoundingBox())
-			{
-				for (auto& Shape : ShapesArray)
-				{
-					Shape->UpdateShapeBounds(ActorTM, FVec3(0));
-				}
-			}
-		}
-		else
-		{
-			ShapesArray.Reset();
-		}
-
-		if(Proxy)
-		{
-			if(FPhysicsSolverBase* PhysicsSolverBase = Proxy->GetSolver<FPhysicsSolverBase>())
-			{
-				PhysicsSolverBase->SetNumDirtyShapes(Proxy, ShapesArray.Num());
-			}
-		}
-	}
-
-	void UpdateShapesArrayFromGeometry(
-		FShapesArray& ShapesArray,
-		TSerializablePtr<FImplicitObject> Geometry,
-		const FRigidTransform3& ActorTM,
-		IPhysicsProxyBase* Proxy)
-	{
-		UpdateShapesArrayFromGeometryImpl(ShapesArray, Geometry, ActorTM, Proxy, false);
-	}
-
-	// Unwrap transformed shapes
-	// @todo(chaos): also unwrap Instanced and Scaled but that requires a lot of knock work because Convexes are usually Instanced or
-	// Scaled so the Scale and Margin is stored on the wrapper (the convex itself is shared).
-	// - support for Margin as per-shape data passed through the collision functions
-	// - support for Scale as per-shape data passed through the collision functions (or ideally in the Transforms)
-	const FImplicitObject* GetInnerGeometryInstanceData(const FImplicitObject* Implicit, const FRigidTransform3** OutRelativeTransformPtr)
-	{
-		if (Implicit != nullptr)
-		{
-			const EImplicitObjectType ImplicitOuterType = Implicit->GetType();
-			if (ImplicitOuterType == TImplicitObjectTransformed<FReal, 3>::StaticType())
-			{
-				// Transformed Implicit
-				const TImplicitObjectTransformed<FReal, 3>* TransformedImplicit = Implicit->template GetObject<const TImplicitObjectTransformed<FReal, 3>>();
-				if (OutRelativeTransformPtr != nullptr)
-				{
-					*OutRelativeTransformPtr = &TransformedImplicit->GetTransform();
-				}
-				return GetInnerGeometryInstanceData(TransformedImplicit->GetTransformedObject(), OutRelativeTransformPtr);
-			}
-			else if ((uint32)ImplicitOuterType & ImplicitObjectType::IsInstanced)
-			{
-				// Instanced Implicit
-				// Currently we only unwrap instanced TriMesh and Heightfields. Convex is the only other type that will be Instanced,
-				// but we do not unwrap Convex because the convex margin is stored in the Instanced so we would need to extract
-				// it and pass it on to the collision detection.
-				const FImplicitObjectInstanced* Instanced = static_cast<const FImplicitObjectInstanced*>(Implicit);
-				EImplicitObjectType InnerType = Instanced->GetInnerObject()->GetType();
-				if (InnerType != FImplicitConvex3::StaticType())
-				{
-					return GetInnerGeometryInstanceData(Instanced->GetInnerObject().Get(), OutRelativeTransformPtr);
-				}
-			}
-			else if ((uint32)ImplicitOuterType & ImplicitObjectType::IsScaled)
-			{
-				// Scaled Implicit
-				//const FImplicitObjectScaled* Scaled = static_cast<const FImplicitObjectScaled*>(Implicit);
-				//OutTransform.Scale *= Scaled->GetScale();
-				//return GetInnerGeometryInstanceData(Scaled->GetInnerObject().Get(), OutRelativeTransformPtr);
-			}
-		}
-		return Implicit;
-	}
-
-	bool PerShapeDataRequiresCachedLeafInfo(const FImplicitObject* Geometry)
-	{
-		// Unwrap the geometry (if it's Transformed or Instanced)
-		const FRigidTransform3* LeafRelativeTransform = nullptr;
-		const FImplicitObject* LeafGeometry = GetInnerGeometryInstanceData(Geometry, &LeafRelativeTransform);
-		return (LeafGeometry != Geometry) && (LeafRelativeTransform != nullptr);
-	}
-
-
-
-	FPerShapeData::FPerShapeData(int32 InShapeIdx)
-		: bHasCachedLeafInfo(false)
-		, Proxy(nullptr)
-		, ShapeIdx(InShapeIdx)
-		, Geometry()
-		, WorldSpaceInflatedShapeBounds(FAABB3(FVec3(0), FVec3(0)))
-	{
-	}
-
-	FPerShapeData::FPerShapeData(int32 InShapeIdx, TSerializablePtr<FImplicitObject> InGeometry)
-		: bHasCachedLeafInfo(false)
-		, Proxy(nullptr)
-		, ShapeIdx(InShapeIdx)
-		, Geometry(InGeometry)
-		, WorldSpaceInflatedShapeBounds(FAABB3(FVec3(0), FVec3(0)))
-	{
-	}
-
-	FPerShapeData::FPerShapeData(FPerShapeData&& Other)
-		: bHasCachedLeafInfo(Other.bHasCachedLeafInfo)
-		, Proxy(MoveTemp(Other.Proxy))
-		, DirtyFlags(MoveTemp(Other.DirtyFlags))
-		, ShapeIdx(MoveTemp(Other.ShapeIdx))
-		, CollisionData(MoveTemp(Other.CollisionData))
-		, Materials(MoveTemp(Other.Materials))
-		, Geometry(MoveTemp(Other.Geometry))
-		, WorldSpaceInflatedShapeBounds(MoveTemp(Other.WorldSpaceInflatedShapeBounds))
-	{
-	}
-
-	FPerShapeData::~FPerShapeData()
-	{
-	}
-
-	TUniquePtr<FPerShapeData> FPerShapeData::CreatePerShapeData(int32 InShapeIdx, TSerializablePtr<FImplicitObject> InGeometry)
-	{
-		return TUniquePtr<FPerShapeData>(new FPerShapeData(InShapeIdx, InGeometry));
-	}
-
-	void FPerShapeData::UpdateGeometry(TUniquePtr<FPerShapeData>& ShapePtr, TSerializablePtr<FImplicitObject> InGeometry)
-	{
-		ShapePtr->Geometry = InGeometry;
-	}
-
-	TUniquePtr<FPerShapeData> FPerShapeData::CreatePerShapeDataInternal(int32 InShapeIdx, TSerializablePtr<FImplicitObject> InGeometry)
-	{
-		const bool bWantLeafGeometryCache = PerShapeDataRequiresCachedLeafInfo(InGeometry.Get());
-		if (bWantLeafGeometryCache)
-		{
-			return TUniquePtr<FPerShapeData>(new FPerShapeDataCachedLeafInfo(InShapeIdx, InGeometry));
-		}
-		else
-		{
-			return TUniquePtr<FPerShapeData>(new FPerShapeData(InShapeIdx, InGeometry));
-		}
-	}
-
-	void FPerShapeData::UpdateGeometryInternal(TUniquePtr<FPerShapeData>& ShapePtr, TSerializablePtr<FImplicitObject> InGeometry)
-	{
-		const bool bWantLeafGeometryCache = PerShapeDataRequiresCachedLeafInfo(InGeometry.Get());
-
-		// Do we need to add or remove the cached leaf data? If so this requires we recreate the object
-		if (bWantLeafGeometryCache != ShapePtr->HasCachedLeafInfo())
-		{
-			if (bWantLeafGeometryCache)
-			{
-				ShapePtr = TUniquePtr<FPerShapeData>(new FPerShapeDataCachedLeafInfo(MoveTemp(*ShapePtr.Get())));
-			}
-			else
-			{
-				ShapePtr = TUniquePtr<FPerShapeData>(new FPerShapeData(MoveTemp(*ShapePtr.Get())));
-			}
-		}
-
-		ShapePtr->Geometry = InGeometry;
-	}
-
-	FPerShapeDataCachedLeafInfo* FPerShapeData::AsCachedLeafInfo()
-	{
-		if (HasCachedLeafInfo())
-		{
-			return static_cast<FPerShapeDataCachedLeafInfo*>(this);
-		}
-		return nullptr;
-	}
-
-	const FPerShapeDataCachedLeafInfo* FPerShapeData::AsCachedLeafInfo() const
-	{
-		if (HasCachedLeafInfo())
-		{
-			return static_cast<const FPerShapeDataCachedLeafInfo*>(this);
-		}
-		return nullptr;
-	}
-
-	void FPerShapeData::UpdateShapeBounds(const FRigidTransform3& WorldTM, const FVec3& BoundsExpansion)
-	{
-		if (Geometry && Geometry->HasBoundingBox())
-		{
-			WorldSpaceInflatedShapeBounds = Geometry->CalculateTransformedBounds(WorldTM).ThickenSymmetrically(BoundsExpansion);
-		}
-	}
-
-	void FPerShapeData::UpdateWorldSpaceState(const FRigidTransform3& WorldTransform, const FVec3& BoundsExpansion)
-	{
-		FRigidTransform3 LeafWorldTransform = WorldTransform;
-		const FRigidTransform3* LeafRelativeTransform = nullptr;
-		const FImplicitObject* LeafGeometry = GetInnerGeometryInstanceData(Geometry.Get(), &LeafRelativeTransform);
-
-		// Calculate the leaf world transform if different from particle transform
-		if (LeafRelativeTransform != nullptr)
-		{
-			LeafWorldTransform = FRigidTransform3::MultiplyNoScale(*LeafRelativeTransform, WorldTransform);
-		}
-
-		// Store the leaf data if we have a cache
-		if (FPerShapeDataCachedLeafInfo* LeafCache = AsCachedLeafInfo())
-		{
-			LeafCache->SetWorldTransform(LeafWorldTransform);
-		}
-
-		// Update the bounds at the world transform
-		if ((LeafGeometry != nullptr) && LeafGeometry->HasBoundingBox())
-		{
-			WorldSpaceInflatedShapeBounds = LeafGeometry->CalculateTransformedBounds(LeafWorldTransform).ThickenSymmetrically(BoundsExpansion);;
-		}
-	}
-
-	const FImplicitObject* FPerShapeData::GetLeafGeometry() const
-	{
-		return GetInnerGeometryInstanceData(Geometry.Get(), nullptr);
-	}
-
-	FRigidTransform3 FPerShapeData::GetLeafRelativeTransform() const
-	{
-		const FRigidTransform3* LeafRelativeTransform = nullptr;
-		GetInnerGeometryInstanceData(Geometry.Get(), &LeafRelativeTransform);
-
-		if (LeafRelativeTransform != nullptr)
-		{
-			return *LeafRelativeTransform;
-		}
-		else
-		{
-			return FRigidTransform3::Identity;
-		}
-	}
-
-	FRigidTransform3 FPerShapeData::GetLeafWorldTransform(const FGeometryParticleHandle* Particle) const
-	{
-		if (const FPerShapeDataCachedLeafInfo* LeafCache = AsCachedLeafInfo())
-		{
-			return LeafCache->GetWorldTransform();
-		}
-		else
-		{
-			FRigidTransform3 LeafWorldTransform = FConstGenericParticleHandle(Particle)->GetTransformPQ();
-
-			const FRigidTransform3* LeafRelativeTransform = nullptr;
-			GetInnerGeometryInstanceData(Geometry.Get(), &LeafRelativeTransform);
-			if (LeafRelativeTransform != nullptr)
-			{
-				LeafWorldTransform = FRigidTransform3::MultiplyNoScale(*LeafRelativeTransform, LeafWorldTransform);
-			}
-
-			return LeafWorldTransform;
-		}
-	}
-	
-	void FPerShapeData::UpdateLeafWorldTransform(FGeometryParticleHandle* Particle)
-	{
-		if (FPerShapeDataCachedLeafInfo* LeafCache = AsCachedLeafInfo())
-		{
-			FRigidTransform3 LeafWorldTransform = FConstGenericParticleHandle(Particle)->GetTransformPQ();
-
-			const FRigidTransform3* LeafRelativeTransform = nullptr;
-			GetInnerGeometryInstanceData(Geometry.Get(), &LeafRelativeTransform);
-			if (LeafRelativeTransform != nullptr)
-			{
-				LeafWorldTransform = FRigidTransform3::MultiplyNoScale(*LeafRelativeTransform, LeafWorldTransform);
-			}
-
-			LeafCache->SetWorldTransform(LeafWorldTransform);
-		}
-	}
-
-	FPerShapeData* FPerShapeData::SerializationFactory(FChaosArchive& Ar, FPerShapeData*)
-	{
-		//todo: need to rework serialization for shapes, for now just give them all shape idx 0
-		return Ar.IsLoading() ? new FPerShapeData(0) : nullptr;
-	}
-
-	void FPerShapeData::Serialize(FChaosArchive& Ar)
-	{
-		Ar.UsingCustomVersion(FExternalPhysicsCustomObjectVersion::GUID);
-		Ar.UsingCustomVersion(FExternalPhysicsMaterialCustomObjectVersion::GUID);
-
-		Ar << Geometry;
-		Ar << CollisionData;
-		Ar << Materials;
-
-		if (Ar.CustomVer(FExternalPhysicsCustomObjectVersion::GUID) >= FExternalPhysicsCustomObjectVersion::SerializeShapeWorldSpaceBounds)
-		{
-			TBox<FReal,3>::SerializeAsAABB(Ar,WorldSpaceInflatedShapeBounds);
-		}
-		else
-		{
-			// This should be set by particle serializing this FPerShapeData.
-			WorldSpaceInflatedShapeBounds = FAABB3(FVec3(0.0f, 0.0f, 0.0f), FVec3(0.0f, 0.0f, 0.0f));
-		}
-
-	}
+	extern void UpdateShapesArrayFromGeometry(FShapeInstanceArray& ShapesArray, TSerializablePtr<FImplicitObject> Geometry, const FRigidTransform3& ActorTM);
 
 	FShapeOrShapesArray::FShapeOrShapesArray(const FGeometryParticleHandle* Particle)
 	{
@@ -410,7 +42,7 @@ namespace Chaos
 	template <typename T, int d, EGeometryParticlesSimType SimType>
 	void TGeometryParticlesImp<T, d, SimType>::UpdateShapesArray(const int32 Index)
 	{
-		UpdateShapesArrayFromGeometryImpl(MShapesArray[Index], MGeometry[Index], FRigidTransform3(X(Index), R(Index)), nullptr, true);
+		UpdateShapesArrayFromGeometry(MShapesArray[Index], MGeometry[Index], FRigidTransform3(X(Index), R(Index)));
 	}
 
 	template <typename T, int d, EGeometryParticlesSimType SimType>
diff --git a/Engine/Source/Runtime/Experimental/Chaos/Private/Chaos/ParticleHandle.cpp b/Engine/Source/Runtime/Experimental/Chaos/Private/Chaos/ParticleHandle.cpp
--- a/Engine/Source/Runtime/Experimental/Chaos/Private/Chaos/ParticleHandle.cpp
+++ b/Engine/Source/Runtime/Experimental/Chaos/Private/Chaos/ParticleHandle.cpp
@@ -13,7 +13,7 @@
 
 namespace Chaos
 {
-	extern void UpdateShapesArrayFromGeometryImpl(FShapesArray& ShapesArray, TSerializablePtr<FImplicitObject> Geometry, const FRigidTransform3& ActorTM, IPhysicsProxyBase* Proxy, bool bAllowCachedLeafInfo);
+	extern void UpdateShapesArrayFromGeometry(FShapeInstanceProxyArray& ShapesArray, TSerializablePtr<FImplicitObject> Geometry, const FRigidTransform3& ActorTM, IPhysicsProxyBase* Proxy);
 
 	void SetObjectStateHelper(IPhysicsProxyBase& Proxy, FPBDRigidParticleHandle& Rigid, EObjectStateType InState, bool bAllowEvents, bool bInvalidate)
 	{
@@ -59,7 +59,7 @@ namespace Chaos
 	template <typename T, int d>
 	void Chaos::TGeometryParticle<T, d>::UpdateShapesArray()
 	{
-		UpdateShapesArrayFromGeometryImpl(MShapesArray, MakeSerializable(MNonFrequentData.Read().Geometry()), FRigidTransform3(X(), R()), Proxy, /*bAllowLeafCache*/false);
+		UpdateShapesArrayFromGeometry(MShapesArray, MakeSerializable(MNonFrequentData.Read().Geometry()), FRigidTransform3(X(), R()), Proxy);
 	}
 
 	template <typename T, int d>
diff --git a/Engine/Source/Runtime/Experimental/Chaos/Private/Chaos/ShapeInstance.cpp b/Engine/Source/Runtime/Experimental/Chaos/Private/Chaos/ShapeInstance.cpp
new file mode 100644
--- /dev/null
+++ b/Engine/Source/Runtime/Experimental/Chaos/Private/Chaos/ShapeInstance.cpp
@@ -0,0 +1,439 @@
+// Copyright Epic Games, Inc. All Rights Reserved.
+#include "Chaos/ShapeInstance.h"
+#include "Chaos/ImplicitObject.h"
+#include "Chaos/ImplicitObjectUnion.h"
+#include "Chaos/ImplicitObjectScaled.h"
+
+namespace Chaos
+{
+	// Create or reuse the shapes in the shapes array and populate them with the Geometry.
+	// If we have a Union it will be unpacked into the ShapesArray.
+	// On the Physics Thread we set bAllowCachedLeafInfo which caches the shapes world space state to optimize collision detection,
+	// but this flag should not be used on the Game Thread because the extended data is not maintained so some PerShapeData functions
+	// will return unitialized values (E.g., GetLeafWorldTransform).
+	template<typename TShapesArrayType>
+	void UpdateShapesArrayFromGeometryImpl(
+		TShapesArrayType& ShapesArray,
+		TSerializablePtr<FImplicitObject> Geometry,
+		const FRigidTransform3& ActorTM,
+		IPhysicsProxyBase* Proxy)
+	{
+		// TShapesArrayType = TArray<TUniquePtr<T>>
+		using FShapePtrType = typename TShapesArrayType::ElementType;	// TUniquePtr<T>
+		using FShapeType = typename FShapePtrType::ElementType;			// T
+
+		if (Geometry)
+		{
+			const int32 OldShapeNum = ShapesArray.Num();
+			if (const auto* Union = Geometry->template GetObject<FImplicitObjectUnion>())
+			{
+				ShapesArray.SetNum(Union->GetObjects().Num());
+
+				for (int32 ShapeIndex = 0; ShapeIndex < ShapesArray.Num(); ++ShapeIndex)
+				{
+					TSerializablePtr<FImplicitObject> ShapeGeometry = MakeSerializable(Union->GetObjects()[ShapeIndex]);
+
+					if (ShapeIndex >= OldShapeNum)
+					{
+						// If newly allocated shape, initialize it.
+						ShapesArray[ShapeIndex] = FShapeType::Make(ShapeIndex, ShapeGeometry);
+					}
+					else if (ShapeGeometry != ShapesArray[ShapeIndex]->GetGeometry())
+					{
+						// Update geometry pointer if it changed
+						FShapeType::UpdateGeometry(ShapesArray[ShapeIndex], ShapeGeometry);
+					}
+				}
+			}
+			else
+			{
+				ShapesArray.SetNum(1);
+				if (OldShapeNum == 0)
+				{
+					ShapesArray[0] = FShapeType::Make(0, Geometry);
+				}
+				else
+				{
+					FShapeType::UpdateGeometry(ShapesArray[0], Geometry);
+				}
+			}
+
+			if (Geometry->HasBoundingBox())
+			{
+				for (auto& Shape : ShapesArray)
+				{
+					Shape->UpdateShapeBounds(ActorTM, FVec3(0));
+				}
+			}
+		}
+		else
+		{
+			ShapesArray.Reset();
+		}
+
+		if (Proxy)
+		{
+			if (FPhysicsSolverBase* PhysicsSolverBase = Proxy->GetSolver<FPhysicsSolverBase>())
+			{
+				PhysicsSolverBase->SetNumDirtyShapes(Proxy, ShapesArray.Num());
+			}
+		}
+	}
+
+	void UpdateShapesArrayFromGeometry(
+		FShapeInstanceArray& ShapesArray, 
+		TSerializablePtr<FImplicitObject> Geometry, 
+		const FRigidTransform3& ActorTM)
+	{
+		UpdateShapesArrayFromGeometryImpl(ShapesArray, Geometry, ActorTM, nullptr);
+	}
+
+	void UpdateShapesArrayFromGeometry(
+		FShapeInstanceProxyArray& ShapesArray, 
+		TSerializablePtr<FImplicitObject> Geometry, 
+		const FRigidTransform3& ActorTM, 
+		IPhysicsProxyBase* Proxy)
+	{
+		UpdateShapesArrayFromGeometryImpl(ShapesArray, Geometry, ActorTM, Proxy);
+	}
+
+	void UpdateShapesArrayFromGeometry(
+		FShapesArray& ShapesArray,
+		TSerializablePtr<FImplicitObject> Geometry,
+		const FRigidTransform3& ActorTM,
+		IPhysicsProxyBase* Proxy)
+	{
+		check(false);
+	}
+
+	// Unwrap transformed shapes
+	// @todo(chaos): also unwrap Instanced and Scaled but that requires a lot of knock work because Convexes are usually Instanced or
+	// Scaled so the Scale and Margin is stored on the wrapper (the convex itself is shared).
+	// - support for Margin as per-shape data passed through the collision functions
+	// - support for Scale as per-shape data passed through the collision functions (or ideally in the Transforms)
+	const FImplicitObject* GetInnerGeometryInstanceData(const FImplicitObject* Implicit, const FRigidTransform3** OutRelativeTransformPtr)
+	{
+		if (Implicit != nullptr)
+		{
+			const EImplicitObjectType ImplicitOuterType = Implicit->GetType();
+			if (ImplicitOuterType == TImplicitObjectTransformed<FReal, 3>::StaticType())
+			{
+				// Transformed Implicit
+				const TImplicitObjectTransformed<FReal, 3>* TransformedImplicit = Implicit->template GetObject<const TImplicitObjectTransformed<FReal, 3>>();
+				if (OutRelativeTransformPtr != nullptr)
+				{
+					*OutRelativeTransformPtr = &TransformedImplicit->GetTransform();
+				}
+				return GetInnerGeometryInstanceData(TransformedImplicit->GetTransformedObject(), OutRelativeTransformPtr);
+			}
+			else if ((uint32)ImplicitOuterType & ImplicitObjectType::IsInstanced)
+			{
+				// Instanced Implicit
+				// Currently we only unwrap instanced TriMesh and Heightfields. Convex is the only other type that will be Instanced,
+				// but we do not unwrap Convex because the convex margin is stored in the Instanced so we would need to extract
+				// it and pass it on to the collision detection.
+				const FImplicitObjectInstanced* Instanced = static_cast<const FImplicitObjectInstanced*>(Implicit);
+				EImplicitObjectType InnerType = Instanced->GetInnerObject()->GetType();
+				if (InnerType != FImplicitConvex3::StaticType())
+				{
+					return GetInnerGeometryInstanceData(Instanced->GetInnerObject().Get(), OutRelativeTransformPtr);
+				}
+			}
+			else if ((uint32)ImplicitOuterType & ImplicitObjectType::IsScaled)
+			{
+				// Scaled Implicit
+				//const FImplicitObjectScaled* Scaled = static_cast<const FImplicitObjectScaled*>(Implicit);
+				//OutTransform.Scale *= Scaled->GetScale();
+				//return GetInnerGeometryInstanceData(Scaled->GetInnerObject().Get(), OutRelativeTransformPtr);
+			}
+		}
+		return Implicit;
+	}
+
+	bool ShapeInstanceWantsLeafCache(const FImplicitObject* Geometry)
+	{
+		// We will cache the world space transform if we are wrapped in a transform implicit
+		const FRigidTransform3* LeafRelativeTransform = nullptr;
+		const FImplicitObject* LeafGeometry = GetInnerGeometryInstanceData(Geometry, &LeafRelativeTransform);
+		return (LeafGeometry != Geometry) && (LeafRelativeTransform != nullptr);
+	}
+
+	///////////////////////////////////////////////////////////////////////////////////////////////
+	///////////////////////////////////////////////////////////////////////////////////////////////
+	// 
+	// Factory methods
+	// 
+	///////////////////////////////////////////////////////////////////////////////////////////////
+	///////////////////////////////////////////////////////////////////////////////////////////////
+
+	TUniquePtr<FPerShapeData> FPerShapeData::CreatePerShapeData(int32 InShapeIdx, TSerializablePtr<FImplicitObject> InGeometry)
+	{
+		return FShapeInstanceProxy::Make(InShapeIdx, InGeometry);
+	}
+
+	void FPerShapeData::UpdateGeometry(TUniquePtr<FPerShapeData>& InOutShapePtr, TSerializablePtr<FImplicitObject> InGeometry)
+	{
+		return FShapeInstanceProxy::UpdateGeometry(reinterpret_cast<FShapeInstanceProxyPtr&>(InOutShapePtr), InGeometry);
+	}
+
+
+
+	TUniquePtr<FShapeInstanceProxy> FShapeInstanceProxy::Make(int32 InShapeIdx, TSerializablePtr<FImplicitObject> InGeometry)
+	{
+		return TUniquePtr<FShapeInstanceProxy>(new FShapeInstanceProxy(InShapeIdx, InGeometry));
+	}
+
+	void FShapeInstanceProxy::UpdateGeometry(TUniquePtr<FShapeInstanceProxy>& InOutShapePtr, TSerializablePtr<FImplicitObject> InGeometry)
+	{
+		InOutShapePtr->Geometry = InGeometry;
+	}
+
+
+
+	TUniquePtr<FShapeInstance> FShapeInstance::Make(int32 InShapeIdx, TSerializablePtr<FImplicitObject> InGeometry)
+	{
+		const bool bWantLeafCache = ShapeInstanceWantsLeafCache(InGeometry.Get());
+		if (bWantLeafCache)
+		{
+			return TUniquePtr<FShapeInstance>(new Private::FShapeInstanceExtended(InShapeIdx, InGeometry));
+		}
+		else
+		{
+			return TUniquePtr<FShapeInstance>(new FShapeInstance(InShapeIdx, InGeometry));
+		}
+	}
+
+	void FShapeInstance::UpdateGeometry(TUniquePtr<FShapeInstance>& InOutShapePtr, TSerializablePtr<FImplicitObject> InGeometry)
+	{
+		const bool bWantLeafCache = ShapeInstanceWantsLeafCache(InGeometry.Get());
+
+		// Do we need to add or remove the cached leaf data? If so this requires we recreate the object
+		const bool bHasLeafCache = (InOutShapePtr->GetType() == EPerShapeDataType::SimExtended);
+		if (bWantLeafCache != bHasLeafCache)
+		{
+			FShapeInstance* ShapeInstance = InOutShapePtr.Get();
+			if (bWantLeafCache)
+			{
+				InOutShapePtr = TUniquePtr<FShapeInstance>(new Private::FShapeInstanceExtended(MoveTemp(*ShapeInstance)));
+			}
+			else
+			{
+				InOutShapePtr = TUniquePtr<FShapeInstance>(new FShapeInstance(MoveTemp(*ShapeInstance)));
+			}
+		}
+
+		InOutShapePtr->Geometry = InGeometry;
+	}
+
+	///////////////////////////////////////////////////////////////////////////////////////////////
+	///////////////////////////////////////////////////////////////////////////////////////////////
+	// 
+	// Serialization
+	// 
+	///////////////////////////////////////////////////////////////////////////////////////////////
+	///////////////////////////////////////////////////////////////////////////////////////////////
+
+	void FPerShapeData::Serialize(FChaosArchive& Ar)
+	{
+		Ar.UsingCustomVersion(FExternalPhysicsCustomObjectVersion::GUID);
+		Ar.UsingCustomVersion(FExternalPhysicsMaterialCustomObjectVersion::GUID);
+
+		Ar << Geometry;
+
+		DownCast([&Ar](auto& ShapeInstance)
+			{
+				Ar << ShapeInstance.CollisionData;
+				Ar << ShapeInstance.Materials;
+			});
+
+		if (Ar.CustomVer(FExternalPhysicsCustomObjectVersion::GUID) >= FExternalPhysicsCustomObjectVersion::SerializeShapeWorldSpaceBounds)
+		{
+			TBox<FReal, 3>::SerializeAsAABB(Ar, WorldSpaceInflatedShapeBounds);
+		}
+		else
+		{
+			// This should be set by particle serializing this FPerShapeData.
+			WorldSpaceInflatedShapeBounds = FAABB3(FVec3(0.0f, 0.0f, 0.0f), FVec3(0.0f, 0.0f, 0.0f));
+		}
+	}
+
+	FPerShapeData* FPerShapeData::SerializationFactory(FChaosArchive& Ar, FPerShapeData*)
+	{
+		// We should not be serializing base class pointers - should always be FShapeInstanceProxy or FShapeInstance
+		ensure(false);
+		return FShapeInstanceProxy::SerializationFactory(Ar, nullptr);
+	}
+
+	FShapeInstanceProxy* FShapeInstanceProxy::SerializationFactory(FChaosArchive& Ar, FShapeInstanceProxy*)
+	{
+		//todo: need to rework serialization for shapes, for now just give them all shape idx 0
+		return Ar.IsLoading() ? new FShapeInstanceProxy(0) : nullptr;
+	}
+
+	FShapeInstance* FShapeInstance::SerializationFactory(FChaosArchive& Ar, FShapeInstance*)
+	{
+		//todo: need to rework serialization for shapes, for now just give them all shape idx 0
+		return Ar.IsLoading() ? new FShapeInstance(0) : nullptr;
+	}
+
+	///////////////////////////////////////////////////////////////////////////////////////////////
+	///////////////////////////////////////////////////////////////////////////////////////////////
+	// 
+	// FShapeInstanceProxy
+	// 
+	///////////////////////////////////////////////////////////////////////////////////////////////
+	///////////////////////////////////////////////////////////////////////////////////////////////
+
+
+	void FShapeInstanceProxy::UpdateShapeBounds(const FRigidTransform3& WorldTransform, const FVec3& BoundsExpansion)
+	{
+		if (Geometry && Geometry->HasBoundingBox())
+		{
+			WorldSpaceInflatedShapeBounds = Geometry->CalculateTransformedBounds(WorldTransform).ThickenSymmetrically(BoundsExpansion);
+		}
+	}
+
+	void FShapeInstanceProxy::UpdateWorldSpaceState(const FRigidTransform3& WorldTransform, const FVec3& BoundsExpansion)
+	{
+		UpdateShapeBounds(WorldTransform, BoundsExpansion);
+	}
+
+	const FImplicitObject* FShapeInstanceProxy::GetLeafGeometry() const
+	{
+		return GetInnerGeometryInstanceData(Geometry.Get(), nullptr);
+	}
+
+	FRigidTransform3 FShapeInstanceProxy::GetLeafRelativeTransform() const
+	{
+		const FRigidTransform3* LeafRelativeTransform = nullptr;
+		GetInnerGeometryInstanceData(Geometry.Get(), &LeafRelativeTransform);
+
+		if (LeafRelativeTransform != nullptr)
+		{
+			return *LeafRelativeTransform;
+		}
+		else
+		{
+			return FRigidTransform3::Identity;
+		}
+	}
+
+	FRigidTransform3 FShapeInstanceProxy::GetLeafWorldTransform(const FGeometryParticleHandle* Particle) const
+	{
+		FRigidTransform3 LeafWorldTransform = FConstGenericParticleHandle(Particle)->GetTransformPQ();
+
+		const FRigidTransform3* LeafRelativeTransform = nullptr;
+		GetInnerGeometryInstanceData(Geometry.Get(), &LeafRelativeTransform);
+		if (LeafRelativeTransform != nullptr)
+		{
+			LeafWorldTransform = FRigidTransform3::MultiplyNoScale(*LeafRelativeTransform, LeafWorldTransform);
+		}
+
+		return LeafWorldTransform;
+	}
+
+	void FShapeInstanceProxy::UpdateLeafWorldTransform(FGeometryParticleHandle* Particle)
+	{
+	}
+
+	///////////////////////////////////////////////////////////////////////////////////////////////
+	///////////////////////////////////////////////////////////////////////////////////////////////
+	// 
+	// FShapeInstance
+	// 
+	///////////////////////////////////////////////////////////////////////////////////////////////
+	///////////////////////////////////////////////////////////////////////////////////////////////
+
+
+	void FShapeInstance::UpdateShapeBounds(const FRigidTransform3& WorldTransform, const FVec3& BoundsExpansion)
+	{
+		if (Geometry && Geometry->HasBoundingBox())
+		{
+			WorldSpaceInflatedShapeBounds = Geometry->CalculateTransformedBounds(WorldTransform).ThickenSymmetrically(BoundsExpansion);
+		}
+	}
+
+	void FShapeInstance::UpdateWorldSpaceState(const FRigidTransform3& WorldTransform, const FVec3& BoundsExpansion)
+	{
+		FRigidTransform3 LeafWorldTransform = WorldTransform;
+		const FRigidTransform3* LeafRelativeTransform = nullptr;
+		const FImplicitObject* LeafGeometry = GetInnerGeometryInstanceData(Geometry.Get(), &LeafRelativeTransform);
+
+		// Calculate the leaf world transform if different from particle transform
+		if (LeafRelativeTransform != nullptr)
+		{
+			LeafWorldTransform = FRigidTransform3::MultiplyNoScale(*LeafRelativeTransform, WorldTransform);
+		}
+
+		// Store the leaf data if we have a cache
+		if (Private::FShapeInstanceExtended* LeafCache = AsShapeInstanceExtended())
+		{
+			LeafCache->SetWorldTransform(LeafWorldTransform);
+		}
+
+		// Update the bounds at the world transform
+		if ((LeafGeometry != nullptr) && LeafGeometry->HasBoundingBox())
+		{
+			WorldSpaceInflatedShapeBounds = LeafGeometry->CalculateTransformedBounds(LeafWorldTransform).ThickenSymmetrically(BoundsExpansion);;
+		}
+	}
+
+	const FImplicitObject* FShapeInstance::GetLeafGeometry() const
+	{
+		return GetInnerGeometryInstanceData(Geometry.Get(), nullptr);
+	}
+
+	FRigidTransform3 FShapeInstance::GetLeafRelativeTransform() const
+	{
+		const FRigidTransform3* LeafRelativeTransform = nullptr;
+		GetInnerGeometryInstanceData(Geometry.Get(), &LeafRelativeTransform);
+
+		if (LeafRelativeTransform != nullptr)
+		{
+			return *LeafRelativeTransform;
+		}
+		else
+		{
+			return FRigidTransform3::Identity;
+		}
+	}
+
+	FRigidTransform3 FShapeInstance::GetLeafWorldTransform(const FGeometryParticleHandle* Particle) const
+	{
+		if (const Private::FShapeInstanceExtended* LeafCache = AsShapeInstanceExtended())
+		{
+			return LeafCache->GetWorldTransform();
+		}
+		else
+		{
+			FRigidTransform3 LeafWorldTransform = FConstGenericParticleHandle(Particle)->GetTransformPQ();
+
+			const FRigidTransform3* LeafRelativeTransform = nullptr;
+			GetInnerGeometryInstanceData(Geometry.Get(), &LeafRelativeTransform);
+			if (LeafRelativeTransform != nullptr)
+			{
+				LeafWorldTransform = FRigidTransform3::MultiplyNoScale(*LeafRelativeTransform, LeafWorldTransform);
+			}
+
+			return LeafWorldTransform;
+		}
+	}
+
+	void FShapeInstance::UpdateLeafWorldTransform(FGeometryParticleHandle* Particle)
+	{
+		if (Private::FShapeInstanceExtended* LeafCache = AsShapeInstanceExtended())
+		{
+			FRigidTransform3 LeafWorldTransform = FConstGenericParticleHandle(Particle)->GetTransformPQ();
+
+			const FRigidTransform3* LeafRelativeTransform = nullptr;
+			GetInnerGeometryInstanceData(Geometry.Get(), &LeafRelativeTransform);
+			if (LeafRelativeTransform != nullptr)
+			{
+				LeafWorldTransform = FRigidTransform3::MultiplyNoScale(*LeafRelativeTransform, LeafWorldTransform);
+			}
+
+			LeafCache->SetWorldTransform(LeafWorldTransform);
+		}
+	}
+
+}
\ No newline at end of file
diff --git a/Engine/Source/Runtime/Experimental/Chaos/Public/Chaos/GeometryParticles.h b/Engine/Source/Runtime/Experimental/Chaos/Public/Chaos/GeometryParticles.h
--- a/Engine/Source/Runtime/Experimental/Chaos/Public/Chaos/GeometryParticles.h
+++ b/Engine/Source/Runtime/Experimental/Chaos/Public/Chaos/GeometryParticles.h
@@ -17,6 +17,7 @@
 #include "Chaos/Properties.h"
 #include "Chaos/Framework/PhysicsProxyBase.h"
 #include "Chaos/Framework/PhysicsSolverBase.h"
+#include "Chaos/ShapeInstance.h"
 
 #ifndef CHAOS_DETERMINISTIC
 #define CHAOS_DETERMINISTIC 1
@@ -27,9 +28,7 @@ namespace Chaos
 {
 	class FConstraintHandle;
 	class FParticleCollisions;
-	class FPerShapeDataCachedLeafInfo;
 
-	using FShapesArray = TArray<TUniquePtr<FPerShapeData>, TInlineAllocator<1>>;
 	using FConstraintHandleArray = TArray<FConstraintHandle*>;
 
 	namespace CVars
@@ -38,315 +37,6 @@ namespace Chaos
 		CHAOS_API extern bool bCCDAxisThresholdUsesProbeShapes;
 	}
 
-	/**
-	 * Per-particle data associated with a geometry used by collision detection,
-	 * and queries. This contains data like the collision/query filters, material
-	 * properties etc.
-	 *
-	 * Every particle holds one FPerShapeData objects for each geometry they use.
-	 * If the particle has a Union of geometries the shapes array will be a flattened
-	 * list of those geometries. (Except ClusteredUnions - these are not flattened
-	 * because they contain their own query acceleration structure.)
-	 * 
-	 * NOTE: Both the game thread and physics thread have FPerShapeData objects for
-	 * each shape. This class is used for both.
-	 * 
-	 * NOTE: It is extremely important to keep this class as small as possible
-	 * because there is one shape instance in the world. This can easily be in
-	 * the millions.
-	 * 
-	 * @todo(chaos): split FPerShapeData into GT and PT versions. The PT version does 
-	 * not need the interop data. The GT version does not need the extended data flag.
-	 * Also try to remove the GT Proxy pointer - this could easily be passed into the
-	 * relevant functions instead.
-	 * 
-	 * @todo(chaos): For the PT version (maybe the GT version) reduce the cost of
-	 * MaterialData for shapes with one material and no masks etc.
-	 * 
-	 */
-	class CHAOS_API FPerShapeData
-	{
-	public:
-
-		static constexpr bool AlwaysSerializable = true;
-
-		UE_DEPRECATED(5.3, "Not used")
-		static bool RequiresCachedLeafInfo(const FImplicitObject* Geometry) { return false; }
-
-		// Game Thread PerShapeData factory/update
-		static TUniquePtr<FPerShapeData> CreatePerShapeData(int32 InShapeIdx, TSerializablePtr<FImplicitObject> InGeometry);
-		static void UpdateGeometry(TUniquePtr<FPerShapeData>& ShapePtr, TSerializablePtr<FImplicitObject> InGeometry);
-
-		// Physics Thread PerShapeData factory/update - Do not use the Ext versions from the game thread because the extended data will not be maintained
-		static TUniquePtr<FPerShapeData> CreatePerShapeDataInternal(int32 InShapeIdx, TSerializablePtr<FImplicitObject> InGeometry);
-		static void UpdateGeometryInternal(TUniquePtr<FPerShapeData>& ShapePtr, TSerializablePtr<FImplicitObject> InGeometry);
-
-		virtual ~FPerShapeData();
-		
-		void UpdateShapeBounds(const FRigidTransform3& WorldTM, const FVec3& BoundsExpansion = FVec3(0));
-
-		static FPerShapeData* SerializationFactory(FChaosArchive& Ar, FPerShapeData*);
-		void Serialize(FChaosArchive& Ar);
-
-		void* GetUserData() const { return CollisionData.Read().UserData; }
-		void SetUserData(void* InUserData)
-		{
-			CollisionData.Modify(true, DirtyFlags, Proxy, ShapeIdx, [InUserData](FCollisionData& Data){ Data.UserData = InUserData; });
-		}
-
-		const FCollisionFilterData& GetQueryData() const { return CollisionData.Read().QueryData; }
-		void SetQueryData(const FCollisionFilterData& InQueryData)
-		{
-			CollisionData.Modify(true, DirtyFlags, Proxy, ShapeIdx, [InQueryData](FCollisionData& Data){ Data.QueryData = InQueryData; });
-		}
-
-		const FCollisionFilterData& GetSimData() const { return CollisionData.Read().SimData; }
-		void SetSimData(const FCollisionFilterData& InSimData)
-		{
-			CollisionData.Modify(true, DirtyFlags, Proxy, ShapeIdx, [InSimData](FCollisionData& Data){ Data.SimData = InSimData; });
-		}
-
-		template <typename Lambda>
-		void ModifySimData(const Lambda& LambdaFunc)
-		{
-			CollisionData.Modify(true,DirtyFlags,Proxy, ShapeIdx,[&LambdaFunc](FCollisionData& Data){ LambdaFunc(Data.SimData);});
-		}
-
-		TSerializablePtr<FImplicitObject> GetGeometry() const { return Geometry; }
-
-		const TAABB<FReal,3>& GetWorldSpaceInflatedShapeBounds() const { return WorldSpaceInflatedShapeBounds; }
-
-		void UpdateWorldSpaceState(const FRigidTransform3& WorldTransform, const FVec3& BoundsExpansion);
-
-		// The leaf shape (with transformed and implicit wrapper removed).
-		const FImplicitObject* GetLeafGeometry() const;
-
-		// The actor-relative transform of the leaf geometry.
-		FRigidTransform3 GetLeafRelativeTransform() const;
-
-		// The world-space transform of the leaf geometry.
-		// If we have non-identity leaf relative transform, is cached from the last call to UpdateWorldSpaceState.
-		// If not cahced, is constructed from arguments.
-		FRigidTransform3 GetLeafWorldTransform(const FGeometryParticleHandle* Particle) const;
-		void UpdateLeafWorldTransform(FGeometryParticleHandle* Particle);
-
-		// Are we caching leaf transform info?
-		bool HasCachedLeafInfo() const { return bHasCachedLeafInfo; }
-
-		const TArray<FMaterialHandle>& GetMaterials() const { return Materials.Read().Materials; }
-		const TArray<FMaterialMaskHandle>& GetMaterialMasks() const { return Materials.Read().MaterialMasks; }
-		const TArray<uint32>& GetMaterialMaskMaps() const { return Materials.Read().MaterialMaskMaps; }
-		const TArray<FMaterialHandle>& GetMaterialMaskMapMaterials() const { return Materials.Read().MaterialMaskMapMaterials; }
-
-		const FShapeDirtyFlags GetDirtyFlags() const { return DirtyFlags; }
-
-		void SetMaterial(FMaterialHandle InMaterial)
-		{
-			Materials.Modify(true, DirtyFlags, Proxy, ShapeIdx, [InMaterial](FMaterialData& Data)
-			{
-				Data.Materials.Reset(1);
-				Data.Materials.Add(InMaterial);
-			});
-		}
-
-		void SetMaterials(const TArray<FMaterialHandle>& InMaterials)
-		{
-			Materials.Modify(true,DirtyFlags,Proxy, ShapeIdx,[&InMaterials](FMaterialData& Data)
-			{
-				Data.Materials = InMaterials;
-			});
-		}
-
-		void SetMaterialMasks(const TArray<FMaterialMaskHandle>& InMaterialMasks)
-		{
-			Materials.Modify(true,DirtyFlags,Proxy, ShapeIdx,[&InMaterialMasks](FMaterialData& Data)
-			{
-				Data.MaterialMasks = InMaterialMasks;
-			});
-		}
-
-		void SetMaterialMaskMaps(const TArray<uint32>& InMaterialMaskMaps)
-		{
-			Materials.Modify(true,DirtyFlags,Proxy, ShapeIdx,[&InMaterialMaskMaps](FMaterialData& Data)
-			{
-				Data.MaterialMaskMaps = InMaterialMaskMaps;
-			});
-		}
-
-		void SetMaterialMaskMapMaterials(const TArray<FMaterialHandle>& InMaterialMaskMapMaterials)
-		{
-			Materials.Modify(true,DirtyFlags,Proxy, ShapeIdx,[&InMaterialMaskMapMaterials](FMaterialData& Data)
-			{
-				Data.MaterialMaskMapMaterials = InMaterialMaskMapMaterials;
-			});
-		}
-
-		template <typename Lambda>
-		void ModifyMaterials(const Lambda& LambdaFunc)
-		{
-			Materials.Modify(true,DirtyFlags,Proxy, ShapeIdx,[&LambdaFunc](FMaterialData& Data)
-			{
-				LambdaFunc(Data.Materials);
-			});
-		}
-
-		template <typename Lambda>
-		void ModifyMaterialMasks(const Lambda& LambdaFunc)
-		{
-			Materials.Modify(true,DirtyFlags,Proxy, ShapeIdx,[&LambdaFunc](FMaterialData& Data)
-			{
-				LambdaFunc(Data.MaterialMasks);
-			});
-		}
-
-		template <typename Lambda>
-		void ModifyMaterialMaskMaps(const Lambda& LambdaFunc)
-		{
-			Materials.Modify(true,DirtyFlags,Proxy, ShapeIdx,[&LambdaFunc](FMaterialData& Data)
-			{
-				LambdaFunc(Data.MaterialMaskMaps);
-			});
-		}
-
-		template <typename Lambda>
-		void ModifyMaterialMaskMapMaterials(const Lambda& LambdaFunc)
-		{
-			Materials.Modify(true,DirtyFlags,Proxy, ShapeIdx,[&LambdaFunc](FMaterialData& Data)
-			{
-				LambdaFunc(Data.MaterialMaskMapMaterials);
-			});
-		}
-
-		bool GetQueryEnabled() const { return CollisionData.Read().bQueryCollision; }
-		void SetQueryEnabled(const bool bEnable)
-		{
-			CollisionData.Modify(true, DirtyFlags, Proxy, ShapeIdx, [bEnable](FCollisionData& Data){ Data.bQueryCollision = bEnable; });
-		}
-
-		bool GetSimEnabled() const { return CollisionData.Read().bSimCollision; }
-		void SetSimEnabled(const bool bEnable)
-		{
-			CollisionData.Modify(true, DirtyFlags, Proxy, ShapeIdx, [bEnable](FCollisionData& Data){ Data.bSimCollision = bEnable; });
-		}
-
-		bool GetIsProbe() const { return CollisionData.Read().bIsProbe; }
-		void SetIsProbe(const bool bIsProbe)
-		{
-			CollisionData.Modify(true, DirtyFlags, Proxy, ShapeIdx, [bIsProbe](FCollisionData& Data){ Data.bIsProbe = bIsProbe; });
-		}
-
-		EChaosCollisionTraceFlag GetCollisionTraceType() const { return CollisionData.Read().CollisionTraceType; }
-		void SetCollisionTraceType(const EChaosCollisionTraceFlag InTraceFlag)
-		{
-			CollisionData.Modify(true,DirtyFlags,Proxy, ShapeIdx,[InTraceFlag](FCollisionData& Data){ Data.CollisionTraceType = InTraceFlag; });
-		}
-
-		const FCollisionData& GetCollisionData() const { return CollisionData.Read(); }
-
-		void SetCollisionData(const FCollisionData& Data)
-		{
-			CollisionData.Write(Data,true,DirtyFlags,Proxy, ShapeIdx);
-		}
-
-		const FMaterialData& GetMaterialData() const { return Materials.Read(); }
-
-		void SetMaterialData(const FMaterialData& Data)
-		{
-			Materials.Write(Data,true,DirtyFlags,Proxy,ShapeIdx);
-		}
-
-		void SyncRemoteData(FDirtyPropertiesManager& Manager, int32 ShapeDataIdx, FShapeDirtyData& RemoteData)
-		{
-			RemoteData.SetFlags(DirtyFlags);
-			CollisionData.SyncRemote(Manager, ShapeDataIdx, RemoteData);
-			Materials.SyncRemote(Manager, ShapeDataIdx, RemoteData);
-			DirtyFlags.Clear();
-		}
-
-		void SetProxy(class IPhysicsProxyBase* InProxy)
-		{
-			Proxy = InProxy;
-			if(Proxy)
-			{
-				if(DirtyFlags.IsDirty())
-				{
-					if(FPhysicsSolverBase* PhysicsSolverBase = Proxy->GetSolver<FPhysicsSolverBase>())
-					{
-						PhysicsSolverBase->AddDirtyProxyShape(Proxy,ShapeIdx);
-					}
-				}
-			}
-		}
-
-		int32 GetShapeIndex() const
-		{
-			return ShapeIdx;
-		}
-
-		void ModifyShapeIndex(int32 NewShapeIndex)
-		{
-			ShapeIdx = NewShapeIndex;
-		}
-
-	protected:
-		// use CreatePerShapeData
-		FPerShapeData(int32 InShapeIdx, TSerializablePtr<FImplicitObject> InGeometry);
-
-		explicit FPerShapeData(FPerShapeData&& Other);
-
-		FPerShapeDataCachedLeafInfo* AsCachedLeafInfo();
-		const FPerShapeDataCachedLeafInfo* AsCachedLeafInfo() const;
-
-		// Can we be downcasted to FPerShapeDataCachedLeafInfo?
-		uint8 bHasCachedLeafInfo : 1;
-
-	private:
-		// Should only be used by SerializationFactory.
-		FPerShapeData(int32 InShapeIdx);
-
-		class IPhysicsProxyBase* Proxy;
-		FShapeDirtyFlags DirtyFlags;
-		int32 ShapeIdx;
-
-		TShapeProperty<FCollisionData,EShapeProperty::CollisionData> CollisionData;
-		TShapeProperty<FMaterialData,EShapeProperty::Materials> Materials;
-
-		TSerializablePtr<FImplicitObject> Geometry;
-		TAABB<FReal,3> WorldSpaceInflatedShapeBounds;
-	};
-
-	class CHAOS_API FPerShapeDataCachedLeafInfo final : public FPerShapeData 
-	{
-		friend class FPerShapeData;
-
-	private:
-		FPerShapeDataCachedLeafInfo(int32 InShapeIdx, TSerializablePtr<FImplicitObject> InGeometry)
-			: FPerShapeData(InShapeIdx, InGeometry)
-		{
-			bHasCachedLeafInfo = true;
-		}
-
-		// Move old shape into shape with cached leaf info
-		FPerShapeDataCachedLeafInfo(FPerShapeData&& PerShapeData)
-			: FPerShapeData(MoveTemp(PerShapeData))
-		{
-			bHasCachedLeafInfo = true;
-		}
-
-		FRigidTransform3 GetWorldTransform() const
-		{
-			return FRigidTransform3(WorldPosition, WorldRotation);
-		}
-
-		void SetWorldTransform(const FRigidTransform3& LeafWorldTransform)
-		{
-			WorldPosition = LeafWorldTransform.GetTranslation();
-			WorldRotation = LeafWorldTransform.GetRotation();
-		}
-
-		FVec3 WorldPosition;
-		FRotation3 WorldRotation;
-	};
 
 	/**
 	* Union between shape and shapes array pointers, used for passing around shapes with
@@ -402,15 +92,6 @@ namespace Chaos
 		bool bIsSingleShape;
 	};
 
-	inline FChaosArchive& operator<<(FChaosArchive& Ar, FPerShapeData& Shape)
-	{
-		Shape.Serialize(Ar);
-		return Ar;
-	}
-
-	UE_DEPRECATED(5.3, "Not for external use")
-	void CHAOS_API UpdateShapesArrayFromGeometry(FShapesArray& ShapesArray, TSerializablePtr<FImplicitObject> Geometry, const FRigidTransform3& ActorTM, IPhysicsProxyBase* Proxy);
-
 	FORCEINLINE uint32 GetTypeHash(const FParticleID& Unique)
 	{
 		return ::GetTypeHash(Unique.GlobalID);
@@ -649,7 +330,7 @@ namespace Chaos
 		CHAOS_API bool HasCollision(const int32 Index) const { return MHasCollision[Index]; }
 		CHAOS_API bool& HasCollision(const int32 Index) { return MHasCollision[Index]; }
 
-		CHAOS_API const FShapesArray& ShapesArray(const int32 Index) const { return MShapesArray[Index]; }
+		CHAOS_API const FShapesArray& ShapesArray(const int32 Index) const { return reinterpret_cast<const FShapesArray&>(MShapesArray[Index]); }
 
 #if CHAOS_DETERMINISTIC
 		CHAOS_API FParticleID ParticleID(const int32 Idx) const { return MParticleIDs[Idx]; }
@@ -993,7 +674,7 @@ public:
 		TArrayCollectionArray<FGeometryParticle*> MGeometryParticle;
 		TArrayCollectionArray<IPhysicsProxyBase*> MPhysicsProxy;
 		TArrayCollectionArray<bool> MHasCollision;
-		TArrayCollectionArray<FShapesArray> MShapesArray;
+		TArrayCollectionArray<FShapeInstanceArray> MShapesArray;
 		TArrayCollectionArray<TAABB<T,d>> MLocalBounds;
 		TArrayCollectionArray<TVector<T,d>> MCCDAxisThreshold;
 		TArrayCollectionArray<TAABB<T, d>> MWorldSpaceInflatedBounds;
diff --git a/Engine/Source/Runtime/Experimental/Chaos/Public/Chaos/ParticleHandle.h b/Engine/Source/Runtime/Experimental/Chaos/Public/Chaos/ParticleHandle.h
--- a/Engine/Source/Runtime/Experimental/Chaos/Public/Chaos/ParticleHandle.h
+++ b/Engine/Source/Runtime/Experimental/Chaos/Public/Chaos/ParticleHandle.h
@@ -2464,7 +2464,7 @@ public:
 	void SetShapesArray(FShapesArray&& InShapesArray)
 	{
 		ensure(InShapesArray.Num() == MShapesArray.Num());
-		MShapesArray = MoveTemp(InShapesArray);
+		MShapesArray = MoveTemp(reinterpret_cast<FShapeInstanceProxyArray&>(InShapesArray));
 	}
 
 	void MergeShapesArray(FShapesArray&& OtherShapesArray)
@@ -2473,7 +2473,7 @@ public:
 		for (TUniquePtr<FPerShapeData>& Shape : OtherShapesArray)
 		{
 			ensure(Idx < MShapesArray.Num());
-			MShapesArray[Idx++] = MoveTemp(Shape);
+			MShapesArray[Idx++] = MoveTemp(reinterpret_cast<FShapeInstanceProxyPtr&>(Shape));
 		}
 	}
 
@@ -2482,7 +2482,7 @@ public:
 
 	TSerializablePtr<FImplicitObject> Geometry() const { return MakeSerializable(MNonFrequentData.Read().Geometry()); }
 
-	const FShapesArray& ShapesArray() const { return MShapesArray; }
+	const FShapesArray& ShapesArray() const { return reinterpret_cast<const FShapesArray&>(MShapesArray); }
 
 	EObjectStateType ObjectState() const;
 
@@ -2654,7 +2654,7 @@ private:
 	TChaosProperty<FParticleNonFrequentData,EChaosProperty::NonFrequentData> MNonFrequentData;
 	void* MUserData;
 
-	FShapesArray MShapesArray;
+	FShapeInstanceProxyArray MShapesArray;
 
 public:
 	// Ryan: FGeometryCollectionPhysicsProxy needs access to GeometrySharedLowLevel(), 
diff --git a/Engine/Source/Runtime/Experimental/Chaos/Public/Chaos/ShapeInstance.h b/Engine/Source/Runtime/Experimental/Chaos/Public/Chaos/ShapeInstance.h
new file mode 100644
--- /dev/null
+++ b/Engine/Source/Runtime/Experimental/Chaos/Public/Chaos/ShapeInstance.h
@@ -0,0 +1,981 @@
+// Copyright Epic Games, Inc. All Rights Reserved.
+#pragma once
+
+#include "Chaos/Core.h"
+
+#include "Chaos/CollisionFilterData.h"
+#include "Chaos/Framework/PhysicsProxyBase.h"
+#include "Chaos/ImplicitFwd.h"
+#include "Chaos/ParticleDirtyFlags.h"
+#include "Chaos/ParticleHandleFwd.h"
+#include "Chaos/PhysicalMaterials.h"
+#include "Chaos/Properties.h"
+#include "Chaos/Serializable.h"
+#include "Chaos/ShapeInstanceFwd.h"
+
+namespace Chaos
+{
+	class FImplicitObject;
+	class FShapeInstance;
+	class FShapeInstanceProxy;
+
+	namespace Private
+	{
+		class FShapeInstanceExtended;
+	}
+
+
+	/**
+	 * FPerShapeData is going to be deprecated. See FShapeInstance and FShapeInstanceProxy
+	 * 
+	 * @todo(chaos): 
+	 * - change ShapesArray() and all code using it to use ShapeInstance or ShapeInstanceProxy as appropriate
+	 * - deprecate FPerShapeData
+	 */
+	class CHAOS_API FPerShapeData
+	{
+	protected:
+		enum class EPerShapeDataType : uint8
+		{
+			Proxy,
+			Sim,
+			SimExtended,
+		};
+
+		EPerShapeDataType GetType() const
+		{
+			return Type;
+		}
+
+		FShapeInstanceProxy* AsShapeInstanceProxy();
+		const FShapeInstanceProxy* AsShapeInstanceProxy() const;
+
+		FShapeInstance* AsShapeInstance();
+		const FShapeInstance* AsShapeInstance() const;
+
+		Private::FShapeInstanceExtended* AsShapeInstanceExtended();
+		const Private::FShapeInstanceExtended* AsShapeInstanceExtended() const;
+
+		// Call a function on the concrete type
+		template<typename TLambda> decltype(auto) DownCast(const TLambda& Lambda);
+		template<typename TLambda> decltype(auto) DownCast(const TLambda& Lambda) const;
+
+	public:
+		static constexpr bool AlwaysSerializable = true;
+
+		UE_DEPRECATED(5.3, "Not used")
+		static bool RequiresCachedLeafInfo(const FImplicitObject* Geometry) { return false; }
+
+		UE_DEPRECATED(5.3, "Call FShapeInstanceProxy::Make for game thread objects, FShapeInstance::Make for physics thread objects")
+		static TUniquePtr<FPerShapeData> CreatePerShapeData(int32 InShapeIdx, TSerializablePtr<FImplicitObject> InGeometry);
+
+		UE_DEPRECATED(5.3, "Call FShapeInstanceProxy::UpdateGeometry for game thread objects, FShapeInstance::UpdateGeometry for physics thread objects")
+		static void UpdateGeometry(TUniquePtr<FPerShapeData>& InOutShapePtr, TSerializablePtr<FImplicitObject> InGeometry);
+
+
+		static FPerShapeData* SerializationFactory(FChaosArchive& Ar, FPerShapeData*);
+
+		virtual ~FPerShapeData() {}
+
+		virtual void Serialize(FChaosArchive& Ar);
+
+		void UpdateShapeBounds(const FRigidTransform3& WorldTM, const FVec3& BoundsExpansion = FVec3(0));
+
+		void* GetUserData() const;
+		void SetUserData(void* InUserData);
+
+		const FCollisionFilterData& GetQueryData() const;
+		void SetQueryData(const FCollisionFilterData& InQueryData);
+
+		const FCollisionFilterData& GetSimData() const;
+		void SetSimData(const FCollisionFilterData& InSimData);
+
+		TSerializablePtr<FImplicitObject> GetGeometry() const;
+
+		const TAABB<FReal, 3>& GetWorldSpaceInflatedShapeBounds() const;
+
+		void UpdateWorldSpaceState(const FRigidTransform3& WorldTransform, const FVec3& BoundsExpansion);
+
+		// The leaf shape (with transformed and implicit wrapper removed).
+		const FImplicitObject* GetLeafGeometry() const;
+
+		// The actor-relative transform of the leaf geometry.
+		FRigidTransform3 GetLeafRelativeTransform() const;
+
+		// The world-space transform of the leaf geometry.
+		// If we have non-identity leaf relative transform, is cached from the last call to UpdateWorldSpaceState.
+		// If not cahced, is constructed from arguments.
+		FRigidTransform3 GetLeafWorldTransform(const FGeometryParticleHandle* Particle) const;
+
+		void UpdateLeafWorldTransform(FGeometryParticleHandle* Particle);
+
+		const TArray<FMaterialHandle>& GetMaterials() const;
+		void SetMaterial(FMaterialHandle InMaterial);
+		void SetMaterials(const TArray<FMaterialHandle>& InMaterials);
+
+		const TArray<FMaterialMaskHandle>& GetMaterialMasks() const;
+		void SetMaterialMasks(const TArray<FMaterialMaskHandle>& InMaterialMasks);
+
+		const TArray<uint32>& GetMaterialMaskMaps() const;
+		void SetMaterialMaskMaps(const TArray<uint32>& InMaterialMaskMaps);
+
+		const TArray<FMaterialHandle>& GetMaterialMaskMapMaterials() const;
+		void SetMaterialMaskMapMaterials(const TArray<FMaterialHandle>& InMaterialMaskMapMaterials);
+
+		const FShapeDirtyFlags GetDirtyFlags() const;
+
+		bool GetQueryEnabled() const;
+		void SetQueryEnabled(const bool bEnable);
+
+		bool GetSimEnabled() const;
+		void SetSimEnabled(const bool bEnable);
+
+		bool GetIsProbe() const;
+		void SetIsProbe(const bool bIsProbe);
+
+		EChaosCollisionTraceFlag GetCollisionTraceType() const;
+		void SetCollisionTraceType(const EChaosCollisionTraceFlag InTraceFlag);
+
+		const FCollisionData& GetCollisionData() const;
+		void SetCollisionData(const FCollisionData& Data);
+
+		const FMaterialData& GetMaterialData() const;
+		void SetMaterialData(const FMaterialData& Data);
+
+		void SyncRemoteData(FDirtyPropertiesManager& Manager, int32 ShapeDataIdx, FShapeDirtyData& RemoteData);
+		
+		void SetProxy(IPhysicsProxyBase* InProxy);
+		
+		int32 GetShapeIndex() const;
+		void ModifyShapeIndex(int32 NewShapeIndex);
+
+		template <typename Lambda> void ModifySimData(const Lambda& LambdaFunc);
+		template <typename Lambda> void ModifyMaterials(const Lambda& LambdaFunc);
+		template <typename Lambda> void ModifyMaterialMasks(const Lambda& LambdaFunc);
+		template <typename Lambda> void ModifyMaterialMaskMaps(const Lambda& LambdaFunc);
+		template <typename Lambda> void ModifyMaterialMaskMapMaterials(const Lambda& LambdaFunc);
+
+	protected:
+		FPerShapeData(const EPerShapeDataType InType, int32 InShapeIdx)
+			: Type(InType)
+			, ShapeIdx(InShapeIdx)
+			, Geometry()
+			, WorldSpaceInflatedShapeBounds(FAABB3(FVec3(0), FVec3(0)))
+		{
+		}
+
+		FPerShapeData(const EPerShapeDataType InType, int32 InShapeIdx, TSerializablePtr<FImplicitObject> InGeometry)
+			: Type(InType)
+			, ShapeIdx(InShapeIdx)
+			, Geometry(InGeometry)
+			, WorldSpaceInflatedShapeBounds(FAABB3(FVec3(0), FVec3(0)))
+		{
+		}
+
+		FPerShapeData(const EPerShapeDataType InType, const FPerShapeData& Other)
+			: Type(InType)
+			, ShapeIdx(Other.ShapeIdx)
+			, Geometry(Other.Geometry)
+			, WorldSpaceInflatedShapeBounds(Other.WorldSpaceInflatedShapeBounds)
+		{
+		}
+
+		EPerShapeDataType Type;
+		int32 ShapeIdx;
+		TSerializablePtr<FImplicitObject> Geometry;
+		TAABB<FReal, 3> WorldSpaceInflatedShapeBounds;
+	};
+
+
+	///////////////////////////////////////////////////////////////////////////////////////////////
+	///////////////////////////////////////////////////////////////////////////////////////////////
+	// 
+	// FShapeInstanceProxy
+	// 
+	///////////////////////////////////////////////////////////////////////////////////////////////
+	///////////////////////////////////////////////////////////////////////////////////////////////
+
+	/*
+	 * NOTE: FShapeInstanceProxy is a Game-Thread object. 
+	 * See FShapeInstance for the physics-thread equivalent
+	 * 
+	 * FShapeInstanceProxy contains the per-shape data associated with a single shape on a particle. 
+	 * This contains data like the collision / query filters, material properties etc.
+	 *
+	 * Every particle holds one FShapeInstanceProxy object for each geometry they use.
+	 * If the particle has a Union of geometries there will be one FShapeInstanceProxy
+	 * for each geometry in the union. (Except ClusteredUnions - these are not flattened
+	 * because they contain their own query acceleration structure.)
+	 *
+	 * NOTE: keep size to a minimum. There can be millions of these in s scene.
+	 *
+	 * @todo(chaos) : try to remove the GT Proxy pointer - this could easily be passed into the
+	 * relevant functions instead.
+	 *
+	 * @todo(chaos) : reduce the cost of MaterialData for shapes with one materialand no masks etc.
+	 */
+	class CHAOS_API FShapeInstanceProxy : public FPerShapeData
+	{
+	public:
+		friend class FPerShapeData;
+
+		static TUniquePtr<FShapeInstanceProxy> Make(int32 InShapeIdx, TSerializablePtr<FImplicitObject> InGeometry);
+		static void UpdateGeometry(TUniquePtr<FShapeInstanceProxy>& InOutShapePtr, TSerializablePtr<FImplicitObject> InGeometry);
+		static FShapeInstanceProxy* SerializationFactory(FChaosArchive& Ar, FShapeInstanceProxy*);
+
+		void UpdateShapeBounds(const FRigidTransform3& WorldTM, const FVec3& BoundsExpansion = FVec3(0));
+
+		void* GetUserData() const { return CollisionData.Read().UserData; }
+		void SetUserData(void* InUserData)
+		{
+			CollisionData.Modify(true, DirtyFlags, Proxy, ShapeIdx, [InUserData](FCollisionData& Data) { Data.UserData = InUserData; });
+		}
+
+		const FCollisionFilterData& GetQueryData() const { return CollisionData.Read().QueryData; }
+		void SetQueryData(const FCollisionFilterData& InQueryData)
+		{
+			CollisionData.Modify(true, DirtyFlags, Proxy, ShapeIdx, [InQueryData](FCollisionData& Data) { Data.QueryData = InQueryData; });
+		}
+
+		const FCollisionFilterData& GetSimData() const { return CollisionData.Read().SimData; }
+		void SetSimData(const FCollisionFilterData& InSimData)
+		{
+			CollisionData.Modify(true, DirtyFlags, Proxy, ShapeIdx, [InSimData](FCollisionData& Data) { Data.SimData = InSimData; });
+		}
+
+		void UpdateWorldSpaceState(const FRigidTransform3& WorldTransform, const FVec3& BoundsExpansion);
+
+		// The leaf shape (with transformed and implicit wrapper removed).
+		const FImplicitObject* GetLeafGeometry() const;
+
+		// The actor-relative transform of the leaf geometry.
+		FRigidTransform3 GetLeafRelativeTransform() const;
+
+		// The world-space transform of the leaf geometry.
+		// If we have non-identity leaf relative transform, is cached from the last call to UpdateWorldSpaceState.
+		// If not cahced, is constructed from arguments.
+		FRigidTransform3 GetLeafWorldTransform(const FGeometryParticleHandle* Particle) const;
+		void UpdateLeafWorldTransform(FGeometryParticleHandle* Particle);
+
+		const TArray<FMaterialHandle>& GetMaterials() const { return Materials.Read().Materials; }
+		const TArray<FMaterialMaskHandle>& GetMaterialMasks() const { return Materials.Read().MaterialMasks; }
+		const TArray<uint32>& GetMaterialMaskMaps() const { return Materials.Read().MaterialMaskMaps; }
+		const TArray<FMaterialHandle>& GetMaterialMaskMapMaterials() const { return Materials.Read().MaterialMaskMapMaterials; }
+
+		const FShapeDirtyFlags GetDirtyFlags() const { return DirtyFlags; }
+
+		void SetMaterial(FMaterialHandle InMaterial)
+		{
+			Materials.Modify(true, DirtyFlags, Proxy, ShapeIdx, [InMaterial](FMaterialData& Data)
+				{
+					Data.Materials.Reset(1);
+					Data.Materials.Add(InMaterial);
+				});
+		}
+
+		void SetMaterials(const TArray<FMaterialHandle>& InMaterials)
+		{
+			Materials.Modify(true, DirtyFlags, Proxy, ShapeIdx, [&InMaterials](FMaterialData& Data)
+				{
+					Data.Materials = InMaterials;
+				});
+		}
+
+		void SetMaterialMasks(const TArray<FMaterialMaskHandle>& InMaterialMasks)
+		{
+			Materials.Modify(true, DirtyFlags, Proxy, ShapeIdx, [&InMaterialMasks](FMaterialData& Data)
+				{
+					Data.MaterialMasks = InMaterialMasks;
+				});
+		}
+
+		void SetMaterialMaskMaps(const TArray<uint32>& InMaterialMaskMaps)
+		{
+			Materials.Modify(true, DirtyFlags, Proxy, ShapeIdx, [&InMaterialMaskMaps](FMaterialData& Data)
+				{
+					Data.MaterialMaskMaps = InMaterialMaskMaps;
+				});
+		}
+
+		void SetMaterialMaskMapMaterials(const TArray<FMaterialHandle>& InMaterialMaskMapMaterials)
+		{
+			Materials.Modify(true, DirtyFlags, Proxy, ShapeIdx, [&InMaterialMaskMapMaterials](FMaterialData& Data)
+				{
+					Data.MaterialMaskMapMaterials = InMaterialMaskMapMaterials;
+				});
+		}
+
+		bool GetQueryEnabled() const { return CollisionData.Read().bQueryCollision; }
+		void SetQueryEnabled(const bool bEnable)
+		{
+			CollisionData.Modify(true, DirtyFlags, Proxy, ShapeIdx, [bEnable](FCollisionData& Data) { Data.bQueryCollision = bEnable; });
+		}
+
+		bool GetSimEnabled() const { return CollisionData.Read().bSimCollision; }
+		void SetSimEnabled(const bool bEnable)
+		{
+			CollisionData.Modify(true, DirtyFlags, Proxy, ShapeIdx, [bEnable](FCollisionData& Data) { Data.bSimCollision = bEnable; });
+		}
+
+		bool GetIsProbe() const { return CollisionData.Read().bIsProbe; }
+		void SetIsProbe(const bool bIsProbe)
+		{
+			CollisionData.Modify(true, DirtyFlags, Proxy, ShapeIdx, [bIsProbe](FCollisionData& Data) { Data.bIsProbe = bIsProbe; });
+		}
+
+		EChaosCollisionTraceFlag GetCollisionTraceType() const { return CollisionData.Read().CollisionTraceType; }
+		void SetCollisionTraceType(const EChaosCollisionTraceFlag InTraceFlag)
+		{
+			CollisionData.Modify(true, DirtyFlags, Proxy, ShapeIdx, [InTraceFlag](FCollisionData& Data) { Data.CollisionTraceType = InTraceFlag; });
+		}
+
+		const FCollisionData& GetCollisionData() const { return CollisionData.Read(); }
+
+		void SetCollisionData(const FCollisionData& Data)
+		{
+			CollisionData.Write(Data, true, DirtyFlags, Proxy, ShapeIdx);
+		}
+
+		const FMaterialData& GetMaterialData() const { return Materials.Read(); }
+
+		void SetMaterialData(const FMaterialData& Data)
+		{
+			Materials.Write(Data, true, DirtyFlags, Proxy, ShapeIdx);
+		}
+
+		void SyncRemoteData(FDirtyPropertiesManager& Manager, int32 ShapeDataIdx, FShapeDirtyData& RemoteData)
+		{
+			RemoteData.SetFlags(DirtyFlags);
+			CollisionData.SyncRemote(Manager, ShapeDataIdx, RemoteData);
+			Materials.SyncRemote(Manager, ShapeDataIdx, RemoteData);
+			DirtyFlags.Clear();
+		}
+
+		void SetProxy(IPhysicsProxyBase* InProxy)
+		{
+			Proxy = InProxy;
+			if (Proxy)
+			{
+				if (DirtyFlags.IsDirty())
+				{
+					if (FPhysicsSolverBase* PhysicsSolverBase = Proxy->GetSolver<FPhysicsSolverBase>())
+					{
+						PhysicsSolverBase->AddDirtyProxyShape(Proxy, ShapeIdx);
+					}
+				}
+			}
+		}
+
+		template <typename Lambda>
+		void ModifySimData(const Lambda& LambdaFunc)
+		{
+			CollisionData.Modify(true, DirtyFlags, Proxy, ShapeIdx, [&LambdaFunc](FCollisionData& Data) { LambdaFunc(Data.SimData); });
+		}
+
+		template <typename Lambda>
+		void ModifyMaterials(const Lambda& LambdaFunc)
+		{
+			Materials.Modify(true, DirtyFlags, Proxy, ShapeIdx, [&LambdaFunc](FMaterialData& Data)
+				{
+					LambdaFunc(Data.Materials);
+				});
+		}
+
+		template <typename Lambda>
+		void ModifyMaterialMasks(const Lambda& LambdaFunc)
+		{
+			Materials.Modify(true, DirtyFlags, Proxy, ShapeIdx, [&LambdaFunc](FMaterialData& Data)
+				{
+					LambdaFunc(Data.MaterialMasks);
+				});
+		}
+
+		template <typename Lambda>
+		void ModifyMaterialMaskMaps(const Lambda& LambdaFunc)
+		{
+			Materials.Modify(true, DirtyFlags, Proxy, ShapeIdx, [&LambdaFunc](FMaterialData& Data)
+				{
+					LambdaFunc(Data.MaterialMaskMaps);
+				});
+		}
+
+		template <typename Lambda>
+		void ModifyMaterialMaskMapMaterials(const Lambda& LambdaFunc)
+		{
+			Materials.Modify(true, DirtyFlags, Proxy, ShapeIdx, [&LambdaFunc](FMaterialData& Data)
+				{
+					LambdaFunc(Data.MaterialMaskMapMaterials);
+				});
+		}
+
+	protected:
+		explicit FShapeInstanceProxy(int32 InShapeIdx)
+			: FPerShapeData(EPerShapeDataType::Proxy, InShapeIdx)
+			, Proxy(nullptr)
+			, DirtyFlags()
+			, CollisionData()
+			, Materials()
+		{
+		}
+
+		FShapeInstanceProxy(int32 InShapeIdx, TSerializablePtr<FImplicitObject> InGeometry)
+			: FPerShapeData(EPerShapeDataType::Proxy, InShapeIdx, InGeometry)
+			, Proxy(nullptr)
+			, DirtyFlags()
+			, CollisionData()
+			, Materials()
+		{
+		}
+
+		IPhysicsProxyBase* Proxy;
+		FShapeDirtyFlags DirtyFlags;
+
+		TShapeProperty<FCollisionData, EShapeProperty::CollisionData> CollisionData;
+		TShapeProperty<FMaterialData, EShapeProperty::Materials> Materials;
+	};
+
+
+	///////////////////////////////////////////////////////////////////////////////////////////////
+	///////////////////////////////////////////////////////////////////////////////////////////////
+	// 
+	// FShapeInstance
+	// 
+	///////////////////////////////////////////////////////////////////////////////////////////////
+	///////////////////////////////////////////////////////////////////////////////////////////////
+
+
+	/*
+	 * NOTE: FShapeInstance is a Physics-Thread object.
+	 * See FShapeInstanceProxy for the game-thread equivalent
+	 *
+	 * FShapeInstance contains the per-shape data associated with a single shape on a particle.
+	 * This contains data like the collision / query filters, material properties etc.
+	 *
+	 * Every particle holds one FShapeInstance object for each geometry they use.
+	 * If the particle has a Union of geometries there will be one FShapeInstance
+	 * for each geometry in the union. (Except ClusteredUnions - these are not flattened
+	 * because they contain their own query acceleration structure.)
+	 *
+	 * NOTE: keep size to a minimum. There can be millions of these in s scene.
+	 *
+	 * @todo(chaos) : reduce the cost of MaterialData for shapes with one materialand no masks etc.
+	 */
+	class CHAOS_API FShapeInstance : public FPerShapeData
+	{
+	public:
+		friend class FPerShapeData;
+
+		static TUniquePtr<FShapeInstance> Make(int32 InShapeIdx, TSerializablePtr<FImplicitObject> InGeometry);
+		static void UpdateGeometry(TUniquePtr<FShapeInstance>& InOutShapePtr, TSerializablePtr<FImplicitObject> InGeometry);
+		static FShapeInstance* SerializationFactory(FChaosArchive& Ar, FShapeInstance*);
+
+		void UpdateShapeBounds(const FRigidTransform3& WorldTM, const FVec3& BoundsExpansion = FVec3(0));
+
+		void* GetUserData() const { return CollisionData.UserData; }
+		void SetUserData(void* InUserData) { CollisionData.UserData = InUserData; }
+
+		const FCollisionFilterData& GetQueryData() const { return CollisionData.QueryData; }
+		void SetQueryData(const FCollisionFilterData& InQueryData) { CollisionData.QueryData = InQueryData; }
+
+		const FCollisionFilterData& GetSimData() const { return CollisionData.SimData; }
+		void SetSimData(const FCollisionFilterData& InSimData) { CollisionData.SimData = InSimData; }
+
+		void UpdateWorldSpaceState(const FRigidTransform3& WorldTransform, const FVec3& BoundsExpansion);
+
+		// The leaf shape (with transformed and implicit wrapper removed).
+		const FImplicitObject* GetLeafGeometry() const;
+
+		// The actor-relative transform of the leaf geometry.
+		FRigidTransform3 GetLeafRelativeTransform() const;
+
+		// The world-space transform of the leaf geometry.
+		// If we have non-identity leaf relative transform, is cached from the last call to UpdateWorldSpaceState.
+		// If not cahced, is constructed from arguments.
+		FRigidTransform3 GetLeafWorldTransform(const FGeometryParticleHandle* Particle) const;
+		void UpdateLeafWorldTransform(FGeometryParticleHandle* Particle);
+
+		const TArray<FMaterialHandle>& GetMaterials() const { return Materials.Materials; }
+		const TArray<FMaterialMaskHandle>& GetMaterialMasks() const { return Materials.MaterialMasks; }
+		const TArray<uint32>& GetMaterialMaskMaps() const { return Materials.MaterialMaskMaps; }
+		const TArray<FMaterialHandle>& GetMaterialMaskMapMaterials() const { return Materials.MaterialMaskMapMaterials; }
+
+		void SetMaterial(FMaterialHandle InMaterial) { Materials.Materials.Reset(1); Materials.Materials.Add(InMaterial); }
+		void SetMaterials(const TArray<FMaterialHandle>& InMaterials) { Materials.Materials = InMaterials; }
+		void SetMaterialMasks(const TArray<FMaterialMaskHandle>& InMaterialMasks) { Materials.MaterialMasks = InMaterialMasks; }
+		void SetMaterialMaskMaps(const TArray<uint32>& InMaterialMaskMaps) { Materials.MaterialMaskMaps = InMaterialMaskMaps; }
+		void SetMaterialMaskMapMaterials(const TArray<FMaterialHandle>& InMaterialMaskMapMaterials) { Materials.MaterialMaskMapMaterials = InMaterialMaskMapMaterials; }
+
+		bool GetQueryEnabled() const { return CollisionData.bQueryCollision; }
+		void SetQueryEnabled(const bool bEnable) { CollisionData.bQueryCollision = bEnable; }
+
+		bool GetSimEnabled() const { return CollisionData.bSimCollision; }
+		void SetSimEnabled(const bool bEnable) { CollisionData.bSimCollision = bEnable; }
+
+		bool GetIsProbe() const { return CollisionData.bIsProbe; }
+		void SetIsProbe(const bool bIsProbe) { CollisionData.bIsProbe = bIsProbe; }
+
+		EChaosCollisionTraceFlag GetCollisionTraceType() const { return CollisionData.CollisionTraceType; }
+		void SetCollisionTraceType(const EChaosCollisionTraceFlag InTraceFlag) { CollisionData.CollisionTraceType = InTraceFlag; }
+
+		const FCollisionData& GetCollisionData() const { return CollisionData; }
+		void SetCollisionData(const FCollisionData& Data) { CollisionData = Data; }
+
+		const FMaterialData& GetMaterialData() const { return Materials; }
+		void SetMaterialData(const FMaterialData& Data) { Materials = Data; }
+
+		// @todo(chaos): remove when FPerShapeData is removed
+		const FShapeDirtyFlags GetDirtyFlags() const { check(false); return FShapeDirtyFlags(); }
+		void SyncRemoteData(FDirtyPropertiesManager& Manager, int32 ShapeDataIdx, FShapeDirtyData& RemoteData) { check(false); }
+		void SetProxy(IPhysicsProxyBase* InProxy) { check(false); }
+
+		template <typename Lambda> void ModifySimData(const Lambda& LambdaFunc) { LambdaFunc(CollisionData.SimData); }
+		template <typename Lambda> void ModifyMaterials(const Lambda& LambdaFunc) { LambdaFunc(Materials.Materials); }
+		template <typename Lambda> void ModifyMaterialMasks(const Lambda& LambdaFunc) { LambdaFunc(Materials.MaterialMasks); }
+		template <typename Lambda> void ModifyMaterialMaskMaps(const Lambda& LambdaFunc) { LambdaFunc(Materials.MaterialMaskMaps); }
+		template <typename Lambda> void ModifyMaterialMaskMapMaterials(const Lambda& LambdaFunc) { LambdaFunc(Materials.MaterialMaskMapMaterials); }
+
+	protected:
+		explicit FShapeInstance(int32 InShapeIdx)
+			: FPerShapeData(EPerShapeDataType::Sim, InShapeIdx)
+			, CollisionData()
+			, Materials()
+		{
+		}
+
+		FShapeInstance(int32 InShapeIdx, TSerializablePtr<FImplicitObject> InGeometry)
+			: FPerShapeData(EPerShapeDataType::Sim, InShapeIdx, InGeometry)
+			, CollisionData()
+			, Materials()
+		{
+		}
+
+		explicit FShapeInstance(FShapeInstance&& Other)
+			: FPerShapeData(EPerShapeDataType::Sim, Other)
+			, CollisionData(MoveTemp(Other.CollisionData))
+			, Materials(MoveTemp(Other.Materials))
+		{
+		}
+
+		FShapeInstance(const EPerShapeDataType InType, int32 InShapeIdx, TSerializablePtr<FImplicitObject> InGeometry)
+			: FPerShapeData(InType, InShapeIdx, InGeometry)
+			, CollisionData()
+			, Materials()
+		{
+		}
+
+		FShapeInstance(const EPerShapeDataType InType, FShapeInstance&& Other)
+			: FPerShapeData(InType, Other)
+			, CollisionData(MoveTemp(Other.CollisionData))
+			, Materials(MoveTemp(Other.Materials))
+		{
+		}
+
+		FCollisionData CollisionData;
+		FMaterialData Materials;
+	};
+
+	namespace Private
+	{
+
+		///////////////////////////////////////////////////////////////////////////////////////////////
+		///////////////////////////////////////////////////////////////////////////////////////////////
+		// 
+		// FShapeInstanceExtended
+		// 
+		///////////////////////////////////////////////////////////////////////////////////////////////
+		///////////////////////////////////////////////////////////////////////////////////////////////
+
+		/**
+		 * An extended version of FShapeInstance (physics-thread shape instance data) that caches
+		 * some world-space state of the shape for use in collision detection. This extended data
+		 * if only required for shapes that have a transform relative to the particle they are
+		 * attached to. It helps by avoiding the need to recalculate the shape transform every
+		 * time it is needed in collision detection, which is once for each other shape we
+		 * may be in contact with.
+		 * 
+		 * NOTE: keep size to a minimum. There can be millions of these in s scene.
+		 */
+		class CHAOS_API FShapeInstanceExtended : public FShapeInstance
+		{
+		public:
+			friend class FShapeInstance;
+
+			FRigidTransform3 GetWorldTransform() const
+			{
+				return FRigidTransform3(WorldPosition, WorldRotation);
+			}
+
+			void SetWorldTransform(const FRigidTransform3& LeafWorldTransform)
+			{
+				WorldPosition = LeafWorldTransform.GetTranslation();
+				WorldRotation = LeafWorldTransform.GetRotation();
+			}
+
+		protected:
+			FShapeInstanceExtended(int32 InShapeIdx, TSerializablePtr<FImplicitObject> InGeometry)
+				: FShapeInstance(EPerShapeDataType::SimExtended, InShapeIdx, InGeometry)
+			{
+			}
+
+			FShapeInstanceExtended(FShapeInstance&& PerShapeData)
+				: FShapeInstance(EPerShapeDataType::SimExtended, MoveTemp(PerShapeData))
+			{
+			}
+
+			FVec3 WorldPosition;
+			FRotation3 WorldRotation;
+		};
+	}
+
+
+	///////////////////////////////////////////////////////////////////////////////////////////////
+	///////////////////////////////////////////////////////////////////////////////////////////////
+	// 
+	// Downcasts
+	// 
+	///////////////////////////////////////////////////////////////////////////////////////////////
+	///////////////////////////////////////////////////////////////////////////////////////////////
+
+	inline FShapeInstanceProxy* FPerShapeData::AsShapeInstanceProxy()
+	{
+		if (Type == EPerShapeDataType::Proxy)
+		{
+			return static_cast<FShapeInstanceProxy*>(this);
+		}
+		return nullptr;
+	}
+	
+	inline const FShapeInstanceProxy* FPerShapeData::AsShapeInstanceProxy() const
+	{
+		if (Type == EPerShapeDataType::Proxy)
+		{
+			return static_cast<const FShapeInstanceProxy*>(this);
+		}
+		return nullptr;
+	}
+
+	inline FShapeInstance* FPerShapeData::AsShapeInstance()
+	{
+		if (Type != EPerShapeDataType::Proxy)
+		{
+			return static_cast<FShapeInstance*>(this);
+		}
+		return nullptr;
+	}
+
+	inline const FShapeInstance* FPerShapeData::AsShapeInstance() const
+	{
+		if (Type != EPerShapeDataType::Proxy)
+		{
+			return static_cast<const FShapeInstance*>(this);
+		}
+		return nullptr;
+	}
+
+	inline Private::FShapeInstanceExtended* FPerShapeData::AsShapeInstanceExtended()
+	{
+		if (Type == EPerShapeDataType::SimExtended)
+		{
+			return static_cast<Private::FShapeInstanceExtended*>(this);
+		}
+		return nullptr;
+	}
+
+	inline const Private::FShapeInstanceExtended* FPerShapeData::AsShapeInstanceExtended() const
+	{
+		if (Type == EPerShapeDataType::SimExtended)
+		{
+			return static_cast<const Private::FShapeInstanceExtended*>(this);
+		}
+		return nullptr;
+	}
+
+	template<typename TLambda>
+	inline decltype(auto) FPerShapeData::DownCast(const TLambda& Lambda)
+	{
+		// NOTE: only FShapeInstanceProxy and FShapeInstance implement the full interface
+		// FShapeInstanceExtended is a hidden derived type as far as we are concerned here
+		if (Type == EPerShapeDataType::Proxy)
+		{
+			return Lambda(*AsShapeInstanceProxy());
+		}
+		else
+		{
+			return Lambda(*AsShapeInstance());
+		}
+	}
+
+	template<typename TLambda>
+	inline  decltype(auto) FPerShapeData::DownCast(const TLambda& Lambda) const
+	{
+		// See comments in non-const DownCast()
+		if (Type == EPerShapeDataType::Proxy)
+		{
+			return Lambda(*AsShapeInstanceProxy());
+		}
+		else
+		{
+			return Lambda(*AsShapeInstance());
+		}
+	}
+
+
+	///////////////////////////////////////////////////////////////////////////////////////////////
+	///////////////////////////////////////////////////////////////////////////////////////////////
+	// 
+	// FPerShapeData implementation
+	// 
+	///////////////////////////////////////////////////////////////////////////////////////////////
+	///////////////////////////////////////////////////////////////////////////////////////////////
+
+
+	inline void FPerShapeData::UpdateShapeBounds(const FRigidTransform3& WorldTM, const FVec3& BoundsExpansion)
+	{
+		DownCast([&WorldTM, &BoundsExpansion](auto& ShapeInstance) { ShapeInstance.UpdateShapeBounds(WorldTM, BoundsExpansion); });
+	}
+
+	inline void* FPerShapeData::GetUserData() const
+	{
+		return DownCast([](auto& ShapeInstance) { return ShapeInstance.GetUserData(); });
+	}
+
+	inline void FPerShapeData::SetUserData(void* InUserData)
+	{
+		DownCast([InUserData](auto& ShapeInstance) { ShapeInstance.SetUserData(InUserData); });
+	}
+
+	inline const FCollisionFilterData& FPerShapeData::GetQueryData() const
+	{
+		return DownCast([](const auto& ShapeInstance) -> const auto& { return ShapeInstance.GetQueryData(); });
+	}
+
+	inline void FPerShapeData::SetQueryData(const FCollisionFilterData& InQueryData)
+	{
+		DownCast([&InQueryData](auto& ShapeInstance) { ShapeInstance.SetQueryData(InQueryData); });
+	}
+
+	inline const FCollisionFilterData& FPerShapeData::GetSimData() const
+	{
+		return DownCast([](const auto& ShapeInstance) -> const auto& { return ShapeInstance.GetSimData(); });
+	}
+
+	inline void FPerShapeData::SetSimData(const FCollisionFilterData& InSimData)
+	{
+		return DownCast([&InSimData](auto& ShapeInstance) { ShapeInstance.SetSimData(InSimData); });
+	}
+
+	inline TSerializablePtr<FImplicitObject> FPerShapeData::GetGeometry() const
+	{
+		return Geometry;
+	}
+
+	inline const TAABB<FReal, 3>& FPerShapeData::GetWorldSpaceInflatedShapeBounds() const
+	{
+		return WorldSpaceInflatedShapeBounds;
+	}
+
+	inline void FPerShapeData::UpdateWorldSpaceState(const FRigidTransform3& WorldTransform, const FVec3& BoundsExpansion)
+	{
+		DownCast([&WorldTransform, &BoundsExpansion](auto& ShapeInstance) { ShapeInstance.UpdateWorldSpaceState(WorldTransform, BoundsExpansion); });
+	}
+
+	inline const FImplicitObject* FPerShapeData::GetLeafGeometry() const
+	{
+		return DownCast([](auto& ShapeInstance) { return ShapeInstance.GetLeafGeometry(); });
+	}
+
+	inline FRigidTransform3 FPerShapeData::GetLeafRelativeTransform() const
+	{
+		return DownCast([](auto& ShapeInstance) { return ShapeInstance.GetLeafRelativeTransform(); });
+	}
+
+	inline FRigidTransform3 FPerShapeData::GetLeafWorldTransform(const FGeometryParticleHandle* Particle) const
+	{
+		return DownCast([Particle](auto& ShapeInstance) { return ShapeInstance.GetLeafWorldTransform(Particle); });
+	}
+
+	inline void FPerShapeData::UpdateLeafWorldTransform(FGeometryParticleHandle* Particle)
+	{
+		DownCast([Particle](auto& ShapeInstance) { ShapeInstance.UpdateLeafWorldTransform(Particle); });
+	}
+
+	inline const TArray<FMaterialHandle>& FPerShapeData::GetMaterials() const
+	{
+		return DownCast([](const auto& ShapeInstance) -> const auto& { return ShapeInstance.GetMaterials(); });
+	}
+
+	inline const TArray<FMaterialMaskHandle>& FPerShapeData::GetMaterialMasks() const
+	{
+		return DownCast([](const auto& ShapeInstance) -> const auto& { return ShapeInstance.GetMaterialMasks(); });
+	}
+
+	inline const TArray<uint32>& FPerShapeData::GetMaterialMaskMaps() const
+	{
+		return DownCast([](const auto& ShapeInstance) -> const auto& { return ShapeInstance.GetMaterialMaskMaps(); });
+	}
+
+	inline const TArray<FMaterialHandle>& FPerShapeData::GetMaterialMaskMapMaterials() const
+	{
+		return DownCast([](const auto& ShapeInstance) -> const auto& { return ShapeInstance.GetMaterialMaskMapMaterials(); });
+	}
+
+	inline const FShapeDirtyFlags FPerShapeData::GetDirtyFlags() const
+	{
+		return DownCast([](auto& ShapeInstance) { return ShapeInstance.GetDirtyFlags(); });
+	}
+
+	inline void FPerShapeData::SetMaterial(FMaterialHandle InMaterial)
+	{
+		DownCast([&InMaterial](auto& ShapeInstance) { ShapeInstance.SetMaterial(InMaterial); });
+	}
+
+	inline void FPerShapeData::SetMaterials(const TArray<FMaterialHandle>& InMaterials)
+	{
+		DownCast([&InMaterials](auto& ShapeInstance) { ShapeInstance.SetMaterials(InMaterials); });
+	}
+
+	inline void FPerShapeData::SetMaterialMasks(const TArray<FMaterialMaskHandle>& InMaterialMasks)
+	{
+		DownCast([&InMaterialMasks](auto& ShapeInstance) { ShapeInstance.SetMaterialMasks(InMaterialMasks); });
+	}
+
+	inline void FPerShapeData::SetMaterialMaskMaps(const TArray<uint32>& InMaterialMaskMaps)
+	{
+		DownCast([&InMaterialMaskMaps](auto& ShapeInstance) { ShapeInstance.SetMaterialMaskMaps(InMaterialMaskMaps); });
+	}
+
+	inline void FPerShapeData::SetMaterialMaskMapMaterials(const TArray<FMaterialHandle>& InMaterialMaskMapMaterials)
+	{
+		DownCast([&InMaterialMaskMapMaterials](auto& ShapeInstance) { ShapeInstance.SetMaterialMaskMapMaterials(InMaterialMaskMapMaterials); });
+	}
+
+	inline bool FPerShapeData::GetQueryEnabled() const
+	{
+		return DownCast([](auto& ShapeInstance) { return ShapeInstance.GetQueryEnabled(); });
+	}
+
+	inline void FPerShapeData::SetQueryEnabled(const bool bEnable)
+	{
+		DownCast([bEnable](auto& ShapeInstance) { ShapeInstance.SetQueryEnabled(bEnable); });
+	}
+
+	inline bool FPerShapeData::GetSimEnabled() const
+	{
+		return DownCast([](auto& ShapeInstance) { return ShapeInstance.GetSimEnabled(); });
+	}
+
+	inline void FPerShapeData::SetSimEnabled(const bool bEnable)
+	{
+		DownCast([bEnable](auto& ShapeInstance) { ShapeInstance.SetSimEnabled(bEnable); });
+	}
+
+	inline bool FPerShapeData::GetIsProbe() const
+	{
+		return DownCast([](auto& ShapeInstance) { return ShapeInstance.GetIsProbe(); });
+	}
+
+	inline void FPerShapeData::SetIsProbe(const bool bIsProbe)
+	{
+		DownCast([bIsProbe](auto& ShapeInstance) { ShapeInstance.SetIsProbe(bIsProbe); });
+	}
+
+	inline EChaosCollisionTraceFlag FPerShapeData::GetCollisionTraceType() const
+	{
+		return DownCast([](auto& ShapeInstance) { return ShapeInstance.GetCollisionTraceType(); });
+	}
+
+	inline void FPerShapeData::SetCollisionTraceType(const EChaosCollisionTraceFlag InTraceFlag)
+	{
+		DownCast([InTraceFlag](auto& ShapeInstance) { ShapeInstance.SetCollisionTraceType(InTraceFlag); });
+	}
+
+	inline const FCollisionData& FPerShapeData::GetCollisionData() const
+	{
+		return DownCast([](const auto& ShapeInstance) -> const auto& { return ShapeInstance.GetCollisionData(); });
+	}
+
+	inline void FPerShapeData::SetCollisionData(const FCollisionData& Data)
+	{
+		DownCast([&Data](auto& ShapeInstance) { ShapeInstance.SetCollisionData(Data); });
+	}
+
+	inline const FMaterialData& FPerShapeData::GetMaterialData() const
+	{
+		return DownCast([](auto& ShapeInstance) -> const auto& { return ShapeInstance.GetMaterialData(); });
+	}
+
+	inline void FPerShapeData::SetMaterialData(const FMaterialData& Data)
+	{
+		DownCast([&Data](auto& ShapeInstance) { ShapeInstance.SetMaterialData(Data); });
+	}
+
+	inline void FPerShapeData::SyncRemoteData(FDirtyPropertiesManager& Manager, int32 ShapeDataIdx, FShapeDirtyData& RemoteData)
+	{
+		DownCast([&Manager, ShapeDataIdx, &RemoteData](auto& ShapeInstance) { ShapeInstance.SyncRemoteData(Manager, ShapeDataIdx, RemoteData); });
+	}
+
+	inline void FPerShapeData::SetProxy(IPhysicsProxyBase* InProxy)
+	{
+		DownCast([InProxy](auto& ShapeInstance) { ShapeInstance.SetProxy(InProxy); });
+	}
+
+	inline int32 FPerShapeData::GetShapeIndex() const
+	{
+		return ShapeIdx;
+	}
+
+	inline void FPerShapeData::ModifyShapeIndex(int32 NewShapeIndex)
+	{
+		ShapeIdx = NewShapeIndex;
+	}
+
+	template <typename Lambda> 
+	void FPerShapeData::ModifySimData(const Lambda& LambdaFunc)
+	{
+		DownCast([&LambdaFunc](auto& ShapeInstance) { ShapeInstance.ModifySimData(LambdaFunc); });
+	}
+
+	template <typename Lambda> 
+	void FPerShapeData::ModifyMaterials(const Lambda& LambdaFunc)
+	{
+		DownCast([&LambdaFunc](auto& ShapeInstance) { ShapeInstance.ModifyMaterials(LambdaFunc); });
+	}
+
+	template <typename Lambda>
+	void FPerShapeData::ModifyMaterialMasks(const Lambda& LambdaFunc)
+	{
+		DownCast([&LambdaFunc](auto& ShapeInstance) { ShapeInstance.ModifyMaterialMasks(LambdaFunc); });
+	}
+
+	template <typename Lambda> 
+	void FPerShapeData::ModifyMaterialMaskMaps(const Lambda& LambdaFunc)
+	{
+		DownCast([&LambdaFunc](auto& ShapeInstance) { ShapeInstance.ModifyMaterialMaskMaps(LambdaFunc); });
+	}
+
+	template <typename Lambda> 
+	void FPerShapeData::ModifyMaterialMaskMapMaterials(const Lambda& LambdaFunc)
+	{
+		DownCast([&LambdaFunc](auto& ShapeInstance) { ShapeInstance.ModifyMaterialMaskMapMaterials(LambdaFunc); });
+	}
+
+
+	///////////////////////////////////////////////////////////////////////////////////////////////
+	///////////////////////////////////////////////////////////////////////////////////////////////
+	// 
+	// Misc stuff
+	// 
+	///////////////////////////////////////////////////////////////////////////////////////////////
+	///////////////////////////////////////////////////////////////////////////////////////////////
+
+
+	inline FChaosArchive& operator<<(FChaosArchive& Ar, FPerShapeData& Shape)
+	{
+		Shape.Serialize(Ar);
+		return Ar;
+	}
+
+	UE_DEPRECATED(5.3, "Not for external use")
+	void CHAOS_API UpdateShapesArrayFromGeometry(FShapesArray& ShapesArray, TSerializablePtr<FImplicitObject> Geometry, const FRigidTransform3& ActorTM, IPhysicsProxyBase* Proxy);
+
+}
\ No newline at end of file
diff --git a/Engine/Source/Runtime/Experimental/Chaos/Public/Chaos/ShapeInstanceFwd.h b/Engine/Source/Runtime/Experimental/Chaos/Public/Chaos/ShapeInstanceFwd.h
new file mode 100644
--- /dev/null
+++ b/Engine/Source/Runtime/Experimental/Chaos/Public/Chaos/ShapeInstanceFwd.h
@@ -0,0 +1,32 @@
+// Copyright Epic Games, Inc. All Rights Reserved.
+#pragma once
+
+#include "Templates/SharedPointer.h"
+#include "Templates/UniquePtr.h"
+
+// @todo(chaos): use TRefCountPtr rather than TSharedPtr below - it is smaller (one atomic int instead of 2 atomic ints and 2 pointers)
+
+namespace Chaos
+{
+	class FImplicitObject;
+	using FImplicitObjectPtr = TSharedPtr<FImplicitObject, ESPMode::ThreadSafe>;
+	using FConstImplicitObjectPtr = TSharedPtr<const FImplicitObject, ESPMode::ThreadSafe>;
+	using FImplicitObjectsArray = TArray<FImplicitObjectPtr>;
+
+	// Legacy names
+	class FPerShapeData;
+	using FShapesArray = TArray<TUniquePtr<FPerShapeData>, TInlineAllocator<1>>;;
+
+	// Game thread ShapeInstance
+	class FShapeInstanceProxy;
+	using FShapeInstanceProxyPtr = TUniquePtr<FShapeInstanceProxy>;
+	using FConstShapeInstanceProxyPtr = TUniquePtr<const FShapeInstanceProxy>;
+	using FShapeInstanceProxyArray = TArray<FShapeInstanceProxyPtr, TInlineAllocator<1>>;
+
+	// Physics thread ShapeInstance
+	class FShapeInstance;
+	using FShapeInstancePtr = TUniquePtr<FShapeInstance>;
+	using FConstShapeInstancePtr = TUniquePtr<const FShapeInstance>;
+	using FShapeInstanceArray = TArray<FShapeInstancePtr, TInlineAllocator<1>>;
+
+}
\ No newline at end of file
diff --git a/Engine/Source/Runtime/Landscape/Private/LandscapeCollision.cpp b/Engine/Source/Runtime/Landscape/Private/LandscapeCollision.cpp
--- a/Engine/Source/Runtime/Landscape/Private/LandscapeCollision.cpp
+++ b/Engine/Source/Runtime/Landscape/Private/LandscapeCollision.cpp
@@ -326,7 +326,7 @@ void ULandscapeHeightfieldCollisionComponent::OnCreatePhysicsState()
 				HeightfieldRef->Heightfield->SetScale(FinalScale * LandscapeComponentTransform.GetScale3D().GetSignVector());
 				TUniquePtr<Chaos::TImplicitObjectTransformed<Chaos::FReal, 3>> ChaosHeightFieldFromCooked = MakeUnique<Chaos::TImplicitObjectTransformed<Chaos::FReal, 3>>(MakeSerializable(HeightfieldRef->Heightfield), Chaos::FRigidTransform3(FTransform::Identity));
 
-				TUniquePtr<Chaos::FPerShapeData> NewShape = Chaos::FPerShapeData::CreatePerShapeData(ShapeArray.Num(), MakeSerializable(ChaosHeightFieldFromCooked));
+				TUniquePtr<Chaos::FPerShapeData> NewShape = Chaos::FShapeInstanceProxy::Make(ShapeArray.Num(), MakeSerializable(ChaosHeightFieldFromCooked));
 
 				// Setup filtering
 				FCollisionFilterData QueryFilterData, SimFilterData;
@@ -350,7 +350,7 @@ void ULandscapeHeightfieldCollisionComponent::OnCreatePhysicsState()
 					HeightfieldRef->HeightfieldSimple->SetScale(FinalSimpleCollisionScale);
 					TUniquePtr<Chaos::TImplicitObjectTransformed<Chaos::FReal, 3>> ChaosSimpleHeightFieldFromCooked = MakeUnique<Chaos::TImplicitObjectTransformed<Chaos::FReal, 3>>(MakeSerializable(HeightfieldRef->HeightfieldSimple), Chaos::FRigidTransform3(FTransform::Identity));
 
-					TUniquePtr<Chaos::FPerShapeData> NewSimpleShape = Chaos::FPerShapeData::CreatePerShapeData(ShapeArray.Num(), MakeSerializable(ChaosSimpleHeightFieldFromCooked));
+					TUniquePtr<Chaos::FPerShapeData> NewSimpleShape = Chaos::FShapeInstanceProxy::Make(ShapeArray.Num(), MakeSerializable(ChaosSimpleHeightFieldFromCooked));
 
 					FCollisionFilterData QueryFilterDataSimple = QueryFilterData;
 					FCollisionFilterData SimFilterDataSimple = SimFilterData;
@@ -372,7 +372,7 @@ void ULandscapeHeightfieldCollisionComponent::OnCreatePhysicsState()
 					HeightfieldRef->EditorHeightfield->SetScale(FinalScale * LandscapeComponentTransform.GetScale3D().GetSignVector());
 					TUniquePtr<Chaos::TImplicitObjectTransformed<Chaos::FReal, 3>> ChaosEditorHeightFieldFromCooked = MakeUnique<Chaos::TImplicitObjectTransformed<Chaos::FReal, 3>>(MakeSerializable(HeightfieldRef->EditorHeightfield), Chaos::FRigidTransform3(FTransform::Identity));
 
-					TUniquePtr<Chaos::FPerShapeData> NewEditorShape = Chaos::FPerShapeData::CreatePerShapeData(ShapeArray.Num(), MakeSerializable(ChaosEditorHeightFieldFromCooked));
+					TUniquePtr<Chaos::FPerShapeData> NewEditorShape = Chaos::FShapeInstanceProxy::Make(ShapeArray.Num(), MakeSerializable(ChaosEditorHeightFieldFromCooked));
 
 					FCollisionResponseContainer CollisionResponse;
 					CollisionResponse.SetAllChannels(ECollisionResponse::ECR_Ignore);
@@ -1446,7 +1446,7 @@ struct FMeshCollisionInitHelper
 		TSharedPtr<Chaos::FTriangleMeshImplicitObject, ESPMode::ThreadSafe> SharedPtrForRefCount(nullptr); // Not shared trimesh, no need for ref counting trimesh.
 		{
 			TUniquePtr<Chaos::TImplicitObjectScaled<Chaos::FTriangleMeshImplicitObject>> ScaledTrimesh = MakeUnique<Chaos::TImplicitObjectScaled<Chaos::FTriangleMeshImplicitObject>>(MakeSerializable(MeshRef->Trimesh), SharedPtrForRefCount, Scale);
-			TUniquePtr<Chaos::FPerShapeData> NewShape = Chaos::FPerShapeData::CreatePerShapeData(ShapeArray.Num(), MakeSerializable(ScaledTrimesh));
+			TUniquePtr<Chaos::FPerShapeData> NewShape = Chaos::FShapeInstanceProxy::Make(ShapeArray.Num(), MakeSerializable(ScaledTrimesh));
 
 			NewShape->SetQueryData(QueryFilter);
 			NewShape->SetSimData(SimulationFilter);
@@ -1461,7 +1461,7 @@ struct FMeshCollisionInitHelper
 		if(!World->IsGameWorld())
 		{
 			TUniquePtr<Chaos::TImplicitObjectScaled<Chaos::FTriangleMeshImplicitObject>> ScaledTrimeshEd = MakeUnique<Chaos::TImplicitObjectScaled<Chaos::FTriangleMeshImplicitObject>>(MakeSerializable(MeshRef->EditorTrimesh), SharedPtrForRefCount, Scale);
-			TUniquePtr<Chaos::FPerShapeData> NewEdShape = Chaos::FPerShapeData::CreatePerShapeData(ShapeArray.Num(), MakeSerializable(ScaledTrimeshEd));
+			TUniquePtr<Chaos::FPerShapeData> NewEdShape = Chaos::FShapeInstanceProxy::Make(ShapeArray.Num(), MakeSerializable(ScaledTrimeshEd));
 
 			NewEdShape->SetQueryData(QueryFilterEd);
 			NewEdShape->SetSimEnabled(false);
diff --git a/Engine/Source/Runtime/PhysicsCore/Public/PhysicsInterfaceDeclaresCore.h b/Engine/Source/Runtime/PhysicsCore/Public/PhysicsInterfaceDeclaresCore.h
--- a/Engine/Source/Runtime/PhysicsCore/Public/PhysicsInterfaceDeclaresCore.h
+++ b/Engine/Source/Runtime/PhysicsCore/Public/PhysicsInterfaceDeclaresCore.h
@@ -5,6 +5,7 @@
 #define MAX_PHYSICS_FIELD_TARGETS 32
 
 #include "ChaosSQTypes.h"
+#include "Chaos/ShapeInstanceFwd.h"
 #include "PhysicsProxy/SingleParticlePhysicsProxyFwd.h"
 
 namespace Chaos
@@ -32,8 +33,6 @@ namespace Chaos
 
 	class FChaosPhysicsMaterial;
 	class FChaosPhysicsMaterialMask;
-
-	using FShapesArray = TArray<TUniquePtr<FPerShapeData>, TInlineAllocator<1>>;
 }
 
 // Temporary dummy types until SQ implemented
