Update render data when resetting geometry collection asset in fracture editor.
Fixes crash caused by scene proxy recreating without the required render data update.

diff --git a/Engine/Plugins/Experimental/ChaosEditor/Source/FractureEditor/Private/FractureToolGenerators.cpp b/Engine/Plugins/Experimental/ChaosEditor/Source/FractureEditor/Private/FractureToolGenerators.cpp
--- a/Engine/Plugins/Experimental/ChaosEditor/Source/FractureEditor/Private/FractureToolGenerators.cpp
+++ b/Engine/Plugins/Experimental/ChaosEditor/Source/FractureEditor/Private/FractureToolGenerators.cpp
@@ -877,8 +877,8 @@ void UFractureToolResetAsset::Execute(TWeakPtr<FFractureEditorModeToolkit> InToo
 
 				FGeometryCollectionClusteringUtility::UpdateHierarchyLevelOfChildren(GeometryCollection, -1);
 				AddSingleRootNodeIfRequired(GeometryCollectionObject);
-				GeometryCollectionComponent->MarkRenderDynamicDataDirty();
-				GeometryCollectionComponent->MarkRenderStateDirty();
+
+				GeometryCollectionObject->RebuildRenderData();
 			}
 
 			GeometryCollectionObject->MarkPackageDirty();
diff --git a/Engine/Source/Runtime/Experimental/GeometryCollectionEngine/Private/GeometryCollection/GeometryCollectionSceneProxy.cpp b/Engine/Source/Runtime/Experimental/GeometryCollectionEngine/Private/GeometryCollection/GeometryCollectionSceneProxy.cpp
--- a/Engine/Source/Runtime/Experimental/GeometryCollectionEngine/Private/GeometryCollection/GeometryCollectionSceneProxy.cpp
+++ b/Engine/Source/Runtime/Experimental/GeometryCollectionEngine/Private/GeometryCollection/GeometryCollectionSceneProxy.cpp
@@ -346,7 +346,11 @@ void FGeometryCollectionSceneProxy::CreateRenderThreadResources()
 		uint16 const* BoneMapData = &MeshResource.BoneMapVertexBuffer.BoneIndex(0);
 		ParallelFor(MeshDescription.NumVertices, [&](int32 i)
 		{
-			const uint16 ProxyIndex = BoneMapData[i];
+			// Note that some fracture undo/redo operations can: recreate scene proxy, then update render data, then recreate proxy again.
+			// In that case we can come here the first time with too few hit proxy objects for the bone map which hasn't updated.
+			// But we then enter here a second time with the render data correct.
+			int16 ProxyIndex = BoneMapData[i];
+			ProxyIndex = HitProxies.IsValidIndex(ProxyIndex) ? ProxyIndex : 0;
 			HitProxyIdBuffer.VertexColor(i) = HitProxies[ProxyIndex]->Id.GetColor();
 		});
 
