Chaos - Attempt to fix static analysis error

diff --git a/Engine/Source/Runtime/Experimental/Chaos/Public/Chaos/ShapeInstanceFwd.h b/Engine/Source/Runtime/Experimental/Chaos/Public/Chaos/ShapeInstanceFwd.h
--- a/Engine/Source/Runtime/Experimental/Chaos/Public/Chaos/ShapeInstanceFwd.h
+++ b/Engine/Source/Runtime/Experimental/Chaos/Public/Chaos/ShapeInstanceFwd.h
@@ -1,6 +1,7 @@
 // Copyright Epic Games, Inc. All Rights Reserved.
 #pragma once
 
+#include "Containers/Array.h"
 #include "Templates/SharedPointer.h"
 #include "Templates/UniquePtr.h"
 
@@ -15,7 +16,7 @@ namespace Chaos
 
 	// Legacy names
 	class FPerShapeData;
-	using FShapesArray = TArray<TUniquePtr<FPerShapeData>, TInlineAllocator<1>>;;
+	using FShapesArray = TArray<TUniquePtr<FPerShapeData>, TInlineAllocator<1>>;
 
 	// Game thread ShapeInstance
 	class FShapeInstanceProxy;
@@ -28,5 +29,4 @@ namespace Chaos
 	using FShapeInstancePtr = TUniquePtr<FShapeInstance>;
 	using FConstShapeInstancePtr = TUniquePtr<const FShapeInstance>;
 	using FShapeInstanceArray = TArray<FShapeInstancePtr, TInlineAllocator<1>>;
-
-}
\ No newline at end of file
+}
