make the fracture mode 'explode amount' slider's exploded vectors affect the rendered bone transforms for geometry collections

diff --git a/Engine/Source/Runtime/Experimental/GeometryCollectionEngine/Private/GeometryCollection/GeometryCollectionComponent.cpp b/Engine/Source/Runtime/Experimental/GeometryCollectionEngine/Private/GeometryCollection/GeometryCollectionComponent.cpp
--- a/Engine/Source/Runtime/Experimental/GeometryCollectionEngine/Private/GeometryCollection/GeometryCollectionComponent.cpp
+++ b/Engine/Source/Runtime/Experimental/GeometryCollectionEngine/Private/GeometryCollection/GeometryCollectionComponent.cpp
@@ -2085,6 +2085,9 @@ void UGeometryCollectionComponent::GetRestTransforms(TArray<FMatrix44f>& OutRest
 {
 	TArray<FMatrix> RestMatrices;
 	GeometryCollectionAlgo::GlobalMatrices(RestCollection->GetGeometryCollection()->Transform, RestCollection->GetGeometryCollection()->Parent, RestMatrices);
+#if WITH_EDITOR
+	UpdateGlobalMatricesWithExplodedVectors(RestMatrices, *(RestCollection->GetGeometryCollection()));
+#endif
 	CopyTransformsWithConversionWhenNeeded(OutRestTransforms, RestMatrices);
 }
 
