fix mesh paint not updating vertex colors for geometry collection + losing exploded vectors on finishing paint transactions

diff --git a/Engine/Plugins/MeshPainting/Source/MeshPaintingToolset/Private/MeshPaintGeometryCollectionAdapter.cpp b/Engine/Plugins/MeshPainting/Source/MeshPaintingToolset/Private/MeshPaintGeometryCollectionAdapter.cpp
--- a/Engine/Plugins/MeshPainting/Source/MeshPaintingToolset/Private/MeshPaintGeometryCollectionAdapter.cpp
+++ b/Engine/Plugins/MeshPainting/Source/MeshPaintingToolset/Private/MeshPaintGeometryCollectionAdapter.cpp
@@ -276,7 +276,7 @@ void FMeshPaintGeometryCollectionComponentAdapter::PostEdit()
 		return;
 	}
 
-	GeometryCollectionComponent->MarkRenderStateDirty();
+	GetGeometryCollectionObject()->RebuildRenderData();
 }
 
 void FMeshPaintGeometryCollectionComponentAdapter::GetVertexColor(int32 VertexIndex, FColor& OutColor, bool bInstance /*= true*/) const
diff --git a/Engine/Source/Runtime/Experimental/GeometryCollectionEngine/Private/GeometryCollection/GeometryCollectionObject.cpp b/Engine/Source/Runtime/Experimental/GeometryCollectionEngine/Private/GeometryCollection/GeometryCollectionObject.cpp
--- a/Engine/Source/Runtime/Experimental/GeometryCollectionEngine/Private/GeometryCollection/GeometryCollectionObject.cpp
+++ b/Engine/Source/Runtime/Experimental/GeometryCollectionEngine/Private/GeometryCollection/GeometryCollectionObject.cpp
@@ -864,7 +864,7 @@ void UGeometryCollection::Serialize(FArchive& Ar)
 #if WITH_EDITOR
 		if (Ar.IsSaving() && !Ar.IsTransacting())
 		{
-			EnsureDataIsCooked(false /*bInitResourcs*/, Ar.IsTransacting(), Ar.IsPersistent());
+			EnsureDataIsCooked(false /*bInitResources*/, Ar.IsTransacting(), Ar.IsPersistent(), false /*bAllowCopyFromDDC*/);
 		}
 #endif
 		if (Ar.IsLoading())
@@ -1424,11 +1424,11 @@ bool UGeometryCollection::Modify(bool bAlwaysMarkDirty /*= true*/)
 	return bSuperResult;
 }
 
-void UGeometryCollection::EnsureDataIsCooked(bool bInitResources, bool bIsTransacting, bool bIsPersistant)
+void UGeometryCollection::EnsureDataIsCooked(bool bInitResources, bool bIsTransacting, bool bIsPersistant, bool bAllowCopyFromDDC)
 {
 	if (StateGuid != LastBuiltSimulationDataGuid)
 	{
-		CreateSimulationDataImp(/*bCopyFromDDC=*/ !bIsTransacting);
+		CreateSimulationDataImp(/*bCopyFromDDC=*/ bAllowCopyFromDDC && !bIsTransacting);
 		LastBuiltSimulationDataGuid = StateGuid;
 	}
 
diff --git a/Engine/Source/Runtime/Experimental/GeometryCollectionEngine/Public/GeometryCollection/GeometryCollectionObject.h b/Engine/Source/Runtime/Experimental/GeometryCollectionEngine/Public/GeometryCollection/GeometryCollectionObject.h
--- a/Engine/Source/Runtime/Experimental/GeometryCollectionEngine/Public/GeometryCollection/GeometryCollectionObject.h
+++ b/Engine/Source/Runtime/Experimental/GeometryCollectionEngine/Public/GeometryCollection/GeometryCollectionObject.h
@@ -346,7 +346,7 @@ public:
 #endif
 
 #if WITH_EDITOR
-	void EnsureDataIsCooked(bool bInitResources, bool bIsTransacting, bool bIsPersistant);
+	void EnsureDataIsCooked(bool bInitResources, bool bIsTransacting, bool bIsPersistant, bool bAllowCopyFromDDC = true);
 #endif
 
 	/** Accessors for internal geometry collection */
