Avoid segfault.  Print stats about captured particles.

diff --git a/Engine/Source/Runtime/Experimental/Chaos/Private/GeometryCollection/Facades/CollectionTetrahedralSkeletalBindingsFacade.cpp b/Engine/Source/Runtime/Experimental/Chaos/Private/GeometryCollection/Facades/CollectionTetrahedralSkeletalBindingsFacade.cpp
--- a/Engine/Source/Runtime/Experimental/Chaos/Private/GeometryCollection/Facades/CollectionTetrahedralSkeletalBindingsFacade.cpp
+++ b/Engine/Source/Runtime/Experimental/Chaos/Private/GeometryCollection/Facades/CollectionTetrahedralSkeletalBindingsFacade.cpp
@@ -105,7 +105,7 @@ namespace GeometryCollection::Facades
 	{
 		auto UEVert3d = [](FVector3f V) { return FVector3d(V.X, V.Y, V.Z); };
 
-		if (IsValid())
+		if (IsValid() && InVertices.Num())
 		{
 			int32 GroupIndex = MeshIdAttribute.Get().Find(InKey);
 			if (GroupIndex != INDEX_NONE)
@@ -129,10 +129,10 @@ namespace GeometryCollection::Facades
 							if (0 <= SkeletonIndex && SkeletonIndex < OutPosition.Num())
 							{
 								const FVector4 Weight(Weights[BindingIndex][0], Weights[BindingIndex][1], Weights[BindingIndex][2], Weights[BindingIndex][3]);
-								FVector X0 = UEVert3d(InVertices[Tet[0]]);
-								FVector X1 = UEVert3d(InVertices[Tet[1]]);
-								FVector X2 = UEVert3d(InVertices[Tet[2]]);
-								FVector X3 = UEVert3d(InVertices[Tet[3]]);
+								FVector X0 = InVertices.IsValidIndex(Tet[0]) ? UEVert3d(InVertices[Tet[0]]) : FVector();
+								FVector X1 = InVertices.IsValidIndex(Tet[0]) ? UEVert3d(InVertices[Tet[1]]) : FVector();
+								FVector X2 = InVertices.IsValidIndex(Tet[0]) ? UEVert3d(InVertices[Tet[2]]) : FVector();
+								FVector X3 = InVertices.IsValidIndex(Tet[0]) ? UEVert3d(InVertices[Tet[3]]) : FVector();
 								OutPosition[SkeletonIndex] = Weight[0] * X0 + Weight[1] * X1 + Weight[2] * X2 + Weight[3] * X3;
 								if (OutInfluence!=nullptr) if(SkeletonIndex<OutInfluence->Num()) (*OutInfluence)[SkeletonIndex] = true;
 							}
