ensure that the OnComponentPhysicsStateChanged gets called properly for geometry collections

diff --git a/Engine/Source/Runtime/Experimental/GeometryCollectionEngine/Private/GeometryCollection/GeometryCollectionComponent.cpp b/Engine/Source/Runtime/Experimental/GeometryCollectionEngine/Private/GeometryCollection/GeometryCollectionComponent.cpp
--- a/Engine/Source/Runtime/Experimental/GeometryCollectionEngine/Private/GeometryCollection/GeometryCollectionComponent.cpp
+++ b/Engine/Source/Runtime/Experimental/GeometryCollectionEngine/Private/GeometryCollection/GeometryCollectionComponent.cpp
@@ -2646,6 +2646,9 @@ void UGeometryCollectionComponent::OnCreatePhysicsState()
 			if (BodyInstance.bSimulatePhysics)
 			{
 				RegisterAndInitializePhysicsProxy();
+
+				// We're skipping over the primitive component so we need to make sure this event gets fired.
+				OnComponentPhysicsStateChanged.Broadcast(this, EComponentPhysicsStateChange::Created);
 			}
 			
 		}
@@ -2958,6 +2961,9 @@ void UGeometryCollectionComponent::OnDestroyPhysicsState()
 		// Discard the pointer (cleanup happens through the scene or dedicated thread)
 		PhysicsProxy = nullptr;
 	}
+
+	// We're skipping over the primitive component so we need to make sure this event gets fired.
+	OnComponentPhysicsStateChanged.Broadcast(this, EComponentPhysicsStateChange::Destroyed);
 }
 
 void UGeometryCollectionComponent::SendRenderDynamicData_Concurrent()
