Iris
* Warn when a replicated object forgets to call CreateAndRegisterFragmentsForObject unintentionally

diff --git a/Engine/Source/Runtime/Experimental/Iris/Core/Private/Iris/ReplicationSystem/ObjectReplicationBridge.cpp b/Engine/Source/Runtime/Experimental/Iris/Core/Private/Iris/ReplicationSystem/ObjectReplicationBridge.cpp
--- a/Engine/Source/Runtime/Experimental/Iris/Core/Private/Iris/ReplicationSystem/ObjectReplicationBridge.cpp
+++ b/Engine/Source/Runtime/Experimental/Iris/Core/Private/Iris/ReplicationSystem/ObjectReplicationBridge.cpp
@@ -224,6 +224,25 @@ UE::Net::FNetRefHandle UObjectReplicationBridge::BeginReplication(UObject* Insta
 
 	const FReplicationFragments& RegisteredFragments = FFragmentRegistrationContextPrivateAccessor::GetReplicationFragments(FragmentRegistrationContext);
 
+#if !(UE_BUILD_SHIPPING || UE_BUILD_TEST)
+	if (RegisteredFragments.IsEmpty() && !FragmentRegistrationContext.IsFragmentlessNetObject())
+	{
+		// Look if the class registered replicated properties
+		TArray<FLifetimeProperty> ReplicatedProps;
+		Instance->GetLifetimeReplicatedProps(ReplicatedProps);
+		
+		if (ReplicatedProps.IsEmpty())
+		{
+			ensureMsgf(!RegisteredFragments.IsEmpty(), TEXT("NetObject %s (class %s) registered no fragments. Call SetIsFragmentlessNetObject if this is intentional."), *GetNameSafe(Instance), *GetNameSafe(Instance->GetClass()));
+		}
+		else
+		{
+			ensureMsgf(!RegisteredFragments.IsEmpty(), TEXT("NetObject %s (class %s) registered no fragments but GetLifetimeReplicatedProps returned %d variables. Make sure to call CreateAndRegisterFragmentsForObject in RegisterReplicationFragments"),
+				*GetNameSafe(Instance), *GetNameSafe(Instance->GetClass()), ReplicatedProps.Num());
+		}
+	}
+#endif
+
 	// We currently identify protocols by local archetype or CDO pointer and verified the protocol id received from server and the hash of the default state
 	const UObject* ArchetypeOrCDOUsedAsKey = Instance->GetArchetype();
 
diff --git a/Engine/Source/Runtime/Experimental/Iris/Core/Private/Iris/ReplicationSystem/ReplicationFragmentUtil.cpp b/Engine/Source/Runtime/Experimental/Iris/Core/Private/Iris/ReplicationSystem/ReplicationFragmentUtil.cpp
--- a/Engine/Source/Runtime/Experimental/Iris/Core/Private/Iris/ReplicationSystem/ReplicationFragmentUtil.cpp
+++ b/Engine/Source/Runtime/Experimental/Iris/Core/Private/Iris/ReplicationSystem/ReplicationFragmentUtil.cpp
@@ -69,6 +69,12 @@ uint32 FReplicationFragmentUtil::CreateAndRegisterFragmentsForObject(UObject* Ob
 		}
 	}
 
+	// If we did not find any fragments to create, tell the context it's known.
+	if (Context.NumFragments() <= 0)
+	{
+		Context.SetIsFragmentlessNetObject(true);
+	}
+
 	return NumCreatedReplicationFragments;
 }
 
diff --git a/Engine/Source/Runtime/Experimental/Iris/Core/Public/Iris/ReplicationSystem/ReplicationFragment.h b/Engine/Source/Runtime/Experimental/Iris/Core/Public/Iris/ReplicationSystem/ReplicationFragment.h
--- a/Engine/Source/Runtime/Experimental/Iris/Core/Public/Iris/ReplicationSystem/ReplicationFragment.h
+++ b/Engine/Source/Runtime/Experimental/Iris/Core/Public/Iris/ReplicationSystem/ReplicationFragment.h
@@ -199,17 +199,26 @@ class FFragmentRegistrationContext
 {
 public:
 	explicit FFragmentRegistrationContext(Private::FReplicationStateDescriptorRegistry* InReplicationStateRegistry, EReplicationFragmentTraits InFragmentTraits)
-	: ReplicationStateRegistry(InReplicationStateRegistry)
-	, FragmentTraits(InFragmentTraits)
+		: ReplicationStateRegistry(InReplicationStateRegistry)
+		, FragmentTraits(InFragmentTraits)
 	{
 	}
 
-	/* Returns the traits */
+	/** Returns the traits */
 	const EReplicationFragmentTraits GetFragmentTraits() const { return FragmentTraits; }
 
-	/* Register ReplicationFragment */
+	/** Register ReplicationFragment */
 	void RegisterReplicationFragment(FReplicationFragment* Fragment, const FReplicationStateDescriptor* Descriptor, void* SrcReplicationStateBuffer) { Fragments.Add({ Descriptor, SrcReplicationStateBuffer, Fragment }); }
 
+	/** Call this when you have a netobject that is replicated but will never contain any RPCs or replicated properties. Prevents the registration code from complaining about potential errors. */
+	void SetIsFragmentlessNetObject(bool bIsFragmentless) { bIsAFragmentlessNetObject = bIsFragmentless; }
+
+	/** Returns true when the netobject knows it won't contain any replicated properties or RPCs */
+	bool IsFragmentlessNetObject() const { return bIsAFragmentlessNetObject; }
+
+	/** Returns the number of fragments registered */
+	int32 NumFragments() const { return Fragments.Num(); }
+
 	Private::FReplicationStateDescriptorRegistry* GetReplicationStateRegistry() const { return ReplicationStateRegistry; }
 
 	friend Private::FFragmentRegistrationContextPrivateAccessor;
@@ -218,6 +227,7 @@ private:
 	FReplicationFragments Fragments;
 	Private::FReplicationStateDescriptorRegistry* ReplicationStateRegistry;
 	const EReplicationFragmentTraits FragmentTraits;
+	bool bIsAFragmentlessNetObject = false;
 };
 
 }
