Iris - Ensure we have a valid string before operating on it.

#rnx

diff --git a/Engine/Source/Runtime/Experimental/Iris/Core/Private/Iris/ReplicationState/ReplicationStateDescriptorBuilder.cpp b/Engine/Source/Runtime/Experimental/Iris/Core/Private/Iris/ReplicationState/ReplicationStateDescriptorBuilder.cpp
--- a/Engine/Source/Runtime/Experimental/Iris/Core/Private/Iris/ReplicationState/ReplicationStateDescriptorBuilder.cpp
+++ b/Engine/Source/Runtime/Experimental/Iris/Core/Private/Iris/ReplicationState/ReplicationStateDescriptorBuilder.cpp
@@ -92,6 +92,9 @@ struct FLazyGetPathNameHelper
 		{
 			Object->GetPathName(nullptr, StringBuilder);
 
+			// Ensure we have a valid string
+			StringBuilder.ToString();
+
 			// Make sure that the path have consistent casing
 			TCString<TCHAR>::Strupr(StringBuilder.GetData(), StringBuilder.Len());
 		}
