Chaos - expose ShapeInstance arrays on path to deprecate FPerShapeData (WIP)

diff --git a/Engine/Source/Runtime/Engine/Private/PhysicsEngine/Experimental/PhysInterface_Chaos.cpp b/Engine/Source/Runtime/Engine/Private/PhysicsEngine/Experimental/PhysInterface_Chaos.cpp
--- a/Engine/Source/Runtime/Engine/Private/PhysicsEngine/Experimental/PhysInterface_Chaos.cpp
+++ b/Engine/Source/Runtime/Engine/Private/PhysicsEngine/Experimental/PhysInterface_Chaos.cpp
@@ -703,9 +703,10 @@ void FPhysInterface_Chaos::AddGeometry(FPhysicsActorHandle& InActor, const FGeom
 			//FPhysInterface_Chaos::SetMaterials(NewHandle, InParams.ComplexMaterials.Num() > 0 ? InParams.ComplexMaterials : SimpleView);
 		}
 
+		// NOTE: Both MergeGeometry and SetGeometry will extend the ShapesInstances array to contain enough elements for
+		// each geometry in the Union. However the shape data will not have been filled in, hence the call to MergeShapeInstance at the end.
+		// todo: we should not be creating unique geometry per actor
 		bool bMergeShapesArray = false;
-		//todo: we should not be creating unique geometry per actor
-		// we always have a union so we can support any future welding operations. (Non-trivial converting the SharedPtr to UniquePtr)
 		{
 			if (InActor->GetGameThreadAPI().Geometry()) // geometry already exists - combine new geometry with the existing
 			{
@@ -714,10 +715,13 @@ void FPhysInterface_Chaos::AddGeometry(FPhysicsActorHandle& InActor, const FGeom
 			}
 			else
 			{
+				// we always have a union so we can support any future welding operations. (Non-trivial converting the SharedPtr to UniquePtr)
 				InActor->GetGameThreadAPI().SetGeometry(MakeUnique<Chaos::FImplicitObjectUnion>(MoveTemp(Geoms)));
 			}
 		}
 
+		// Update the newly added shapes with the collision filters, materials etc
+		// NOTE: MergeShapes overwrites the last N shapes (see comments above)
 		if (bMergeShapesArray)
 		{
 			InActor->GetGameThreadAPI().MergeShapesArray(MoveTemp(Shapes));
diff --git a/Engine/Source/Runtime/Experimental/Chaos/Public/Chaos/GeometryParticles.h b/Engine/Source/Runtime/Experimental/Chaos/Public/Chaos/GeometryParticles.h
--- a/Engine/Source/Runtime/Experimental/Chaos/Public/Chaos/GeometryParticles.h
+++ b/Engine/Source/Runtime/Experimental/Chaos/Public/Chaos/GeometryParticles.h
@@ -332,6 +332,8 @@ namespace Chaos
 
 		CHAOS_API const FShapesArray& ShapesArray(const int32 Index) const { return reinterpret_cast<const FShapesArray&>(MShapesArray[Index]); }
 
+		CHAOS_API const FShapeInstanceArray& ShapeInstances(const int32 Index) const { return MShapesArray[Index]; }
+
 #if CHAOS_DETERMINISTIC
 		CHAOS_API FParticleID ParticleID(const int32 Idx) const { return MParticleIDs[Idx]; }
 		CHAOS_API FParticleID& ParticleID(const int32 Idx) { return MParticleIDs[Idx]; }
diff --git a/Engine/Source/Runtime/Experimental/Chaos/Public/Chaos/ParticleHandle.h b/Engine/Source/Runtime/Experimental/Chaos/Public/Chaos/ParticleHandle.h
--- a/Engine/Source/Runtime/Experimental/Chaos/Public/Chaos/ParticleHandle.h
+++ b/Engine/Source/Runtime/Experimental/Chaos/Public/Chaos/ParticleHandle.h
@@ -551,6 +551,7 @@ public:
 	void SetDynamicGeometry(TUniquePtr<FImplicitObject>&& Unique) { GeometryParticles->SetDynamicGeometry(ParticleIdx, MoveTemp(Unique)); }
 
 	const FShapesArray& ShapesArray() const { return GeometryParticles->ShapesArray(ParticleIdx); }
+	const FShapeInstanceArray& ShapeInstances() const { return GeometryParticles->ShapeInstances(ParticleIdx); }
 
 	const TAABB<T, d>& LocalBounds() const { return GeometryParticles->LocalBounds(ParticleIdx); }
 	void SetLocalBounds(const TAABB<T, d>& NewBounds) { GeometryParticles->LocalBounds(ParticleIdx) = NewBounds; }
@@ -2105,6 +2106,7 @@ public:
 	}
 
 	const FShapesArray& ShapesArray() const { return MHandle->ShapesArray(); }
+	const FShapeInstanceArray& ShapeInstances() const { return MHandle->ShapeInstances(); }
 
 	static constexpr EParticleType StaticType()
 	{
@@ -2469,21 +2471,46 @@ public:
 
 	void MergeShapesArray(FShapesArray&& OtherShapesArray)
 	{
-		int Idx = MShapesArray.Num() - OtherShapesArray.Num();
-		for (TUniquePtr<FPerShapeData>& Shape : OtherShapesArray)
+		MergeShapeInstances(reinterpret_cast<FShapeInstanceProxyArray&&>(OtherShapesArray));
+	}
+
+	const FShapesArray& ShapesArray() const { return reinterpret_cast<const FShapesArray&>(MShapesArray); }
+
+
+	TSerializablePtr<FImplicitObject> Geometry() const { return MakeSerializable(MNonFrequentData.Read().Geometry()); }
+
+	const FShapeInstanceProxyArray& ShapeInstances() const
+	{ 
+		return MShapesArray;
+	}
+
+	//Note: this must be called after setting geometry. This API seems bad. Should probably be part of setting geometry
+	void SetShapeInstances(FShapeInstanceProxyArray&& InShapes)
+	{ 
+		// This is only called after setting or merging geometry which resizes the shapes array
+		// so we should always have enough elements at this point
+		check(InShapes.Num() == MShapesArray.Num());
+		MShapesArray = MoveTemp(InShapes);
+	}
+
+	// Overwite the last N shapes with the data from InShapes. This is a highly dubious operation
+	// provided as a utility for use in FPhysInterface_Chaos::AddGeometry. Ideally we would remove this
+	void MergeShapeInstances(FShapeInstanceProxyArray&& InShapes)
+	{
+		check(InShapes.Num() <= MShapesArray.Num());
+		int Idx = MShapesArray.Num() - InShapes.Num();
+		for (FShapeInstanceProxyPtr& Shape : InShapes)
 		{
 			ensure(Idx < MShapesArray.Num());
-			MShapesArray[Idx++] = MoveTemp(reinterpret_cast<FShapeInstanceProxyPtr&>(Shape));
+			MShapesArray[Idx++] = MoveTemp(Shape);
 		}
 	}
 
+
 	void SetIgnoreAnalyticCollisionsImp(FImplicitObject* Implicit, bool bIgnoreAnalyticCollisions);
+	
 	CHAOS_API void SetIgnoreAnalyticCollisions(bool bIgnoreAnalyticCollisions);
 
-	TSerializablePtr<FImplicitObject> Geometry() const { return MakeSerializable(MNonFrequentData.Read().Geometry()); }
-
-	const FShapesArray& ShapesArray() const { return reinterpret_cast<const FShapesArray&>(MShapesArray); }
-
 	EObjectStateType ObjectState() const;
 
 	EParticleType ObjectType() const
@@ -2491,6 +2518,7 @@ public:
 		return Type;
 	}
 
+	
 
 	const TKinematicGeometryParticle<T, d>* CastToKinematicParticle() const;
 	TKinematicGeometryParticle<T, d>* CastToKinematicParticle();
