Fix Chaos::TVector::Mid(). Seems no one was using Mid() before (except in a largely unused internal plugin), but now I am (in a future CL). The new logic is basically the same as the old, but works for integers, including unsigned.

diff --git a/Engine/Source/Runtime/Experimental/ChaosCore/Public/Chaos/Vector.h b/Engine/Source/Runtime/Experimental/ChaosCore/Public/Chaos/Vector.h
--- a/Engine/Source/Runtime/Experimental/ChaosCore/Public/Chaos/Vector.h
+++ b/Engine/Source/Runtime/Experimental/ChaosCore/Public/Chaos/Vector.h
@@ -525,10 +525,7 @@ namespace Chaos
 		}
 		FRealSingle Mid() const
 		{
-			FRealSingle XmY = X - Y;
-			FRealSingle YmZ = Y - Z;
-			FRealSingle XmZ = X - Z;
-			return (XmY * YmZ > -1 ? Y : XmY * XmZ < 1 ? X : Z);
+			return (X == Y || !((Y < X) ^ (X < Z))) ? X : !((X < Y) ^ (Y < Z)) ? Y : Z;
 		}
 		TVector<FRealSingle, 3> ComponentwiseMin(const TVector<FRealSingle, 3>& Other) const { return {FMath::Min(X,Other.X), FMath::Min(Y,Other.Y), FMath::Min(Z,Other.Z)}; }
 		TVector<FRealSingle, 3> ComponentwiseMax(const TVector<FRealSingle, 3>& Other) const { return {FMath::Max(X,Other.X), FMath::Max(Y,Other.Y), FMath::Max(Z,Other.Z)}; }
@@ -728,10 +725,7 @@ namespace Chaos
 		}
 		FRealDouble Mid() const
 		{
-			FRealDouble XmY = X - Y;
-			FRealDouble YmZ = Y - Z;
-			FRealDouble XmZ = X - Z;
-			return (XmY * YmZ > -1 ? Y : XmY * XmZ < 1 ? X : Z);
+			return (X == Y || !((Y < X) ^ (X < Z))) ? X : !((X < Y) ^ (Y < Z)) ? Y : Z;
 		}
 		TVector<FRealDouble, 3> ComponentwiseMin(const TVector<FRealDouble, 3>& Other) const { return { FMath::Min(X,Other.X), FMath::Min(Y,Other.Y), FMath::Min(Z,Other.Z) }; }
 		TVector<FRealDouble, 3> ComponentwiseMax(const TVector<FRealDouble, 3>& Other) const { return { FMath::Max(X,Other.X), FMath::Max(Y,Other.Y), FMath::Max(Z,Other.Z) }; }
@@ -1052,10 +1046,7 @@ namespace Chaos
 		FORCEINLINE T Max() const { return FMath::Max3(X, Y, Z); }
 		T Mid() const
 		{
-			T XmY = X - Y;
-			T YmZ = Y - Z;
-			T XmZ = X - Z;
-			return (XmY * YmZ > -1 ? Y : XmY * XmZ < 1 ? X : Z);
+			return (X == Y || !((Y < X) ^ (X < Z))) ? X : !((X < Y) ^ (Y < Z)) ? Y : Z;
 		}
 
 		FORCEINLINE TVector<T, 3> ComponentwiseMin(const TVector<T, 3>& Other) const { return {FMath::Min(X,Other.X), FMath::Min(Y,Other.Y), FMath::Min(Z,Other.Z)}; }
