Dataflow - Fixed the pin tooltips when the tooltip is the node doc comment.

#rnx

diff --git a/Engine/Source/Runtime/Experimental/Dataflow/Core/Private/Dataflow/DataflowNode.cpp b/Engine/Source/Runtime/Experimental/Dataflow/Core/Private/Dataflow/DataflowNode.cpp
--- a/Engine/Source/Runtime/Experimental/Dataflow/Core/Private/Dataflow/DataflowNode.cpp
+++ b/Engine/Source/Runtime/Experimental/Dataflow/Core/Private/Dataflow/DataflowNode.cpp
@@ -472,7 +472,11 @@ FString FDataflowNode::GetPinToolTip(const FName& PropertyName)
 
 							if (NumElems == 2)
 							{
-								return OutArr[1];
+								return OutArr[1];  // Return tooltip meta text
+							}
+							else if (NumElems == 1)
+							{
+								return OutArr[0];  // Return doc comment
 							}
 						}
 					}
