Tentative fix for code warning in ObjectReplicationBridge

diff --git a/Engine/Source/Runtime/Experimental/Iris/Core/Private/Iris/ReplicationSystem/ObjectReplicationBridge.cpp b/Engine/Source/Runtime/Experimental/Iris/Core/Private/Iris/ReplicationSystem/ObjectReplicationBridge.cpp
--- a/Engine/Source/Runtime/Experimental/Iris/Core/Private/Iris/ReplicationSystem/ObjectReplicationBridge.cpp
+++ b/Engine/Source/Runtime/Experimental/Iris/Core/Private/Iris/ReplicationSystem/ObjectReplicationBridge.cpp
@@ -233,11 +233,11 @@ UE::Net::FNetRefHandle UObjectReplicationBridge::BeginReplication(UObject* Insta
 		
 		if (ReplicatedProps.IsEmpty())
 		{
-			ensureMsgf(!RegisteredFragments.IsEmpty(), TEXT("NetObject %s (class %s) registered no fragments. Call SetIsFragmentlessNetObject if this is intentional."), *GetNameSafe(Instance), *GetNameSafe(Instance->GetClass()));
+			ensureMsgf(false, TEXT("NetObject %s (class %s) registered no fragments. Call SetIsFragmentlessNetObject if this is intentional."), *GetNameSafe(Instance), *GetNameSafe(Instance->GetClass()));
 		}
 		else
 		{
-			ensureMsgf(!RegisteredFragments.IsEmpty(), TEXT("NetObject %s (class %s) registered no fragments but GetLifetimeReplicatedProps returned %d variables. Make sure to call CreateAndRegisterFragmentsForObject in RegisterReplicationFragments"),
+			ensureMsgf(false, TEXT("NetObject %s (class %s) registered no fragments but GetLifetimeReplicatedProps returned %d variables. Make sure to call CreateAndRegisterFragmentsForObject in RegisterReplicationFragments"),
 				*GetNameSafe(Instance), *GetNameSafe(Instance->GetClass()), ReplicatedProps.Num());
 		}
 	}
