Improved reader performance in cases where R >> W

diff --git a/Engine/Source/Runtime/Experimental/Chaos/Public/Framework/Threading.h b/Engine/Source/Runtime/Experimental/Chaos/Public/Framework/Threading.h
--- a/Engine/Source/Runtime/Experimental/Chaos/Public/Framework/Threading.h
+++ b/Engine/Source/Runtime/Experimental/Chaos/Public/Framework/Threading.h
@@ -512,6 +512,8 @@ FORCEINLINE void EnsureIsInGameThreadContext()
 
 		TRwFifoLock()
 			: NumReaders(0)
+			, NumWriters(0)
+			, UpdateNumFlag(0)
 		{
 			//ThreadingPrivate::CreateLockThreadData(this);
 		}
@@ -525,18 +527,29 @@ FORCEINLINE void EnsureIsInGameThreadContext()
 		{
 			if(ThreadingPrivate::GetThreadReadDepth(this) == 0)
 			{
-				TMutexScopeLock<MutexType> Guard(Mutex);
+				// Acquire the update num flag locklessly
+				uint32 ExpectedUpdateNumFlag = 0;
+				while (!UpdateNumFlag.compare_exchange_strong(ExpectedUpdateNumFlag, 1)) {}
 
-				// We lock for this increment to halt if there's a writer waiting to enter the lock
-				// In this case we will be forced to wait till the write completes
+				if (NumWriters.load() != 0)
+				{
+					// Release update num flag, then wait
+					UpdateNumFlag.store(0);
+					TMutexScopeLock<MutexType> Guard(Mutex);
+					// Reacquire update num flag
+					while (!UpdateNumFlag.compare_exchange_strong(ExpectedUpdateNumFlag, 1)) {}
+				}
+
+				// Increment reader count, then finally release the update num flag
 				++NumReaders;
+				UpdateNumFlag.store(0);
 			}
 			else
 			{
 				// Only require a lock on the first acquisition. this allows recursive reads even while a
 				// writer is holding the lock waiting to enter. The writer will be allowed to proceed when
 				// all write scopes end
-			++NumReaders;
+				++NumReaders;
 			}
 
 #if PHYSICS_THREAD_CONTEXT
@@ -565,6 +578,12 @@ FORCEINLINE void EnsureIsInGameThreadContext()
 #endif
 
 			Mutex.Lock();
+			// Acquire the update num flag locklessly
+			uint32 ExpectedUpdateNumFlag = 0;
+			while (!UpdateNumFlag.compare_exchange_strong(ExpectedUpdateNumFlag, 1)) {}
+			++NumWriters;
+			// Release the update num flag
+			UpdateNumFlag.store(0);
 
 			// Spin until all readers are finished
 			for(;;)
@@ -606,6 +625,7 @@ FORCEINLINE void EnsureIsInGameThreadContext()
 		{
 			CHAOS_RECORD_LEAVE_WRITE_LOCK;
 
+			--NumWriters;
 			Mutex.Unlock();
 #if PHYSICS_THREAD_CONTEXT
 			// Write lock is released, the gamethread context is gone
@@ -616,6 +636,8 @@ FORCEINLINE void EnsureIsInGameThreadContext()
 	private:
 		MutexType Mutex;
 		std::atomic<uint32> NumReaders;
+		std::atomic<uint32> NumWriters;
+		std::atomic<uint32> UpdateNumFlag;
 	};
 
 	/**
