Add flag for strip render data on cook to Geometry Collection.
This can be used if we have a non-native rendering replacement such as using ISM pools.
Added p.GeometryCollectionAssetForceStripRenderDataOnCook as a per-project way to force this.

diff --git a/Engine/Source/Runtime/Experimental/GeometryCollectionEngine/Private/GeometryCollection/DerivedDataGeometryCollectionCooker.cpp b/Engine/Source/Runtime/Experimental/GeometryCollectionEngine/Private/GeometryCollection/DerivedDataGeometryCollectionCooker.cpp
--- a/Engine/Source/Runtime/Experimental/GeometryCollectionEngine/Private/GeometryCollection/DerivedDataGeometryCollectionCooker.cpp
+++ b/Engine/Source/Runtime/Experimental/GeometryCollectionEngine/Private/GeometryCollection/DerivedDataGeometryCollectionCooker.cpp
@@ -100,7 +100,7 @@ bool FDerivedDataGeometryCollectionRenderDataCooker::Build(TArray<uint8>& OutDat
 		Chaos::FErrorReporter ErrorReporter(GeometryCollection.GetName());
 
 		TUniquePtr<FGeometryCollectionRenderData> RenderData = FGeometryCollectionRenderData::Create(*Collection, GeometryCollection.EnableNanite, GeometryCollection.bUseFullPrecisionUVs);
-		RenderData->Serialize(ChaosAr, &GeometryCollection);
+		RenderData->Serialize(ChaosAr, GeometryCollection);
 
 		if (false && ErrorReporter.EncounteredAnyErrors())
 		{
diff --git a/Engine/Source/Runtime/Experimental/GeometryCollectionEngine/Private/GeometryCollection/GeometryCollectionObject.cpp b/Engine/Source/Runtime/Experimental/GeometryCollectionEngine/Private/GeometryCollection/GeometryCollectionObject.cpp
--- a/Engine/Source/Runtime/Experimental/GeometryCollectionEngine/Private/GeometryCollection/GeometryCollectionObject.cpp
+++ b/Engine/Source/Runtime/Experimental/GeometryCollectionEngine/Private/GeometryCollection/GeometryCollectionObject.cpp
@@ -47,7 +47,7 @@ bool GeometryCollectionAssetForceStripOnCook = false;
 FAutoConsoleVariableRef CVarGeometryCollectionBypassPhysicsAttributes(
 	TEXT("p.GeometryCollectionAssetForceStripOnCook"),
 	GeometryCollectionAssetForceStripOnCook,
-	TEXT("Bypass the construction of simulation properties when all bodies are simply cached. for playback."));
+	TEXT("Bypass the construction of simulation properties when all bodies are simply cached for playback."));
 
 bool bGeometryCollectionEnableForcedConvexGenerationInSerialize = true;
 FAutoConsoleVariableRef CVarGeometryCollectionEnableForcedConvexGenerationInSerialize(
@@ -90,6 +90,7 @@ UGeometryCollection::UGeometryCollection(const FObjectInitializer& ObjectInitial
 	, ConnectionGraphBoundsFilteringMargin(0)
 	, bUseFullPrecisionUVs(false)
 	, bStripOnCook(false)
+	, bStripRenderDataOnCook(false)
 	, EnableNanite(false)
 #if WITH_EDITORONLY_DATA
 	, CollisionType_DEPRECATED(ECollisionTypeEnum::Chaos_Volumetric)
@@ -923,7 +924,7 @@ void UGeometryCollection::Serialize(FArchive& Ar)
 				RenderData = MakeUnique<FGeometryCollectionRenderData>();
 			}
 
-			RenderData->Serialize(ChaosAr, this);
+			RenderData->Serialize(ChaosAr, *this);
 		}
 	}
 
@@ -1089,7 +1090,7 @@ void UGeometryCollection::CreateRenderDataImp(bool bCopyFromDDC)
 			Chaos::FChaosArchive ChaosAr(Ar);
 
 			RenderData = MakeUnique<FGeometryCollectionRenderData>();
-			RenderData->Serialize(ChaosAr, this);
+			RenderData->Serialize(ChaosAr, *this);
 		}
 	}
 }
@@ -1244,7 +1245,7 @@ void UGeometryCollection::InitResources()
 {
 	if (RenderData)
 	{
-		RenderData->InitResources(this);
+		RenderData->InitResources(*this);
 	}
 }
 
@@ -1447,7 +1448,7 @@ void UGeometryCollection::EnsureDataIsCooked(bool bInitResources, bool bIsTransa
 		{
 			if (RenderData)
 			{
-				RenderData->InitResources(this);
+				RenderData->InitResources(*this);
 			}
 		}
 	
diff --git a/Engine/Source/Runtime/Experimental/GeometryCollectionEngine/Private/GeometryCollection/GeometryCollectionRenderData.cpp b/Engine/Source/Runtime/Experimental/GeometryCollectionEngine/Private/GeometryCollection/GeometryCollectionRenderData.cpp
--- a/Engine/Source/Runtime/Experimental/GeometryCollectionEngine/Private/GeometryCollection/GeometryCollectionRenderData.cpp
+++ b/Engine/Source/Runtime/Experimental/GeometryCollectionEngine/Private/GeometryCollection/GeometryCollectionRenderData.cpp
@@ -82,9 +82,9 @@ void FBoneMapVertexBuffer::Serialize(FArchive& Ar, bool bInNeedsCPUAccess)
 }
 
 
-void FGeometryCollectionMeshResources::InitResources(UGeometryCollection* Owner)
+void FGeometryCollectionMeshResources::InitResources(UGeometryCollection const& Owner)
 {
-	FName OwnerName = Owner->GetFName();
+	FName OwnerName = Owner.GetFName();
 	IndexBuffer.SetOwnerName(OwnerName);
 	PositionVertexBuffer.SetOwnerName(OwnerName);
 	StaticMeshVertexBuffer.SetOwnerName(OwnerName);
@@ -645,9 +645,17 @@ FGeometryCollectionRenderData::~FGeometryCollectionRenderData()
 	Fence.Wait();
 }
 
-void FGeometryCollectionRenderData::Serialize(FArchive& Ar, UGeometryCollection* Owner)
+// Temporary cvar driven option to force strip render data on cook if we are also stripping other GC data.
+bool GeometryCollectionAssetForceStripRenderDataOnCook = false;
+FAutoConsoleVariableRef CVarGeometryCollectionStripRenderData(
+	TEXT("p.GeometryCollectionAssetForceStripRenderDataOnCook"),
+	GeometryCollectionAssetForceStripRenderDataOnCook,
+	TEXT("Bypass the cook of render data when other geom collection data is also being stripped."));
+
+void FGeometryCollectionRenderData::Serialize(FArchive& Ar, UGeometryCollection& Owner)
 {
-	if (Owner->bStripOnCook && Ar.IsCooking())
+	const bool bStripRenderData = Owner.bStripRenderDataOnCook || (GeometryCollectionAssetForceStripRenderDataOnCook && Owner.bStripOnCook);
+	if (bStripRenderData && Ar.IsCooking())
 	{
 		// Don't cook rendering data.
 		// This is used if we expect to use a non-native GC render path such as the ISM pool.
@@ -676,14 +684,14 @@ void FGeometryCollectionRenderData::Serialize(FArchive& Ar, UGeometryCollection*
 		if (Ar.IsSaving())
 		{
 			// Nanite data is 1:1 with each geometry group in the collection.
-			const int32 NumGeometryGroups = Owner->NumElements(FGeometryCollection::GeometryGroup);
-			if (!Owner->EnableNanite || NumGeometryGroups != NaniteResource.HierarchyRootOffsets.Num())
+			const int32 NumGeometryGroups = Owner.NumElements(FGeometryCollection::GeometryGroup);
+			if (!Owner.EnableNanite || NumGeometryGroups != NaniteResource.HierarchyRootOffsets.Num())
 			{
 				Ar.SetError();
 			}
 		}
 
-		NaniteResource.Serialize(Ar, Owner, true);
+		NaniteResource.Serialize(Ar, &Owner, true);
 	}
 	else if (Ar.IsLoading())
 	{
@@ -691,7 +699,7 @@ void FGeometryCollectionRenderData::Serialize(FArchive& Ar, UGeometryCollection*
 	}
 }
 
-void FGeometryCollectionRenderData::InitResources(UGeometryCollection* Owner)
+void FGeometryCollectionRenderData::InitResources(UGeometryCollection const& Owner)
 {
 	if (bIsInitialized)
 	{
@@ -705,7 +713,7 @@ void FGeometryCollectionRenderData::InitResources(UGeometryCollection* Owner)
 
 	if (bHasNaniteData)
 	{
-		NaniteResource.InitResources(Owner);
+		NaniteResource.InitResources(&Owner);
 	}
 
 	bIsInitialized = true;
diff --git a/Engine/Source/Runtime/Experimental/GeometryCollectionEngine/Private/GeometryCollection/GeometryCollectionRenderData.h b/Engine/Source/Runtime/Experimental/GeometryCollectionEngine/Private/GeometryCollection/GeometryCollectionRenderData.h
--- a/Engine/Source/Runtime/Experimental/GeometryCollectionEngine/Private/GeometryCollection/GeometryCollectionRenderData.h
+++ b/Engine/Source/Runtime/Experimental/GeometryCollectionEngine/Private/GeometryCollection/GeometryCollectionRenderData.h
@@ -58,7 +58,7 @@ struct FGeometryCollectionMeshResources
 	FBoneMapVertexBuffer BoneMapVertexBuffer;
 
 	/** Initialize the render resources. */
-	void InitResources(UGeometryCollection* Owner);
+	void InitResources(UGeometryCollection const& Owner);
 	/** Release the render resources. */
 	void ReleaseResources();
 	/** Serialization. */
@@ -123,10 +123,10 @@ public:
 	}
 
 	/** Serialization. */
-	void Serialize(FArchive& Ar, UGeometryCollection* Owner);
+	void Serialize(FArchive& Ar, UGeometryCollection& Owner);
 
 	/** Initialize the render resources. */
-	void InitResources(UGeometryCollection* Owner);
+	void InitResources(UGeometryCollection const& Owner);
 
 	/** Releases the render resources. */
 	void ReleaseResources();
diff --git a/Engine/Source/Runtime/Experimental/GeometryCollectionEngine/Public/GeometryCollection/GeometryCollectionObject.h b/Engine/Source/Runtime/Experimental/GeometryCollectionEngine/Public/GeometryCollection/GeometryCollectionObject.h
--- a/Engine/Source/Runtime/Experimental/GeometryCollectionEngine/Public/GeometryCollection/GeometryCollectionObject.h
+++ b/Engine/Source/Runtime/Experimental/GeometryCollectionEngine/Public/GeometryCollection/GeometryCollectionObject.h
@@ -502,6 +502,19 @@ public:
 	UPROPERTY(EditAnywhere, Category = "Rendering")
 	bool bUseFullPrecisionUVs = false;
 
+	/**
+	 * Strip unnecessary data from the Geometry Collection to keep the memory footprint as small as possible.
+	 */
+	UPROPERTY(EditAnywhere, BlueprintReadWrite, Category = "Rendering", meta = (DisplayName = "Strip Source Data On Cook"))
+	bool bStripOnCook;
+
+	/**
+	 * Strip unnecessary render data from the Geometry Collection to keep the memory footprint as small as possible.
+	 * This may be used if the cooked build uses an alternative render path such as ISMPools.
+	 */
+	UPROPERTY(EditAnywhere, BlueprintReadWrite, Category = "Rendering")
+	bool bStripRenderDataOnCook;
+
 	/** list of unique static mesh / materials pairs for auto instancing*/
 	UPROPERTY(EditAnywhere, Category = "Rendering")
 	TArray<FGeometryCollectionAutoInstanceMesh> AutoInstanceMeshes;
@@ -510,11 +523,7 @@ public:
 	UPROPERTY(BlueprintReadWrite, EditAnywhere, Category = "Rendering")
 	FGeometryCollectionProxyMeshData RootProxyData;
 
-	/**
-	 * Strip unnecessary data from the Geometry Collection to keep the memory footprint as small as possible.
-	 */
-	UPROPERTY(EditAnywhere, BlueprintReadWrite, Category = "Nanite")
-	bool bStripOnCook;
+
 
 	UFUNCTION(BlueprintCallable, Category = "Nanite")
 	void SetEnableNanite(bool bValue);
