remove legacy guard against having 0 materials on geometry collection components

diff --git a/Engine/Source/Runtime/Experimental/GeometryCollectionEngine/Private/GeometryCollection/GeometryCollectionComponent.cpp b/Engine/Source/Runtime/Experimental/GeometryCollectionEngine/Private/GeometryCollection/GeometryCollectionComponent.cpp
--- a/Engine/Source/Runtime/Experimental/GeometryCollectionEngine/Private/GeometryCollection/GeometryCollectionComponent.cpp
+++ b/Engine/Source/Runtime/Experimental/GeometryCollectionEngine/Private/GeometryCollection/GeometryCollectionComponent.cpp
@@ -4030,11 +4030,9 @@ void UGeometryCollectionComponent::CalculateGlobalMatrices()
 #endif
 }
 
-// #todo(dmp): for backwards compatibility with existing maps, we need to have a default of 3 materials.  Otherwise
-// some existing test scenes will crash
 int32 UGeometryCollectionComponent::GetNumMaterials() const
 {
-	return !RestCollection || RestCollection->Materials.Num() == 0 ? 3 : RestCollection->Materials.Num();
+	return !RestCollection ? 0 : RestCollection->Materials.Num();
 }
 
 UMaterialInterface* UGeometryCollectionComponent::GetMaterial(int32 MaterialIndex) const
