Chaos - reduce ShapeInstance size on physics thread
- dedupe material getting code
- fix physics material usage for mesh-based shapes (was always taking material[0] before!)

diff --git a/Engine/Source/Programs/HeadlessChaos/Private/HeadlessChaosTestGraph.cpp b/Engine/Source/Programs/HeadlessChaos/Private/HeadlessChaosTestGraph.cpp
--- a/Engine/Source/Programs/HeadlessChaos/Private/HeadlessChaosTestGraph.cpp
+++ b/Engine/Source/Programs/HeadlessChaos/Private/HeadlessChaosTestGraph.cpp
@@ -639,7 +639,8 @@ namespace ChaosTest {
 	void HelpTickConstraints(FPBDRigidsSOAs& SOAs, const TArray<TPBDRigidParticleHandle<FReal, 3>*>& Particles,
 		Private::FPBDIslandManager& Graph, TMockGraphConstraints<0> Constraints,
 		const TArrayCollectionArray<TSerializablePtr<FChaosPhysicsMaterial>>& PhysicsMaterials,
-		const THandleArray<FChaosPhysicsMaterial>& PhysicalMaterials)
+		const TArrayCollectionArray<TUniquePtr<FChaosPhysicsMaterial>>& PerParticlePhysicsMaterials,
+		const THandleArray<FChaosPhysicsMaterial>& SimMaterials)
 	{
 		SOAs.ClearTransientDirty();
 
@@ -650,7 +651,7 @@ namespace ChaosTest {
 		Graph.UpdateIslands(SOAs);
 		for (int32 IslandIndex = 0; IslandIndex < Graph.NumIslands(); ++IslandIndex)
 		{
-			const bool bSleeped = Graph.SleepInactive(IslandIndex,PhysicsMaterials,PhysicalMaterials);
+			const bool bSleeped = Graph.SleepInactive(IslandIndex, PhysicsMaterials, PerParticlePhysicsMaterials, SimMaterials);
 
 			if(bSleeped)
 			{
@@ -702,8 +703,10 @@ namespace ChaosTest {
 			FPBDRigidsSOAs SOAs(UniqueIndices);
 			TArray<TPBDRigidParticleHandle<FReal, 3>*> Particles = SOAs.CreateDynamicParticles(NumParticles);
 			TArrayCollectionArray<TSerializablePtr<FChaosPhysicsMaterial>> PhysicsMaterials;
+			TArrayCollectionArray<TUniquePtr<FChaosPhysicsMaterial>> PerParticlePhysicsMaterials;
 			THandleArray<FChaosPhysicsMaterial> PhysicalMaterials;
- 			SOAs.GetParticleHandles().AddArray(&PhysicsMaterials);
+			SOAs.GetParticleHandles().AddArray(&PhysicsMaterials);
+			SOAs.GetParticleHandles().AddArray(&PerParticlePhysicsMaterials);
 
 			for (int32 Idx = 0; Idx < NumParticles; ++Idx)
 			{
@@ -749,7 +752,7 @@ namespace ChaosTest {
 					Particles[ParticleIndex]->ResetSmoothedVelocities();
 				}
 
-				HelpTickConstraints(SOAs,Particles,Graph,Constraints,PhysicsMaterials,PhysicalMaterials);
+				HelpTickConstraints(SOAs,Particles,Graph,Constraints,PhysicsMaterials,PerParticlePhysicsMaterials,PhysicalMaterials);
 			
 				// Particles 0-2 are always awake
 				EXPECT_FALSE(Particles[0]->Sleeping());
@@ -795,8 +798,10 @@ namespace ChaosTest {
 			FPBDRigidsSOAs SOAs(UniqueIndices);
 			TArray<TPBDRigidParticleHandle<FReal, 3>*> Particles = SOAs.CreateDynamicParticles(NumParticles);
 			TArrayCollectionArray<TSerializablePtr<FChaosPhysicsMaterial>> PhysicsMaterials;
+			TArrayCollectionArray<TUniquePtr<FChaosPhysicsMaterial>> PerParticlePhysicsMaterials;
 			THandleArray<FChaosPhysicsMaterial> PhysicalMaterials;
  			SOAs.GetParticleHandles().AddArray(&PhysicsMaterials);
+			SOAs.GetParticleHandles().AddArray(&PerParticlePhysicsMaterials);
 
 			for (int32 Idx = 0; Idx < NumParticles; ++Idx)
 			{
@@ -853,8 +858,8 @@ namespace ChaosTest {
 					}
 				}
 
-				HelpTickConstraints(SOAs,Particles,Graph,Constraints,PhysicsMaterials,PhysicalMaterials);
-			
+				HelpTickConstraints(SOAs, Particles, Graph, Constraints, PhysicsMaterials, PerParticlePhysicsMaterials, PhysicalMaterials);
+
 				// Particles 0-2 are always awake
 				EXPECT_FALSE(Particles[0]->Sleeping());
 				EXPECT_FALSE(Particles[1]->Sleeping());
@@ -902,8 +907,10 @@ namespace ChaosTest {
 			FPBDRigidsSOAs SOAs(UniqueIndices);
 			TArray<TPBDRigidParticleHandle<FReal,3>*> Particles = SOAs.CreateDynamicParticles(NumParticles);
 			TArrayCollectionArray<TSerializablePtr<FChaosPhysicsMaterial>> PhysicsMaterials;
+			TArrayCollectionArray<TUniquePtr<FChaosPhysicsMaterial>> PerParticlePhysicsMaterials;
 			THandleArray<FChaosPhysicsMaterial> PhysicalMaterials;
 			SOAs.GetParticleHandles().AddArray(&PhysicsMaterials);
+			SOAs.GetParticleHandles().AddArray(&PerParticlePhysicsMaterials);
 
 			for(int32 Idx = 0; Idx < NumParticles; ++Idx)
 			{
@@ -968,7 +975,7 @@ namespace ChaosTest {
 					}
 				}
 
-				HelpTickConstraints(SOAs,Particles,Graph,Constraints,PhysicsMaterials,PhysicalMaterials);
+				HelpTickConstraints(SOAs, Particles, Graph, Constraints, PhysicsMaterials, PerParticlePhysicsMaterials, PhysicalMaterials);
 
 				// Particle 2 is always awake
 				EXPECT_FALSE(Particles[2]->Sleeping());
diff --git a/Engine/Source/Runtime/Engine/Private/PhysicsEngine/Experimental/PhysInterface_Chaos.cpp b/Engine/Source/Runtime/Engine/Private/PhysicsEngine/Experimental/PhysInterface_Chaos.cpp
--- a/Engine/Source/Runtime/Engine/Private/PhysicsEngine/Experimental/PhysInterface_Chaos.cpp
+++ b/Engine/Source/Runtime/Engine/Private/PhysicsEngine/Experimental/PhysInterface_Chaos.cpp
@@ -683,6 +683,9 @@ void FPhysInterface_Chaos::AddGeometry(FPhysicsActorHandle& InActor, const FGeom
 {
 	TRACE_CPUPROFILER_EVENT_SCOPE(FPhysInterface_Chaos::AddGeometry);
 	LLM_SCOPE(ELLMTag::ChaosGeometry);
+
+	// @todo(chaos): we should not be creating unique geometry per actor
+	// @todo(chaos): we are creating the Shapes array twice. Once here and again in SetGeometry or MergeGeometry. Fix this.
 	TArray<TUniquePtr<Chaos::FImplicitObject>> Geoms;
 	Chaos::FShapesArray Shapes;
 	ChaosInterface::CreateGeometry(InParams, Geoms, Shapes);
diff --git a/Engine/Source/Runtime/Experimental/Chaos/Private/Chaos/Collision/PBDCollisionConstraint.cpp b/Engine/Source/Runtime/Experimental/Chaos/Private/Chaos/Collision/PBDCollisionConstraint.cpp
--- a/Engine/Source/Runtime/Experimental/Chaos/Private/Chaos/Collision/PBDCollisionConstraint.cpp
+++ b/Engine/Source/Runtime/Experimental/Chaos/Private/Chaos/Collision/PBDCollisionConstraint.cpp
@@ -127,8 +127,8 @@ namespace Chaos
 		OutConstraint.Particle[1] = Particle1;
 		OutConstraint.Implicit[0] = Implicit0;
 		OutConstraint.Implicit[1] = Implicit1;
-		OutConstraint.Shape[0] = Shape0;
-		OutConstraint.Shape[1] = Shape1;
+		OutConstraint.Shape[0] = (Shape0 != nullptr) ? Shape0->AsShapeInstance() : nullptr;
+		OutConstraint.Shape[1] = (Shape1 != nullptr) ? Shape1->AsShapeInstance() : nullptr;
 		OutConstraint.Simplicial[0] = Simplicial0;
 		OutConstraint.Simplicial[1] = Simplicial1;
 
@@ -216,7 +216,7 @@ namespace Chaos
 		: ImplicitTransform{ FRigidTransform3(), FRigidTransform3() }
 		, Particle{ Particle0, Particle1 }
 		, Implicit{ Implicit0, Implicit1 }
-		, Shape{ Shape0, Shape1 }
+		, Shape{ nullptr, nullptr }
 		, Simplicial{ Simplicial0, Simplicial1 }
 		, AccumulatedImpulse(0)
 		, ContainerCookie()
@@ -240,6 +240,18 @@ namespace Chaos
 		, CCDEnablePenetration(0)
 		, CCDTargetPenetration(0)
 	{
+		// These downcasts should never fail if the input shape is not null
+		// @todo(chaos): remove when FPerShapeData is deprecated
+		check((Shape0 == nullptr) || (Shape0->AsShapeInstance() != nullptr));
+		check((Shape1 == nullptr) || (Shape1->AsShapeInstance() != nullptr));
+		if (Shape0 != nullptr)
+		{
+			Shape[0] = Shape0->AsShapeInstance();
+		}
+		if (Shape1 != nullptr)
+		{
+			Shape[1] = Shape1->AsShapeInstance();
+		}
 	}
 
 	FPBDCollisionConstraint::~FPBDCollisionConstraint()
diff --git a/Engine/Source/Runtime/Experimental/Chaos/Private/Chaos/Island/IslandManager.cpp b/Engine/Source/Runtime/Experimental/Chaos/Private/Chaos/Island/IslandManager.cpp
--- a/Engine/Source/Runtime/Experimental/Chaos/Private/Chaos/Island/IslandManager.cpp
+++ b/Engine/Source/Runtime/Experimental/Chaos/Private/Chaos/Island/IslandManager.cpp
@@ -2,11 +2,12 @@
 #include "Chaos/Island/IslandManager.h"
 #include "Chaos/Island/SolverIsland.h"
 #include "Chaos/Island/IslandGraph.h"
-#include "Chaos/ParticleHandle.h"
 #include "Chaos/ConstraintHandle.h"
+#include "Chaos/ParticleHandle.h"
 #include "Chaos/PBDConstraintContainer.h"
 #include "Chaos/PBDRigidParticles.h"
 #include "Chaos/PBDRigidsSOAs.h"
+#include "Chaos/PhysicsMaterialUtilities.h"
 #include "ChaosLog.h"
 
 #include "Framework/Threading.h"
@@ -73,24 +74,6 @@ bool IsStationaryParticle(const FGeometryParticleHandle* ParticleHandle)
 	}
 }
 
-/** Check if a particle is dynamic or sleeping */
-inline const FChaosPhysicsMaterial* GetPhysicsMaterial(const TArrayCollectionArray<TSerializablePtr<FChaosPhysicsMaterial>>& ParticleMaterialAttributes,
-														    const THandleArray<FChaosPhysicsMaterial>& SolverPhysicsMaterials, FPBDRigidParticleHandle* RigidParticleHandle)
-{
-	const FChaosPhysicsMaterial* PhysicsMaterial = RigidParticleHandle->AuxilaryValue(ParticleMaterialAttributes).Get();
-	if (!PhysicsMaterial && RigidParticleHandle->ShapesArray().Num())
-	{
-		if (FPerShapeData* PerShapeData = RigidParticleHandle->ShapesArray()[0].Get())
-		{
-			if (PerShapeData->GetMaterials().Num())
-			{
-				PhysicsMaterial = SolverPhysicsMaterials.Get(PerShapeData->GetMaterials()[0].InnerHandle);
-			}
-		}
-	}
-	return PhysicsMaterial;
-}
-
 /** Check if an island is sleeping or not given a linear/angular velocities and sleeping thresholds */
 inline bool IsIslandSleeping(const FReal MaxLinearSpeed2, const FReal MaxAngularSpeed2,
 								  const FReal LinearSleepingThreshold, const FReal AngularSleepingThreshold,
@@ -118,10 +101,16 @@ inline bool IsIslandSleeping(const FReal MaxLinearSpeed2, const FReal MaxAngular
 }
 	
 /** Compute sleeping thresholds given a solver island  */
-inline bool ComputeSleepingThresholds(FPBDIsland* Island,
-	const TArrayCollectionArray<TSerializablePtr<FChaosPhysicsMaterial>>& PerParticleMaterialAttributes,
-	const THandleArray<FChaosPhysicsMaterial>& SolverPhysicsMaterials, FReal& LinearSleepingThreshold,
-	FReal& AngularSleepingThreshold, FReal& MaxLinearSpeed2, FReal& MaxAngularSpeed2, int32& SleepCounterThreshold)
+inline bool ComputeSleepingThresholds(
+	FPBDIsland* Island,
+	const TArrayCollectionArray<TSerializablePtr<FChaosPhysicsMaterial>>& PhysicsMaterials,
+	const TArrayCollectionArray<TUniquePtr<FChaosPhysicsMaterial>>& PerParticlePhysicsMaterials,
+	const THandleArray<FChaosPhysicsMaterial>& SimMaterials,
+	FReal& LinearSleepingThreshold, 
+	FReal& AngularSleepingThreshold, 
+	FReal& MaxLinearSpeed2, 
+	FReal& MaxAngularSpeed2, 
+	int32& SleepCounterThreshold)
 {
 	LinearSleepingThreshold = FLT_MAX;  
 	AngularSleepingThreshold = FLT_MAX; 
@@ -154,7 +143,7 @@ inline bool ComputeSleepingThresholds(FPBDIsland* Island,
 					MaxLinearSpeed2 = FMath::Max(PBDRigid->VSmooth().SizeSquared(), MaxLinearSpeed2);
 					MaxAngularSpeed2 = FMath::Max(PBDRigid->WSmooth().SizeSquared(), MaxAngularSpeed2);
 
-					const FChaosPhysicsMaterial* PhysicsMaterial = GetPhysicsMaterial(PerParticleMaterialAttributes, SolverPhysicsMaterials, PBDRigid);
+					const FChaosPhysicsMaterial* PhysicsMaterial = Private::GetFirstPhysicsMaterial(PBDRigid, PhysicsMaterials, PerParticlePhysicsMaterials, &SimMaterials);
 
 					const FReal LocalSleepingLinearThreshold = PhysicsMaterial ? PhysicsMaterial->SleepingLinearThreshold :
 										ChaosSolverCollisionDefaultLinearSleepThresholdCVar;
@@ -771,8 +760,9 @@ void FPBDIslandManager::ValidateIslands() const
 }
 
 bool FPBDIslandManager::SleepInactive(const int32 IslandIndex,
-	const TArrayCollectionArray<TSerializablePtr<FChaosPhysicsMaterial>>& PerParticleMaterialAttributes,
-	const THandleArray<FChaosPhysicsMaterial>& SolverPhysicsMaterials)
+	const TArrayCollectionArray<TSerializablePtr<FChaosPhysicsMaterial>>& PhysicsMaterials,
+	const TArrayCollectionArray<TUniquePtr<FChaosPhysicsMaterial>>& PerParticlePhysicsMaterials,
+	const THandleArray<FChaosPhysicsMaterial>& SimMaterials)
 {
 	// Only the persistent islands could start sleeping
 	const int32 GraphIndex = GetGraphIndex(IslandIndex);
@@ -786,8 +776,8 @@ bool FPBDIslandManager::SleepInactive(const int32 IslandIndex,
 	int32 SleepCounterThreshold = 0, NumDynamicParticles = 0;
 
 	// Compute of the linear/angular velocities + thresholds to make islands sleeping 
-	if (ComputeSleepingThresholds(Islands[GraphIndex].Get(), PerParticleMaterialAttributes,
-		SolverPhysicsMaterials, LinearSleepingThreshold, AngularSleepingThreshold,
+	if (ComputeSleepingThresholds(Islands[GraphIndex].Get(), PhysicsMaterials, PerParticlePhysicsMaterials,
+		SimMaterials, LinearSleepingThreshold, AngularSleepingThreshold,
 		MaxLinearSpeed2, MaxAngularSpeed2, SleepCounterThreshold))
 	{
 		int32 SleepCounter = Islands[GraphIndex]->GetSleepCounter();
diff --git a/Engine/Source/Runtime/Experimental/Chaos/Private/Chaos/PBDCollisionConstraints.cpp b/Engine/Source/Runtime/Experimental/Chaos/Private/Chaos/PBDCollisionConstraints.cpp
--- a/Engine/Source/Runtime/Experimental/Chaos/Private/Chaos/PBDCollisionConstraints.cpp
+++ b/Engine/Source/Runtime/Experimental/Chaos/Private/Chaos/PBDCollisionConstraints.cpp
@@ -2,10 +2,10 @@
 
 #include "Chaos/PBDCollisionConstraints.h"
 
+#include "Chaos/CastingUtilities.h"
 #include "Chaos/ChaosPerfTest.h"
 #include "Chaos/ContactModification.h"
 #include "Chaos/MidPhaseModification.h"
-#include "Chaos/PBDCollisionConstraintsContact.h"
 #include "Chaos/CCDModification.h"
 #include "Chaos/CollisionResolution.h"
 #include "Chaos/Collision/CollisionPruning.h"
@@ -17,14 +17,13 @@
 #include "Chaos/Evolution/ABTestingConstraintContainerSolver.h"
 #include "Chaos/GeometryQueries.h"
 #include "Chaos/Island/IslandManager.h"
-#include "Chaos/SpatialAccelerationCollection.h"
+#include "Chaos/PBDCollisionConstraintsContact.h"
 #include "Chaos/PBDRigidsSOAs.h"
-#include "Chaos/CastingUtilities.h"
+#include "Chaos/PhysicsMaterialUtilities.h"
+#include "Chaos/SpatialAccelerationCollection.h"
 #include "ChaosLog.h"
 #include "ChaosStats.h"
 #include "ProfilingDebugging/ScopedTimers.h"
-#include "Algo/Sort.h"
-#include "Algo/StableSort.h"
 
 //PRAGMA_DISABLE_OPTIMIZATION
 
@@ -406,56 +405,19 @@ namespace Chaos
 		return ConstraintAllocator.GetConstConstraints();
 	}
 
-	const FChaosPhysicsMaterial* GetPhysicsMaterial(const TGeometryParticleHandle<FReal, 3>* Particle, const FImplicitObject* Geom, const TArrayCollectionArray<TSerializablePtr<FChaosPhysicsMaterial>>& PhysicsMaterials, const TArrayCollectionArray<TUniquePtr<FChaosPhysicsMaterial>>& PerParticlePhysicsMaterials, const THandleArray<FChaosPhysicsMaterial>* const SimMaterials)
+	void FPBDCollisionConstraints::UpdateConstraintMaterialProperties(FPBDCollisionConstraint& Constraint)
 	{
-		// Use the per-particle material if it exists
-		const FChaosPhysicsMaterial* UniquePhysicsMaterial = Particle->AuxilaryValue(PerParticlePhysicsMaterials).Get();
-		if (UniquePhysicsMaterial != nullptr)
+		// We only support one material shared by all manifold points for now, even when our 
+		// 4 manifold points are on different triangles of a mesh for example
+		int32 ShapeFaceIndex = INDEX_NONE;
+		if (Constraint.NumManifoldPoints() > 0)
 		{
-			return UniquePhysicsMaterial;
-		}
-		const FChaosPhysicsMaterial* PhysicsMaterial = Particle->AuxilaryValue(PhysicsMaterials).Get();
-		if (PhysicsMaterial != nullptr)
-		{
-			return PhysicsMaterial;
+			ShapeFaceIndex = Constraint.GetManifoldPoint(0).ContactPoint.FaceIndex;
 		}
 
-		// If no particle material, see if the shape has one
-		// @todo(chaos): handle materials for meshes etc
-		for (const TUniquePtr<FPerShapeData>& ShapeData : Particle->ShapesArray())
-		{
-			const FImplicitObject* OuterShapeGeom = ShapeData->GetGeometry().Get();
-			const FImplicitObject* InnerShapeGeom = Utilities::ImplicitChildHelper(OuterShapeGeom);
-			if (Geom == OuterShapeGeom || Geom == InnerShapeGeom)
-			{
-				if (ShapeData->GetMaterials().Num() > 0)
-				{
-					if(SimMaterials)
-					{
-						return SimMaterials->Get(ShapeData->GetMaterials()[0].InnerHandle);
-					}
-					else
-					{
-						UE_LOG(LogChaos, Warning, TEXT("Attempted to resolve a material for a constraint but we do not have a sim material container."));
-					}
-				}
-				else
-				{
-					// This shape doesn't have a material assigned
-					return nullptr;
-				}
-			}
-		}
-
-		// The geometry used for this particle does not belong to the particle.
-		// This can happen in the case of fracture.
-		return nullptr;
-	}
-
-	void FPBDCollisionConstraints::UpdateConstraintMaterialProperties(FPBDCollisionConstraint& Constraint)
-	{
-		const FChaosPhysicsMaterial* PhysicsMaterial0 = GetPhysicsMaterial(Constraint.Particle[0], Constraint.Implicit[0], MPhysicsMaterials, MPerParticlePhysicsMaterials, SimMaterials);
-		const FChaosPhysicsMaterial* PhysicsMaterial1 = GetPhysicsMaterial(Constraint.Particle[1], Constraint.Implicit[1], MPhysicsMaterials, MPerParticlePhysicsMaterials, SimMaterials);
+		// This is a bit dodgy - we pass the FaceIndex to both material requests, knowing that at most one of the shapes will use it
+		const FChaosPhysicsMaterial* PhysicsMaterial0 = Private::GetPhysicsMaterial(Constraint.Particle[0], Constraint.GetShape0(), ShapeFaceIndex, MPhysicsMaterials, MPerParticlePhysicsMaterials, SimMaterials);
+		const FChaosPhysicsMaterial* PhysicsMaterial1 = Private::GetPhysicsMaterial(Constraint.Particle[1], Constraint.GetShape1(), ShapeFaceIndex, MPhysicsMaterials, MPerParticlePhysicsMaterials, SimMaterials);
 
 		FReal MaterialRestitution = 0;
 		FReal MaterialRestitutionThreshold = 0;
diff --git a/Engine/Source/Runtime/Experimental/Chaos/Private/Chaos/PBDRigidsEvolutionGBF.cpp b/Engine/Source/Runtime/Experimental/Chaos/Private/Chaos/PBDRigidsEvolutionGBF.cpp
--- a/Engine/Source/Runtime/Experimental/Chaos/Private/Chaos/PBDRigidsEvolutionGBF.cpp
+++ b/Engine/Source/Runtime/Experimental/Chaos/Private/Chaos/PBDRigidsEvolutionGBF.cpp
@@ -311,7 +311,7 @@ void FPBDRigidsEvolutionGBF::BuildDisabledParticles(const int32 Island, TArray<T
 	}
 
 	// Turn off if not moving
-	SleepedIslands[Island] = GetIslandManager().SleepInactive(Island, PhysicsMaterials, SolverPhysicsMaterials);
+	SleepedIslands[Island] = GetIslandManager().SleepInactive(Island, PhysicsMaterials, PerParticlePhysicsMaterials, SolverPhysicsMaterials);
 }
 
 void FPBDRigidsEvolutionGBF::UpdateCollisionSolverType()
diff --git a/Engine/Source/Runtime/Experimental/Chaos/Private/Chaos/PhysicsMaterialUtilities.cpp b/Engine/Source/Runtime/Experimental/Chaos/Private/Chaos/PhysicsMaterialUtilities.cpp
new file mode 100644
--- /dev/null
+++ b/Engine/Source/Runtime/Experimental/Chaos/Private/Chaos/PhysicsMaterialUtilities.cpp
@@ -0,0 +1,63 @@
+// Copyright Epic Games, Inc. All Rights Reserved.
+#include "Chaos/PhysicsMaterialUtilities.h"
+#include "Chaos/ParticleHandle.h"
+#include "Chaos/ShapeInstance.h"
+
+namespace Chaos
+{
+	namespace Private
+	{
+		const FChaosPhysicsMaterial* GetPhysicsMaterial(
+			const TGeometryParticleHandle<FReal, 3>* Particle,
+			const FShapeInstance* Shape,
+			const int32 ShapeFaceIndex,
+			const TArrayCollectionArray<TSerializablePtr<FChaosPhysicsMaterial>>& PhysicsMaterials,
+			const TArrayCollectionArray<TUniquePtr<FChaosPhysicsMaterial>>& PerParticlePhysicsMaterials,
+			const THandleArray<FChaosPhysicsMaterial>* const SimMaterials)
+		{
+			// Use the per-particle material if it exists
+			const FChaosPhysicsMaterial* UniquePhysicsMaterial = Particle->AuxilaryValue(PerParticlePhysicsMaterials).Get();
+			if (UniquePhysicsMaterial != nullptr)
+			{
+				return UniquePhysicsMaterial;
+			}
+
+			// Use the.... other?.... per particle material if it exists. 
+			// @todo(chaos): I assume this is to support writable unique materials, but there must be a better 
+			// way that adding another pointer to every particle. E.g., a ref count or shared setting on the material.
+			const FChaosPhysicsMaterial* PhysicsMaterial = Particle->AuxilaryValue(PhysicsMaterials).Get();
+			if (PhysicsMaterial != nullptr)
+			{
+				return PhysicsMaterial;
+			}
+
+			// If no particle material, see if the shape has one
+			if ((Shape != nullptr) && (SimMaterials != nullptr) && (Shape->NumMaterials() > 0))
+			{
+				// By default we take the first material (all non-mesh shapes will use material 0)
+				int32 ShapeMaterialIndex = 0;
+
+				if ((Shape->NumMaterials() > 1) && (ShapeFaceIndex != INDEX_NONE))
+				{
+					// Multi-shape object (tri mesh, heightfield, etc) - use the face index
+					ShapeMaterialIndex = Shape->GetGeometry()->GetMaterialIndex(ShapeFaceIndex);
+				}
+
+				if (ShapeMaterialIndex < Shape->NumMaterials())
+				{
+					return SimMaterials->Get(Shape->GetMaterial(ShapeMaterialIndex).InnerHandle);
+				}
+			}
+
+			// Either no material is assigned, or the geometry used for this particle 
+			// does not belong to the particle. This can happen in the case of fracture.
+			return nullptr;
+		}
+
+		const FChaosPhysicsMaterial* GetFirstPhysicsMaterial(const TGeometryParticleHandle<FReal, 3>* Particle, const TArrayCollectionArray<TSerializablePtr<FChaosPhysicsMaterial>>& PhysicsMaterials, const TArrayCollectionArray<TUniquePtr<FChaosPhysicsMaterial>>& PerParticlePhysicsMaterials, const THandleArray<FChaosPhysicsMaterial>* const SimMaterials)
+		{
+			const FShapeInstance* Shape = (Particle->ShapeInstances().IsEmpty()) ? nullptr : Particle->ShapeInstances()[0].Get();
+			return GetPhysicsMaterial(Particle, Shape, INDEX_NONE, PhysicsMaterials, PerParticlePhysicsMaterials, SimMaterials);
+		}
+	}
+}
\ No newline at end of file
diff --git a/Engine/Source/Runtime/Experimental/Chaos/Private/Chaos/ShapeInstance.cpp b/Engine/Source/Runtime/Experimental/Chaos/Private/Chaos/ShapeInstance.cpp
--- a/Engine/Source/Runtime/Experimental/Chaos/Private/Chaos/ShapeInstance.cpp
+++ b/Engine/Source/Runtime/Experimental/Chaos/Private/Chaos/ShapeInstance.cpp
@@ -243,9 +243,10 @@ namespace Chaos
 		DownCast([&Ar](auto& ShapeInstance)
 			{
 				Ar << ShapeInstance.CollisionData;
-				Ar << ShapeInstance.Materials;
 			});
 
+		SerializeMaterials(Ar);
+
 		if (Ar.CustomVer(FExternalPhysicsCustomObjectVersion::GUID) >= FExternalPhysicsCustomObjectVersion::SerializeShapeWorldSpaceBounds)
 		{
 			TBox<FReal, 3>::SerializeAsAABB(Ar, WorldSpaceInflatedShapeBounds);
@@ -270,12 +271,42 @@ namespace Chaos
 		return Ar.IsLoading() ? new FShapeInstanceProxy(0) : nullptr;
 	}
 
+	void FShapeInstanceProxy::SerializeMaterials(FChaosArchive& Ar)
+	{
+		Ar << Materials;
+	}
+
 	FShapeInstance* FShapeInstance::SerializationFactory(FChaosArchive& Ar, FShapeInstance*)
 	{
 		//todo: need to rework serialization for shapes, for now just give them all shape idx 0
 		return Ar.IsLoading() ? new FShapeInstance(0) : nullptr;
 	}
 
+	void FShapeInstance::SerializeMaterials(FChaosArchive& Ar)
+	{
+		// NOTE: FShapeInstanceProxy and FShapeInstance used to be serialized as FPerShapeData
+		// so we must ensure what we serialize here is binary compatible with FShapeInstanceProxy::SerializeMaterials
+		if (Ar.IsLoading())
+		{
+			FMaterialData Data;
+			Ar << Data;
+			SetMaterialData(Data);
+		}
+		else
+		{
+			if (bIsSingleMaterial)
+			{
+				FMaterialData Data;
+				Data.Materials.Add(Material.MaterialHandle);
+				Ar << Data;
+			}
+			else
+			{
+				Ar << GetMaterialDataImpl();
+			}
+		}
+	}
+
 	///////////////////////////////////////////////////////////////////////////////////////////////
 	///////////////////////////////////////////////////////////////////////////////////////////////
 	// 
diff --git a/Engine/Source/Runtime/Experimental/Chaos/Private/EventDefaults.cpp b/Engine/Source/Runtime/Experimental/Chaos/Private/EventDefaults.cpp
--- a/Engine/Source/Runtime/Experimental/Chaos/Private/EventDefaults.cpp
+++ b/Engine/Source/Runtime/Experimental/Chaos/Private/EventDefaults.cpp
@@ -32,26 +32,24 @@ namespace Chaos
 	 */
 	FMaterialHandle ResolveMaterial(const FPerShapeData* InShape, const FPBDCollisionConstraint& InConstraint)
 	{
-		if(InShape)
-		{
-			const TArray<FMaterialHandle>& MaterialArray = InShape->GetMaterials();
+		// @todo(chaos): this does not handle PerParticleMaterials so may return the wrong value when that is used
 
+		if (InShape && (InShape->NumMaterials() > 0))
+		{
 			// Simple case, one material. All primitives (Box, convex, sphere etc.) should only have one material.
 			// Heightfield and Trimesh can have one or more and will require data from the contacts to resolve.
-			if(MaterialArray.Num() == 1)
+			if (InShape->NumMaterials() == 1)
 			{
-				return MaterialArray[0];
+				return InShape->GetMaterial(0);
 			}
-			else
+			else if (InConstraint.NumManifoldPoints() > 0)
 			{
-				// Check the manifold points, return the first valid material based on the face that the contact hit.
-				for(const FManifoldPoint& Point : InConstraint.GetManifoldPoints())
+				// We only support one material per manifold (just use the first manifold point)
+				const int32 ShapeFaceIndex = InConstraint.GetManifoldPoint(0).ContactPoint.FaceIndex;
+				const int32 ShapeMaterialIndex = InShape->GetGeometry()->GetMaterialIndex(ShapeFaceIndex);
+				if (ShapeMaterialIndex < InShape->NumMaterials())
 				{
-					const int32 PotentialMaterialIndex = InShape->GetGeometry()->GetMaterialIndex(Point.ContactPoint.FaceIndex);
-					if(MaterialArray.IsValidIndex(PotentialMaterialIndex))
-					{
-						return MaterialArray[PotentialMaterialIndex];
-					}
+					return InShape->GetMaterial(ShapeMaterialIndex);
 				}
 			}
 		}
diff --git a/Engine/Source/Runtime/Experimental/Chaos/Public/Chaos/Collision/PBDCollisionConstraint.h b/Engine/Source/Runtime/Experimental/Chaos/Public/Chaos/Collision/PBDCollisionConstraint.h
--- a/Engine/Source/Runtime/Experimental/Chaos/Public/Chaos/Collision/PBDCollisionConstraint.h
+++ b/Engine/Source/Runtime/Experimental/Chaos/Public/Chaos/Collision/PBDCollisionConstraint.h
@@ -27,10 +27,11 @@ namespace Chaos
 	class FParticlePairMidPhase;
 	class FPBDCollisionConstraint;
 	class FPBDCollisionConstraints;
+	class FPerShapeData;
+	class FShapeInstance;
 	class FSingleShapePairCollisionDetector;
 	class FSolverBody;
 	class FSolverBodyContainer;
-	class FPerShapeData;
 
 	UE_DEPRECATED(4.27, "Use FPBDCollisionConstraint instead")
 	typedef FPBDCollisionConstraint FRigidBodyPointContactConstraint;
@@ -44,7 +45,8 @@ namespace Chaos
 	{
 	public:
 		FPBDCollisionConstraintMaterial()
-			: MaterialDynamicFriction(0)
+			: FaceIndex(INDEX_NONE)
+			, MaterialDynamicFriction(0)
 			, MaterialStaticFriction(0)
 			, MaterialRestitution(0)
 			, DynamicFriction(0)
@@ -58,7 +60,11 @@ namespace Chaos
 		{
 		}
 
+		// The face index that the material was extracted from
+		int32 FaceIndex;
+
 		// Material properties pulled from the materials of the two shapes involved in the contact
+		// @todo(chaos): we can remove these now and just recollect the material data when a modifier was applied
 		FRealSingle MaterialDynamicFriction;
 		FRealSingle MaterialStaticFriction;
 		FRealSingle MaterialRestitution;
@@ -282,9 +288,9 @@ namespace Chaos
 		const FImplicitObject* GetImplicit1() const { return Implicit[1]; }
 		const FImplicitObject* GetImplicit(const int32 ParticleIndex) const { check((ParticleIndex >= 0) && (ParticleIndex < 2)); return Implicit[ParticleIndex]; }
 
-		const FPerShapeData* GetShape0() const { return Shape[0]; }
-		const FPerShapeData* GetShape1() const { return Shape[1]; }
-		const FPerShapeData* GetShape(const int32 ParticleIndex) const { check((ParticleIndex >= 0) && (ParticleIndex < 2)); return Shape[ParticleIndex]; }
+		const FShapeInstance* GetShape0() const { return Shape[0]; }
+		const FShapeInstance* GetShape1() const { return Shape[1]; }
+		const FShapeInstance* GetShape(const int32 ParticleIndex) const { check((ParticleIndex >= 0) && (ParticleIndex < 2)); return Shape[ParticleIndex]; }
 
 		const FBVHParticles* GetCollisionParticles0() const { return Simplicial[0]; }
 		const FBVHParticles* GetCollisionParticles1() const { return Simplicial[1]; }
@@ -509,6 +515,18 @@ namespace Chaos
 					}
 				}
 			}
+
+			// If we have a new face we may need to change the material
+			// @todo(chaos): support per-manifold point materials?
+			if (NumManifoldPoints() > 0)
+			{
+				if (ManifoldPoints[0].ContactPoint.FaceIndex != Material.FaceIndex)
+				{
+					Material.FaceIndex = ManifoldPoints[0].ContactPoint.FaceIndex;
+					Flags.bMaterialSet = false;
+				}
+			}
+
 		}
 
 		void UpdateManifoldContacts();
@@ -768,7 +786,7 @@ namespace Chaos
 		
 		FGeometryParticleHandle* Particle[2];
 		const FImplicitObject* Implicit[2];
-		const FPerShapeData* Shape[2];
+		const FShapeInstance* Shape[2];
 		const FBVHParticles* Simplicial[2];
 
 	public:
diff --git a/Engine/Source/Runtime/Experimental/Chaos/Public/Chaos/Island/IslandManager.h b/Engine/Source/Runtime/Experimental/Chaos/Public/Chaos/Island/IslandManager.h
--- a/Engine/Source/Runtime/Experimental/Chaos/Public/Chaos/Island/IslandManager.h
+++ b/Engine/Source/Runtime/Experimental/Chaos/Public/Chaos/Island/IslandManager.h
@@ -209,7 +209,11 @@ public:
 	  * @param SolverPhysicsMaterials Solver materials that could be used to find the sleeping thresholds
 	  * @return return true if the island is going to sleep
 	  */
-	bool SleepInactive(const int32 IslandIndex, const TArrayCollectionArray<TSerializablePtr<FChaosPhysicsMaterial>>& PerParticleMaterialAttributes, const THandleArray<FChaosPhysicsMaterial>& SolverPhysicsMaterials);
+	bool SleepInactive(
+		const int32 IslandIndex, 
+		const TArrayCollectionArray<TSerializablePtr<FChaosPhysicsMaterial>>& PhysicsMaterials, 
+		const TArrayCollectionArray<TUniquePtr<FChaosPhysicsMaterial>>& PerParticlePhysicsMaterials,
+		const THandleArray<FChaosPhysicsMaterial>& SimMaterials);
 
 	/**
 	 * Get the set of islands that the particle belongs to. Will be 0-1 island for dynamics and could be any number for non-dynamics.
diff --git a/Engine/Source/Runtime/Experimental/Chaos/Public/Chaos/PhysicalMaterials.h b/Engine/Source/Runtime/Experimental/Chaos/Public/Chaos/PhysicalMaterials.h
--- a/Engine/Source/Runtime/Experimental/Chaos/Public/Chaos/PhysicalMaterials.h
+++ b/Engine/Source/Runtime/Experimental/Chaos/Public/Chaos/PhysicalMaterials.h
@@ -27,6 +27,8 @@ namespace Chaos
 
 		bool IsValid() const { return Get() != nullptr; }
 
+		void Reset() { InnerHandle = FChaosMaterialHandle(); }
+
 		friend bool operator ==(const FMaterialHandle& A, const FMaterialHandle& B) { return A.InnerHandle == B.InnerHandle; }
 
 		friend FArchive& operator <<(FArchive& Ar, FMaterialHandle& Value) { Ar << Value.InnerHandle; return Ar; }
diff --git a/Engine/Source/Runtime/Experimental/Chaos/Public/Chaos/PhysicsMaterialUtilities.h b/Engine/Source/Runtime/Experimental/Chaos/Public/Chaos/PhysicsMaterialUtilities.h
new file mode 100644
--- /dev/null
+++ b/Engine/Source/Runtime/Experimental/Chaos/Public/Chaos/PhysicsMaterialUtilities.h
@@ -0,0 +1,34 @@
+// Copyright Epic Games, Inc. All Rights Reserved.
+#pragma once
+
+#include "Chaos/ArrayCollectionArray.h"
+#include "Chaos/Framework/Handles.h"
+#include "Chaos/ParticleHandleFwd.h"
+#include "Chaos/Serializable.h"
+#include "Chaos/ShapeInstanceFwd.h"
+#include "Templates/UniquePtr.h"
+
+class FChaosPhysicsMaterial;
+
+namespace Chaos
+{
+	namespace Private
+	{
+		// Get the physical material for the specified Shape on the particle.
+		// NOTE: each shape may have its own material, but the particle may also have an override.
+		extern CHAOS_API const FChaosPhysicsMaterial* GetPhysicsMaterial(
+			const FGeometryParticleHandle* Particle,
+			const FShapeInstance* Shape,
+			const int32 FaceIndex,
+			const TArrayCollectionArray<TSerializablePtr<FChaosPhysicsMaterial>>& PhysicsMaterials,
+			const TArrayCollectionArray<TUniquePtr<FChaosPhysicsMaterial>>& PerParticlePhysicsMaterials,
+			const THandleArray<FChaosPhysicsMaterial>* const SimMaterials);
+
+		// Get the first physical material on the particle (a multi-shape particle may have many)
+		extern CHAOS_API const FChaosPhysicsMaterial* GetFirstPhysicsMaterial(
+			const FGeometryParticleHandle* Particle,
+			const TArrayCollectionArray<TSerializablePtr<FChaosPhysicsMaterial>>& PhysicsMaterials,
+			const TArrayCollectionArray<TUniquePtr<FChaosPhysicsMaterial>>& PerParticlePhysicsMaterials,
+			const THandleArray<FChaosPhysicsMaterial>* const SimMaterials);
+	}
+}
diff --git a/Engine/Source/Runtime/Experimental/Chaos/Public/Chaos/ShapeInstance.h b/Engine/Source/Runtime/Experimental/Chaos/Public/Chaos/ShapeInstance.h
--- a/Engine/Source/Runtime/Experimental/Chaos/Public/Chaos/ShapeInstance.h
+++ b/Engine/Source/Runtime/Experimental/Chaos/Public/Chaos/ShapeInstance.h
@@ -47,18 +47,19 @@ namespace Chaos
 			return Type;
 		}
 
-		FShapeInstanceProxy* AsShapeInstanceProxy();
-		const FShapeInstanceProxy* AsShapeInstanceProxy() const;
-
-		FShapeInstance* AsShapeInstance();
-		const FShapeInstance* AsShapeInstance() const;
+		// Call a function on the concrete type
+		template<typename TLambda> decltype(auto) DownCast(const TLambda& Lambda);
+		template<typename TLambda> decltype(auto) DownCast(const TLambda& Lambda) const;
 
 		Private::FShapeInstanceExtended* AsShapeInstanceExtended();
 		const Private::FShapeInstanceExtended* AsShapeInstanceExtended() const;
 
-		// Call a function on the concrete type
-		template<typename TLambda> decltype(auto) DownCast(const TLambda& Lambda);
-		template<typename TLambda> decltype(auto) DownCast(const TLambda& Lambda) const;
+	public:
+		// Downcasts exposed until we deprecate and remove FPerShapeData
+		FShapeInstanceProxy* AsShapeInstanceProxy();
+		const FShapeInstanceProxy* AsShapeInstanceProxy() const;
+		FShapeInstance* AsShapeInstance();
+		const FShapeInstance* AsShapeInstance() const;
 
 	public:
 		static constexpr bool AlwaysSerializable = true;
@@ -109,6 +110,9 @@ namespace Chaos
 
 		void UpdateLeafWorldTransform(FGeometryParticleHandle* Particle);
 
+		int32 NumMaterials() const;
+		const FMaterialHandle& GetMaterial(const int32 Index) const;
+
 		const TArray<FMaterialHandle>& GetMaterials() const;
 		void SetMaterial(FMaterialHandle InMaterial);
 		void SetMaterials(const TArray<FMaterialHandle>& InMaterials);
@@ -158,6 +162,7 @@ namespace Chaos
 	protected:
 		FPerShapeData(const EPerShapeDataType InType, int32 InShapeIdx)
 			: Type(InType)
+			, bIsSingleMaterial(false)
 			, ShapeIdx(InShapeIdx)
 			, Geometry()
 			, WorldSpaceInflatedShapeBounds(FAABB3(FVec3(0), FVec3(0)))
@@ -166,6 +171,7 @@ namespace Chaos
 
 		FPerShapeData(const EPerShapeDataType InType, int32 InShapeIdx, TSerializablePtr<FImplicitObject> InGeometry)
 			: Type(InType)
+			, bIsSingleMaterial(false)
 			, ShapeIdx(InShapeIdx)
 			, Geometry(InGeometry)
 			, WorldSpaceInflatedShapeBounds(FAABB3(FVec3(0), FVec3(0)))
@@ -174,13 +180,17 @@ namespace Chaos
 
 		FPerShapeData(const EPerShapeDataType InType, const FPerShapeData& Other)
 			: Type(InType)
+			, bIsSingleMaterial(Other.bIsSingleMaterial)
 			, ShapeIdx(Other.ShapeIdx)
 			, Geometry(Other.Geometry)
 			, WorldSpaceInflatedShapeBounds(Other.WorldSpaceInflatedShapeBounds)
 		{
 		}
 
+		virtual void SerializeMaterials(FChaosArchive& Ar) = 0;
+
 		EPerShapeDataType Type;
+		uint8 bIsSingleMaterial : 1;	// For use by FShapeInstance (here because the space is available for free)
 		int32 ShapeIdx;
 		TSerializablePtr<FImplicitObject> Geometry;
 		TAABB<FReal, 3> WorldSpaceInflatedShapeBounds;
@@ -257,6 +267,9 @@ namespace Chaos
 		FRigidTransform3 GetLeafWorldTransform(const FGeometryParticleHandle* Particle) const;
 		void UpdateLeafWorldTransform(FGeometryParticleHandle* Particle);
 
+		int32 NumMaterials() const { return Materials.Read(). Materials.Num(); }
+		const FMaterialHandle& GetMaterial(const int32 Index) const { return Materials.Read().Materials[Index]; }
+
 		const TArray<FMaterialHandle>& GetMaterials() const { return Materials.Read().Materials; }
 		const TArray<FMaterialMaskHandle>& GetMaterialMasks() const { return Materials.Read().MaterialMasks; }
 		const TArray<uint32>& GetMaterialMaskMaps() const { return Materials.Read().MaterialMaskMaps; }
@@ -427,6 +440,9 @@ namespace Chaos
 		{
 		}
 
+		virtual void SerializeMaterials(FChaosArchive& Ar) override final;
+
+
 		IPhysicsProxyBase* Proxy;
 		FShapeDirtyFlags DirtyFlags;
 
@@ -458,7 +474,10 @@ namespace Chaos
 	 *
 	 * NOTE: keep size to a minimum. There can be millions of these in s scene.
 	 *
-	 * @todo(chaos) : reduce the cost of MaterialData for shapes with one materialand no masks etc.
+	 * NOTE: this version has reduced-memory material storage compared to FShapeInstanceProxy, 
+	 * but increased cost when switching from a single-material to a multi-material shape
+	 * (though this should not be happening much if at all).
+	 * 
 	 */
 	class CHAOS_API FShapeInstance : public FPerShapeData
 	{
@@ -469,6 +488,14 @@ namespace Chaos
 		static void UpdateGeometry(TUniquePtr<FShapeInstance>& InOutShapePtr, TSerializablePtr<FImplicitObject> InGeometry);
 		static FShapeInstance* SerializationFactory(FChaosArchive& Ar, FShapeInstance*);
 
+		virtual ~FShapeInstance()
+		{
+			if (!bIsSingleMaterial && (Material.MaterialData != nullptr))
+			{
+				delete Material.MaterialData;
+			}
+		}
+
 		void UpdateShapeBounds(const FRigidTransform3& WorldTM, const FVec3& BoundsExpansion = FVec3(0));
 
 		void* GetUserData() const { return CollisionData.UserData; }
@@ -494,16 +521,42 @@ namespace Chaos
 		FRigidTransform3 GetLeafWorldTransform(const FGeometryParticleHandle* Particle) const;
 		void UpdateLeafWorldTransform(FGeometryParticleHandle* Particle);
 
-		const TArray<FMaterialHandle>& GetMaterials() const { return Materials.Materials; }
-		const TArray<FMaterialMaskHandle>& GetMaterialMasks() const { return Materials.MaterialMasks; }
-		const TArray<uint32>& GetMaterialMaskMaps() const { return Materials.MaterialMaskMaps; }
-		const TArray<FMaterialHandle>& GetMaterialMaskMapMaterials() const { return Materials.MaterialMaskMapMaterials; }
+		int32 NumMaterials() const
+		{
+			if (bIsSingleMaterial)
+			{
+				return Material.MaterialHandle.IsValid() ? 1 : 0;
+			}
+			else
+			{
+				return (Material.MaterialData != nullptr) ? Material.MaterialData->Materials.Num() : 0;
+			}
+		}
 
-		void SetMaterial(FMaterialHandle InMaterial) { Materials.Materials.Reset(1); Materials.Materials.Add(InMaterial); }
-		void SetMaterials(const TArray<FMaterialHandle>& InMaterials) { Materials.Materials = InMaterials; }
-		void SetMaterialMasks(const TArray<FMaterialMaskHandle>& InMaterialMasks) { Materials.MaterialMasks = InMaterialMasks; }
-		void SetMaterialMaskMaps(const TArray<uint32>& InMaterialMaskMaps) { Materials.MaterialMaskMaps = InMaterialMaskMaps; }
-		void SetMaterialMaskMapMaterials(const TArray<FMaterialHandle>& InMaterialMaskMapMaterials) { Materials.MaterialMaskMapMaterials = InMaterialMaskMapMaterials; }
+		const FMaterialHandle& GetMaterial(const int32 Index) const
+		{
+			if (bIsSingleMaterial)
+			{
+				check(Index == 0);
+				check(Material.MaterialHandle.IsValid());
+				return Material.MaterialHandle;
+			}
+			else
+			{
+				return GetMaterialDataImpl().Materials[Index];
+			}
+		}
+
+		const TArray<FMaterialHandle>& GetMaterials() const { return GetMaterialDataImpl().Materials; }
+		const TArray<FMaterialMaskHandle>& GetMaterialMasks() const { return GetMaterialDataImpl().MaterialMasks; }
+		const TArray<uint32>& GetMaterialMaskMaps() const { return GetMaterialDataImpl().MaterialMaskMaps; }
+		const TArray<FMaterialHandle>& GetMaterialMaskMapMaterials() const { return GetMaterialDataImpl().MaterialMaskMapMaterials; }
+
+		void SetMaterial(FMaterialHandle InMaterial) { SetMaterialImpl(InMaterial); }
+		void SetMaterials(const TArray<FMaterialHandle>& InMaterials) { GetMaterialDataImpl().Materials = InMaterials; }
+		void SetMaterialMasks(const TArray<FMaterialMaskHandle>& InMaterialMasks) { GetMaterialDataImpl().MaterialMasks = InMaterialMasks; }
+		void SetMaterialMaskMaps(const TArray<uint32>& InMaterialMaskMaps) { GetMaterialDataImpl().MaterialMaskMaps = InMaterialMaskMaps; }
+		void SetMaterialMaskMapMaterials(const TArray<FMaterialHandle>& InMaterialMaskMapMaterials) { GetMaterialDataImpl().MaterialMaskMapMaterials = InMaterialMaskMapMaterials; }
 
 		bool GetQueryEnabled() const { return CollisionData.bQueryCollision; }
 		void SetQueryEnabled(const bool bEnable) { CollisionData.bQueryCollision = bEnable; }
@@ -520,8 +573,8 @@ namespace Chaos
 		const FCollisionData& GetCollisionData() const { return CollisionData; }
 		void SetCollisionData(const FCollisionData& Data) { CollisionData = Data; }
 
-		const FMaterialData& GetMaterialData() const { return Materials; }
-		void SetMaterialData(const FMaterialData& Data) { Materials = Data; }
+		const FMaterialData& GetMaterialData() const { return GetMaterialDataImpl(); }
+		void SetMaterialData(const FMaterialData& Data) { SetMaterialDataImpl(Data); }
 
 		// @todo(chaos): remove when FPerShapeData is removed
 		const FShapeDirtyFlags GetDirtyFlags() const { check(false); return FShapeDirtyFlags(); }
@@ -529,49 +582,137 @@ namespace Chaos
 		void SetProxy(IPhysicsProxyBase* InProxy) { check(false); }
 
 		template <typename Lambda> void ModifySimData(const Lambda& LambdaFunc) { LambdaFunc(CollisionData.SimData); }
-		template <typename Lambda> void ModifyMaterials(const Lambda& LambdaFunc) { LambdaFunc(Materials.Materials); }
-		template <typename Lambda> void ModifyMaterialMasks(const Lambda& LambdaFunc) { LambdaFunc(Materials.MaterialMasks); }
-		template <typename Lambda> void ModifyMaterialMaskMaps(const Lambda& LambdaFunc) { LambdaFunc(Materials.MaterialMaskMaps); }
-		template <typename Lambda> void ModifyMaterialMaskMapMaterials(const Lambda& LambdaFunc) { LambdaFunc(Materials.MaterialMaskMapMaterials); }
+		template <typename Lambda> void ModifyMaterials(const Lambda& LambdaFunc) { LambdaFunc(GetMaterialDataImpl().Materials); }
+		template <typename Lambda> void ModifyMaterialMasks(const Lambda& LambdaFunc) { LambdaFunc(GetMaterialDataImpl().MaterialMasks); }
+		template <typename Lambda> void ModifyMaterialMaskMaps(const Lambda& LambdaFunc) { LambdaFunc(GetMaterialDataImpl().MaterialMaskMaps); }
+		template <typename Lambda> void ModifyMaterialMaskMapMaterials(const Lambda& LambdaFunc) { LambdaFunc(GetMaterialDataImpl().MaterialMaskMapMaterials); }
 
 	protected:
 		explicit FShapeInstance(int32 InShapeIdx)
 			: FPerShapeData(EPerShapeDataType::Sim, InShapeIdx)
 			, CollisionData()
-			, Materials()
 		{
+			bIsSingleMaterial = true;
 		}
 
 		FShapeInstance(int32 InShapeIdx, TSerializablePtr<FImplicitObject> InGeometry)
 			: FPerShapeData(EPerShapeDataType::Sim, InShapeIdx, InGeometry)
 			, CollisionData()
-			, Materials()
 		{
+			bIsSingleMaterial = true;
 		}
 
 		explicit FShapeInstance(FShapeInstance&& Other)
 			: FPerShapeData(EPerShapeDataType::Sim, Other)
 			, CollisionData(MoveTemp(Other.CollisionData))
-			, Materials(MoveTemp(Other.Materials))
 		{
+			if (bIsSingleMaterial)
+			{
+				Material.MaterialHandle = Other.Material.MaterialHandle;
+			}
+			else
+			{
+				Material.MaterialData = Other.Material.MaterialData;
+				Other.Material.MaterialData = nullptr;
+			}
 		}
 
 		FShapeInstance(const EPerShapeDataType InType, int32 InShapeIdx, TSerializablePtr<FImplicitObject> InGeometry)
 			: FPerShapeData(InType, InShapeIdx, InGeometry)
 			, CollisionData()
-			, Materials()
 		{
+			bIsSingleMaterial = true;
 		}
 
 		FShapeInstance(const EPerShapeDataType InType, FShapeInstance&& Other)
 			: FPerShapeData(InType, Other)
 			, CollisionData(MoveTemp(Other.CollisionData))
-			, Materials(MoveTemp(Other.Materials))
 		{
+			if (bIsSingleMaterial)
+			{
+				Material.MaterialHandle = Other.Material.MaterialHandle;
+			}
+			else
+			{
+				Material.MaterialData = Other.Material.MaterialData;
+				Other.Material.MaterialData = nullptr;
+			}
+		}
+
+		FMaterialData& GetMaterialDataImpl()
+		{
+			if (bIsSingleMaterial)
+			{
+				const FMaterialHandle ExistingMaterial = Material.MaterialHandle;
+
+				Material.MaterialData = new FMaterialData();
+				bIsSingleMaterial = false;
+
+				if (ExistingMaterial.IsValid())
+				{
+					// Ideally we do not get here if we have been setup as a single-material shape.
+					// However the API combination of SetMaterial() and GetMaterialData() requires
+					// we be able to upgrade from a single material to multi-material shape.
+					// We should try to remove this if possible.
+					UE_LOG(LogChaos, Warning, TEXT("Perf/memory warning: request for Material Array on single-material ShapeInstance"));
+
+					Material.MaterialData->Materials.Add(ExistingMaterial);
+				}
+			}
+			return *Material.MaterialData;
+		}
+
+		const FMaterialData& GetMaterialDataImpl() const
+		{
+			return const_cast<FShapeInstance*>(this)->GetMaterialDataImpl();
 		}
 
+		void SetMaterialImpl(const FMaterialHandle& InMaterial)
+		{
+			if (!bIsSingleMaterial)
+			{
+				// Destroy existing MaterialData if present
+				if (Material.MaterialData != nullptr)
+				{
+					delete Material.MaterialData;
+				}
+				bIsSingleMaterial = true;
+			}
+
+			Material.MaterialHandle = InMaterial;
+		}
+
+		void SetMaterialDataImpl(const FMaterialData& Data)
+		{ 
+			// If we have only one material and no masks, we can use the single material handle
+			if ((Data.Materials.Num() <= 1) && Data.MaterialMasks.IsEmpty() && Data.MaterialMaskMaps.IsEmpty() && Data.MaterialMaskMapMaterials.IsEmpty())
+			{
+				const FMaterialHandle MaterialHandle = (Data.Materials.Num() > 0) ? Data.Materials[0] : FMaterialHandle();
+				SetMaterialImpl(MaterialHandle);
+			}
+			else
+			{
+				// This could probably be a Move except it would change the SetMaterialData API
+				// As it stands this will duplicate (and eventually discard) the material and mask arrays
+				// although this is probably not called on an already-initialized material data very often
+				GetMaterialDataImpl() = Data;
+			}
+		}
+
+		virtual void SerializeMaterials(FChaosArchive& Ar) override final;
+
+		union FMaterialUnion
+		{
+			FMaterialUnion() : MaterialHandle() {}		// Default to single-shape mode
+			~FMaterialUnion() {}						// Destruction handled by FShapeInstance
+
+			FMaterialHandle MaterialHandle;				// Set if we have only 1 material, no masks etc
+			FMaterialData* MaterialData;				// Set if we have multiple materials or any masks
+		};
+
 		FCollisionData CollisionData;
-		FMaterialData Materials;
+		mutable FMaterialUnion Material;
+
 	};
 
 	namespace Private
@@ -622,8 +763,9 @@ namespace Chaos
 			{
 			}
 
-			FVec3 WorldPosition;
+			// NOTE: FRotation is 16-byte aligned so it goes first
 			FRotation3 WorldRotation;
+			FVec3 WorldPosition;
 		};
 	}
 
@@ -799,6 +941,16 @@ namespace Chaos
 		DownCast([Particle](auto& ShapeInstance) { ShapeInstance.UpdateLeafWorldTransform(Particle); });
 	}
 
+	inline int32 FPerShapeData::NumMaterials() const
+	{
+		return DownCast([](const auto& ShapeInstance) { return ShapeInstance.NumMaterials(); });
+	}
+
+	inline const FMaterialHandle& FPerShapeData::GetMaterial(const int32 Index) const
+	{
+		return DownCast([Index](const auto& ShapeInstance) -> const auto& { return ShapeInstance.GetMaterial(Index); });
+	}
+
 	inline const TArray<FMaterialHandle>& FPerShapeData::GetMaterials() const
 	{
 		return DownCast([](const auto& ShapeInstance) -> const auto& { return ShapeInstance.GetMaterials(); });
