Move the QueryOnly Physics objects into its own acceleration structure
Make sure that it does not get included in the physics thread.

diff --git a/Engine/Source/Runtime/Engine/Private/PhysicsEngine/BodyInstance.cpp b/Engine/Source/Runtime/Engine/Private/PhysicsEngine/BodyInstance.cpp
--- a/Engine/Source/Runtime/Engine/Private/PhysicsEngine/BodyInstance.cpp
+++ b/Engine/Source/Runtime/Engine/Private/PhysicsEngine/BodyInstance.cpp
@@ -1177,7 +1177,7 @@ void FInitBodiesHelperBase::CreateActor_AssumesLocked(FBodyInstance* Instance, c
 	ActorParams.DebugName = Instance->CharDebugName.IsValid() ? Instance->CharDebugName->GetData() : nullptr;
 #endif
 	ActorParams.bEnableGravity = Instance->bEnableGravity;
-	ActorParams.bQueryOnly = bDisableSim;
+	ActorParams.bQueryOnly = CollisionType == ECollisionEnabled::QueryOnly;
 	ActorParams.Scene = PhysScene;
 
 	if (IsStatic())
diff --git a/Engine/Source/Runtime/Experimental/Chaos/Private/Chaos/PBDRigidsEvolution.cpp b/Engine/Source/Runtime/Experimental/Chaos/Private/Chaos/PBDRigidsEvolution.cpp
--- a/Engine/Source/Runtime/Experimental/Chaos/Private/Chaos/PBDRigidsEvolution.cpp
+++ b/Engine/Source/Runtime/Experimental/Chaos/Private/Chaos/PBDRigidsEvolution.cpp
@@ -29,6 +29,9 @@ namespace Chaos
 	CHAOS_API int32 FixBadAccelerationStructureRemoval = 1;
 	FAutoConsoleVariableRef CVarFixBadAccelerationStructureRemoval(TEXT("p.FixBadAccelerationStructureRemoval"), FixBadAccelerationStructureRemoval, TEXT(""));
 
+	CHAOS_API int32 AccelerationStructureIsolateQueryOnlyObjects = 0;
+	FAutoConsoleVariableRef CVarAccelerationStructureIsolateQueryOnlyObjects(TEXT("p.Chaos.AccelerationStructureIsolateQueryOnlyObjects"), AccelerationStructureIsolateQueryOnlyObjects, TEXT("Set to 1: QueryOnly Objects will not be moved to acceleration structures on the Physics Thread"));
+
 	CHAOS_API int32 AccelerationStructureSplitStaticAndDynamic = 1;
 	FAutoConsoleVariableRef CVarAccelerationStructureSplitStaticAndDynamic(TEXT("p.Chaos.AccelerationStructureSplitStaticDynamic"), AccelerationStructureSplitStaticAndDynamic, TEXT("Set to 1: Sort Dynamic and Static bodies into seperate acceleration structures, any other value will disable the feature"));
 
@@ -69,7 +72,7 @@ namespace Chaos
 
 		static bool IsDynamicTree(FSpatialAccelerationIdx SpatialAccelerationIdx)
 		{
-			return AccelerationStructureUseDynamicTree && SpatialAccelerationIdx.InnerIdx == 1;
+			return AccelerationStructureUseDynamicTree && (SpatialAccelerationIdx.InnerIdx == ESpatialAccelerationCollectionBucketInnerIdx::Dynamic || SpatialAccelerationIdx.InnerIdx == ESpatialAccelerationCollectionBucketInnerIdx::DynamicQueryOnly);
 		}
 
 		TUniquePtr<ISpatialAccelerationCollection<FAccelerationStructureHandle, FReal, 3>> CreateEmptyCollection() override
@@ -82,11 +85,24 @@ namespace Chaos
 			for (uint16 BucketIdx = 0; BucketIdx < NumBuckets; ++BucketIdx)
 			{
 				// For static bodies
-				Collection->AddSubstructure(CreateAccelerationPerBucket_Threaded(Empty, BucketIdx, true, false), BucketIdx, 0);
+				Collection->AddSubstructure(CreateAccelerationPerBucket_Threaded(Empty, BucketIdx, true, false, false), BucketIdx, ESpatialAccelerationCollectionBucketInnerIdx::Default);
+
+				// Always create this if QueryOnlyObjects are isolated to ensure all consecutive indices are created (within buckets)
+				if (AccelerationStructureSplitStaticAndDynamic == 1 || AccelerationStructureIsolateQueryOnlyObjects == 1) 
+				{
+					// Non static bodies
+					Collection->AddSubstructure(CreateAccelerationPerBucket_Threaded(Empty, BucketIdx, true, IsDynamicTree(FSpatialAccelerationIdx{BucketIdx, ESpatialAccelerationCollectionBucketInnerIdx::Dynamic }), true), BucketIdx, ESpatialAccelerationCollectionBucketInnerIdx::Dynamic);
+				}
+			}
+
+			if (AccelerationStructureIsolateQueryOnlyObjects && NumBuckets == 1)
+			{
+				constexpr uint16 BucketIdx = 0;
+				Collection->AddSubstructure(CreateAccelerationPerBucket_Threaded(Empty, BucketIdx, true, false, false), BucketIdx, ESpatialAccelerationCollectionBucketInnerIdx::DefaultQueryOnly);
 				if (AccelerationStructureSplitStaticAndDynamic == 1)
 				{
 					// Non static bodies
-					Collection->AddSubstructure(CreateAccelerationPerBucket_Threaded(Empty, BucketIdx, true, IsDynamicTree(FSpatialAccelerationIdx{BucketIdx, 1})), BucketIdx, 1);
+					Collection->AddSubstructure(CreateAccelerationPerBucket_Threaded(Empty, BucketIdx, true, IsDynamicTree(FSpatialAccelerationIdx{ BucketIdx,  ESpatialAccelerationCollectionBucketInnerIdx::DynamicQueryOnly }), false), BucketIdx, ESpatialAccelerationCollectionBucketInnerIdx::DynamicQueryOnly);
 				}
 			}
 
@@ -135,7 +151,7 @@ namespace Chaos
 			}
 		}
 
-		virtual TUniquePtr<ISpatialAcceleration<FAccelerationStructureHandle, FReal, 3>> CreateAccelerationPerBucket_Threaded(const TConstParticleView<FSpatialAccelerationCache>& Particles, uint16 BucketIdx, bool ForceFullBuild, bool bDynamicTree) override
+		virtual TUniquePtr<ISpatialAcceleration<FAccelerationStructureHandle, FReal, 3>> CreateAccelerationPerBucket_Threaded(const TConstParticleView<FSpatialAccelerationCache>& Particles, uint16 BucketIdx, bool ForceFullBuild, bool bDynamicTree, bool bBuildOverlapCache) override
 		{
 			// TODO: Unduplicate switch statement here with IsBucketTimeSliced and refactor so that bucket index mapping is better.
 			switch (BucketIdx)
@@ -150,9 +166,9 @@ namespace Chaos
 				{
 					if (bDynamicTree)
 					{
-						return MakeUnique<AABBDynamicTreeType>(Particles, BroadPhaseConfig.MaxChildrenInLeaf, BroadPhaseConfig.MaxTreeDepth, BroadPhaseConfig.MaxPayloadSize, ForceFullBuild ? 0 : BroadPhaseConfig.IterationsPerTimeSlice, true);
+						return MakeUnique<AABBDynamicTreeType>(Particles, BroadPhaseConfig.MaxChildrenInLeaf, BroadPhaseConfig.MaxTreeDepth, BroadPhaseConfig.MaxPayloadSize, ForceFullBuild ? 0 : BroadPhaseConfig.IterationsPerTimeSlice, true, bBuildOverlapCache);
 					}
-					return MakeUnique<AABBTreeType>(Particles, BroadPhaseConfig.MaxChildrenInLeaf, BroadPhaseConfig.MaxTreeDepth, BroadPhaseConfig.MaxPayloadSize, ForceFullBuild ? 0 : BroadPhaseConfig.IterationsPerTimeSlice, false, AccelerationStructureUseDirtyTreeInsteadOfGrid == 1);
+					return MakeUnique<AABBTreeType>(Particles, BroadPhaseConfig.MaxChildrenInLeaf, BroadPhaseConfig.MaxTreeDepth, BroadPhaseConfig.MaxPayloadSize, ForceFullBuild ? 0 : BroadPhaseConfig.IterationsPerTimeSlice, false, AccelerationStructureUseDirtyTreeInsteadOfGrid == 1, bBuildOverlapCache);
 				}
 				else if (BroadPhaseConfig.BroadphaseType == FBroadPhaseConfig::TreeOfGridAndGrid || BroadPhaseConfig.BroadphaseType == FBroadPhaseConfig::TreeOfGrid)
 				{
@@ -435,7 +451,7 @@ namespace Chaos
 
 				auto ParticleView = MakeConstParticleView(MoveTemp(ViewsPerBucket[SpatialAccIdx.Bucket][SpatialAccIdx.InnerIdx]));
 				// Create and run the first slice here
-				auto NewStruct = SpatialCollectionFactory.CreateAccelerationPerBucket_Threaded(ParticleView, SpatialAccIdx.Bucket, IsForceFullBuild, false);
+				auto NewStruct = SpatialCollectionFactory.CreateAccelerationPerBucket_Threaded(ParticleView, SpatialAccIdx.Bucket, IsForceFullBuild, false, false);
 				{
 					NewStruct->ClearShouldRebuild();  // Mark as pristine
 				}
@@ -473,7 +489,7 @@ namespace Chaos
 					SCOPE_CYCLE_COUNTER(STAT_CreateNonSlicedStructures);
 
 					auto ParticleView = MakeConstParticleView(MoveTemp(ViewsPerBucket[SpatialAccIdx.Bucket][SpatialAccIdx.InnerIdx]));
-					auto NewStruct = SpatialCollectionFactory.CreateAccelerationPerBucket_Threaded(ParticleView,SpatialAccIdx.Bucket,IsForceFullBuild, false);
+					auto NewStruct = SpatialCollectionFactory.CreateAccelerationPerBucket_Threaded(ParticleView,SpatialAccIdx.Bucket,IsForceFullBuild, false, false);
 
 					ISpatialAcceleration<FAccelerationStructureHandle, FReal, 3>* AccelerationSUBStructureCopy = GetSubStructureCopy(SpatialAccIdx);
 					if (AccelerationSUBStructureCopy)
@@ -768,7 +784,7 @@ namespace Chaos
 			{
 				FlushInternalAccelerationQueue();
 
-				if (GAccelerationStructureCacheOverlappingLeaves)
+				if (GAccelerationStructureCacheOverlappingLeaves  && AccelerationStructureUseDynamicTree)
 				{
 					// Caching of the overlapping leaves
 					InternalAcceleration->CacheOverlappingLeaves();
diff --git a/Engine/Source/Runtime/Experimental/Chaos/Private/PhysicsProxy/GeometryCollectionPhysicsProxy.cpp b/Engine/Source/Runtime/Experimental/Chaos/Private/PhysicsProxy/GeometryCollectionPhysicsProxy.cpp
--- a/Engine/Source/Runtime/Experimental/Chaos/Private/PhysicsProxy/GeometryCollectionPhysicsProxy.cpp
+++ b/Engine/Source/Runtime/Experimental/Chaos/Private/PhysicsProxy/GeometryCollectionPhysicsProxy.cpp
@@ -57,6 +57,10 @@
 #define GC_PHYSICSPROXY_CHECK_FOR_NAN(Vec)
 #endif
 
+namespace Chaos{
+	extern int32 AccelerationStructureSplitStaticAndDynamic;
+	extern int32 AccelerationStructureIsolateQueryOnlyObjects;
+}
 
 float CollisionParticlesPerObjectFractionDefault = 1.0f;
 FAutoConsoleVariableRef CVarCollisionParticlesPerObjectFractionDefault(
diff --git a/Engine/Source/Runtime/Experimental/Chaos/Public/Chaos/AABBTree.h b/Engine/Source/Runtime/Experimental/Chaos/Public/Chaos/AABBTree.h
--- a/Engine/Source/Runtime/Experimental/Chaos/Public/Chaos/AABBTree.h
+++ b/Engine/Source/Runtime/Experimental/Chaos/Public/Chaos/AABBTree.h
@@ -749,6 +749,7 @@ public:
 		, MaxPayloadBounds(DefaultMaxPayloadBounds)
 		, MaxNumToProcess(DefaultMaxNumToProcess)
 		, bShouldRebuild(true)
+		, bBuildOverlapCache(true)
 	{
 		GetCVars();
 	}
@@ -816,7 +817,7 @@ public:
 	}
 
 	template <typename TParticles>
-	TAABBTree(const TParticles& Particles, int32 InMaxChildrenInLeaf = DefaultMaxChildrenInLeaf, int32 InMaxTreeDepth = DefaultMaxTreeDepth, T InMaxPayloadBounds = DefaultMaxPayloadBounds, int32 InMaxNumToProcess = DefaultMaxNumToProcess, bool bInDynamicTree = false, bool bInUseDirtyTree = false)
+	TAABBTree(const TParticles& Particles, int32 InMaxChildrenInLeaf = DefaultMaxChildrenInLeaf, int32 InMaxTreeDepth = DefaultMaxTreeDepth, T InMaxPayloadBounds = DefaultMaxPayloadBounds, int32 InMaxNumToProcess = DefaultMaxNumToProcess, bool bInDynamicTree = false, bool bInUseDirtyTree = false, bool bInBuildOverlapCache = true)
 		: ISpatialAcceleration<TPayloadType, T, 3>(StaticType)
 		, bDynamicTree(bInDynamicTree)
 		, MaxChildrenInLeaf(InMaxChildrenInLeaf)
@@ -824,6 +825,7 @@ public:
 		, MaxPayloadBounds(InMaxPayloadBounds)
 		, MaxNumToProcess(InMaxNumToProcess)
 		, bShouldRebuild(true)
+		, bBuildOverlapCache(bInBuildOverlapCache)
 	{
 		if (bInUseDirtyTree)
 		{
@@ -834,7 +836,7 @@ public:
 	}
 
 	template <typename ParticleView>
-	void Reinitialize(const ParticleView& Particles, int32 InMaxChildrenInLeaf = DefaultMaxChildrenInLeaf, int32 InMaxTreeDepth = DefaultMaxTreeDepth, T InMaxPayloadBounds = DefaultMaxPayloadBounds, int32 InMaxNumToProcess = DefaultMaxNumToProcess, bool bInDynamicTree = false)
+	void Reinitialize(const ParticleView& Particles, int32 InMaxChildrenInLeaf = DefaultMaxChildrenInLeaf, int32 InMaxTreeDepth = DefaultMaxTreeDepth, T InMaxPayloadBounds = DefaultMaxPayloadBounds, int32 InMaxNumToProcess = DefaultMaxNumToProcess, bool bInDynamicTree = false, bool bInbBuildOverlapCache = true)
 	{
 		bDynamicTree = bInDynamicTree;
 		MaxChildrenInLeaf = InMaxChildrenInLeaf;
@@ -842,6 +844,7 @@ public:
 		MaxPayloadBounds = InMaxPayloadBounds;
 		MaxNumToProcess = InMaxNumToProcess;
 		bShouldRebuild = true;
+		bBuildOverlapCache = bInbBuildOverlapCache;
 		GenerateTree(Particles);
 	}
 
@@ -2098,6 +2101,7 @@ public:
 		if (Ar.IsLoading())
 		{
 			bDynamicTree = false;
+			bBuildOverlapCache = true;
 			RootNode = INDEX_NONE;
 			FirstFreeInternalNode = INDEX_NONE;
 			FirstFreeLeafNode = INDEX_NONE;
@@ -2335,6 +2339,11 @@ public:
 	/** Cache for each leaves all the overlapping leaves*/
 	virtual void CacheOverlappingLeaves() override
 	{
+		if (!bBuildOverlapCache)
+		{
+			return;
+		}
+
 		// Dev settings to switch easily algorithms
 		// Will switch to cvars if the leaf version could be faster 
 		const bool bCachingRoot = true;
@@ -3427,6 +3436,7 @@ private:
 		, MaxNumToProcess(Other.MaxNumToProcess)
 		, NumProcessedThisSlice(Other.NumProcessedThisSlice)
 		, bShouldRebuild(Other.bShouldRebuild)
+		, bBuildOverlapCache(Other.bBuildOverlapCache)
 		, OverlappingLeaves(Other.OverlappingLeaves)
 		, OverlappingOffsets(Other.OverlappingOffsets)
 		, OverlappingPairs(Other.OverlappingPairs)
@@ -3486,6 +3496,7 @@ private:
 			MaxNumToProcess = Rhs.MaxNumToProcess;
 			NumProcessedThisSlice = Rhs.NumProcessedThisSlice;
 			bShouldRebuild = Rhs.bShouldRebuild;
+			bBuildOverlapCache = Rhs.bBuildOverlapCache;
 			if (Rhs.DirtyElementTree)
 			{
 				check(DirtyElementTree); // We should have allocated this already
@@ -3541,6 +3552,7 @@ private:
 
 	bool bShouldRebuild;  // Contract: this can only ever be cleared by calling the ClearShouldRebuild method
 
+	bool bBuildOverlapCache;
 	/** Flat array of overlapping leaves.  */
 	TArray<int32> OverlappingLeaves;
 
diff --git a/Engine/Source/Runtime/Experimental/Chaos/Public/Chaos/Collision/SpatialAccelerationBroadPhase.h b/Engine/Source/Runtime/Experimental/Chaos/Public/Chaos/Collision/SpatialAccelerationBroadPhase.h
--- a/Engine/Source/Runtime/Experimental/Chaos/Public/Chaos/Collision/SpatialAccelerationBroadPhase.h
+++ b/Engine/Source/Runtime/Experimental/Chaos/Public/Chaos/Collision/SpatialAccelerationBroadPhase.h
@@ -287,7 +287,7 @@ namespace Chaos
 		}
 
 		/**
-		 * Generate all overlapping pairs and spawn a misphase object to handle collisions for each of them
+		 * Generate all overlapping pairs and spawn a midphase object to handle collisions for each of them
 		 */
 		void ProduceOverlaps(
 			FReal Dt, 
diff --git a/Engine/Source/Runtime/Experimental/Chaos/Public/Chaos/PBDRigidsEvolution.h b/Engine/Source/Runtime/Experimental/Chaos/Public/Chaos/PBDRigidsEvolution.h
--- a/Engine/Source/Runtime/Experimental/Chaos/Public/Chaos/PBDRigidsEvolution.h
+++ b/Engine/Source/Runtime/Experimental/Chaos/Public/Chaos/PBDRigidsEvolution.h
@@ -256,7 +256,7 @@ struct CHAOS_API ISpatialAccelerationCollectionFactory
 	virtual bool IsBucketTimeSliced(uint16 BucketIdx) const = 0;
 
 	//Chaos creates new acceleration structures per bucket. Factory can change underlying type at runtime as well as number of buckets to AB test
-	virtual TUniquePtr<ISpatialAcceleration<FAccelerationStructureHandle, FReal, 3>> CreateAccelerationPerBucket_Threaded(const TConstParticleView<FSpatialAccelerationCache>& Particles, uint16 BucketIdx, bool ForceFullBuild, bool bDynamicTree) = 0;
+	virtual TUniquePtr<ISpatialAcceleration<FAccelerationStructureHandle, FReal, 3>> CreateAccelerationPerBucket_Threaded(const TConstParticleView<FSpatialAccelerationCache>& Particles, uint16 BucketIdx, bool ForceFullBuild, bool bDynamicTree, bool bBuildOverlapCache) = 0;
 
 	//Mask indicating which bucket is active. Spatial indices in inactive buckets fallback to bucket 0. Bit 0 indicates bucket 0 is active, Bit 1 indicates bucket 1 is active, etc...
 	virtual uint8 GetActiveBucketsMask() const = 0;
diff --git a/Engine/Source/Runtime/Experimental/Chaos/Public/Chaos/ParticleDirtyFlags.h b/Engine/Source/Runtime/Experimental/Chaos/Public/Chaos/ParticleDirtyFlags.h
--- a/Engine/Source/Runtime/Experimental/Chaos/Public/Chaos/ParticleDirtyFlags.h
+++ b/Engine/Source/Runtime/Experimental/Chaos/Public/Chaos/ParticleDirtyFlags.h
@@ -595,6 +595,7 @@ struct FCollisionData
 	}
 
 	bool HasCollisionData() const { return bSimCollision || bQueryCollision; }
+	bool HasQueryOnlyData() const { return !bSimCollision && bQueryCollision; }
 
 	void Serialize(FChaosArchive& Ar)
 	{
diff --git a/Engine/Source/Runtime/Experimental/Chaos/Public/Chaos/ParticleHandle.h b/Engine/Source/Runtime/Experimental/Chaos/Public/Chaos/ParticleHandle.h
--- a/Engine/Source/Runtime/Experimental/Chaos/Public/Chaos/ParticleHandle.h
+++ b/Engine/Source/Runtime/Experimental/Chaos/Public/Chaos/ParticleHandle.h
@@ -73,21 +73,11 @@ void GeometryParticleDefaultConstruct(FConcrete& Concrete, const FGeometryPartic
 	Concrete.SetEnabledDuringResim(true);
 }
 
-
-extern CHAOS_API int32 AccelerationStructureSplitStaticAndDynamic;
 template <typename T, int d, typename FConcrete>
 void KinematicGeometryParticleDefaultConstruct(FConcrete& Concrete, const FKinematicGeometryParticleParameters& Params)
 {
 	Concrete.SetV(TVector<T, d>(0));
-	Concrete.SetW(TVector<T, d>(0));
-	if (AccelerationStructureSplitStaticAndDynamic == 1)
-	{
-		Concrete.SetSpatialIdx(FSpatialAccelerationIdx{ 0,1 });
-	}
-	else
-	{
-		Concrete.SetSpatialIdx(FSpatialAccelerationIdx{ 0,0 });
-	}
+	Concrete.SetW(TVector<T, d>(0));	
 }
 template <typename T, int d, typename FConcrete>
 void PBDRigidParticleDefaultConstruct(FConcrete& Concrete, const FPBDRigidParticleParameters& Params)
diff --git a/Engine/Source/Runtime/Experimental/Chaos/Public/Chaos/SpatialAccelerationCollection.h b/Engine/Source/Runtime/Experimental/Chaos/Public/Chaos/SpatialAccelerationCollection.h
--- a/Engine/Source/Runtime/Experimental/Chaos/Public/Chaos/SpatialAccelerationCollection.h
+++ b/Engine/Source/Runtime/Experimental/Chaos/Public/Chaos/SpatialAccelerationCollection.h
@@ -1,6 +1,7 @@
 // Copyright Epic Games, Inc. All Rights Reserved.
 #pragma once
 #include "Chaos/ISpatialAccelerationCollection.h"
+#include "Chaos/PBDRigidsEvolution.h"
 #include "Chaos/Box.h"
 #include "Chaos/Collision/SpatialAccelerationBroadPhase.h"
 #include "Chaos/Collision/StatsData.h"
@@ -12,6 +13,15 @@
 namespace Chaos
 {
 
+// Bucket inner indices define for specific use cases		
+enum ESpatialAccelerationCollectionBucketInnerIdx
+{
+	Default = 0, // If acceleration structures are not split up by body type, they end up in this bucket, otherwise static bodies will go here
+	Dynamic = 1, // For dynamic bodies if they use a separate acceleration structure
+	DefaultQueryOnly = 2, // Query only bodies if they are separated
+	DynamicQueryOnly = 3 // Dynamic Query only bodies if they are separated
+};
+
 template <typename T>
 void FreeObjHelper(T*& RawPtr)
 {
@@ -354,18 +364,16 @@ struct TSpatialAccelerationCollectionHelper
 };
 
 template <typename SpatialAccelerationCollection>
-typename TEnableIf<std::is_same_v<typename SpatialAccelerationCollection::TPayloadType, FAccelerationStructureHandle>, void>::Type PBDComputeConstraintsLowLevel_Helper(FReal Dt, const SpatialAccelerationCollection& Accel, FSpatialAccelerationBroadPhase& BroadPhase, Private::FCollisionConstraintAllocator* Allocator, const FCollisionDetectorSettings& Settings, IResimCacheBase* ResimCache)
+typename std::enable_if_t<std::is_same_v<typename SpatialAccelerationCollection::TPayloadType, FAccelerationStructureHandle>, void> PBDComputeConstraintsLowLevel_Helper(FReal Dt, const SpatialAccelerationCollection& Accel, FSpatialAccelerationBroadPhase& BroadPhase, Private::FCollisionConstraintAllocator* Allocator, const FCollisionDetectorSettings& Settings, IResimCacheBase* ResimCache)
 {
 	BroadPhase.ProduceOverlaps(Dt, Accel, Allocator, Settings, ResimCache);
 }
 
 template <typename SpatialAccelerationCollection>
-typename TEnableIf<!std::is_same_v<typename SpatialAccelerationCollection::TPayloadType, FAccelerationStructureHandle>, void>::Type PBDComputeConstraintsLowLevel_Helper(FReal Dt, const SpatialAccelerationCollection& Accel, FSpatialAccelerationBroadPhase& BroadPhase, Private::FCollisionConstraintAllocator* Allocator, const FCollisionDetectorSettings& Settings, IResimCacheBase* ResimCache)
+typename std::enable_if_t<!std::is_same_v<typename SpatialAccelerationCollection::TPayloadType, FAccelerationStructureHandle>, void> PBDComputeConstraintsLowLevel_Helper(FReal Dt, const SpatialAccelerationCollection& Accel, FSpatialAccelerationBroadPhase& BroadPhase, Private::FCollisionConstraintAllocator* Allocator, const FCollisionDetectorSettings& Settings, IResimCacheBase* ResimCache)
 {
 }
 
-
-
 template <typename ... TSpatialAccelerationTypes>
 class CHAOS_API TSpatialAccelerationCollection : public
 	ISpatialAccelerationCollection<typename std::tuple_element<0, std::tuple< TSpatialAccelerationTypes...>>::type::PayloadType,
@@ -376,7 +384,7 @@ public:
 	using FirstAccelerationType = typename std::tuple_element<0, std::tuple< TSpatialAccelerationTypes...>>::type;
 	using TPayloadType = typename FirstAccelerationType::PayloadType;
 	using T = typename FirstAccelerationType::TType;
-	static constexpr int d = FirstAccelerationType::D;
+	static constexpr int d = FirstAccelerationType::D;	
 
 	TSpatialAccelerationCollection()
 	{
@@ -513,11 +521,29 @@ public:
 	}
 
 	template <typename SQVisitor>
-	void Overlap(const TAABB<T, 3>& QueryBounds, SQVisitor& Visitor) const
+	typename std::enable_if_t<!std::is_same_v<SQVisitor, typename Private::FSimOverlapVisitor>, void> Overlap(const TAABB<T, 3>& QueryBounds, SQVisitor& Visitor) const
 	{
 		TSpatialAccelerationCollectionHelper<0, NumTypes, decltype(Types), TPayloadType, T, d>::OverlapFast(Types, QueryBounds, Visitor);
 	}
 
+	template <typename SQVisitor>
+	typename std::enable_if_t<std::is_same_v<SQVisitor, typename Private::FSimOverlapVisitor>, void> Overlap(const TAABB<T, 3>& QueryBounds, SQVisitor& Visitor) const
+	{		
+		const TSpatialCollectionBucket<TSpatialAccelerationBucketEntry<TPayloadType, T, d>>& Bucket = Buckets[0];
+		{
+			int BucketInnerIdx = 0;
+			for (const TSpatialAccelerationBucketEntry<TPayloadType, T, d>& BucketEntry : Bucket.Objects)
+			{
+				if (BucketInnerIdx != ESpatialAccelerationCollectionBucketInnerIdx::DefaultQueryOnly && BucketInnerIdx != ESpatialAccelerationCollectionBucketInnerIdx::DynamicQueryOnly)
+				{
+					auto AccelStructure = static_cast<Chaos::TAABBTree<TPayloadType, Chaos::TAABBTreeLeafArray<Chaos::FAccelerationStructureHandle, true, T>, true, T>*>(BucketEntry.Acceleration.Get());
+					AccelStructure->OverlapFast(QueryBounds, Visitor);
+				}
+				BucketInnerIdx++;
+			}
+		}
+	}
+
 	TArray<TPayloadBoundsElement<TPayloadType, T>> GlobalObjects() const
 	{
 		TArray<TPayloadBoundsElement<TPayloadType, T>> ObjList;
diff --git a/Engine/Source/Runtime/Landscape/Private/LandscapeCollision.cpp b/Engine/Source/Runtime/Landscape/Private/LandscapeCollision.cpp
--- a/Engine/Source/Runtime/Landscape/Private/LandscapeCollision.cpp
+++ b/Engine/Source/Runtime/Landscape/Private/LandscapeCollision.cpp
@@ -312,7 +312,7 @@ void ULandscapeHeightfieldCollisionComponent::OnCreatePhysicsState()
 				FActorCreationParams Params;
 				Params.InitialTM = LandscapeComponentTransform;
 				Params.InitialTM.SetScale3D(FVector(0));
-				Params.bQueryOnly = true;
+				Params.bQueryOnly = false;
 				Params.bStatic = true;
 				Params.Scene = GetWorld()->GetPhysicsScene();
 				FPhysicsActorHandle PhysHandle;
@@ -1435,7 +1435,7 @@ struct FMeshCollisionInitHelper
 		FActorCreationParams Params;
 		Params.InitialTM = ComponentToWorld;
 		Params.InitialTM.SetScale3D(FVector::OneVector);
-		Params.bQueryOnly = true;
+		Params.bQueryOnly = false;
 		Params.bStatic = true;
 		Params.Scene = PhysScene;
 
diff --git a/Engine/Source/Runtime/PhysicsCore/Private/ChaosEngineInterface.cpp b/Engine/Source/Runtime/PhysicsCore/Private/ChaosEngineInterface.cpp
--- a/Engine/Source/Runtime/PhysicsCore/Private/ChaosEngineInterface.cpp
+++ b/Engine/Source/Runtime/PhysicsCore/Private/ChaosEngineInterface.cpp
@@ -29,6 +29,12 @@ FPhysicsDelegatesCore::FOnUpdatePhysXMaterial FPhysicsDelegatesCore::OnUpdatePhy
 #include "PhysicsProxy/SingleParticlePhysicsProxy.h"
 #include "Chaos/CastingUtilities.h"
 
+namespace Chaos
+{
+	extern CHAOS_API int32 AccelerationStructureSplitStaticAndDynamic;
+	extern CHAOS_API int32 AccelerationStructureIsolateQueryOnlyObjects;
+}
+
 bool bEnableChaosJointConstraints = true;
 FAutoConsoleVariableRef CVarEnableChaosJointConstraints(TEXT("p.ChaosSolverEnableJointConstraints"), bEnableChaosJointConstraints, TEXT("Enable Joint Constraints defined within the Physics Asset Editor"));
 
@@ -1762,6 +1768,47 @@ void FChaosEngineInterface::CreateActor(const FActorCreationParams& InParams,FPh
 		Particle = MoveTemp(Rigid);
 	}
 
+	// Set the particle acceleration structure spatial index here
+	{
+		FSpatialAccelerationIdx SpatialIndex{0, ESpatialAccelerationCollectionBucketInnerIdx::Default };
+		if (AccelerationStructureSplitStaticAndDynamic == 1)
+		{
+			if (AccelerationStructureIsolateQueryOnlyObjects == 1)
+			{
+				if (InParams.bStatic && InParams.bQueryOnly)
+				{
+					SpatialIndex = FSpatialAccelerationIdx{0, ESpatialAccelerationCollectionBucketInnerIdx::DefaultQueryOnly};
+				}
+				else if (!InParams.bStatic && InParams.bQueryOnly)
+				{
+					SpatialIndex = FSpatialAccelerationIdx{ 0, ESpatialAccelerationCollectionBucketInnerIdx::DynamicQueryOnly};
+				}
+				else if (!InParams.bStatic && !InParams.bQueryOnly)
+				{
+					SpatialIndex = FSpatialAccelerationIdx{ 0, ESpatialAccelerationCollectionBucketInnerIdx::Dynamic };
+				}
+			}
+			else
+			{
+				if (!InParams.bStatic)
+				{
+					SpatialIndex = FSpatialAccelerationIdx{ 0, ESpatialAccelerationCollectionBucketInnerIdx::Dynamic };
+				}				
+			}
+		}
+		else
+		{
+			if (AccelerationStructureIsolateQueryOnlyObjects == 1)
+			{
+				if (InParams.bQueryOnly)
+				{
+					SpatialIndex = FSpatialAccelerationIdx{ 0, ESpatialAccelerationCollectionBucketInnerIdx::DefaultQueryOnly };
+				}
+			}		
+		}
+		Particle->SetSpatialIdx(SpatialIndex);
+	}
+
 	Handle = Chaos::FSingleParticlePhysicsProxy::Create(MoveTemp(Particle));
 	Chaos::FRigidBodyHandle_External& Body_External = Handle->GetGameThreadAPI();
 
