Static analysis fix

diff --git a/Engine/Source/Runtime/Experimental/Chaos/Public/Chaos/SpatialAccelerationCollection.h b/Engine/Source/Runtime/Experimental/Chaos/Public/Chaos/SpatialAccelerationCollection.h
--- a/Engine/Source/Runtime/Experimental/Chaos/Public/Chaos/SpatialAccelerationCollection.h
+++ b/Engine/Source/Runtime/Experimental/Chaos/Public/Chaos/SpatialAccelerationCollection.h
@@ -1,7 +1,6 @@
 // Copyright Epic Games, Inc. All Rights Reserved.
 #pragma once
 #include "Chaos/ISpatialAccelerationCollection.h"
-#include "Chaos/PBDRigidsEvolution.h"
 #include "Chaos/Box.h"
 #include "Chaos/Collision/SpatialAccelerationBroadPhase.h"
 #include "Chaos/Collision/StatsData.h"
