This adds extra ensures to catch any bugs that will result in a malformed AABBTree

diff --git a/Engine/Source/Runtime/Experimental/Chaos/Public/Chaos/AABBTree.h b/Engine/Source/Runtime/Experimental/Chaos/Public/Chaos/AABBTree.h
--- a/Engine/Source/Runtime/Experimental/Chaos/Public/Chaos/AABBTree.h
+++ b/Engine/Source/Runtime/Experimental/Chaos/Public/Chaos/AABBTree.h
@@ -431,6 +431,7 @@ struct TAABBTreeLeafArray : public TBoundsWrapperHelper<TPayloadType, T, bComput
 				Elems.RemoveAtSwap(Idx);
 				break;
 			}
+			ensure(Idx != Elems.Num() - 1); // Make sure the payload was actually in here
 		}
 		bDirtyLeaf = true;
 	}
@@ -748,8 +749,9 @@ public:
 		, MaxTreeDepth(DefaultMaxTreeDepth)
 		, MaxPayloadBounds(DefaultMaxPayloadBounds)
 		, MaxNumToProcess(DefaultMaxNumToProcess)
+		, bModifyingTreeMultiThreadingFastCheck(false)
 		, bShouldRebuild(true)
-		, bBuildOverlapCache(true)
+		, bBuildOverlapCache(true)		
 	{
 		GetCVars();
 	}
@@ -824,8 +826,9 @@ public:
 		, MaxTreeDepth(InMaxTreeDepth)
 		, MaxPayloadBounds(InMaxPayloadBounds)
 		, MaxNumToProcess(InMaxNumToProcess)
+		, bModifyingTreeMultiThreadingFastCheck(false)
 		, bShouldRebuild(true)
-		, bBuildOverlapCache(bInBuildOverlapCache)
+		, bBuildOverlapCache(bInBuildOverlapCache)		
 	{
 		if (bInUseDirtyTree)
 		{
@@ -843,8 +846,9 @@ public:
 		MaxTreeDepth = InMaxTreeDepth;
 		MaxPayloadBounds = InMaxPayloadBounds;
 		MaxNumToProcess = InMaxNumToProcess;
+		bModifyingTreeMultiThreadingFastCheck = false;
 		bShouldRebuild = true;
-		bBuildOverlapCache = bInbBuildOverlapCache;
+		bBuildOverlapCache = bInbBuildOverlapCache;		
 		GenerateTree(Particles);
 	}
 
@@ -1192,7 +1196,7 @@ public:
 			AllocatedNodeIdx = Nodes.AddUninitialized(1);;
 			Nodes[AllocatedNodeIdx].bLeaf = false;
 		}
-
+		ensure(Nodes[AllocatedNodeIdx].bLeaf == false);
 		return AllocatedNodeIdx;
 	}
 
@@ -1235,6 +1239,7 @@ public:
 		Nodes[AllocatedNodeIdx].ChildrenBounds[0] = ExpandedBounds;
 
 		Nodes[AllocatedNodeIdx].ParentNode = INDEX_NONE;
+		ensure(Nodes[AllocatedNodeIdx].bLeaf == true);
 		
 		return NodeAndLeafIndices{ AllocatedNodeIdx , LeafIndex };
 	}
@@ -1243,6 +1248,7 @@ public:
 	{
 		Nodes[NodeIdx].ChildrenNodes[1] = FirstFreeInternalNode;
 		FirstFreeInternalNode = NodeIdx;
+		ensure(Nodes[NodeIdx].bLeaf == false);
 	}
 
 	void  DeAllocateLeafNode(int32 NodeIdx)
@@ -1252,6 +1258,7 @@ public:
 
 		Nodes[NodeIdx].ChildrenNodes[1] = FirstFreeLeafNode;
 		FirstFreeLeafNode = NodeIdx;
+		ensure(Nodes[NodeIdx].bLeaf == true);
 	}
 
 	// Is the input node Child 0 or Child 1?
@@ -1266,6 +1273,7 @@ public:
 		}
 		else
 		{
+			ensure(Nodes[ParentIdx].ChildrenNodes[1] == NodeIdx);
 			return 1;
 		}
 	}
@@ -1424,8 +1432,10 @@ public:
 			// Modify NodeIdx 
 			Nodes[NodeIdx].ChildrenNodes[AuntLocalChildIdx] = BestGrandChildToSwap;
 			// Modify BestGrandChildToSwap
+			ensure(BestGrandChildToSwap != NodeIdx);
 			Nodes[BestGrandChildToSwap].ParentNode = NodeIdx;
 			// Modify BestAuntToSwap
+			ensure(BestAuntToSwap != MotherOfBestGrandChild);
 			Nodes[BestAuntToSwap].ParentNode = MotherOfBestGrandChild;
 			// Modify MotherOfBestGrandChild
 			Nodes[MotherOfBestGrandChild].ChildrenNodes[GrandChildLocalChildIdx] = BestAuntToSwap;
@@ -1449,7 +1459,7 @@ public:
 		//	DynamicTreeDebugStats();
 		//}
 
-		// New tree?
+		// Empty tree case
 		if(RootNode == INDEX_NONE)
 		{
 			NodeAndLeafIndices NewIndices = AllocateLeafNodeAndLeaf(Payload, NewBounds);
@@ -1479,6 +1489,7 @@ public:
 		const int32 OldParent = Nodes[BestSibling].ParentNode;
 		const int32 NewParent = AllocateInternalNode();
 		FNode& NewParentNode = Nodes[NewParent];
+		ensure(NewParent != OldParent);
 		NewParentNode.ParentNode = OldParent;
 		NewParentNode.ChildrenNodes[0] = BestSibling;
 		NewParentNode.ChildrenNodes[1] = NewLeafIndices.NodeIdx;
@@ -1499,7 +1510,9 @@ public:
 			RootNode = NewParent;
 		}
 
+		ensure(BestSibling != NewParent);
 		Nodes[BestSibling].ParentNode = NewParent;
+		ensure(NewLeafIndices.NodeIdx != NewParent);
 		Nodes[NewLeafIndices.NodeIdx].ParentNode = NewParent;
 
 		UpdateAncestorBounds(NewParent, true);
@@ -1518,9 +1531,9 @@ public:
 		{
 			RotateNode(NodeIdx,true); 
 		}*/
-		
 		while (ParentNodeIdx != INDEX_NONE)
 		{
+			ensure(CurrentNodeIdx != ParentNodeIdx);
 			int32 ChildIndex = WhichChildAmI(CurrentNodeIdx);
 			Nodes[ParentNodeIdx].ChildrenBounds[ChildIndex] = Nodes[CurrentNodeIdx].ChildrenBounds[0];
 			if (!Nodes[CurrentNodeIdx].bLeaf)
@@ -1540,6 +1553,8 @@ public:
 
 	void RemoveLeafNode(int32 LeafNodeIdx, const TPayloadType& Payload)
 	{
+		ensure(Nodes[LeafNodeIdx].bLeaf == true);
+
 		int32 LeafIdx = Nodes[LeafNodeIdx].ChildrenNodes[0];
 
 		if (Leaves[LeafIdx].GetElementCount() > 1)
@@ -1552,6 +1567,10 @@ public:
 			UpdateAncestorBounds(LeafNodeIdx);
 			return;
 		}
+		else
+		{
+			Leaves[LeafIdx].RemoveElement(Payload); // Just to check if the element was in there to begin with
+		}
 
 		int32 ParentNodeIdx = Nodes[LeafNodeIdx].ParentNode;
 
@@ -1570,6 +1589,7 @@ public:
 			{
 				RootNode = SiblingNodeIdx;
 			}
+			ensure(SiblingNodeIdx != GrandParentNodeIdx);
 			Nodes[SiblingNodeIdx].ParentNode = GrandParentNodeIdx;
 			UpdateAncestorBounds(SiblingNodeIdx);
 			DeAllocateInternalNode(ParentNodeIdx);
@@ -1583,6 +1603,8 @@ public:
 
 	virtual void RemoveElement(const TPayloadType& Payload)
 	{
+		ensure(bModifyingTreeMultiThreadingFastCheck == false);
+		bModifyingTreeMultiThreadingFastCheck = true;
 		if (ensure(bMutable))
 		{
 			if (FAABBTreePayloadInfo* PayloadInfo = PayloadToInfo.Find(Payload))
@@ -1635,10 +1657,13 @@ public:
 				bShouldRebuild = true;
 			}
 		}
+		bModifyingTreeMultiThreadingFastCheck = false;
 	}
 
 	virtual void UpdateElement(const TPayloadType& Payload, const TAABB<T, 3>& NewBounds, bool bInHasBounds) override
 	{
+		ensure(bModifyingTreeMultiThreadingFastCheck == false);
+		bModifyingTreeMultiThreadingFastCheck = true;
 #if !WITH_EDITOR
 		//CSV_SCOPED_TIMING_STAT(ChaosPhysicsTimers, AABBTreeUpdateElement)
 		//CSV_CUSTOM_STAT(ChaosPhysicsTimers, 1, 1, ECsvCustomStatOp::Accumulate);
@@ -1673,6 +1698,7 @@ public:
 								// We still need to update the constituent bounds
 								Leaves[PayloadInfo->LeafIdx].UpdateElement(Payload, NewBounds, bHasBounds);
 								Leaves[PayloadInfo->LeafIdx].RecomputeBounds();
+								bModifyingTreeMultiThreadingFastCheck = false;
 								return;
 							}
 						}
@@ -1683,6 +1709,7 @@ public:
 							{
 								// We still need to update the constituent bounds
 								Leaves[PayloadInfo->LeafIdx].UpdateElement(Payload, NewBounds, bHasBounds);
+								bModifyingTreeMultiThreadingFastCheck = false;
 								return;
 							}
 						}
@@ -1842,6 +1869,7 @@ public:
 			UE_LOG(LogChaos, Verbose, TEXT("Bounding volume exceeded maximum dirty elements (%d dirty of max %d) and is forcing a tree rebuild."), DirtyElements.Num(), MaxDirtyElements);
 			ReoptimizeTree();
 		}
+		bModifyingTreeMultiThreadingFastCheck = false;
 	}
 
 	int32 NumDirtyElements() const
@@ -2100,8 +2128,9 @@ public:
 		// Dynamic trees are not serialized/deserialized for now
 		if (Ar.IsLoading())
 		{
+			bModifyingTreeMultiThreadingFastCheck = false;
 			bDynamicTree = false;
-			bBuildOverlapCache = true;
+			bBuildOverlapCache = true;			
 			RootNode = INDEX_NONE;
 			FirstFreeInternalNode = INDEX_NONE;
 			FirstFreeLeafNode = INDEX_NONE;
@@ -3435,8 +3464,9 @@ private:
 		, MaxPayloadBounds(Other.MaxPayloadBounds)
 		, MaxNumToProcess(Other.MaxNumToProcess)
 		, NumProcessedThisSlice(Other.NumProcessedThisSlice)
+		, bModifyingTreeMultiThreadingFastCheck(Other.bModifyingTreeMultiThreadingFastCheck)
 		, bShouldRebuild(Other.bShouldRebuild)
-		, bBuildOverlapCache(Other.bBuildOverlapCache)
+		, bBuildOverlapCache(Other.bBuildOverlapCache)		
 		, OverlappingLeaves(Other.OverlappingLeaves)
 		, OverlappingOffsets(Other.OverlappingOffsets)
 		, OverlappingPairs(Other.OverlappingPairs)
@@ -3495,8 +3525,9 @@ private:
 			MaxPayloadBounds = Rhs.MaxPayloadBounds;
 			MaxNumToProcess = Rhs.MaxNumToProcess;
 			NumProcessedThisSlice = Rhs.NumProcessedThisSlice;
+			bModifyingTreeMultiThreadingFastCheck = Rhs.bModifyingTreeMultiThreadingFastCheck;
 			bShouldRebuild = Rhs.bShouldRebuild;
-			bBuildOverlapCache = Rhs.bBuildOverlapCache;
+			bBuildOverlapCache = Rhs.bBuildOverlapCache;			
 			if (Rhs.DirtyElementTree)
 			{
 				check(DirtyElementTree); // We should have allocated this already
@@ -3550,6 +3581,8 @@ private:
 	TArray<int32> WorkPoolFreeList;
 	TArray<FWorkSnapshot> WorkPool;
 
+	bool bModifyingTreeMultiThreadingFastCheck; // Used for fast but not perfect multithreading sanity check
+
 	bool bShouldRebuild;  // Contract: this can only ever be cleared by calling the ClearShouldRebuild method
 
 	bool bBuildOverlapCache;
@@ -3570,7 +3603,8 @@ private:
 	 * Tuple of Node/Index/Cost here avoids cache miss when accessing the node and its cost.
 	 */
 	using FNodeIndexAndCost = TTuple<FNode&, int32, FReal>;
-	TArray<FNodeIndexAndCost> PriorityQ;
+	TArray<FNodeIndexAndCost> PriorityQ;	
+
 };
 
 template<typename TPayloadType, typename TLeafType, bool bMutable, typename T>
