Potential fix for crash

diff --git a/Engine/Source/Runtime/Experimental/Chaos/Public/Chaos/AABBTree.h b/Engine/Source/Runtime/Experimental/Chaos/Public/Chaos/AABBTree.h
--- a/Engine/Source/Runtime/Experimental/Chaos/Public/Chaos/AABBTree.h
+++ b/Engine/Source/Runtime/Experimental/Chaos/Public/Chaos/AABBTree.h
@@ -1569,7 +1569,7 @@ public:
 		}
 		else
 		{
-			Leaves[LeafIdx].RemoveElement(Payload); // Just to check if the element was in there to begin with
+			//Leaves[LeafIdx].RemoveElement(Payload); // Just to check if the element was in there to begin with
 		}
 
 		int32 ParentNodeIdx = Nodes[LeafNodeIdx].ParentNode;
