ChaosFlesh : Add scaling for R in XPBD corotated constraints

diff --git a/Engine/Source/Runtime/Experimental/Chaos/Public/Chaos/XPBDCorotatedConstraints.h b/Engine/Source/Runtime/Experimental/Chaos/Public/Chaos/XPBDCorotatedConstraints.h
--- a/Engine/Source/Runtime/Experimental/Chaos/Public/Chaos/XPBDCorotatedConstraints.h
+++ b/Engine/Source/Runtime/Experimental/Chaos/Public/Chaos/XPBDCorotatedConstraints.h
@@ -316,8 +316,6 @@ namespace Chaos::Softs
 					}
 				}
 			}
-
-
 		}
 
 		TVec4<TVector<T, 3>> GetPolarGradient(const PMatrix<T, 3, 3>& Fe, const PMatrix<T, 3, 3>& Re, const PMatrix<T, 3, 3>& DmInvT, const T C1) const
@@ -508,6 +506,8 @@ namespace Chaos::Softs
 
 			Chaos::PolarDecomposition(Fe, Re, Se);
 
+			Re *= FGenericPlatformMath::Pow(AlphaJArray[ElementIndex], (T)1. / (T)3.);
+
 			T C1 = (T)0.;
 			for (int i = 0; i < 3; i++)
 			{
