- Initialize Spatial indices as soon as possible to make sure particles will go to correct acceleration structures
- Undo previous speculative fix
#[fyi] Benn.Gallagher, Chris.Caulfield, Cedric.Caillaud

diff --git a/Engine/Source/Runtime/Experimental/Chaos/Public/Chaos/AABBTree.h b/Engine/Source/Runtime/Experimental/Chaos/Public/Chaos/AABBTree.h
--- a/Engine/Source/Runtime/Experimental/Chaos/Public/Chaos/AABBTree.h
+++ b/Engine/Source/Runtime/Experimental/Chaos/Public/Chaos/AABBTree.h
@@ -1569,7 +1569,7 @@ public:
 		}
 		else
 		{
-			//Leaves[LeafIdx].RemoveElement(Payload); // Just to check if the element was in there to begin with
+			Leaves[LeafIdx].RemoveElement(Payload); // Just to check if the element was in there to begin with
 		}
 
 		int32 ParentNodeIdx = Nodes[LeafNodeIdx].ParentNode;
diff --git a/Engine/Source/Runtime/Experimental/Chaos/Public/Chaos/ParticleHandle.h b/Engine/Source/Runtime/Experimental/Chaos/Public/Chaos/ParticleHandle.h
--- a/Engine/Source/Runtime/Experimental/Chaos/Public/Chaos/ParticleHandle.h
+++ b/Engine/Source/Runtime/Experimental/Chaos/Public/Chaos/ParticleHandle.h
@@ -73,11 +73,21 @@ void GeometryParticleDefaultConstruct(FConcrete& Concrete, const FGeometryPartic
 	Concrete.SetEnabledDuringResim(true);
 }
 
+
+extern CHAOS_API int32 AccelerationStructureSplitStaticAndDynamic;
 template <typename T, int d, typename FConcrete>
 void KinematicGeometryParticleDefaultConstruct(FConcrete& Concrete, const FKinematicGeometryParticleParameters& Params)
 {
 	Concrete.SetV(TVector<T, d>(0));
-	Concrete.SetW(TVector<T, d>(0));	
+	Concrete.SetW(TVector<T, d>(0));
+	if (AccelerationStructureSplitStaticAndDynamic == 1)
+	{
+		Concrete.SetSpatialIdx(FSpatialAccelerationIdx{ 0,1 });
+	}
+	else
+	{
+		Concrete.SetSpatialIdx(FSpatialAccelerationIdx{ 0,0 });
+	}
 }
 template <typename T, int d, typename FConcrete>
 void PBDRigidParticleDefaultConstruct(FConcrete& Concrete, const FPBDRigidParticleParameters& Params)
