Moved CHAOS_API on FShapeInstance to work around issues compiling for Windows with Clang

An optimization to UE_LOG exposed an issue in Clang when targeting Windows. Log statements in inline functions of imported types fail to compile.

#rnx

diff --git a/Engine/Source/Runtime/Experimental/Chaos/Public/Chaos/ShapeInstance.h b/Engine/Source/Runtime/Experimental/Chaos/Public/Chaos/ShapeInstance.h
--- a/Engine/Source/Runtime/Experimental/Chaos/Public/Chaos/ShapeInstance.h
+++ b/Engine/Source/Runtime/Experimental/Chaos/Public/Chaos/ShapeInstance.h
@@ -479,14 +479,14 @@ namespace Chaos
 	 * (though this should not be happening much if at all).
 	 * 
 	 */
-	class CHAOS_API FShapeInstance : public FPerShapeData
+	class FShapeInstance : public FPerShapeData
 	{
 	public:
 		friend class FPerShapeData;
 
-		static TUniquePtr<FShapeInstance> Make(int32 InShapeIdx, TSerializablePtr<FImplicitObject> InGeometry);
-		static void UpdateGeometry(TUniquePtr<FShapeInstance>& InOutShapePtr, TSerializablePtr<FImplicitObject> InGeometry);
-		static FShapeInstance* SerializationFactory(FChaosArchive& Ar, FShapeInstance*);
+		CHAOS_API static TUniquePtr<FShapeInstance> Make(int32 InShapeIdx, TSerializablePtr<FImplicitObject> InGeometry);
+		CHAOS_API static void UpdateGeometry(TUniquePtr<FShapeInstance>& InOutShapePtr, TSerializablePtr<FImplicitObject> InGeometry);
+		CHAOS_API static FShapeInstance* SerializationFactory(FChaosArchive& Ar, FShapeInstance*);
 
 		virtual ~FShapeInstance()
 		{
@@ -496,7 +496,7 @@ namespace Chaos
 			}
 		}
 
-		void UpdateShapeBounds(const FRigidTransform3& WorldTM, const FVec3& BoundsExpansion = FVec3(0));
+		CHAOS_API void UpdateShapeBounds(const FRigidTransform3& WorldTM, const FVec3& BoundsExpansion = FVec3(0));
 
 		void* GetUserData() const { return CollisionData.UserData; }
 		void SetUserData(void* InUserData) { CollisionData.UserData = InUserData; }
@@ -507,19 +507,19 @@ namespace Chaos
 		const FCollisionFilterData& GetSimData() const { return CollisionData.SimData; }
 		void SetSimData(const FCollisionFilterData& InSimData) { CollisionData.SimData = InSimData; }
 
-		void UpdateWorldSpaceState(const FRigidTransform3& WorldTransform, const FVec3& BoundsExpansion);
+		CHAOS_API void UpdateWorldSpaceState(const FRigidTransform3& WorldTransform, const FVec3& BoundsExpansion);
 
 		// The leaf shape (with transformed and implicit wrapper removed).
-		const FImplicitObject* GetLeafGeometry() const;
+		CHAOS_API const FImplicitObject* GetLeafGeometry() const;
 
 		// The actor-relative transform of the leaf geometry.
-		FRigidTransform3 GetLeafRelativeTransform() const;
+		CHAOS_API FRigidTransform3 GetLeafRelativeTransform() const;
 
 		// The world-space transform of the leaf geometry.
 		// If we have non-identity leaf relative transform, is cached from the last call to UpdateWorldSpaceState.
 		// If not cahced, is constructed from arguments.
-		FRigidTransform3 GetLeafWorldTransform(const FGeometryParticleHandle* Particle) const;
-		void UpdateLeafWorldTransform(FGeometryParticleHandle* Particle);
+		CHAOS_API FRigidTransform3 GetLeafWorldTransform(const FGeometryParticleHandle* Particle) const;
+		CHAOS_API void UpdateLeafWorldTransform(FGeometryParticleHandle* Particle);
 
 		int32 NumMaterials() const
 		{
@@ -699,7 +699,7 @@ namespace Chaos
 			}
 		}
 
-		virtual void SerializeMaterials(FChaosArchive& Ar) override final;
+		CHAOS_API virtual void SerializeMaterials(FChaosArchive& Ar) override final;
 
 		union FMaterialUnion
 		{
