Chaos: Fixed heightfield overlaps not taking hole materials into account. All overlaps of cells were reported as a hit - ignoring whether or not those cells were painted out of the landscape.

diff --git a/Engine/Source/Runtime/Experimental/Chaos/Private/Chaos/HeightField.cpp b/Engine/Source/Runtime/Experimental/Chaos/Private/Chaos/HeightField.cpp
--- a/Engine/Source/Runtime/Experimental/Chaos/Private/Chaos/HeightField.cpp
+++ b/Engine/Source/Runtime/Experimental/Chaos/Private/Chaos/HeightField.cpp
@@ -1967,7 +1967,7 @@ namespace Chaos
 					const int32 SingleIndex = Cell[1] * (GeomData.NumCols) + Cell[0];
 					GeomData.GetPointsAndBoundsScaled(SingleIndex, Points, CellBounds);
 
-					if(CellBounds.Intersects(QueryBounds))
+					if(CellBounds.Intersects(QueryBounds) && !IsHole(Cell[0], Cell[1]))
 					{
 						bOverlaps |= OverlapTriangleMTD(Points[0], Points[1], Points[3], OutMTD);
 						bOverlaps |= OverlapTriangleMTD(Points[0], Points[3], Points[2], OutMTD);
@@ -2028,7 +2028,7 @@ namespace Chaos
 				{
 					const int32 SingleIndex = Cell[1] * (GeomData.NumCols) + Cell[0];
 					GeomData.GetPointsAndBoundsScaledSimd(SingleIndex, Points, CellBounds);
-					if (CellBounds.Intersects(QueryBoundsSimd))
+					if (CellBounds.Intersects(QueryBoundsSimd) && !IsHole(Cell[0], Cell[1]))
 					{
 						// pre-transform the triangle in overlap geometry space
 						for (int32 Index = 0; Index < 4; ++Index)
