Fix crash in FNaniteGeometryCollectionSceneProxy creation.
Issue seems to have been from bad merge of work in other stream.

diff --git a/Engine/Source/Runtime/Experimental/GeometryCollectionEngine/Private/GeometryCollection/GeometryCollectionSceneProxy.cpp b/Engine/Source/Runtime/Experimental/GeometryCollectionEngine/Private/GeometryCollection/GeometryCollectionSceneProxy.cpp
--- a/Engine/Source/Runtime/Experimental/GeometryCollectionEngine/Private/GeometryCollection/GeometryCollectionSceneProxy.cpp
+++ b/Engine/Source/Runtime/Experimental/GeometryCollectionEngine/Private/GeometryCollection/GeometryCollectionSceneProxy.cpp
@@ -1129,7 +1129,7 @@ FNaniteGeometryCollectionSceneProxy::FNaniteGeometryCollectionSceneProxy(UGeomet
 			{
 				FGeometryNaniteData& Instance = GeometryNaniteData[GeometryIndex];
 				Instance.HierarchyOffset = GeometryCollection->GetNaniteHierarchyOffset(GeometryIndex);
-				SetInstanceLocalBounds(GeometryIndex, FRenderBounds(), false);
+				Instance.LocalBounds = BoundingBoxes[TransformIndex];
 			}
 		}
 	}
