Geometry collection : fix issue where implicit shapes can have two level of transformed implicit causing collision issues with clusters using union of children shapes

diff --git a/Engine/Source/Runtime/Experimental/Chaos/Private/PhysicsProxy/GeometryCollectionPhysicsProxy.cpp b/Engine/Source/Runtime/Experimental/Chaos/Private/PhysicsProxy/GeometryCollectionPhysicsProxy.cpp
--- a/Engine/Source/Runtime/Experimental/Chaos/Private/PhysicsProxy/GeometryCollectionPhysicsProxy.cpp
+++ b/Engine/Source/Runtime/Experimental/Chaos/Private/PhysicsProxy/GeometryCollectionPhysicsProxy.cpp
@@ -102,6 +102,7 @@ FAutoConsoleVariableRef CVarGeometryCollectionAlwaysGenerateConnectionGraph(
 DEFINE_LOG_CATEGORY_STATIC(UGCC_LOG, Error, All);
 
 static const FSharedSimulationSizeSpecificData& GetSizeSpecificData(const TArray<FSharedSimulationSizeSpecificData>& SizeSpecificData, const FGeometryCollection& RestCollection, const int32 TransformIndex, const FBox& BoundingBox);
+static TUniquePtr<Chaos::FImplicitObject> MakeTransformImplicitObject(const Chaos::FImplicitObject& ImplicitObject, const Chaos::FRigidTransform3& Transform);
 
 //==============================================================================
 // FGeometryCollectionResults
@@ -583,44 +584,24 @@ void FGeometryCollectionPhysicsProxy::Initialize(Chaos::FPBDRigidsEvolutionBase
 							TArray<TUniquePtr<Chaos::FImplicitObject>> ChildImplicits;
 							for (const int32& ChildIndex : GameThreadCollection.Children[ParentToFixIndex])
 							{
-								using FImplicitObjectTransformed = Chaos::TImplicitObjectTransformed<Chaos::FReal, 3>;
-
-								FParticle* ChildParticle = GTParticles[ChildIndex].Get();
 								const FGeometryDynamicCollection::FSharedImplicit& ChildImplicit = GameThreadCollection.Implicits[ChildIndex];
 								if (ChildImplicit)
 								{
 									const Chaos::FRigidTransform3 ChildShapeTransform = GameThreadCollection.MassToLocal[ChildIndex] * GameThreadCollection.Transform[ChildIndex];
 									const Chaos::FRigidTransform3 RelativeShapeTransform = ChildShapeTransform.GetRelativeTransform(ParentShapeTransform);
 
-									// assumption that we only have can only have one level of union for any child
-									if (ChildImplicit->GetType() == Chaos::ImplicitObjectType::Union)
+									TUniquePtr<Chaos::FImplicitObject> TransformedChildImplicit = MakeTransformImplicitObject(*ChildImplicit, RelativeShapeTransform);
+
+									// if this remains a union we need to unpack it 
+									if (TransformedChildImplicit->IsUnderlyingUnion())
 									{
-										if (Chaos::FImplicitObjectUnion* Union = ChildImplicit->GetObject<Chaos::FImplicitObjectUnion>())
-										{
-											for (const TUniquePtr<Chaos::FImplicitObject>& ImplicitObject : Union->GetObjects())
-											{
-												TUniquePtr<Chaos::FImplicitObject> CopiedChildImplicit = ImplicitObject->DeepCopy();
-												FImplicitObjectTransformed* TransformedChildImplicit = new FImplicitObjectTransformed(MoveTemp(CopiedChildImplicit), RelativeShapeTransform);
-												ChildImplicits.Add(TUniquePtr<FImplicitObjectTransformed>(TransformedChildImplicit));
-											}
-										}
+										// we just move the array fo children in the new array 
+										Chaos::FImplicitObjectUnion& Union = static_cast<Chaos::FImplicitObjectUnion&>(*TransformedChildImplicit);
+										ChildImplicits.Append(MoveTemp(Union.GetObjects()));
 									}
 									else
 									{
-										TUniquePtr<Chaos::FImplicitObject> CopiedChildImplicit = GameThreadCollection.Implicits[ChildIndex]->DeepCopy();
-										if (CopiedChildImplicit->GetType() == Chaos::ImplicitObjectType::Transformed)
-										{
-											if (FImplicitObjectTransformed* Transformed = CopiedChildImplicit->GetObject<FImplicitObjectTransformed>())
-											{
-												Transformed->SetTransform(Transformed->GetTransform() * RelativeShapeTransform);
-												ChildImplicits.Add(TUniquePtr<FImplicitObjectTransformed>(Transformed));
-											}
-										}
-										else
-										{
-											FImplicitObjectTransformed* TransformedChildImplicit = new FImplicitObjectTransformed(MoveTemp(CopiedChildImplicit), RelativeShapeTransform);
-											ChildImplicits.Add(TUniquePtr<FImplicitObjectTransformed>(TransformedChildImplicit));
-										}
+										ChildImplicits.Add(MoveTemp(TransformedChildImplicit));
 									}
 								}
 							}
diff --git a/Engine/Source/Runtime/Experimental/Chaos/Public/Chaos/ImplicitObjectUnion.h b/Engine/Source/Runtime/Experimental/Chaos/Public/Chaos/ImplicitObjectUnion.h
--- a/Engine/Source/Runtime/Experimental/Chaos/Public/Chaos/ImplicitObjectUnion.h
+++ b/Engine/Source/Runtime/Experimental/Chaos/Public/Chaos/ImplicitObjectUnion.h
@@ -147,6 +147,8 @@ class CHAOS_API FImplicitObjectUnion : public FImplicitObject
 
 	const TArray<TUniquePtr<FImplicitObject>>& GetObjects() const { return MObjects; }
 
+	TArray<TUniquePtr<FImplicitObject>>& GetObjects() { return MObjects; }
+
 	virtual uint32 GetTypeHash() const override
 	{
 		uint32 Result = 0;
