Geometry collection: make removal random generated timer values consistent between client and server
- make bone colors unique per component based on number of bones
- remove dependency to PhysXUserdata

diff --git a/Engine/Source/Runtime/Experimental/GeometryCollectionEngine/Private/GeometryCollection/GeometryCollectionComponent.cpp b/Engine/Source/Runtime/Experimental/GeometryCollectionEngine/Private/GeometryCollection/GeometryCollectionComponent.cpp
--- a/Engine/Source/Runtime/Experimental/GeometryCollectionEngine/Private/GeometryCollection/GeometryCollectionComponent.cpp
+++ b/Engine/Source/Runtime/Experimental/GeometryCollectionEngine/Private/GeometryCollection/GeometryCollectionComponent.cpp
@@ -2565,7 +2565,7 @@ void UGeometryCollectionComponent::OnCreatePhysicsState()
 		const bool bValidCollection = DynamicCollection && DynamicCollection->Transform.Num() > 0;
 		if (bValidWorld && bValidCollection)
 		{
-			FPhysxUserData::Set<UPrimitiveComponent>(&PhysicsUserData, this);
+			FChaosUserData::Set<UPrimitiveComponent>(&PhysicsUserData, this);
 
 			// If the Component is set to Dynamic, we look to the RestCollection for initial dynamic state override per transform.
 			TManagedArray<int32> & DynamicState = DynamicCollection->DynamicState;
@@ -3184,11 +3184,14 @@ FScopedColorEdit::FScopedColorEdit(UGeometryCollectionComponent* InComponent, bo
 {
 	if (RandomColors.Num() == 0)
 	{
-		FMath::RandInit(2019);
+		// predictable colors based on the component
+		FRandomStream Random(GetTypeHash(InComponent));
 		for (int i = 0; i < 100; i++)
 		{
-			const FColor Color(FMath::Rand() % 100 + 5, FMath::Rand() % 100 + 5, FMath::Rand() % 100 + 5, 255);
-			RandomColors.Push(FLinearColor(Color));
+			const uint8 R = static_cast<uint8>(Random.FRandRange(5, 105));
+			const uint8 G = static_cast<uint8>(Random.FRandRange(5, 105));
+			const uint8 B = static_cast<uint8>(Random.FRandRange(5, 105));
+			RandomColors.Push(FLinearColor(FColor(R, G, B, 255)));
 		}
 	}
 }
diff --git a/Engine/Source/Runtime/Experimental/GeometryCollectionEngine/Private/GeometryCollection/GeometryCollectionEngineRemoval.cpp b/Engine/Source/Runtime/Experimental/GeometryCollectionEngine/Private/GeometryCollection/GeometryCollectionEngineRemoval.cpp
--- a/Engine/Source/Runtime/Experimental/GeometryCollectionEngine/Private/GeometryCollection/GeometryCollectionEngineRemoval.cpp
+++ b/Engine/Source/Runtime/Experimental/GeometryCollectionEngine/Private/GeometryCollection/GeometryCollectionEngineRemoval.cpp
@@ -46,6 +46,11 @@ void FGeometryCollectionRemoveOnBreakDynamicFacade::SetAttributeValues(const Geo
 
 		BreakTimerAttribute.Fill(DisabledBreakTimer);
 
+		// make sure we generate random value consistently between client and server 
+		// we can use the length of the transform group for that
+		const int32 RandomSeed = Children.Num();
+		FRandomStream Random(RandomSeed);
+
 		TManagedArray<float>& PostBreakDuration = PostBreakDurationAttribute.Modify();
 		for (int32 Idx = 0; Idx < PostBreakDuration.Num(); ++Idx)
 		{
@@ -53,7 +58,7 @@ void FGeometryCollectionRemoveOnBreakDynamicFacade::SetAttributeValues(const Geo
 			const FVector2f PostBreakTimer = RemoveOnBreakData.GetBreakTimer();
 			const float MinBreakTime = FMath::Max(0.0f, PostBreakTimer.X);
 			const float MaxBreakTime = FMath::Max(MinBreakTime, PostBreakTimer.Y);
-			PostBreakDuration[Idx] = RemoveOnBreakData.IsEnabled() ? FMath::RandRange(MinBreakTime, MaxBreakTime) : DisabledPostBreakDuration;
+			PostBreakDuration[Idx] = RemoveOnBreakData.IsEnabled() ? Random.FRandRange(MinBreakTime, MaxBreakTime) : DisabledPostBreakDuration;
 		}
 
 		TManagedArray<float>& BreakRemovalDuration = BreakRemovalDurationAttribute.Modify();
@@ -65,7 +70,7 @@ void FGeometryCollectionRemoveOnBreakDynamicFacade::SetAttributeValues(const Geo
 			const float MaxRemovalTime = FMath::Max(MinRemovalTime, RemovalTimer.Y);
 			const bool bIsCluster = (Children[Idx].Num() > 0);
 			const bool bUseClusterCrumbling = (bIsCluster && RemoveOnBreakData.GetClusterCrumbling());
-			BreakRemovalDuration[Idx] = bUseClusterCrumbling ? CrumblingRemovalTimer : FMath::RandRange(MinRemovalTime, MaxRemovalTime);
+			BreakRemovalDuration[Idx] = bUseClusterCrumbling ? CrumblingRemovalTimer : Random.FRandRange(MinRemovalTime, MaxRemovalTime);
 		}
 	}
 }
diff --git a/Engine/Source/Runtime/Experimental/GeometryCollectionEngine/Public/GeometryCollection/GeometryCollectionComponent.h b/Engine/Source/Runtime/Experimental/GeometryCollectionEngine/Public/GeometryCollection/GeometryCollectionComponent.h
--- a/Engine/Source/Runtime/Experimental/GeometryCollectionEngine/Public/GeometryCollection/GeometryCollectionComponent.h
+++ b/Engine/Source/Runtime/Experimental/GeometryCollectionEngine/Public/GeometryCollection/GeometryCollectionComponent.h
@@ -27,7 +27,6 @@
 #include "Chaos/PBDRigidsEvolutionFwd.h"
 #include "EngineDefines.h"
 #include "Math/MathFwd.h"
-#include "PhysxUserData.h"
 
 #include "GeometryCollectionComponent.generated.h"
 
@@ -1174,7 +1173,7 @@ private:
 
 	FCollisionFilterData InitialSimFilter;
 	FCollisionFilterData InitialQueryFilter;
-	FPhysxUserData PhysicsUserData;
+	FChaosUserData PhysicsUserData;
 
 #if WITH_EDITORONLY_DATA
 	UPROPERTY(Transient)
