add accessor functions to the FReadPhysicsObjectInterface to convert the physics object handle to the raw particle approriate for the context

diff --git a/Engine/Source/Runtime/Experimental/Chaos/Private/Chaos/PhysicsObjectInterface.cpp b/Engine/Source/Runtime/Experimental/Chaos/Private/Chaos/PhysicsObjectInterface.cpp
--- a/Engine/Source/Runtime/Experimental/Chaos/Private/Chaos/PhysicsObjectInterface.cpp
+++ b/Engine/Source/Runtime/Experimental/Chaos/Private/Chaos/PhysicsObjectInterface.cpp
@@ -173,6 +173,48 @@ namespace Chaos
 		return {};
 	}
 
+	template<EThreadContext Id>
+	TArray<TThreadParticle<Id>*> FReadPhysicsObjectInterface<Id>::GetAllParticles(TArrayView<FPhysicsObjectHandle> InObjects)
+	{
+		TArray<TThreadParticle<Id>*> Particles;
+		Particles.Reserve(InObjects.Num());
+
+		for (const FPhysicsObjectHandle& Handle : InObjects)
+		{
+			if (TThreadParticle<Id>* Particle = Handle->GetParticle<Id>())
+			{
+				Particles.Add(Particle);
+			}
+		}
+
+		return Particles;
+	}
+
+	template<EThreadContext Id>
+	TArray<TThreadRigidParticle<Id>*> FReadPhysicsObjectInterface<Id>::GetAllRigidParticles(TArrayView<FPhysicsObjectHandle> InObjects)
+	{
+		TArray<TThreadRigidParticle<Id>*> Particles;
+		Particles.Reserve(InObjects.Num());
+
+		for (const FPhysicsObjectHandle& Handle : InObjects)
+		{
+			if (!Handle)
+			{
+				continue;
+			}
+
+			if (TThreadParticle<Id>* Particle = Handle->GetParticle<Id>())
+			{
+				if (TThreadRigidParticle<Id>* RigidParticle = Particle->CastToRigidParticle())
+				{
+					Particles.Add(RigidParticle);
+				}
+			}
+		}
+
+		return Particles;
+	}
+
 	template<EThreadContext Id>
 	TArray<FPerShapeData*> FReadPhysicsObjectInterface<Id>::GetAllShapes(TArrayView<FPhysicsObjectHandle> InObjects)
 	{
diff --git a/Engine/Source/Runtime/Experimental/Chaos/Public/Chaos/PhysicsObjectInterface.h b/Engine/Source/Runtime/Experimental/Chaos/Public/Chaos/PhysicsObjectInterface.h
--- a/Engine/Source/Runtime/Experimental/Chaos/Public/Chaos/PhysicsObjectInterface.h
+++ b/Engine/Source/Runtime/Experimental/Chaos/Public/Chaos/PhysicsObjectInterface.h
@@ -64,6 +64,8 @@ namespace Chaos
 		FQuat GetR(FPhysicsObjectHandle Object);
 		FSpatialAccelerationIdx GetSpatialIndex(FPhysicsObjectHandle Object);
 
+		TArray<TThreadParticle<Id>*> GetAllParticles(TArrayView<FPhysicsObjectHandle> InObjects);
+		TArray<TThreadRigidParticle<Id>*> GetAllRigidParticles(TArrayView<FPhysicsObjectHandle> InObjects);
 		TArray<FPerShapeData*> GetAllShapes(TArrayView<FPhysicsObjectHandle> InObjects);
 		bool GetPhysicsObjectOverlap(FPhysicsObjectHandle ObjectA, FPhysicsObjectHandle ObjectB, bool bTraceComplex, Chaos::FOverlapInfo& OutOverlap);
 		bool GetPhysicsObjectOverlapWithTransform(FPhysicsObjectHandle ObjectA, const FTransform& InTransformA, FPhysicsObjectHandle ObjectB, const FTransform& InTransformB, bool bTraceComplex, Chaos::FOverlapInfo& OutOverlap);
