expose getters for the XR and Velocities property on geometry particles

diff --git a/Engine/Source/Runtime/Experimental/Chaos/Public/Chaos/ParticleHandle.h b/Engine/Source/Runtime/Experimental/Chaos/Public/Chaos/ParticleHandle.h
--- a/Engine/Source/Runtime/Experimental/Chaos/Public/Chaos/ParticleHandle.h
+++ b/Engine/Source/Runtime/Experimental/Chaos/Public/Chaos/ParticleHandle.h
@@ -2358,6 +2358,7 @@ public:
 	const TRotation<T, d>& R() const { return MXR.Read().R(); }
 	void SetR(const TRotation<T, d>& InR, bool bInvalidate = true);
 
+	const FParticlePositionRotation& XR() const { return MXR.Read(); }
 	void SetXR(const FParticlePositionRotation& InXR, bool bInvalidate = true)
 	{
 		MXR.Write(InXR,bInvalidate,MDirtyFlags,Proxy);
@@ -2573,6 +2574,11 @@ public:
 		return MDirtyFlags.IsDirty(CheckBits);
 	}
 
+	void ForceDirty(EChaosPropertyFlags CheckBits)
+	{
+		MDirtyFlags.MarkDirty(CheckBits);
+	}
+
 	const FDirtyChaosPropertyFlags& DirtyFlags() const
 	{
 		return MDirtyFlags;
@@ -2780,6 +2786,7 @@ public:
 		MKinematicTarget.Clear(MDirtyFlags, Proxy);
 	}
 
+	const FParticleVelocities& Velocities() const { return MVelocities.Read(); }
 	void SetVelocities(const FParticleVelocities& InVelocities,bool bInvalidate = true)
 	{
 		MVelocities.Write(InVelocities,bInvalidate,MDirtyFlags,Proxy);
