Updated the "AutoInstances" meshes to not add to the ISMPool until they are needed.
Deprecated the FSoftObjectPath "StaticMesh" property in FGeometryCollectionAutoInstanceMesh and replaced it with a TObjectPtr<UStaticMesh> Mesh property. This fixes an issue where asserts would fire because we were doing sync loads after the loadscreen.
These changes saved roughly 100MBs in our test level.

diff --git a/Engine/Plugins/Experimental/GeometryCollectionPlugin/Source/GeometryCollectionNodes/Private/Dataflow/Nodes/GeometryCollectionAssetNodes.cpp b/Engine/Plugins/Experimental/GeometryCollectionPlugin/Source/GeometryCollectionNodes/Private/Dataflow/Nodes/GeometryCollectionAssetNodes.cpp
--- a/Engine/Plugins/Experimental/GeometryCollectionPlugin/Source/GeometryCollectionNodes/Private/Dataflow/Nodes/GeometryCollectionAssetNodes.cpp
+++ b/Engine/Plugins/Experimental/GeometryCollectionPlugin/Source/GeometryCollectionNodes/Private/Dataflow/Nodes/GeometryCollectionAssetNodes.cpp
@@ -148,7 +148,7 @@ void FCreateGeometryCollectionFromSourcesDataflowNode::Evaluate(Dataflow::FConte
 
 		// todo(chaos) if the source is a geometry collection this will not work properly 
 		FGeometryCollectionAutoInstanceMesh InstancedMesh;
-		InstancedMesh.StaticMesh = Source.SourceGeometryObject;
+		InstancedMesh.Mesh = Cast<UStaticMesh>(Source.SourceGeometryObject.TryLoad());
 		InstancedMesh.Materials = Source.SourceMaterial;
 		const int32 InstancedMeshIndex = OutInstancedMeshes.AddUnique(InstancedMesh);
 
diff --git a/Engine/Source/Runtime/Experimental/GeometryCollectionEngine/Private/GeometryCollection/GeometryCollectionComponent.cpp b/Engine/Source/Runtime/Experimental/GeometryCollectionEngine/Private/GeometryCollection/GeometryCollectionComponent.cpp
--- a/Engine/Source/Runtime/Experimental/GeometryCollectionEngine/Private/GeometryCollection/GeometryCollectionComponent.cpp
+++ b/Engine/Source/Runtime/Experimental/GeometryCollectionEngine/Private/GeometryCollection/GeometryCollectionComponent.cpp
@@ -129,6 +129,9 @@ FAutoConsoleVariableRef CVarChaosGCInitConstantDataParallelForBatchSize(TEXT("p.
 int32 MaxGeometryCollectionAsyncPhysicsTickIdleTimeMs = 30;
 FAutoConsoleVariableRef CVarMaxGeometryCollectionAsyncPhysicsTickIdleTimeMs(TEXT("p.Chaos.GC.MaxGeometryCollectionAsyncPhysicsTickIdleTimeMs"), MaxGeometryCollectionAsyncPhysicsTickIdleTimeMs, TEXT("Amount of time in milliseconds before the async tick turns off when it is otherwise not doing anything."));
 
+bool bChaos_GC_DeferAddingAutoInstancesToISMPool = true;
+FAutoConsoleVariableRef CVarDeferAddingAutoInstancesToISMPool(TEXT("p.Chaos.GC.DeferAddingAutoInstancesToISMPool"), bChaos_GC_DeferAddingAutoInstancesToISMPool, TEXT("When enabled, auto instances will be added to the ISM pool after the GC is broken"));
+
 DEFINE_LOG_CATEGORY_STATIC(UGCC_LOG, Error, All);
 
 extern FGeometryCollectionDynamicDataPool GDynamicDataPool;
@@ -4251,59 +4254,6 @@ void UGeometryCollectionComponent::RegisterToISMPool()
 			if (RestCollection)
 			{
 				ISMPoolMeshGroupIndex = ISMPoolComp->CreateMeshGroup();
-				if (bChaos_GC_UseISMPoolForNonFracturedParts)
-				{
-					if (RestCollection->GetGeometryCollection())
-					{
-						// if we use ISM pool for the hierarchy we must hide the component for rendering 
-						bCanRenderComponent = false;
-
-						// fisrt count the instance per mesh 
-						TArray<int32> InstanceCounts;
-						InstanceCounts.AddZeroed(RestCollection->AutoInstanceMeshes.Num());
-
-						const TManagedArray<TSet<int32>>& Children = RestCollection->GetGeometryCollection()->Children;
-
-						const GeometryCollection::Facades::FCollectionInstancedMeshFacade InstancedMeshFacade(*RestCollection->GetGeometryCollection());
-						if (InstancedMeshFacade.IsValid())
-						{
-							for (int32 TransformIndex = 0; TransformIndex < InstancedMeshFacade.GetNumIndices(); TransformIndex++)
-							{
-								const int32 AutoInstanceMeshIndex = InstancedMeshFacade.GetIndex(TransformIndex);
-								if (Children[TransformIndex].Num() == 0)
-								{
-									InstanceCounts[AutoInstanceMeshIndex]++;
-								}
-							}
-						}
-
-						// now register each mesh 
-						for (int32 MeshIndex = 0; MeshIndex < RestCollection->AutoInstanceMeshes.Num(); MeshIndex++)
-						{
-							const FGeometryCollectionAutoInstanceMesh& AutoInstanceMesh = RestCollection->AutoInstanceMeshes[MeshIndex];
-							if (UStaticMesh* StaticMesh = Cast<UStaticMesh>(AutoInstanceMesh.StaticMesh.TryLoad()))
-							{
-								bool bMaterialOverride = false;
-								for (int32 MatIndex = 0; MatIndex < AutoInstanceMesh.Materials.Num(); MatIndex++)
-								{
-									const UMaterialInterface* OriginalMaterial = StaticMesh->GetMaterial(MatIndex);
-									if (OriginalMaterial != AutoInstanceMesh.Materials[MatIndex])
-									{
-										bMaterialOverride = true;
-										break;
-									}
-								}
-								FGeometryCollectionStaticMeshInstance StaticMeshInstance;
-								StaticMeshInstance.StaticMesh = StaticMesh;
-								if (bMaterialOverride)
-								{
-									StaticMeshInstance.MaterialsOverrides = AutoInstanceMesh.Materials;
-								}
-								ISMPoolComp->AddMeshToGroup(ISMPoolMeshGroupIndex, StaticMeshInstance, InstanceCounts[MeshIndex], bChaos_GC_UseHierarchicalISMForLeafMeshes);
-							}
-						}
-					}
-				}
 
 				// root proxy if available 
 				// TODO : if ISM pool is not available : uses a standard static mesh component
@@ -4320,6 +4270,18 @@ void UGeometryCollectionComponent::RegisterToISMPool()
 						ISMPoolRootProxyMeshIds.Add(ISMPoolComp->AddMeshToGroup(ISMPoolMeshGroupIndex, StaticMeshInstance, 1, bChaos_GC_UseHierarchicalISMForProxyMesh));
 					}
 				}
+
+				if (bChaos_GC_UseISMPoolForNonFracturedParts)
+				{
+					// if we use ISM pool for the hierarchy we must hide the component for rendering 
+					bCanRenderComponent = false;
+
+					const bool bHasRootProxyMesh = !ISMPoolRootProxyMeshIds.IsEmpty();
+					if (!(bChaos_GC_DeferAddingAutoInstancesToISMPool && bHasRootProxyMesh))
+					{
+						AddAutoInstancesToISMPool();
+					}
+				}
 			}
 
 			SetVisibility(bCanRenderComponent);
@@ -4343,6 +4305,68 @@ void UGeometryCollectionComponent::UnregisterFromISMPool()
 	}
 }
 
+void UGeometryCollectionComponent::AddAutoInstancesToISMPool()
+{
+	if (ISMPoolAutoInstancesMeshIds.Num() > 0)
+	{
+		return;
+	}
+
+	if (UGeometryCollectionISMPoolComponent* ISMPoolComp = AssignedISMPool->GetISMPoolComp())
+	{
+		if (bChaos_GC_UseISMPoolForNonFracturedParts)
+		{
+			if (RestCollection->GetGeometryCollection())
+			{
+				// first count the instance per mesh 
+				TArray<int32> InstanceCounts;
+				InstanceCounts.AddZeroed(RestCollection->AutoInstanceMeshes.Num());
+
+				const TManagedArray<TSet<int32>>& Children = RestCollection->GetGeometryCollection()->Children;
+
+				const GeometryCollection::Facades::FCollectionInstancedMeshFacade InstancedMeshFacade(*RestCollection->GetGeometryCollection());
+				if (InstancedMeshFacade.IsValid())
+				{
+					for (int32 TransformIndex = 0; TransformIndex < InstancedMeshFacade.GetNumIndices(); TransformIndex++)
+					{
+						const int32 AutoInstanceMeshIndex = InstancedMeshFacade.GetIndex(TransformIndex);
+						if (Children[TransformIndex].Num() == 0)
+						{
+							InstanceCounts[AutoInstanceMeshIndex]++;
+						}
+					}
+				}
+
+				// now register each mesh 
+				for (int32 MeshIndex = 0; MeshIndex < RestCollection->AutoInstanceMeshes.Num(); MeshIndex++)
+				{
+					const FGeometryCollectionAutoInstanceMesh& AutoInstanceMesh = RestCollection->AutoInstanceMeshes[MeshIndex];
+					if (const UStaticMesh* StaticMesh = AutoInstanceMesh.Mesh)
+					{
+						bool bMaterialOverride = false;
+						for (int32 MatIndex = 0; MatIndex < AutoInstanceMesh.Materials.Num(); MatIndex++)
+						{
+							const UMaterialInterface* OriginalMaterial = StaticMesh->GetMaterial(MatIndex);
+							if (OriginalMaterial != AutoInstanceMesh.Materials[MatIndex])
+							{
+								bMaterialOverride = true;
+								break;
+							}
+						}
+						FGeometryCollectionStaticMeshInstance StaticMeshInstance;
+						StaticMeshInstance.StaticMesh = const_cast<UStaticMesh*>(StaticMesh);
+						if (bMaterialOverride)
+						{
+							StaticMeshInstance.MaterialsOverrides = AutoInstanceMesh.Materials;
+						}
+						ISMPoolAutoInstancesMeshIds.Add(ISMPoolComp->AddMeshToGroup(ISMPoolMeshGroupIndex, StaticMeshInstance, InstanceCounts[MeshIndex], bChaos_GC_UseHierarchicalISMForLeafMeshes));
+					}
+				}
+			}
+		}
+	}
+}
+
 void UGeometryCollectionComponent::RefreshISMPoolInstances()
 {
 	if (CanUseISMPool())
@@ -4390,6 +4414,8 @@ void UGeometryCollectionComponent::RefreshISMPoolInstances()
 							}
 							else if (bChaos_GC_UseISMPoolForNonFracturedParts)
 							{
+								AddAutoInstancesToISMPool();
+
 								// make sure this mesh is invisible 
 								// todo : should be event based instead of doing it every frame
 								if (bHasRootProxyMesh && GlobalMatrices.IsValidIndex(RootIndex))
@@ -4414,7 +4440,7 @@ void UGeometryCollectionComponent::RefreshISMPoolInstances()
 											InstanceTransforms.Add(FTransform(GlobalMatrices[TransformIndex]) * ComponentTransform);
 										}
 									}
-									ISMPoolComp->BatchUpdateInstancesTransforms(ISMPoolMeshGroupIndex, MeshIndex, 0, InstanceTransforms, bWorlSpace, bMarkRenderStateDirty, bTeleport);
+									ISMPoolComp->BatchUpdateInstancesTransforms(ISMPoolMeshGroupIndex, ISMPoolAutoInstancesMeshIds[MeshIndex], 0, InstanceTransforms, bWorlSpace, bMarkRenderStateDirty, bTeleport);
 								}
 							}
 						}
diff --git a/Engine/Source/Runtime/Experimental/GeometryCollectionEngine/Private/GeometryCollection/GeometryCollectionEngineConversion.cpp b/Engine/Source/Runtime/Experimental/GeometryCollectionEngine/Private/GeometryCollection/GeometryCollectionEngineConversion.cpp
--- a/Engine/Source/Runtime/Experimental/GeometryCollectionEngine/Private/GeometryCollection/GeometryCollectionEngineConversion.cpp
+++ b/Engine/Source/Runtime/Experimental/GeometryCollectionEngine/Private/GeometryCollection/GeometryCollectionEngineConversion.cpp
@@ -486,7 +486,7 @@ void FGeometryCollectionEngineConversion::AppendAutoInstanceMeshIndices(UGeometr
 			// create the schema if necessary
 			InstancedMeshFacade.DefineSchema();
 	
-			const int32 AutoInstanceMeshIndex = GeometryCollectionObject->FindOrAddAutoInstanceMesh(*StaticMesh, Materials);
+			const int32 AutoInstanceMeshIndex = GeometryCollectionObject->FindOrAddAutoInstanceMesh(StaticMesh, Materials);
 			for (int32 TransformIndex = FromTransformIndex; TransformIndex < NewNumOfTransforms; TransformIndex++)
 			{
 				InstancedMeshFacade.SetIndex(TransformIndex, AutoInstanceMeshIndex);
diff --git a/Engine/Source/Runtime/Experimental/GeometryCollectionEngine/Private/GeometryCollection/GeometryCollectionObject.cpp b/Engine/Source/Runtime/Experimental/GeometryCollectionEngine/Private/GeometryCollection/GeometryCollectionObject.cpp
--- a/Engine/Source/Runtime/Experimental/GeometryCollectionEngine/Private/GeometryCollection/GeometryCollectionObject.cpp
+++ b/Engine/Source/Runtime/Experimental/GeometryCollectionEngine/Private/GeometryCollection/GeometryCollectionObject.cpp
@@ -1298,7 +1298,7 @@ void UGeometryCollection::RemoveExemplars(const TArray<int32>& SortedRemovalIndi
 
 bool FGeometryCollectionAutoInstanceMesh::operator ==(const FGeometryCollectionAutoInstanceMesh& Other) const
 {
-	return (StaticMesh == Other.StaticMesh) && (Materials == Other.Materials);
+	return (Mesh == Other.Mesh) && (Materials == Other.Materials);
 }
 
 /** find or add a auto instance mesh and return its index */
@@ -1313,10 +1313,10 @@ int32 UGeometryCollection::FindOrAddAutoInstanceMesh(const FGeometryCollectionAu
 	return AutoInstanceMeshes.AddUnique(AutoInstanecMesh);
 }
 
-int32 UGeometryCollection::FindOrAddAutoInstanceMesh(const UStaticMesh& StaticMesh, const TArray<UMaterialInterface*>& MeshMaterials)
+int32 UGeometryCollection::FindOrAddAutoInstanceMesh(const UStaticMesh* StaticMesh, const TArray<UMaterialInterface*>& MeshMaterials)
 {
 	FGeometryCollectionAutoInstanceMesh NewMesh;
-	NewMesh.StaticMesh = FSoftObjectPath(&StaticMesh);
+	NewMesh.Mesh = StaticMesh;
 	NewMesh.Materials = MeshMaterials;
 	return AutoInstanceMeshes.AddUnique(NewMesh);
 }
@@ -1476,6 +1476,19 @@ void UGeometryCollection::PostLoad()
 		}
 		RootProxy_DEPRECATED = nullptr;
 	}
+
+	for (int32 MeshIndex = 0; MeshIndex < AutoInstanceMeshes.Num(); MeshIndex++)
+	{
+		FGeometryCollectionAutoInstanceMesh& AutoInstanceMesh = AutoInstanceMeshes[MeshIndex];
+		if (!AutoInstanceMesh.StaticMesh_DEPRECATED.IsNull())
+		{
+			if (UStaticMesh* StaticMesh = Cast<UStaticMesh>(AutoInstanceMesh.StaticMesh_DEPRECATED.TryLoad()))
+			{
+				AutoInstanceMesh.Mesh = TObjectPtr<UStaticMesh>(StaticMesh);
+			}
+			AutoInstanceMesh.StaticMesh_DEPRECATED = nullptr;
+		}
+	}
 #endif
 }
 
diff --git a/Engine/Source/Runtime/Experimental/GeometryCollectionEngine/Public/GeometryCollection/GeometryCollectionComponent.h b/Engine/Source/Runtime/Experimental/GeometryCollectionEngine/Public/GeometryCollection/GeometryCollectionComponent.h
--- a/Engine/Source/Runtime/Experimental/GeometryCollectionEngine/Public/GeometryCollection/GeometryCollectionComponent.h
+++ b/Engine/Source/Runtime/Experimental/GeometryCollectionEngine/Public/GeometryCollection/GeometryCollectionComponent.h
@@ -1069,6 +1069,7 @@ protected:
 	TObjectPtr<AGeometryCollectionISMPoolActor> AssignedISMPool = nullptr;
 	int32 ISMPoolMeshGroupIndex = INDEX_NONE;
 	TArray<int32> ISMPoolRootProxyMeshIds;
+	TArray<int32> ISMPoolAutoInstancesMeshIds;
 
 	/** Populate the dynamic particle data for the render thread. */
 	FGeometryCollectionDynamicData* InitDynamicData(bool bInitialization);
@@ -1258,6 +1259,8 @@ private:
 
 	void OnPostPhysicsSync();
 
+	void AddAutoInstancesToISMPool();
+
 	/** The clusters we need to replicate */
 	TUniquePtr<TSet<Chaos::FPBDRigidClusteredParticleHandle*>> ClustersToRep;
 
diff --git a/Engine/Source/Runtime/Experimental/GeometryCollectionEngine/Public/GeometryCollection/GeometryCollectionEngineConversion.h b/Engine/Source/Runtime/Experimental/GeometryCollectionEngine/Public/GeometryCollection/GeometryCollectionEngineConversion.h
--- a/Engine/Source/Runtime/Experimental/GeometryCollectionEngine/Public/GeometryCollection/GeometryCollectionEngineConversion.h
+++ b/Engine/Source/Runtime/Experimental/GeometryCollectionEngine/Public/GeometryCollection/GeometryCollectionEngineConversion.h
@@ -39,7 +39,7 @@ public:
 	 * @param GeometryCollectionObject geometry collection to add to 
 	 * @param FromTransformIndex transform index to start from 
 	 * @param StaticMesh  static mesh to add reference to 
-	 * @param Materials materials correcponding to the static mesh instance to get the index from 
+	 * @param Materials materials corresponding to the static mesh instance to get the index from 
 	 */
 	static void AppendAutoInstanceMeshIndices(UGeometryCollection* GeometryCollectionObject, int32 FromTransformIndex, const UStaticMesh* StaticMesh, const TArray<UMaterialInterface*>& Materials);
 
@@ -169,7 +169,7 @@ public:
 
 	/**
 	*  Appends a GeometryCollectionSource to a GeometryCollection
-	*  @param GeometryCollectionSource : geometry collectiun source object ( from UGeometryCollection collection asset ) 
+	*  @param GeometryCollectionSource : geometry collection source object ( from UGeometryCollection collection asset ) 
 	*  @param GeometryCollectionInOut : GeometryCollection to append to 
 	*  @param MaterialsInOut  : array of materials to append to
 	*/
diff --git a/Engine/Source/Runtime/Experimental/GeometryCollectionEngine/Public/GeometryCollection/GeometryCollectionObject.h b/Engine/Source/Runtime/Experimental/GeometryCollectionEngine/Public/GeometryCollection/GeometryCollectionObject.h
--- a/Engine/Source/Runtime/Experimental/GeometryCollectionEngine/Public/GeometryCollection/GeometryCollectionObject.h
+++ b/Engine/Source/Runtime/Experimental/GeometryCollectionEngine/Public/GeometryCollection/GeometryCollectionObject.h
@@ -64,8 +64,13 @@ struct GEOMETRYCOLLECTIONENGINE_API FGeometryCollectionAutoInstanceMesh
 {
 	GENERATED_BODY()
 
+#if WITH_EDITORONLY_DATA
+	UPROPERTY(meta = (DeprecatedProperty, DeprecationMessage = "Use Mesh instead."))
+	FSoftObjectPath StaticMesh_DEPRECATED;
+#endif
+
 	UPROPERTY(BlueprintReadWrite, EditAnywhere, Category = "AutoInstance", meta = (AllowedClasses = "/Script/Engine.StaticMesh"))
-	FSoftObjectPath StaticMesh;
+	TObjectPtr<const UStaticMesh> Mesh;
 
 	UPROPERTY(BlueprintReadWrite, EditAnywhere, Category = "AutoInstance")
 	TArray<TObjectPtr<UMaterialInterface>> Materials;
@@ -408,7 +413,7 @@ public:
 	int32 FindOrAddAutoInstanceMesh(const FGeometryCollectionAutoInstanceMesh& AutoInstanecMesh);
 
 	/** find or add a auto instance mesh from a mesh and alist of material and return its index */
-	int32 FindOrAddAutoInstanceMesh(const UStaticMesh& StaticMesh, const TArray<UMaterialInterface*>& Materials);
+	int32 FindOrAddAutoInstanceMesh(const UStaticMesh* StaticMesh, const TArray<UMaterialInterface*>& Materials);
 
 	/** Produce a deep copy of GeometryCollection member, stripped of data unecessary for gameplay. */
 	TSharedPtr<FGeometryCollection, ESPMode::ThreadSafe> GenerateMinimalGeometryCollection() const;
