geometry collection physics proxies should not attempt to clean up particles that belong to a cluster union physics proxy

diff --git a/Engine/Source/Runtime/Experimental/Chaos/Private/PhysicsProxy/GeometryCollectionPhysicsProxy.cpp b/Engine/Source/Runtime/Experimental/Chaos/Private/PhysicsProxy/GeometryCollectionPhysicsProxy.cpp
--- a/Engine/Source/Runtime/Experimental/Chaos/Private/PhysicsProxy/GeometryCollectionPhysicsProxy.cpp
+++ b/Engine/Source/Runtime/Experimental/Chaos/Private/PhysicsProxy/GeometryCollectionPhysicsProxy.cpp
@@ -2206,7 +2206,9 @@ void FGeometryCollectionPhysicsProxy::OnRemoveFromSolver(Chaos::FPBDRigidsSolver
 		{
 			if (FClusterHandle* ParentCluster = Evolution->GetRigidClustering().DestroyClusterParticle(Handle))
 			{
-				if (ParentCluster->InternalCluster())
+				// We do not want to do anything relating to any cluster unions that might own this geometry collection.
+				IPhysicsProxyBase* ParentClusterProxy = ParentCluster->PhysicsProxy();
+				if (ParentCluster->InternalCluster() && ParentClusterProxy == this)
 				{
 					ClustersToRebuild.Add(ParentCluster);
 				}
