changes to the physics proxy base to support the cluster union proxy

- Added the new cluster union proxy type and timestamp.
- Moved GetSolverSyncTimestamp from the single particle physics proxy to the base to allow it to be used more generally.
- Add ability to set a ParentProxy on a physics proxy as some behavior needs to change if owned by a cluster union.

diff --git a/Engine/Source/Runtime/Experimental/Chaos/Private/Framework/PhysicsProxyBase.cpp b/Engine/Source/Runtime/Experimental/Chaos/Private/Framework/PhysicsProxyBase.cpp
--- a/Engine/Source/Runtime/Experimental/Chaos/Private/Framework/PhysicsProxyBase.cpp
+++ b/Engine/Source/Runtime/Experimental/Chaos/Private/Framework/PhysicsProxyBase.cpp
@@ -9,4 +9,14 @@ IPhysicsProxyBase::~IPhysicsProxyBase()
 	{
 		GetSolver<Chaos::FPhysicsSolverBase>()->RemoveDirtyProxy(this);
 	}
+}
+
+int32 IPhysicsProxyBase::GetSolverSyncTimestamp_External() const
+{
+	if (Chaos::FPhysicsSolverBase* SolverBase = GetSolverBase())
+	{
+		return SolverBase->GetMarshallingManager().GetExternalTimestamp_External();
+	}
+
+	return INDEX_NONE;
 }
\ No newline at end of file
diff --git a/Engine/Source/Runtime/Experimental/Chaos/Public/Chaos/Framework/PhysicsProxyBase.h b/Engine/Source/Runtime/Experimental/Chaos/Public/Chaos/Framework/PhysicsProxyBase.h
--- a/Engine/Source/Runtime/Experimental/Chaos/Public/Chaos/Framework/PhysicsProxyBase.h
+++ b/Engine/Source/Runtime/Experimental/Chaos/Public/Chaos/Framework/PhysicsProxyBase.h
@@ -17,6 +17,7 @@ enum class EPhysicsProxyType : uint32
 	SuspensionConstraintType = 9,
 	CharacterGroundConstraintType = 10,
 	SingleParticleProxy,
+	ClusterUnionProxy,
 	Count
 };
 
@@ -58,6 +59,15 @@ struct CHAOS_API FGeometryCollectionProxyTimestamp: public FProxyTimestampBase
 	// ( including kinematic targeting )
 };
 
+struct CHAOS_API FClusterUnionProxyTimestamp : public FProxyTimestampBase
+{
+	TTimestampProperty<Chaos::FVec3> OverWriteX;
+	TTimestampProperty<Chaos::FRotation3> OverWriteR;
+	TTimestampProperty<Chaos::FVec3> OverWriteV;
+	TTimestampProperty<Chaos::FVec3> OverWriteW;
+	TTimestampProperty<bool> OverWriteAnchored;
+	TTimestampProperty<TArray<FTransform>> OverWriteChildToParent;
+};
 
 class CHAOS_API IPhysicsProxyBase
 {
@@ -110,6 +120,9 @@ public:
 
 	int32 GetIgnoreDataOnStep_Internal() const { return IgnoreDataOnStep_Internal; }
 
+	IPhysicsProxyBase* GetParentProxy() const { return ParentProxy; }
+	void SetParentProxy(IPhysicsProxyBase* InProxy) { ParentProxy = InProxy; }
+
 protected:
 	// Ensures that derived classes can successfully call this destructor
 	// but no one can delete using a IPhysicsProxyBase*
@@ -128,11 +141,14 @@ protected:
 private:
 	int32 DirtyIdx;
 	TSharedPtr<FProxyTimestampBase,ESPMode::ThreadSafe> SyncTimestamp;
+	IPhysicsProxyBase* ParentProxy = nullptr;
 protected:
 	/** Proxy type */
 	EPhysicsProxyType Type;
 	int32 InitializedOnStep = INDEX_NONE;
 	int32 IgnoreDataOnStep_Internal = INDEX_NONE;
+
+	int32 GetSolverSyncTimestamp_External() const;
 };
 
 struct PhysicsProxyWrapper
diff --git a/Engine/Source/Runtime/Experimental/Chaos/Public/PhysicsProxy/SingleParticlePhysicsProxy.h b/Engine/Source/Runtime/Experimental/Chaos/Public/PhysicsProxy/SingleParticlePhysicsProxy.h
--- a/Engine/Source/Runtime/Experimental/Chaos/Public/PhysicsProxy/SingleParticlePhysicsProxy.h
+++ b/Engine/Source/Runtime/Experimental/Chaos/Public/PhysicsProxy/SingleParticlePhysicsProxy.h
@@ -1258,19 +1258,6 @@ public:
 
 		return false;
 	}
-
-private:
-
-	int32 GetSolverSyncTimestamp_External() const
-	{
-		if (FPhysicsSolverBase* SolverBase = GetSolverBase())
-		{
-			return SolverBase->GetMarshallingManager().GetExternalTimestamp_External();
-		}
-
-		return INDEX_NONE;
-	}
-
 };
 
 static_assert(sizeof(FRigidBodyHandle_External) == sizeof(FSingleParticlePhysicsProxy), "Derived types only used to constrain API, all data lives in base class ");
