clearing a chaos property flag should only remove a dirty proxy if that proxy has no dirty shapes

diff --git a/Engine/Source/Runtime/Experimental/Chaos/Public/Chaos/ChaosMarshallingManager.h b/Engine/Source/Runtime/Experimental/Chaos/Public/Chaos/ChaosMarshallingManager.h
--- a/Engine/Source/Runtime/Experimental/Chaos/Public/Chaos/ChaosMarshallingManager.h
+++ b/Engine/Source/Runtime/Experimental/Chaos/Public/Chaos/ChaosMarshallingManager.h
@@ -92,7 +92,7 @@ public:
 		
 	}
 
-
+	// Forcefully removes the proxy from being dirty.
 	void Remove(IPhysicsProxyBase* Base)
 	{
 		const int32 Idx = Base->GetDirtyIdx();
@@ -119,6 +119,24 @@ public:
 		}
 	}
 
+	// Only does the removal if no shapes are dirty.
+	void RemoveIfNoShapesAreDirty(IPhysicsProxyBase* Base)
+	{
+		const int32 Idx = Base->GetDirtyIdx();
+		if (Idx != INDEX_NONE)
+		{
+			FDirtyProxiesBucket& Bucket = DirtyProxyBuckets[(uint32)Base->GetType()];
+			if (Bucket.ProxiesData.IsValidIndex(Idx))
+			{
+				FDirtyProxy& ProxyData = Bucket.ProxiesData[Idx];
+				if (ProxyData.ShapeDataIndices.IsEmpty())
+				{
+					Remove(Base);
+				}
+			}
+		}
+	}
+
 	void Reset()
 	{
 		for(FDirtyProxiesBucket& Bucket : DirtyProxyBuckets)
diff --git a/Engine/Source/Runtime/Experimental/Chaos/Public/Chaos/Framework/PhysicsSolverBase.h b/Engine/Source/Runtime/Experimental/Chaos/Public/Chaos/Framework/PhysicsSolverBase.h
--- a/Engine/Source/Runtime/Experimental/Chaos/Public/Chaos/Framework/PhysicsSolverBase.h
+++ b/Engine/Source/Runtime/Experimental/Chaos/Public/Chaos/Framework/PhysicsSolverBase.h
@@ -243,6 +243,11 @@ namespace Chaos
 			MarshallingManager.GetProducerData_External()->DirtyProxiesDataBuffer.Remove(ProxyBaseIn);
 		}
 
+		void RemoveDirtyProxyIfNoShapesAreDirty(IPhysicsProxyBase* ProxyBaseIn)
+		{
+			MarshallingManager.GetProducerData_External()->DirtyProxiesDataBuffer.RemoveIfNoShapesAreDirty(ProxyBaseIn);
+		}
+
 		const FDirtyProxiesBucketInfo& GetDirtyProxyBucketInfo_External()
 		{
 			return MarshallingManager.GetProducerData_External()->DirtyProxiesDataBuffer.GetDirtyProxyBucketInfo();
diff --git a/Engine/Source/Runtime/Experimental/Chaos/Public/Chaos/Properties.h b/Engine/Source/Runtime/Experimental/Chaos/Public/Chaos/Properties.h
--- a/Engine/Source/Runtime/Experimental/Chaos/Public/Chaos/Properties.h
+++ b/Engine/Source/Runtime/Experimental/Chaos/Public/Chaos/Properties.h
@@ -48,7 +48,7 @@ public:
 		{
 			if (FPhysicsSolverBase* PhysicsSolverBase = Proxy->GetSolver<FPhysicsSolverBase>())
 			{
-				PhysicsSolverBase->RemoveDirtyProxy(Proxy);
+				PhysicsSolverBase->RemoveDirtyProxyIfNoShapesAreDirty(Proxy);
 			}
 		}
 	}
