add ability to destroy/cleanup a cluster union as well as proper ways of setting the unique index on cluster uinon particles as well as we now set the parennt proxy on child particles

diff --git a/Engine/Source/Runtime/Experimental/Chaos/Private/Chaos/ClusterUnionManager.cpp b/Engine/Source/Runtime/Experimental/Chaos/Private/Chaos/ClusterUnionManager.cpp
--- a/Engine/Source/Runtime/Experimental/Chaos/Private/Chaos/ClusterUnionManager.cpp
+++ b/Engine/Source/Runtime/Experimental/Chaos/Private/Chaos/ClusterUnionManager.cpp
@@ -17,7 +17,7 @@ namespace Chaos
 	}
 
 	DECLARE_CYCLE_STAT(TEXT("FClusterUnionManager::CreateNewClusterUnion"), STAT_CreateNewClusterUnion, STATGROUP_Chaos);
-	FClusterUnionIndex FClusterUnionManager::CreateNewClusterUnion(const FClusterCreationParameters& Parameters, FClusterUnionExplicitIndex ExplicitIndex)
+	FClusterUnionIndex FClusterUnionManager::CreateNewClusterUnion(const FClusterCreationParameters& Parameters, FClusterUnionExplicitIndex ExplicitIndex, const FUniqueIdx* UniqueIndex)
 	{
 		SCOPE_CYCLE_COUNTER(STAT_CreateNewClusterUnion);
 		FClusterUnionIndex NewIndex = ClaimNextUnionIndex();
@@ -26,17 +26,33 @@ namespace Chaos
 		FClusterUnion NewUnion;
 		NewUnion.ExplicitIndex = ExplicitIndex;
 		NewUnion.SharedGeometry = ForceRecreateClusterUnionSharedGeometry(NewUnion);
-		NewUnion.InternalCluster = MClustering.CreateClusterParticle(-NewIndex, {}, Parameters, NewUnion.SharedGeometry);
+		NewUnion.InternalCluster = MClustering.CreateClusterParticle(-NewIndex, {}, Parameters, NewUnion.SharedGeometry, nullptr, UniqueIndex);
 		NewUnion.Parameters = Parameters;
 		if (ensure(NewUnion.InternalCluster != nullptr))
 		{
 			NewUnion.InternalCluster->SetInternalCluster(true);
+
+			// No bounds for now since we don't have particles. When/if we do get particles later, updating the
+			// geometry should switch this flag back on.
+			NewUnion.InternalCluster->SetHasBounds(false);
 		}
 
 		ClusterUnions.Add(NewIndex, NewUnion);
 		return NewIndex;
 	}
 
+	DECLARE_CYCLE_STAT(TEXT("FClusterUnionManager::DestroyClusterUnion"), STAT_DestroyClusterUnion, STATGROUP_Chaos);
+	void FClusterUnionManager::DestroyClusterUnion(FClusterUnionIndex Index)
+	{
+		SCOPE_CYCLE_COUNTER(STAT_DestroyClusterUnion);
+
+		if (FClusterUnion* ClusterUnion = FindClusterUnion(Index))
+		{
+			ClusterUnion->ChildParticles.Empty();
+			MClustering.DestroyClusterParticle(ClusterUnion->InternalCluster);
+		}
+	}
+
 	DECLARE_CYCLE_STAT(TEXT("FClusterUnionManager::ForceRecreateClusterUnionSharedGeometry"), STAT_ForceRecreateClusterUnionSharedGeometry, STATGROUP_Chaos);
 	TSharedPtr<FImplicitObject, ESPMode::ThreadSafe> FClusterUnionManager::ForceRecreateClusterUnionSharedGeometry(const FClusterUnion& Union)
 	{
@@ -165,6 +181,12 @@ namespace Chaos
 		);
 		MClustering.SetGenerateClusterBreaking(false);
 
+		// If a physics proxy was set already on the cluster we want to make sure that doesn't change.
+		// This is needed to eventually be able to introduce a new physics proxy that gets attached to the
+		// cluster union particle so that it can communicate with the game thread.
+		IPhysicsProxyBase* OldProxy = Cluster->InternalCluster->PhysicsProxy();
+		EPhysicsProxyType OldProxyType = OldProxy ? OldProxy->GetType() : EPhysicsProxyType::NoneType;
+
 		// If we're a new cluster, we need to determine whether to start the cluster in a sleeping or dynamic state.
 		// Only stay sleeping if all the particles we add are also sleeping.
 		const bool bIsNewCluster = Cluster->ChildParticles.IsEmpty();
@@ -187,14 +209,21 @@ namespace Chaos
 
 			bIsSleeping &= Handle->ObjectState() == EObjectStateType::Sleeping;
 
-			if (FPBDRigidClusteredParticleHandle* ClusterHandle = Handle->CastToClustered(); ClusterHandle && bReleaseClustersFirst)
+			if (FPBDRigidClusteredParticleHandle* ClusterHandle = Handle->CastToClustered())
 			{
-				TSet<FPBDRigidParticleHandle*> Children = MClustering.ReleaseClusterParticles(ClusterHandle, true);
-				FinalParticlesToAdd.Append(Children.Array());
+				if (bReleaseClustersFirst)
+				{
+					TSet<FPBDRigidParticleHandle*> Children = MClustering.ReleaseClusterParticles(ClusterHandle, true);
+					FinalParticlesToAdd.Append(Children.Array());
 
-				for (FPBDRigidParticleHandle* Child : Children)
+					for (FPBDRigidParticleHandle* Child : Children)
+					{
+						ChildToParentMap.Add(Child, ClusterHandle);
+					}
+				}
+				else
 				{
-					ChildToParentMap.Add(Child, ClusterHandle);
+					FinalParticlesToAdd.Add(Handle);
 				}
 
 				bIsAnchored |= ClusterHandle->IsAnchored();
@@ -212,13 +241,22 @@ namespace Chaos
 
 		Cluster->ChildParticles.Append(FinalParticlesToAdd);
 
-		// If a physics proxy was set already on the cluster we want to make sure that doesn't change.
-		// This is needed to eventually be able to introduce a new physics proxy that gets attached to the
-		// cluster union particle so that it can communicate with the game thread.
-		IPhysicsProxyBase* OldProxy = Cluster->InternalCluster->PhysicsProxy();
-
 		MClustering.AddParticlesToCluster(Cluster->InternalCluster, FinalParticlesToAdd, ChildToParentMap);
 
+		// For all the particles that have been added to the cluster we need to set their parent proxy to the
+		// cluster's proxy if it exists. We need the proxy type check because for non-cluster union proxy backed unions,
+		// the cluster union's proxy will be the proxy of the most recently added particle.
+		if (OldProxy && OldProxyType == EPhysicsProxyType::ClusterUnionProxy)
+		{
+			for (FPBDRigidParticleHandle* Particle : FinalParticlesToAdd)
+			{
+				if (Particle && Particle->PhysicsProxy())
+				{
+					Particle->PhysicsProxy()->SetParentProxy(OldProxy);
+				}
+			}
+		}
+
 		// Update cluster properties.
 		FMatrix33 ClusterInertia(0);
 
diff --git a/Engine/Source/Runtime/Experimental/Chaos/Public/Chaos/ClusterUnionManager.h b/Engine/Source/Runtime/Experimental/Chaos/Public/Chaos/ClusterUnionManager.h
--- a/Engine/Source/Runtime/Experimental/Chaos/Public/Chaos/ClusterUnionManager.h
+++ b/Engine/Source/Runtime/Experimental/Chaos/Public/Chaos/ClusterUnionManager.h
@@ -5,7 +5,6 @@
 #include "Chaos/ParticleHandleFwd.h"
 #include "Containers/Array.h"
 #include "Containers/Map.h"
-#include "Containers/Set.h"
 
 namespace Chaos
 {
@@ -33,7 +32,7 @@ namespace Chaos
 		TSharedPtr<FImplicitObject, ESPMode::ThreadSafe> SharedGeometry;
 
 		// All the particles that belong to this cluster.
-		TSet<FPBDRigidParticleHandle*> ChildParticles;
+		TArray<FPBDRigidParticleHandle*> ChildParticles;
 
 		// An explicit index set by the user if any.
 		FClusterUnionExplicitIndex ExplicitIndex;
@@ -53,7 +52,10 @@ namespace Chaos
 		FClusterUnionManager(FRigidClustering& InClustering, FPBDRigidsEvolutionGBF& InEvolution);
 
 		// Creates a new cluster union with an automatically assigned cluster union index.
-		FClusterUnionIndex CreateNewClusterUnion(const FClusterCreationParameters& Parameters, FClusterUnionExplicitIndex ExplicitIndex = INDEX_NONE);
+		FClusterUnionIndex CreateNewClusterUnion(const FClusterCreationParameters& Parameters, FClusterUnionExplicitIndex ExplicitIndex = INDEX_NONE, const FUniqueIdx* UniqueIndex = nullptr);
+
+		// Destroy a given cluster union.
+		void DestroyClusterUnion(FClusterUnionIndex Index);
 
 		// Add a new operation to the queue.
 		void AddPendingExplicitIndexOperation(FClusterUnionExplicitIndex Index, EClusterUnionOperation Op, const TArray<FPBDRigidParticleHandle*>& Particles);
@@ -113,7 +115,6 @@ namespace Chaos
 
 		// Actually performs the change specified in the FClusterUnionOperationData structure.
 		void HandleAddOperation(FClusterUnionIndex ClusterIndex, const TArray<FPBDRigidParticleHandle*>& Particles, bool bReleaseClustersFirst);
-
 	};
 
 }
\ No newline at end of file
