Destruction builder : set import data on geometry collection assets

diff --git a/Engine/Source/Runtime/Experimental/GeometryCollectionEngine/Public/GeometryCollection/GeometryCollectionObject.h b/Engine/Source/Runtime/Experimental/GeometryCollectionEngine/Public/GeometryCollection/GeometryCollectionObject.h
--- a/Engine/Source/Runtime/Experimental/GeometryCollectionEngine/Public/GeometryCollection/GeometryCollectionObject.h
+++ b/Engine/Source/Runtime/Experimental/GeometryCollectionEngine/Public/GeometryCollection/GeometryCollectionObject.h
@@ -682,7 +682,7 @@ public:
 
 #if WITH_EDITORONLY_DATA
 	/** Importing data and options used for this geometry collection */
-	UPROPERTY(EditAnywhere, Instanced, Category = ImportSettings)
+	UPROPERTY(EditAnywhere, BlueprintReadWrite, Instanced, Category = ImportSettings)
 	TObjectPtr<class UAssetImportData> AssetImportData;
 
 	/** Information for thumbnail rendering */
