Chaos Cloth - Fixed Legacy Backstop.

#rnx

diff --git a/Engine/Source/Runtime/Experimental/Chaos/Public/Chaos/PBDSphericalConstraint.h b/Engine/Source/Runtime/Experimental/Chaos/Public/Chaos/PBDSphericalConstraint.h
--- a/Engine/Source/Runtime/Experimental/Chaos/Public/Chaos/PBDSphericalConstraint.h
+++ b/Engine/Source/Runtime/Experimental/Chaos/Public/Chaos/PBDSphericalConstraint.h
@@ -169,7 +169,7 @@ public:
 		, ParticleOffset(InParticleOffset)
 		, Scale(MeshScale)
 		, bEnabled(true)
-		, bUseLegacyBackstop(false)
+		, bUseLegacyBackstop(GetUseLegacyBackstop(PropertyCollection, false))  // Only set the legacy backstop in constructor
 	{
 		check(InSphereRadii.Num() == InParticleCount);
 		check(InSphereOffsetDistances.Num() == InParticleCount);
@@ -353,6 +353,7 @@ private:
 
 	UE_CHAOS_DECLARE_PROPERTYCOLLECTION_NAME(BackstopDistance, float);
 	UE_CHAOS_DECLARE_PROPERTYCOLLECTION_NAME(BackstopRadius, float);
+	UE_CHAOS_DECLARE_PROPERTYCOLLECTION_NAME(UseLegacyBackstop, bool);
 };
 
 }  // End namespace Chaos::Softs
