Chaos: Change IsHole for Heightfields to never return holes if we have no materials on the heightfield (holes are a material property)

diff --git a/Engine/Source/Runtime/Experimental/Chaos/Private/Chaos/HeightField.cpp b/Engine/Source/Runtime/Experimental/Chaos/Private/Chaos/HeightField.cpp
--- a/Engine/Source/Runtime/Experimental/Chaos/Private/Chaos/HeightField.cpp
+++ b/Engine/Source/Runtime/Experimental/Chaos/Private/Chaos/HeightField.cpp
@@ -873,7 +873,8 @@ namespace Chaos
 
 	bool Chaos::FHeightField::IsHole(int32 InIndex) const
 	{
-		return GetMaterialIndex(InIndex) == TNumericLimits<uint8>::Max();
+		const bool bHasMaterials = GeomData.MaterialIndices.Num() > 0;
+		return bHasMaterials && GetMaterialIndex(InIndex) == TNumericLimits<uint8>::Max();
 	}
 
 	bool Chaos::FHeightField::IsHole(int32 InCellX, int32 InCellY) const
