add the cluster union physics proxy to properly marshal data for them between the game thread and physics thread as well as the game thread actor/component to make cluster unions server authoritative

- Modified the FPhysicsSolverBase::PullPhysicsStateForEachDirtyProxy_External to take in a struct that can be used to statically dispatch the appropriate action per proxy type instead of needing to add a new lambda for each new proxy type.

diff --git a/Engine/Source/Runtime/Engine/Classes/PhysicsEngine/ClusterUnionActor.h b/Engine/Source/Runtime/Engine/Classes/PhysicsEngine/ClusterUnionActor.h
new file mode 100644
--- /dev/null
+++ b/Engine/Source/Runtime/Engine/Classes/PhysicsEngine/ClusterUnionActor.h
@@ -0,0 +1,24 @@
+// Copyright Epic Games, Inc. All Rights Reserved.
+
+#pragma once
+
+#include "GameFramework/Actor.h"
+#include "ClusterUnionActor.generated.h"
+
+class UClusterUnionComponent;
+
+/**
+ * A lightweight actor that can be used to own a cluster union component.
+ */
+UCLASS()
+class ENGINE_API AClusterUnionActor : public AActor
+{
+	GENERATED_BODY()
+public:
+	AClusterUnionActor(const FObjectInitializer& ObjectInitializer);
+
+protected:
+	/** The pivot used while building. */
+	UPROPERTY(EditAnywhere, BlueprintReadOnly, Category = "Cluster Union")
+	TObjectPtr<UClusterUnionComponent> ClusterUnion;
+};
\ No newline at end of file
diff --git a/Engine/Source/Runtime/Engine/Classes/PhysicsEngine/ClusterUnionComponent.h b/Engine/Source/Runtime/Engine/Classes/PhysicsEngine/ClusterUnionComponent.h
new file mode 100644
--- /dev/null
+++ b/Engine/Source/Runtime/Engine/Classes/PhysicsEngine/ClusterUnionComponent.h
@@ -0,0 +1,196 @@
+// Copyright Epic Games, Inc. All Rights Reserved.
+
+#pragma once
+
+#include "Chaos/PhysicsObject.h"
+#include "Components/PrimitiveComponent.h"
+#include "Containers/Array.h"
+#include "Containers/Map.h"
+#include "Containers/Set.h"
+#include "Engine/EngineTypes.h"
+#include "Logging/LogMacros.h"
+#include "PhysicsInterfaceTypesCore.h"
+#include "PhysicsProxy/ClusterUnionPhysicsProxy.h"
+#include "UObject/ObjectKey.h"
+
+#include "ClusterUnionComponent.generated.h"
+
+class AActor;
+class FPhysScene_Chaos;
+class UClusterUnionReplicatedProxyComponent;
+
+DECLARE_LOG_CATEGORY_EXTERN(LogClusterUnion, Log, All);
+
+namespace Chaos
+{
+	class FPBDRigidsSolver;
+}
+
+USTRUCT()
+struct FClusteredComponentData
+{
+	GENERATED_BODY()
+
+	TArray<Chaos::FPhysicsObjectHandle> PhysicsObjects;
+
+	// Using a TWeakObjectPtr here because the UClusterUnionReplicatedProxyComponent will have a pointer back
+	// and we don't want to get into a situation where a circular reference occurs.
+	UPROPERTY()
+	TWeakObjectPtr<UClusterUnionReplicatedProxyComponent> ReplicatedProxyComponent;
+
+	UPROPERTY()
+	bool bWasReplicating = true;
+};
+
+USTRUCT()
+struct FClusteredActorData
+{
+	GENERATED_BODY()
+
+	UPROPERTY()
+	TSet<TObjectPtr<UPrimitiveComponent>> Components;
+
+	UPROPERTY()
+	bool bWasReplicatingMovement = true;
+};
+
+USTRUCT()
+struct FClusterUnionReplicatedData
+{
+	GENERATED_BODY()
+
+	UPROPERTY()
+	FVector_NetQuantize100 LinVel;
+
+	UPROPERTY()
+	FVector_NetQuantize100 AngVel;
+
+	UPROPERTY()
+	bool bIsAnchored = false;
+};
+
+USTRUCT()
+struct FClusterUnionPendingAddData
+{
+	GENERATED_BODY()
+
+	UPROPERTY()
+	TArray<FName> BoneNames;
+};
+
+/**
+ * This does the bulk of the work exposing a physics cluster union to the game thread.
+ * This component needs to be a primitive component primarily because of how physics
+ * proxies need to be registered with the solver with an association with a primitive component.
+ * This component can be used as part of AClusterUnionActor or on its own as its list of
+ * clustered components/actors can be specified dynamically at runtime and/or statically
+ * on asset creation.
+ */
+UCLASS()
+class ENGINE_API UClusterUnionComponent : public UPrimitiveComponent
+{
+	GENERATED_BODY()
+public:
+	UClusterUnionComponent(const FObjectInitializer& ObjectInitializer);
+
+	UFUNCTION(BlueprintCallable, Category="Cluster Union")
+	void AddComponentToCluster(UPrimitiveComponent* InComponent, const TArray<FName>& BoneNames);
+
+	UFUNCTION(BlueprintCallable, Category = "Cluster Union")
+	void SetIsAnchored(bool bIsAnchored);
+
+	// The SyncVelocitiesFromPhysics will set replicated state using data from the physics thread. 
+	void SyncVelocitiesFromPhysics(const FVector& LinearVelocity, const FVector& AngularVelocity);
+	// The SyncIsAnchoredFromPhysics will set replicated state using data from the physics thread.
+	void SyncIsAnchoredFromPhysics(bool bIsAnchored);
+	// SyncChildToParentFromProxy will take the ChildToParent transform data from the proxy and make sure that gets replicated to the correct replicated proxy components.
+	void SyncChildToParentFromProxy();
+
+	UFUNCTION()
+	bool IsComponentAdded(UPrimitiveComponent* Component) { return ComponentToPhysicsObjects.Contains(Component); }
+
+	friend class UClusterUnionReplicatedProxyComponent;
+protected:
+
+	// This should only be called on the client when replication happens.
+	UFUNCTION()
+	void ForceSetChildToParent(UPrimitiveComponent* InComponent, const TArray<FName>& BoneNames, const TArray<FTransform>& ChildToParent);
+
+private:
+	// These are the statically clustered components. These should
+	// be specified in the editor and never change.
+	UPROPERTY(EditAnywhere, Category = "Cluster Union")
+	TArray<FComponentReference> ClusteredComponentsReferences;
+	
+	// We need to keep track of the mapping of primitive components to physics objects.
+	// This way we know the right physics objects to pass when removing the component (because
+	// it's possible to get a different list of physics objects when we get to removal). A
+	// side benefit here is being able to track which components are clustered.
+	TMap<TObjectKey<UPrimitiveComponent>, FClusteredComponentData> ComponentToPhysicsObjects;
+
+	// Similarly, we need to keep track of the mapping from physics object handles to primitive components.
+	// The physics proxy will only have data about these handles, so when we get an update from the proxy
+	// about a particular handle we will want to map it back to a component here.
+	TMap<Chaos::FPhysicsObjectHandle, UPrimitiveComponent*> PhysicsObjectToComponent;
+
+	// Also keep track of which actors we are clustering and their components. We make modifications on
+	// actors that get clustered so we need to make sure we undo those changes only once all its clustered
+	// components are removed from the cluster.
+	TMap<TObjectKey<AActor>, FClusteredActorData> ActorToComponents;
+
+	// Sometimes we might be in the process of waiting for a component to create it physics state before adding to the cluster.
+	// Make sure we don't try to add the component multiples times while the add is pending.
+	TMap<TObjectKey<UPrimitiveComponent>, FClusterUnionPendingAddData> PendingComponentsToAdd;
+
+	// Data that can be changed at runtime to keep state about the cluster union consistent between the server and client.
+	UPROPERTY(ReplicatedUsing=OnRep_RigidState)
+	FClusterUnionReplicatedData ReplicatedRigidState;
+
+	// Handles changes to ReplicatedRigidState. Note that this function does not handle replication of X/R since we make use
+	// of the scene component's default replication for that.
+	UFUNCTION()
+	void OnRep_RigidState();
+
+	FPhysScene_Chaos* GetChaosScene() const;
+
+	Chaos::FClusterUnionPhysicsProxy* PhysicsProxy;
+
+	// User data to be able to tie the cluster particle back to this component.
+	FChaosUserData PhysicsUserData;
+
+	// Need to handle the fact that this component may or may not be initialized prior to the components referenced in
+	// ClusteredComponentsReferences. This function lets us listen to OnComponentPhysicsStateChanged on the incoming
+	// primitive component so that once the physics state is properly created we can begin the process of adding it.
+	UFUNCTION()
+	void HandleComponentPhysicsStateChange(UPrimitiveComponent* ChangedComponent, EComponentPhysicsStateChange StateChange);
+
+	// Whether or not this code is running on the server.
+	UFUNCTION()
+	bool IsAuthority() const;
+
+	//~ Begin UActorComponent Interface
+public:
+	virtual void OnCreatePhysicsState() override;
+	virtual void OnDestroyPhysicsState() override;
+	virtual bool ShouldCreatePhysicsState() const override;
+	virtual bool HasValidPhysicsState() const override;
+	virtual void GetLifetimeReplicatedProps(TArray<FLifetimeProperty>& OutLifetimeProps) const override;
+	//~ End UActorComponent Interface
+
+	//~ Begin USceneComponent Interface
+public:
+	virtual void OnUpdateTransform(EUpdateTransformFlags UpdateTransformFlags, ETeleportType Teleport = ETeleportType::None) override;
+	//~ End USceneComponent Interface
+
+	//~ Begin IPhysicsComponent Interface.
+public:
+	virtual Chaos::FPhysicsObject* GetPhysicsObjectById(int32 Id) const override;
+	virtual Chaos::FPhysicsObject* GetPhysicsObjectByName(const FName& Name) const override;
+	virtual TArray<Chaos::FPhysicsObject*> GetAllPhysicsObjects() const override;
+	//~ End IPhysicsComponent Interface.
+
+	//~ Begin UObject Interface.
+public:
+	static void AddReferencedObjects(UObject* InThis, FReferenceCollector& Collector);
+	//~ End UObject Interface.
+};
\ No newline at end of file
diff --git a/Engine/Source/Runtime/Engine/Classes/PhysicsEngine/ClusterUnionReplicatedProxyComponent.h b/Engine/Source/Runtime/Engine/Classes/PhysicsEngine/ClusterUnionReplicatedProxyComponent.h
new file mode 100644
--- /dev/null
+++ b/Engine/Source/Runtime/Engine/Classes/PhysicsEngine/ClusterUnionReplicatedProxyComponent.h
@@ -0,0 +1,93 @@
+// Copyright Epic Games, Inc. All Rights Reserved.
+
+#pragma once
+
+#include "Components/ActorComponent.h"
+#include "GenericPlatform/GenericPlatformMisc.h"
+
+#include "ClusterUnionReplicatedProxyComponent.generated.h"
+
+class UClusterUnionComponent;
+class UPrimitiveComponent;
+
+/**
+ * This component lets us store replicated information about how any particular UPrimitiveComponent
+ * should be attached to its parent cluster union. The benefits of using a separate components are:
+ * 
+ *	1) It lets us avoid adding any additional overhead into the UPrimitiveComponent.
+ *  2) It lets the replicated information have the same net relevancy as the actor being added to the cluster union
+ *     rather than having the same net relevancy as the cluster union (i.e. in the case of replicating this data in
+ *     an array in the ClusterUnionComponent).
+ *  3) It lets us pinpoint what exactly is being added/removed (vs if all this data was stored in an array) which lets
+ *     us be a bit more efficient in terms of modifying the cluster union.
+ */
+UCLASS()
+class ENGINE_API UClusterUnionReplicatedProxyComponent : public UActorComponent
+{
+	GENERATED_BODY()
+public:
+	UClusterUnionReplicatedProxyComponent(const FObjectInitializer& ObjectInitializer);
+
+	UFUNCTION()
+	void SetParentClusterUnion(UClusterUnionComponent* InComponent);
+
+	UFUNCTION()
+	void SetChildClusteredComponent(UPrimitiveComponent* InComponent);
+
+	UFUNCTION()
+	void SetParticleBoneNames(const TArray<FName>& InNames);
+
+	UFUNCTION()
+	void SetParticleChildToParent(const FName& BoneName, const FTransform& ChildToParent);
+
+protected:
+
+	UFUNCTION()
+	void OnRep_ParentClusterUnion();
+
+	UFUNCTION()
+	void OnRep_ChildClusteredComponent();
+
+	UFUNCTION()
+	void OnRep_ParticleBoneNames();
+
+	UFUNCTION()
+	void OnRep_ParticleChildToParents();
+
+private:
+	UPROPERTY(ReplicatedUsing=OnRep_ParentClusterUnion)
+	TObjectPtr<UClusterUnionComponent> ParentClusterUnion;
+
+	UPROPERTY()
+	bool bNetUpdateParentClusterUnion;
+
+	UPROPERTY(ReplicatedUsing=OnRep_ChildClusteredComponent)
+	TObjectPtr<UPrimitiveComponent> ChildClusteredComponent;
+
+	UPROPERTY()
+	bool bNetUpdateChildClusteredComponent;
+
+	UPROPERTY(ReplicatedUsing=OnRep_ParticleBoneNames)
+	TArray<FName> ParticleBoneNames;
+
+	UPROPERTY()
+	bool bNetUpdateParticleBoneNames;
+
+	UPROPERTY(ReplicatedUsing=OnRep_ParticleChildToParents)
+	TArray<FTransform> ParticleChildToParents;
+
+	UPROPERTY()
+	bool bNetUpdateParticleChildToParents;
+
+	void DeferUntilChildClusteredComponentInParentUnion(TFunction<void()> Func);
+
+	//~ Begin UActorComponent Interface
+public:
+	virtual void GetLifetimeReplicatedProps(TArray<FLifetimeProperty>& OutLifetimeProps) const override;
+	//~ End UActorComponent Interface
+
+	//~ Begin UObject Interface
+public:
+	virtual void PostRepNotifies() override;
+	//~ End UObject Interface
+};
\ No newline at end of file
diff --git a/Engine/Source/Runtime/Engine/Private/PhysicsEngine/ClusterUnionActor.cpp b/Engine/Source/Runtime/Engine/Private/PhysicsEngine/ClusterUnionActor.cpp
new file mode 100644
--- /dev/null
+++ b/Engine/Source/Runtime/Engine/Private/PhysicsEngine/ClusterUnionActor.cpp
@@ -0,0 +1,15 @@
+// Copyright Epic Games, Inc. All Rights Reserved.
+#include "PhysicsEngine/ClusterUnionActor.h"
+
+#include "PhysicsEngine/ClusterUnionComponent.h"
+
+#include UE_INLINE_GENERATED_CPP_BY_NAME(ClusterUnionActor)
+
+AClusterUnionActor::AClusterUnionActor(const FObjectInitializer& ObjectInitializer)
+	: AActor(ObjectInitializer)
+{
+	ClusterUnion = CreateDefaultSubobject<UClusterUnionComponent>(TEXT("ClusterUnion"));
+
+	SetRemoteRoleForBackwardsCompat(ROLE_SimulatedProxy);
+	bReplicates = true;
+}
\ No newline at end of file
diff --git a/Engine/Source/Runtime/Engine/Private/PhysicsEngine/ClusterUnionComponent.cpp b/Engine/Source/Runtime/Engine/Private/PhysicsEngine/ClusterUnionComponent.cpp
new file mode 100644
--- /dev/null
+++ b/Engine/Source/Runtime/Engine/Private/PhysicsEngine/ClusterUnionComponent.cpp
@@ -0,0 +1,419 @@
+// Copyright Epic Games, Inc. All Rights Reserved.
+#include "PhysicsEngine/ClusterUnionComponent.h"
+
+#include "Engine/World.h"
+#include "Net/Core/PushModel/PushModel.h"
+#include "Net/UnrealNetwork.h"
+#include "Physics/Experimental/PhysScene_Chaos.h"
+#include "PhysicsEngine/ClusterUnionReplicatedProxyComponent.h"
+#include "PhysicsEngine/PhysicsObjectExternalInterface.h"
+#include "PhysicsProxy/ClusterUnionPhysicsProxy.h"
+#include "TimerManager.h"
+
+#include UE_INLINE_GENERATED_CPP_BY_NAME(ClusterUnionComponent)
+
+DEFINE_LOG_CATEGORY(LogClusterUnion);
+
+UClusterUnionComponent::UClusterUnionComponent(const FObjectInitializer& ObjectInitializer)
+	: UPrimitiveComponent(ObjectInitializer)
+{
+	PhysicsProxy = nullptr;
+	SetIsReplicatedByDefault(true);
+}
+
+FPhysScene_Chaos* UClusterUnionComponent::GetChaosScene() const
+{
+	if (AActor* Owner = GetOwner(); Owner && Owner->GetWorld())
+	{
+		return Owner->GetWorld()->GetPhysicsScene();
+	}
+
+	check(GWorld);
+	return GWorld->GetPhysicsScene();
+}
+
+void UClusterUnionComponent::AddComponentToCluster(UPrimitiveComponent* InComponent, const TArray<FName>& BoneNames)
+{
+	if (!InComponent || !PhysicsProxy)
+	{
+		return;
+	}
+
+	if (!InComponent->HasValidPhysicsState())
+	{
+		if (!PendingComponentsToAdd.Contains(InComponent))
+		{
+			// Early out - defer adding the component to the cluster until the component has a valid physics state.
+			FClusterUnionPendingAddData Data;
+			Data.BoneNames = BoneNames;
+			PendingComponentsToAdd.Add(InComponent, Data);
+			InComponent->OnComponentPhysicsStateChanged.AddDynamic(this, &UClusterUnionComponent::HandleComponentPhysicsStateChange);
+		}
+		return;
+	}
+
+	PendingComponentsToAdd.Remove(InComponent);
+
+	TArray<Chaos::FPhysicsObjectHandle> Objects;
+
+	if (BoneNames.IsEmpty())
+	{
+		Objects = InComponent->GetAllPhysicsObjects();
+	}
+	else
+	{
+		Objects.Reserve(BoneNames.Num());
+		for (const FName& Name : BoneNames)
+		{
+			Objects.Add(InComponent->GetPhysicsObjectByName(Name));
+		}
+	}
+
+	FLockedReadPhysicsObjectExternalInterface Interface = FPhysicsObjectExternalInterface::LockRead(Objects);
+
+	if (BoneNames.IsEmpty())
+	{
+		Objects = Objects.FilterByPredicate(
+			[&Interface](Chaos::FPhysicsObjectHandle Handle)
+			{
+				return !Interface->AreAllDisabled({ &Handle, 1 });
+			}
+		);
+	}
+
+	if (Objects.IsEmpty())
+	{
+		UE_LOG(LogClusterUnion, Warning, TEXT("Trying to add a component [%p] with no physics objects to a cluster union...ignoring"), InComponent)
+		return;
+	}
+
+	// Force the component and its parent actor to stop replicating movement.
+	// Setting the component to not replicate should be sufficient since a simulating
+	// component shouldn't be doing much more than replicating its position anyway.
+	if (AActor* Owner = InComponent->GetOwner())
+	{
+		if (FClusteredActorData* Data = ActorToComponents.Find(Owner))
+		{
+			Data->Components.Add(InComponent);
+		}
+		else
+		{
+			FClusteredActorData NewData;
+			NewData.Components.Add(InComponent);
+			NewData.bWasReplicatingMovement = Owner->IsReplicatingMovement();
+			ActorToComponents.Add(Owner, NewData);
+
+			if (IsAuthority())
+			{
+				Owner->SetReplicatingMovement(false);
+			}
+		}
+	}
+
+	FClusteredComponentData Data;
+	Data.PhysicsObjects = Objects;
+	Data.bWasReplicating = InComponent->GetIsReplicated();
+
+	if (IsAuthority())
+	{
+		InComponent->SetIsReplicated(false);
+		if (AActor* Owner = InComponent->GetOwner())
+		{
+			// Create a replicated proxy component and add it to the actor being added to the cluster.
+			// This component will take care of replicating this addition into the cluster.
+			TObjectPtr<UClusterUnionReplicatedProxyComponent> ReplicatedProxy = NewObject<UClusterUnionReplicatedProxyComponent>(Owner);
+			if (ensure(ReplicatedProxy))
+			{
+				ReplicatedProxy->RegisterComponent();
+				ReplicatedProxy->SetParentClusterUnion(this);
+				ReplicatedProxy->SetChildClusteredComponent(InComponent);
+
+				TArray<FName> PhysicsObjectBoneNames;
+				PhysicsObjectBoneNames.Reserve(Objects.Num());
+
+				for (Chaos::FPhysicsObjectHandle PhysicsObject : Objects)
+				{
+					PhysicsObjectBoneNames.Add(Chaos::FPhysicsObjectInterface::GetName(PhysicsObject));
+				}
+
+				ReplicatedProxy->SetParticleBoneNames(PhysicsObjectBoneNames);
+				ReplicatedProxy->SetIsReplicated(true);
+			}
+
+			Data.ReplicatedProxyComponent = ReplicatedProxy;
+			Owner->FlushNetDormancy();
+		}
+	}
+
+	ComponentToPhysicsObjects.Add(InComponent, Data);
+
+	for (Chaos::FPhysicsObjectHandle Handle : Objects)
+	{
+		PhysicsObjectToComponent.Add(Handle, InComponent);
+	}
+	PhysicsProxy->AddPhysicsObjects_External(Objects);
+}
+
+void UClusterUnionComponent::SetIsAnchored(bool bIsAnchored)
+{
+	if (!PhysicsProxy)
+	{
+		return;
+	}
+
+	PhysicsProxy->SetIsAnchored_External(bIsAnchored);
+}
+
+bool UClusterUnionComponent::IsAuthority() const
+{
+	ENetMode Mode = GetNetMode();
+	if (Mode == ENetMode::NM_Standalone)
+	{
+		return true;
+	}
+
+	if (AActor* Owner = GetOwner())
+	{
+		return Owner->GetLocalRole() == ROLE_Authority && Mode != NM_Client;
+	}
+
+	return false;
+}
+
+void UClusterUnionComponent::OnCreatePhysicsState()
+{
+	USceneComponent::OnCreatePhysicsState();
+
+	// If we've already created the physics proxy we shouldn't do this again.
+	if (PhysicsProxy)
+	{
+		return;
+	}
+
+	// If we're not actually playing/needing this to simulate (e.g. in the editor) there should be no reason to create this proxy.
+	const bool bValidWorld = GetWorld() && (GetWorld()->IsGameWorld() || GetWorld()->IsPreviewWorld());
+	if (!bValidWorld)
+	{
+		return;
+	}
+
+	// TODO: Expose these parameters via the component.
+	Chaos::FClusterCreationParameters Parameters{ 0.3f, 100, false, false };
+	Parameters.ConnectionMethod = Chaos::FClusterCreationParameters::EConnectionMethod::DelaunayTriangulation;
+
+	FChaosUserData::Set<UPrimitiveComponent>(&PhysicsUserData, this);
+	PhysicsProxy = new Chaos::FClusterUnionPhysicsProxy{this, Parameters, static_cast<void*>(&PhysicsUserData), IsAuthority() ? Chaos::EThreadContext::Internal : Chaos::EThreadContext::External };
+	PhysicsProxy->Initialize_External();
+	if (FPhysScene_Chaos* Scene = GetChaosScene())
+	{
+		Scene->AddObject(this, PhysicsProxy);
+	}
+
+	// It's just logically easier to be consistent on the client to go through the replication route.
+	if (IsAuthority())
+	{
+		for (const FComponentReference& ComponentReference : ClusteredComponentsReferences)
+		{
+			if (!ComponentReference.OtherActor.IsValid())
+			{
+				continue;
+			}
+
+			AddComponentToCluster(Cast<UPrimitiveComponent>(ComponentReference.GetComponent(ComponentReference.OtherActor.Get())), {});
+		}
+	}
+}
+
+void UClusterUnionComponent::OnDestroyPhysicsState()
+{
+	USceneComponent::OnDestroyPhysicsState();
+
+	if (!PhysicsProxy)
+	{
+		return;
+	}
+
+	if (FPhysScene_Chaos* Scene = GetChaosScene())
+	{
+		Scene->RemoveObject(PhysicsProxy);
+	}
+
+	PhysicsProxy = nullptr;
+}
+
+void UClusterUnionComponent::OnUpdateTransform(EUpdateTransformFlags UpdateTransformFlags, ETeleportType Teleport)
+{
+	USceneComponent::OnUpdateTransform(UpdateTransformFlags, Teleport);
+
+	if (PhysicsProxy && !(UpdateTransformFlags & EUpdateTransformFlags::SkipPhysicsUpdate))
+	{
+		// If the component transform changes, we need to make sure this update is reflected on the physics thread as well.
+		// This code path is generally used when setting the transform manually or when it's set via replication.
+		const FTransform Transform = GetComponentTransform();
+		PhysicsProxy->SetXR_External(Transform.GetLocation(), Transform.GetRotation());
+	}
+}
+
+bool UClusterUnionComponent::ShouldCreatePhysicsState() const
+{
+	return true;
+}
+
+bool UClusterUnionComponent::HasValidPhysicsState() const
+{
+	return PhysicsProxy != nullptr;
+}
+
+Chaos::FPhysicsObject* UClusterUnionComponent::GetPhysicsObjectById(int32 Id) const
+{
+	if (!PhysicsProxy)
+	{
+		return nullptr;
+	}
+	return PhysicsProxy->GetPhysicsObjectHandle();
+}
+
+Chaos::FPhysicsObject* UClusterUnionComponent::GetPhysicsObjectByName(const FName& Name) const
+{
+	return GetPhysicsObjectById(0);
+}
+
+TArray<Chaos::FPhysicsObject*> UClusterUnionComponent::GetAllPhysicsObjects() const
+{
+	return { GetPhysicsObjectById(0) };
+}
+
+void UClusterUnionComponent::HandleComponentPhysicsStateChange(UPrimitiveComponent* ChangedComponent, EComponentPhysicsStateChange StateChange)
+{
+	// TODO: Maybe we should handle the destroyed state change too?
+	if (!ChangedComponent || StateChange != EComponentPhysicsStateChange::Created)
+	{
+		return;
+	}
+
+	ChangedComponent->OnComponentPhysicsStateChanged.RemoveDynamic(this, &UClusterUnionComponent::HandleComponentPhysicsStateChange);
+
+	if (FClusterUnionPendingAddData* PendingData = PendingComponentsToAdd.Find(ChangedComponent))
+	{
+		AddComponentToCluster(ChangedComponent, PendingData->BoneNames);
+	}
+}
+
+void UClusterUnionComponent::SyncVelocitiesFromPhysics(const FVector& LinearVelocity, const FVector& AngularVelocity)
+{
+	if (!IsAuthority())
+	{
+		return;
+	}
+
+	ReplicatedRigidState.LinVel = LinearVelocity;
+	ReplicatedRigidState.AngVel = AngularVelocity;
+	MARK_PROPERTY_DIRTY_FROM_NAME(UClusterUnionComponent, ReplicatedRigidState, this);
+}
+
+void UClusterUnionComponent::SyncIsAnchoredFromPhysics(bool bIsAnchored)
+{
+	if (!IsAuthority())
+	{
+		return;
+	}
+
+	ReplicatedRigidState.bIsAnchored = bIsAnchored;
+	MARK_PROPERTY_DIRTY_FROM_NAME(UClusterUnionComponent, ReplicatedRigidState, this);
+}
+
+void UClusterUnionComponent::SyncChildToParentFromProxy()
+{
+	if (!PhysicsProxy)
+	{
+		return;
+	}
+
+	const TArray<Chaos::FPhysicsObjectHandle>& ChildPhysicsObjects = PhysicsProxy->GetChildPhysicsObjects_External();
+	const Chaos::FClusterUnionSyncedData& SyncedData = PhysicsProxy->GetSyncedData_External();
+	if (ChildPhysicsObjects.Num() != SyncedData.ChildToParent.Num())
+	{
+		return;
+	}
+
+	for (int32 Index = 0; Index < ChildPhysicsObjects.Num(); ++Index)
+	{
+		Chaos::FPhysicsObjectHandle PhysicsObject = ChildPhysicsObjects[Index];
+		const FTransform& ChildToParent = SyncedData.ChildToParent[Index];
+
+		if (UPrimitiveComponent** Component = PhysicsObjectToComponent.Find(PhysicsObject); Component && *Component)
+		{
+			if (FClusteredComponentData* Data = ComponentToPhysicsObjects.Find(*Component); Data && Data->ReplicatedProxyComponent.IsValid())
+			{
+				Data->ReplicatedProxyComponent.Get()->SetParticleChildToParent(Chaos::FPhysicsObjectInterface::GetName(PhysicsObject), ChildToParent);
+			}
+		}
+	}
+}
+
+void UClusterUnionComponent::OnRep_RigidState()
+{
+	if (!PhysicsProxy)
+	{
+		return;
+	}
+
+	PhysicsProxy->SetLinearVelocity_External(ReplicatedRigidState.LinVel);
+	PhysicsProxy->SetAngularVelocity_External(ReplicatedRigidState.AngVel);
+	PhysicsProxy->SetIsAnchored_External(ReplicatedRigidState.bIsAnchored);
+}
+
+void UClusterUnionComponent::GetLifetimeReplicatedProps(TArray<FLifetimeProperty>& OutLifetimeProps) const
+{
+	UPrimitiveComponent::GetLifetimeReplicatedProps(OutLifetimeProps);
+
+	FDoRepLifetimeParams Params;
+	Params.bIsPushBased = true;
+
+	DOREPLIFETIME_WITH_PARAMS_FAST(UClusterUnionComponent, ReplicatedRigidState, Params);
+}
+
+void UClusterUnionComponent::ForceSetChildToParent(UPrimitiveComponent* InComponent, const TArray<FName>& BoneNames, const TArray<FTransform>& ChildToParent)
+{
+	if (IsAuthority() || !PhysicsProxy || !ensure(InComponent) || !ensure(BoneNames.Num() == ChildToParent.Num()))
+	{
+		return;
+	}
+
+	for (int32 Index = 0; Index < BoneNames.Num(); ++Index)
+	{
+		Chaos::FPhysicsObjectHandle Handle = InComponent->GetPhysicsObjectByName(BoneNames[Index]);
+		PhysicsProxy->SetChildToParent_External(Handle, ChildToParent[Index]);
+	}
+}
+
+void UClusterUnionComponent::AddReferencedObjects(UObject* InThis, FReferenceCollector& Collector)
+{
+	UPrimitiveComponent::AddReferencedObjects(InThis, Collector);
+
+	UClusterUnionComponent* This = CastChecked<UClusterUnionComponent>(InThis);
+
+	{
+		const UScriptStruct* ScriptStruct = FClusteredComponentData::StaticStruct();
+		for (TPair<TObjectKey<UPrimitiveComponent>, FClusteredComponentData>& Kvp : This->ComponentToPhysicsObjects)
+		{
+			Collector.AddReferencedObjects(ScriptStruct, reinterpret_cast<void*>(&Kvp.Value), This, nullptr);
+		}
+	}
+
+	{
+		const UScriptStruct* ScriptStruct = FClusteredActorData::StaticStruct();
+		for (TPair<TObjectKey<AActor>, FClusteredActorData>& Kvp : This->ActorToComponents)
+		{
+			Collector.AddReferencedObjects(ScriptStruct, reinterpret_cast<void*>(&Kvp.Value), This, nullptr);
+		}
+	}
+
+	{
+		const UScriptStruct* ScriptStruct = FClusterUnionPendingAddData::StaticStruct();
+		for (TPair<TObjectKey<UPrimitiveComponent>, FClusterUnionPendingAddData>& Kvp : This->PendingComponentsToAdd)
+		{
+			Collector.AddReferencedObjects(ScriptStruct, reinterpret_cast<void*>(&Kvp.Value), This, nullptr);
+		}
+	}
+}
\ No newline at end of file
diff --git a/Engine/Source/Runtime/Engine/Private/PhysicsEngine/ClusterUnionReplicatedProxyComponent.cpp b/Engine/Source/Runtime/Engine/Private/PhysicsEngine/ClusterUnionReplicatedProxyComponent.cpp
new file mode 100644
--- /dev/null
+++ b/Engine/Source/Runtime/Engine/Private/PhysicsEngine/ClusterUnionReplicatedProxyComponent.cpp
@@ -0,0 +1,148 @@
+// Copyright Epic Games, Inc. All Rights Reserved.
+
+#include "PhysicsEngine/ClusterUnionReplicatedProxyComponent.h"
+
+#include "Net/Core/PushModel/PushModel.h"
+#include "Net/UnrealNetwork.h"
+#include "PhysicsEngine/ClusterUnionComponent.h"
+#include "TimerManager.h"
+
+#include UE_INLINE_GENERATED_CPP_BY_NAME(ClusterUnionReplicatedProxyComponent)
+
+UClusterUnionReplicatedProxyComponent::UClusterUnionReplicatedProxyComponent(const FObjectInitializer& ObjectInitializer)
+	: UActorComponent(ObjectInitializer)
+{
+	SetIsReplicatedByDefault(true);
+
+	ParentClusterUnion = nullptr;
+	ChildClusteredComponent = nullptr;
+
+	bNetUpdateParentClusterUnion = false;
+	bNetUpdateChildClusteredComponent = false;
+	bNetUpdateParticleBoneNames = false;
+	bNetUpdateParticleChildToParents = false;
+}
+
+void UClusterUnionReplicatedProxyComponent::GetLifetimeReplicatedProps(TArray<FLifetimeProperty>& OutLifetimeProps) const
+{
+	UActorComponent::GetLifetimeReplicatedProps(OutLifetimeProps);
+
+	FDoRepLifetimeParams Params;
+	Params.bIsPushBased = true;
+
+	DOREPLIFETIME_WITH_PARAMS_FAST(UClusterUnionReplicatedProxyComponent, ParentClusterUnion, Params);
+	DOREPLIFETIME_WITH_PARAMS_FAST(UClusterUnionReplicatedProxyComponent, ChildClusteredComponent, Params);
+	DOREPLIFETIME_WITH_PARAMS_FAST(UClusterUnionReplicatedProxyComponent, ParticleBoneNames, Params);
+	DOREPLIFETIME_WITH_PARAMS_FAST(UClusterUnionReplicatedProxyComponent, ParticleChildToParents, Params);
+}
+
+void UClusterUnionReplicatedProxyComponent::SetParentClusterUnion(UClusterUnionComponent* InComponent)
+{
+	ParentClusterUnion = InComponent;
+	MARK_PROPERTY_DIRTY_FROM_NAME(UClusterUnionReplicatedProxyComponent, ParentClusterUnion, this);
+}
+
+void UClusterUnionReplicatedProxyComponent::SetChildClusteredComponent(UPrimitiveComponent* InComponent)
+{
+	ChildClusteredComponent = InComponent;
+	MARK_PROPERTY_DIRTY_FROM_NAME(UClusterUnionReplicatedProxyComponent, ChildClusteredComponent, this);
+}
+
+void UClusterUnionReplicatedProxyComponent::SetParticleBoneNames(const TArray<FName>& InNames)
+{
+	ParticleBoneNames = InNames;
+
+	ParticleChildToParents.Empty();
+	ParticleChildToParents.Reserve(InNames.Num());
+	for (int32 Index = 0; Index < InNames.Num(); ++Index)
+	{
+		ParticleChildToParents.Add(FTransform::Identity);
+	}
+
+	MARK_PROPERTY_DIRTY_FROM_NAME(UClusterUnionReplicatedProxyComponent, ParticleBoneNames, this);
+}
+
+void UClusterUnionReplicatedProxyComponent::SetParticleChildToParent(const FName& BoneName, const FTransform& ChildToParent)
+{
+	int32 Index = INDEX_NONE;
+	if (ParticleBoneNames.Find(BoneName, Index))
+	{
+		ParticleChildToParents[Index] = ChildToParent;
+		MARK_PROPERTY_DIRTY_FROM_NAME(UClusterUnionReplicatedProxyComponent, ParticleChildToParents, this);
+	}
+}
+
+void UClusterUnionReplicatedProxyComponent::OnRep_ParentClusterUnion()
+{
+	bNetUpdateParentClusterUnion = true;
+}
+
+void UClusterUnionReplicatedProxyComponent::OnRep_ChildClusteredComponent()
+{
+	bNetUpdateChildClusteredComponent = true;
+}
+
+void UClusterUnionReplicatedProxyComponent::OnRep_ParticleBoneNames()
+{
+	bNetUpdateParticleBoneNames = true;
+}
+
+void UClusterUnionReplicatedProxyComponent::OnRep_ParticleChildToParents()
+{
+	bNetUpdateParticleChildToParents = true;
+}
+
+void UClusterUnionReplicatedProxyComponent::PostRepNotifies()
+{
+	UActorComponent::PostRepNotifies();
+
+	// These three properties should only get set once when the component is created.
+	const bool bIsInitialReplication = bNetUpdateParentClusterUnion || bNetUpdateChildClusteredComponent || bNetUpdateParticleBoneNames;
+	const bool bIsValid = ParentClusterUnion && ChildClusteredComponent && !ParticleBoneNames.IsEmpty();
+	if (bIsInitialReplication)
+	{
+		if (bIsValid)
+		{
+			ParentClusterUnion->AddComponentToCluster(ChildClusteredComponent, ParticleBoneNames);
+		}
+
+		bNetUpdateParentClusterUnion = false;
+		bNetUpdateChildClusteredComponent = false;
+		bNetUpdateParticleBoneNames = false;
+	}
+
+	if (bIsValid && bNetUpdateParticleChildToParents && ParticleBoneNames.Num() == ParticleChildToParents.Num())
+	{
+		// This particular bit can't happen utnil *after* we add the component to the cluster union. There's an additional deferral
+		// in AddComponentToCluster that we have to wait for.
+		DeferUntilChildClusteredComponentInParentUnion(
+			[this]()
+			{
+				ParentClusterUnion->ForceSetChildToParent(ChildClusteredComponent, ParticleBoneNames, ParticleChildToParents);
+			}
+		);
+		bNetUpdateParticleChildToParents = false;
+	}
+}
+
+void UClusterUnionReplicatedProxyComponent::DeferUntilChildClusteredComponentInParentUnion(TFunction<void()> Func)
+{
+	if (!ParentClusterUnion || !ChildClusteredComponent)
+	{
+		return;
+	}
+
+	if (ParentClusterUnion->IsComponentAdded(ChildClusteredComponent))
+	{
+		Func();
+	}
+	else
+	{
+		GetOwner()->GetWorldTimerManager().SetTimerForNextTick(
+			[this, Func]()
+			{
+				DeferUntilChildClusteredComponentInParentUnion(Func);
+			}
+		);
+	}
+}
\ No newline at end of file
diff --git a/Engine/Source/Runtime/Engine/Private/PhysicsEngine/Experimental/PhysScene_Chaos.cpp b/Engine/Source/Runtime/Engine/Private/PhysicsEngine/Experimental/PhysScene_Chaos.cpp
--- a/Engine/Source/Runtime/Engine/Private/PhysicsEngine/Experimental/PhysScene_Chaos.cpp
+++ b/Engine/Source/Runtime/Engine/Private/PhysicsEngine/Experimental/PhysScene_Chaos.cpp
@@ -11,13 +11,14 @@
 #include "Components/SkeletalMeshComponent.h"
 #include "Physics/PhysicsInterfaceUtils.h"
 #include "PhysicsReplication.h"
+#include "PhysicsEngine/ClusterUnionComponent.h"
 #include "PhysicsEngine/ConstraintInstance.h"
 #include "PhysicsEngine/PhysicsCollisionHandler.h"
 #include "PhysicsEngine/PhysicsObjectExternalInterface.h"
 
 #include "ChaosSolversModule.h"
 
-
+#include "PhysicsProxy/ClusterUnionPhysicsProxy.h"
 #include "PhysicsProxy/SkeletalMeshPhysicsProxy.h"
 #include "PhysicsProxy/StaticMeshPhysicsProxy.h"
 #include "EventsData.h"
@@ -551,6 +552,14 @@ void FPhysScene_Chaos::AddObject(UPrimitiveComponent* Component, FGeometryCollec
 	Solver->RegisterObject(InObject);
 }
 
+void FPhysScene_Chaos::AddObject(UPrimitiveComponent* Component, Chaos::FClusterUnionPhysicsProxy* InObject)
+{
+	AddToComponentMaps(Component, InObject);
+
+	Chaos::FPhysicsSolver* Solver = GetSolver();
+	Solver->RegisterObject(InObject);
+}
+
 template<typename ObjectType>
 void RemovePhysicsProxy(ObjectType* InObject, Chaos::FPhysicsSolver* InSolver, FChaosSolversModule* InModule)
 {
@@ -654,6 +663,18 @@ void FPhysScene_Chaos::RemoveObject(FGeometryCollectionPhysicsProxy* InObject)
 	RemoveFromComponentMaps(InObject);
 }
 
+void FPhysScene_Chaos::RemoveObject(Chaos::FClusterUnionPhysicsProxy* InObject)
+{
+	Chaos::FPhysicsSolver* Solver = InObject->GetSolver<Chaos::FPhysicsSolver>();
+
+	if (Solver)
+	{
+		Solver->UnregisterObject(InObject);
+	}
+
+	RemoveFromComponentMaps(InObject);
+}
+
 FPhysicsReplication* FPhysScene_Chaos::GetPhysicsReplication()
 {
 	return PhysicsReplication;
@@ -1798,55 +1819,58 @@ void FPhysScene_Chaos::KillVisualDebugger()
 void FPhysScene_Chaos::OnSyncBodies(Chaos::FPhysicsSolverBase* Solver)
 {
 	using namespace Chaos;
-	TArray<FPhysScenePendingComponentTransform_Chaos> PendingTransforms;
-	TSet<FGeometryCollectionPhysicsProxy*> GCProxies;
 
-	FPhysicsCommand::ExecuteWrite(this, [this, &Solver, &PendingTransforms](FPhysScene* PhysScene)
+	struct FDispatcher
 	{
-		auto RigidLambda = [&PhysScene, &PendingTransforms](Chaos::FSingleParticlePhysicsProxy* Proxy)
+		FPhysScene_Chaos* Outer;
+		FPhysScene* PhysScene;
+		Chaos::FPhysicsSolverBase* Solver;
+		TArray<FPhysScenePendingComponentTransform_Chaos> PendingTransforms;
+
+		void operator()(FSingleParticlePhysicsProxy* Proxy)
 		{
 			FPBDRigidParticle* DirtyParticle = Proxy->GetRigidParticleUnsafe();
 
-			if(FBodyInstance* BodyInstance = FPhysicsUserData::Get<FBodyInstance>(DirtyParticle->UserData()))
+			if (FBodyInstance* BodyInstance = FPhysicsUserData::Get<FBodyInstance>(DirtyParticle->UserData()))
 			{
-				if(BodyInstance->OwnerComponent.IsValid())
+				if (BodyInstance->OwnerComponent.IsValid())
 				{
 					if (Chaos::SyncKinematicOnGameThread == 0 && BodyInstance->IsInstanceSimulatingPhysics() == false)
 					{
 						return;
 					}
 					UPrimitiveComponent* OwnerComponent = BodyInstance->OwnerComponent.Get();
-					if(OwnerComponent != nullptr)
+					if (OwnerComponent != nullptr)
 					{
 						bool bPendingMove = false;
-							FRigidTransform3 NewTransform(DirtyParticle->X(),DirtyParticle->R());
+						FRigidTransform3 NewTransform(DirtyParticle->X(), DirtyParticle->R());
 
-							if(!NewTransform.EqualsNoScale(OwnerComponent->GetComponentTransform()))
+						if (!NewTransform.EqualsNoScale(OwnerComponent->GetComponentTransform()))
+						{
+							if (BodyInstance->InstanceBodyIndex == INDEX_NONE)
 							{
-								if (BodyInstance->InstanceBodyIndex == INDEX_NONE)
-								{
 
 								bPendingMove = true;
 								const FVector MoveBy = NewTransform.GetLocation() - OwnerComponent->GetComponentTransform().GetLocation();
 								const FQuat NewRotation = NewTransform.GetRotation();
-								PendingTransforms.Add(FPhysScenePendingComponentTransform_Chaos(OwnerComponent,MoveBy,NewRotation,Proxy->GetWakeEvent()));
+								PendingTransforms.Add(FPhysScenePendingComponentTransform_Chaos(OwnerComponent, MoveBy, NewRotation, Proxy->GetWakeEvent()));
 
 							}
 
-								PhysScene->UpdateActorInAccelerationStructure(BodyInstance->ActorHandle);
+							PhysScene->UpdateActorInAccelerationStructure(BodyInstance->ActorHandle);
 						}
 
-						if(Proxy->GetWakeEvent() != Chaos::EWakeEventEntry::None && !bPendingMove)
+						if (Proxy->GetWakeEvent() != Chaos::EWakeEventEntry::None && !bPendingMove)
 						{
-							PendingTransforms.Add(FPhysScenePendingComponentTransform_Chaos(OwnerComponent,Proxy->GetWakeEvent()));
+							PendingTransforms.Add(FPhysScenePendingComponentTransform_Chaos(OwnerComponent, Proxy->GetWakeEvent()));
 						}
 						Proxy->ClearEvents();
 					}
 				}
 			}
-		};
+		}
 
-		auto ConstraintLambda = [&PendingTransforms](FJointConstraintPhysicsProxy* Proxy)
+		void operator()(FJointConstraintPhysicsProxy* Proxy)
 		{
 			Chaos::FJointConstraint* Constraint = Proxy->GetConstraint();
 
@@ -1871,10 +1895,9 @@ void FPhysScene_Chaos::OnSyncBodies(Chaos::FPhysicsSolverBase* Solver)
 
 				Constraint->GetOutputData().bDriveTargetChanged = false;
 			}
+		}
 
-		};
-
-		auto GeometryCollectionLambda = [this](FGeometryCollectionPhysicsProxy* Proxy)
+		void operator()(FGeometryCollectionPhysicsProxy* Proxy)
 		{
 			// Don't pass in anything here so we don't end up locking anything because we can assume the scene is already locked.
 			FLockedWritePhysicsObjectExternalInterface Interface = FPhysicsObjectExternalInterface::LockWrite({});
@@ -1898,15 +1921,50 @@ void FPhysScene_Chaos::OnSyncBodies(Chaos::FPhysicsSolverBase* Solver)
 					return !ShouldGCParticleBeInAccelerationStructure(Handle);
 				}
 			);
-			Interface->AddToSpatialAcceleration(ActiveHandles, GetSpacialAcceleration());
-			Interface->RemoveFromSpatialAcceleration(DisabledHandles, GetSpacialAcceleration());
-		};
 
-		Solver->PullPhysicsStateForEachDirtyProxy_External(RigidLambda, ConstraintLambda, GeometryCollectionLambda);
+			Interface->AddToSpatialAcceleration(ActiveHandles, Outer->GetSpacialAcceleration());
+			Interface->RemoveFromSpatialAcceleration(DisabledHandles, Outer->GetSpacialAcceleration());
+		}
+
+		void operator()(FClusterUnionPhysicsProxy* Proxy)
+		{
+			FLockedWritePhysicsObjectExternalInterface Interface = FPhysicsObjectExternalInterface::LockWrite({});
+			Chaos::FPhysicsObjectHandle Handle = Proxy->GetPhysicsObjectHandle();
+
+			// Cluster unions should always have their owner be the cluster union component.
+			FClusterUnionPhysicsProxy::FExternalParticle* DirtyParticle = Proxy->GetParticle_External();
+			if (UClusterUnionComponent* ParentComponent = Cast<UClusterUnionComponent>(FChaosUserData::Get<UPrimitiveComponent>(DirtyParticle->UserData())))
+			{
+				const FRigidTransform3 NewTransform(DirtyParticle->X(), DirtyParticle->R());
+
+				if (!NewTransform.EqualsNoScale(ParentComponent->GetComponentTransform()))
+				{
+					const FVector MoveBy = NewTransform.GetLocation() - ParentComponent->GetComponentTransform().GetLocation();
+					PendingTransforms.Add(FPhysScenePendingComponentTransform_Chaos(ParentComponent, MoveBy, NewTransform.GetRotation(), DirtyParticle->GetWakeEvent()));
+
+					// Any changes in the linear/angular velocity of the particle is meaningless if the X/R doesn't change too.
+					ParentComponent->SyncVelocitiesFromPhysics(DirtyParticle->V(), DirtyParticle->W());
+
+					Interface->AddToSpatialAcceleration({ &Handle, 1 }, Outer->GetSpacialAcceleration());
+				}
+
+				ParentComponent->SyncIsAnchoredFromPhysics(Proxy->IsAnchored_External());
+				ParentComponent->SyncChildToParentFromProxy();
+				DirtyParticle->ClearEvents();
+			}
+		}
+	};
 
+	FDispatcher Dispatcher;
+	FPhysicsCommand::ExecuteWrite(this, [this, Solver, &Dispatcher](FPhysScene* PhysScene)
+	{
+		Dispatcher.Outer = this;
+		Dispatcher.PhysScene = PhysScene;
+		Dispatcher.Solver = Solver;
+		Solver->PullPhysicsStateForEachDirtyProxy_External(Dispatcher);
 	});
 
-	for (const FPhysScenePendingComponentTransform_Chaos& ComponentTransform : PendingTransforms)
+	for (const FPhysScenePendingComponentTransform_Chaos& ComponentTransform : Dispatcher.PendingTransforms)
 	{
 		if (ComponentTransform.OwningComp != nullptr)
 		{
diff --git a/Engine/Source/Runtime/Engine/Public/Physics/Experimental/PhysScene_Chaos.h b/Engine/Source/Runtime/Engine/Public/Physics/Experimental/PhysScene_Chaos.h
--- a/Engine/Source/Runtime/Engine/Public/Physics/Experimental/PhysScene_Chaos.h
+++ b/Engine/Source/Runtime/Engine/Public/Physics/Experimental/PhysScene_Chaos.h
@@ -45,6 +45,7 @@ struct FConstraintInstanceBase;
 namespace Chaos
 {
 	class FPhysicsProxy;
+	class FClusterUnionPhysicsProxy;
 
 	struct FCollisionEventData;
 
@@ -129,6 +130,7 @@ public:
 	void AddObject(UPrimitiveComponent* Component, FStaticMeshPhysicsProxy* InObject);
 	void AddObject(UPrimitiveComponent* Component, Chaos::FSingleParticlePhysicsProxy* InObject);
 	void AddObject(UPrimitiveComponent* Component, FGeometryCollectionPhysicsProxy* InObject);
+	void AddObject(UPrimitiveComponent* Component, Chaos::FClusterUnionPhysicsProxy* InObject);
 	
 	void AddToComponentMaps(UPrimitiveComponent* Component, IPhysicsProxyBase* InObject);
 	void RemoveFromComponentMaps(IPhysicsProxyBase* InObject);
@@ -141,6 +143,7 @@ public:
 	void RemoveObject(FStaticMeshPhysicsProxy* InObject);
 	void RemoveObject(Chaos::FSingleParticlePhysicsProxy* InObject);
 	void RemoveObject(FGeometryCollectionPhysicsProxy* InObject);
+	void RemoveObject(Chaos::FClusterUnionPhysicsProxy* InObject);
 
 	FPhysicsReplication* GetPhysicsReplication();
 	void SetPhysicsReplication(FPhysicsReplication* InPhysicsReplication);
diff --git a/Engine/Source/Runtime/Experimental/Chaos/Private/Chaos/ChaosMarshallingManager.cpp b/Engine/Source/Runtime/Experimental/Chaos/Private/Chaos/ChaosMarshallingManager.cpp
--- a/Engine/Source/Runtime/Experimental/Chaos/Private/Chaos/ChaosMarshallingManager.cpp
+++ b/Engine/Source/Runtime/Experimental/Chaos/Private/Chaos/ChaosMarshallingManager.cpp
@@ -251,6 +251,11 @@ void FPushPhysicsData::CopySubstepData(const FPushPhysicsData& FirstStepData)
 
 			Dirty.Proxy->ResetDirtyIdx();	//dirty idx is only used temporarily
 		}
+		else if (Dirty.Proxy->GetType() == EPhysicsProxyType::ClusterUnionProxy)
+		{
+			DirtyProxiesDataBuffer.Add(Dirty.Proxy);
+			Dirty.Proxy->ResetDirtyIdx();
+		}
 	});
 
 	//make sure inputs are available to every sub-step
diff --git a/Engine/Source/Runtime/Experimental/Chaos/Private/Chaos/Framework/ChaosResultsManager.cpp b/Engine/Source/Runtime/Experimental/Chaos/Private/Chaos/Framework/ChaosResultsManager.cpp
--- a/Engine/Source/Runtime/Experimental/Chaos/Private/Chaos/Framework/ChaosResultsManager.cpp
+++ b/Engine/Source/Runtime/Experimental/Chaos/Private/Chaos/Framework/ChaosResultsManager.cpp
@@ -1,6 +1,7 @@
 // Copyright Epic Games, Inc. All Rights Reserved.
 
 #include "Chaos/Framework/ChaosResultsManager.h"
+#include "PhysicsProxy/ClusterUnionPhysicsProxy.h"
 #include "PhysicsProxy/SingleParticlePhysicsProxy.h"
 #include "PhysicsProxy/GeometryCollectionPhysicsProxy.h"
 #include "Chaos/ChaosMarshallingManager.h"
@@ -68,6 +69,7 @@ namespace Chaos
 	{
 		ResetInterpolations<FSingleParticlePhysicsProxy>(RigidInterpolations);
 		ResetInterpolations<FGeometryCollectionPhysicsProxy>(GeometryCollectionInterpolations);
+		ResetInterpolations<FClusterUnionPhysicsProxy>(ClusterUnionInterpolations);
 
 		//purposely leave Prev and Next alone as we use those for rebuild
 	}
@@ -124,8 +126,8 @@ namespace Chaos
 		const int32 Timestamp = PullData.SolverTimestamp;
 
 		SetPrevNextDataHelperTyped<Mode>(PullData.DirtyRigids, Results.RigidInterpolations);
-
 		SetPrevNextDataHelperTyped<Mode>(PullData.DirtyGeometryCollections, Results.GeometryCollectionInterpolations);
+		SetPrevNextDataHelperTyped<Mode>(PullData.DirtyClusterUnions, Results.ClusterUnionInterpolations);
 
 		// update resim target for rigids
 		for (const FDirtyRigidParticleData& Data : PullData.DirtyRigids)
diff --git a/Engine/Source/Runtime/Experimental/Chaos/Private/Chaos/PhysicsObjectInternal.h b/Engine/Source/Runtime/Experimental/Chaos/Private/Chaos/PhysicsObjectInternal.h
--- a/Engine/Source/Runtime/Experimental/Chaos/Private/Chaos/PhysicsObjectInternal.h
+++ b/Engine/Source/Runtime/Experimental/Chaos/Private/Chaos/PhysicsObjectInternal.h
@@ -5,6 +5,7 @@
 #include "Containers/Set.h"
 #include "Framework/Threading.h"
 
+#include "PhysicsProxy/ClusterUnionPhysicsProxy.h"
 #include "PhysicsProxy/SingleParticlePhysicsProxy.h"
 #include "PhysicsProxy/GeometryCollectionPhysicsProxy.h"
 
@@ -44,10 +45,6 @@ namespace Chaos
 		FPhysicsObjectHandle GetRootObject()
 		{
 			EPhysicsProxyType ProxyType = Proxy->GetType();
-			if (ProxyType == EPhysicsProxyType::SingleParticleProxy)
-			{
-				return this;
-			}
 
 			TThreadParticle<Id>* Particle = GetParticle<Id>();
 			FPhysicsObjectHandle CurrentObject = this;
@@ -71,10 +68,31 @@ namespace Chaos
 			{
 				FGeometryCollectionPhysicsProxy* GeometryCollectionProxy = static_cast<FGeometryCollectionPhysicsProxy*>(Proxy);
 				FGeometryDynamicCollection& Collection = GetGeometryCollectionDynamicCollection<Id>(*GeometryCollectionProxy);
+
 				if (int32 Index = Collection.Parent[BodyIndex]; Index != INDEX_NONE)
 				{
 					return GeometryCollectionProxy->GetPhysicsObjectByIndex(Index);
 				}
+
+				// Yes this is an explicit fall through. Once we get to the root particle on the geometry collection we need to check
+				// if it's in a cluster union same with all the other proxies.
+				[[fallthrough]];
+			}
+			case EPhysicsProxyType::ClusterUnionProxy:
+			case EPhysicsProxyType::SingleParticleProxy:
+			{
+				// At this point, we need to be able to check if we're part of a cluster union.
+				// However, note that we're after more than just the parent particle, we're after
+				// the parent physics object handle which lives on the proxy. Hence why physics proxies
+				// need to keep track of their parent proxy. We do, however, make the assumption that the
+				// only possible valid parent proxy is a cluster union proxy though; hence being able to directly
+				// retrieve the physics object handle.
+				if (IPhysicsProxyBase* ParentProxy = Proxy->GetParentProxy(); ParentProxy && ParentProxy->GetType() == EPhysicsProxyType::ClusterUnionProxy)
+				{
+					FClusterUnionPhysicsProxy* ParentClusterProxy = static_cast<FClusterUnionPhysicsProxy*>(ParentProxy);
+					return ParentClusterProxy->GetPhysicsObjectHandle();
+				}
+				break;
 			}
 			default:
 				break;
@@ -103,6 +121,8 @@ namespace Chaos
 				{
 					switch (ProxyType)
 					{
+					case EPhysicsProxyType::ClusterUnionProxy:
+						return static_cast<FClusterUnionPhysicsProxy*>(Proxy)->GetParticle_External();
 					case EPhysicsProxyType::GeometryCollectionType:
 						return static_cast<FGeometryCollectionPhysicsProxy*>(Proxy)->GetParticleByIndex_External(BodyIndex);
 					case EPhysicsProxyType::SingleParticleProxy:
@@ -115,6 +135,8 @@ namespace Chaos
 				{
 					switch (ProxyType)
 					{
+					case EPhysicsProxyType::ClusterUnionProxy:
+						return static_cast<FClusterUnionPhysicsProxy*>(Proxy)->GetParticle_Internal();
 					case EPhysicsProxyType::GeometryCollectionType:
 						return static_cast<FGeometryCollectionPhysicsProxy*>(Proxy)->GetParticleByIndex_Internal(BodyIndex);
 					case EPhysicsProxyType::SingleParticleProxy:
@@ -152,6 +174,18 @@ namespace Chaos
 				FGeometryDynamicCollection& Collection = GetGeometryCollectionDynamicCollection<Id>(*static_cast<FGeometryCollectionPhysicsProxy*>(Proxy));
 				return !Collection.Children[BodyIndex].IsEmpty();
 			}
+			case EPhysicsProxyType::ClusterUnionProxy:
+			{
+				FClusterUnionPhysicsProxy* ClusterProxy = static_cast<FClusterUnionPhysicsProxy*>(Proxy);
+				if constexpr (Id == EThreadContext::External)
+				{
+					return ClusterProxy->HasChildren_External();
+				}
+				else
+				{
+					return ClusterProxy->HasChildren_Internal();
+				}
+			}
 			default:
 				break;
 			}
diff --git a/Engine/Source/Runtime/Experimental/Chaos/Private/Chaos/PullPhysicsDataImp.cpp b/Engine/Source/Runtime/Experimental/Chaos/Private/Chaos/PullPhysicsDataImp.cpp
--- a/Engine/Source/Runtime/Experimental/Chaos/Private/Chaos/PullPhysicsDataImp.cpp
+++ b/Engine/Source/Runtime/Experimental/Chaos/Private/Chaos/PullPhysicsDataImp.cpp
@@ -9,6 +9,7 @@ void FPullPhysicsData::Reset()
 {
 	DirtyRigids.Reset();
 	DirtyGeometryCollections.Reset();
+	DirtyClusterUnions.Reset();
 	DirtyJointConstraints.Reset();
 }
 
diff --git a/Engine/Source/Runtime/Experimental/Chaos/Private/PBDRigidsSolver.cpp b/Engine/Source/Runtime/Experimental/Chaos/Private/PBDRigidsSolver.cpp
--- a/Engine/Source/Runtime/Experimental/Chaos/Private/PBDRigidsSolver.cpp
+++ b/Engine/Source/Runtime/Experimental/Chaos/Private/PBDRigidsSolver.cpp
@@ -15,6 +15,7 @@
 #include "HAL/FileManager.h"
 #include "Misc/ScopeLock.h"
 #include "PhysicsProxy/CharacterGroundConstraintProxy.h"
+#include "PhysicsProxy/ClusterUnionPhysicsProxy.h"
 #include "PhysicsProxy/SingleParticlePhysicsProxy.h"
 #include "PhysicsProxy/SkeletalMeshPhysicsProxy.h"
 #include "PhysicsProxy/StaticMeshPhysicsProxy.h"
@@ -669,12 +670,47 @@ namespace Chaos
 			{
 				GeometryCollectionPhysicsProxies_Internal.RemoveSingle(InProxy);
 				InProxy->SyncBeforeDestroy();
-				InProxy->OnRemoveFromSolver(this);
 				InProxy->ResetDirtyIdx();
 				PendingDestroyGeometryCollectionPhysicsProxy.Add(InProxy);
 			});
 	}
 
+	void FPBDRigidsSolver::RegisterObject(FClusterUnionPhysicsProxy* Proxy)
+	{
+		if (!Proxy)
+		{
+			return;
+		}
+
+		if (FPBDRigidParticle* Particle = Proxy->GetParticle_External())
+		{
+			Particle->SetUniqueIdx(GetEvolution()->GenerateUniqueIdx());
+		}
+		Proxy->SetSolver(this);
+		AddDirtyProxy(Proxy);
+	}
+
+	void FPBDRigidsSolver::UnregisterObject(FClusterUnionPhysicsProxy* Proxy)
+	{
+		if (!Proxy)
+		{
+			return;
+		}
+
+		Proxy->MarkDeleted();
+		RemoveDirtyProxy(Proxy);
+
+		EnqueueCommandImmediate(
+			[Proxy, this]()
+			{
+				ClusterUnionPhysicsProxies_Internal.RemoveSingle(Proxy);
+				Proxy->SyncBeforeDestroy();
+				Proxy->ResetDirtyIdx();
+				PendingDestroyClusterUnionProxy.Add(Proxy);
+			}
+		);
+	}
+
 	void FPBDRigidsSolver::RegisterObject(Chaos::FJointConstraint* GTConstraint)
 	{
 		LLM_SCOPE(ELLMTag::ChaosConstraint);
@@ -860,6 +896,23 @@ namespace Chaos
 
 	void FPBDRigidsSolver::DestroyPendingProxies_Internal()
 	{
+		MarshallingManager.GetCurrentPullData_Internal()->DirtyClusterUnions.Reset();
+		for (FClusterUnionPhysicsProxy* Proxy : PendingDestroyClusterUnionProxy)
+		{
+			if (!Proxy)
+			{
+				continue;
+			}
+
+			if (FPBDRigidsEvolutionGBF* Evolution = GetEvolution())
+			{
+				// Remove the cluster union this proxy manages which will in turn also the destroy the necessary particles.
+				Evolution->GetRigidClustering().GetClusterUnionManager().DestroyClusterUnion(Proxy->GetClusterUnionIndex());
+			}
+			delete Proxy;
+		}
+		PendingDestroyClusterUnionProxy.Reset();
+
 		// If we have any callback objects watching for particle deregistrations,
 		// send an array of proxies that are about to be deleted.
 		if (UnregistrationWatchers.Num() > 0)
@@ -900,7 +953,10 @@ namespace Chaos
 		bool bResetCollisionConstraints=false;
 		for (auto Proxy : PendingDestroyGeometryCollectionPhysicsProxy)
 		{
-			//ensure(Proxy->GetHandle_LowLevel() == nullptr);	//should have already cleared this out
+			// Removing the geometry collection from the solver a bit delayed. This lets the cluster union do its cleanup first before
+			// the geometry collection if they're all being destroyed at the same time.
+			Proxy->OnRemoveFromSolver(this);
+
 			MarshallingManager.GetCurrentPullData_Internal()->DirtyGeometryCollections.Reset();
 			bResetCollisionConstraints = true;
 			delete Proxy;
@@ -1060,6 +1116,15 @@ namespace Chaos
 				Proxy->PushStateOnGameThread(this);
 				break;
 			}
+			case EPhysicsProxyType::ClusterUnionProxy:
+			{
+				FClusterUnionPhysicsProxy* Proxy = static_cast<FClusterUnionPhysicsProxy*>(Dirty.Proxy);
+				FClusterUnionPhysicsProxy::FExternalParticle* Particle = Proxy->GetParticle_External();
+				Proxy->SyncRemoteData(*Manager, DataIdx, Dirty.PropertyData);
+				Proxy->ClearAccumulatedData();
+				Proxy->ResetDirtyIdx();
+				break;
+			}
 			case EPhysicsProxyType::JointConstraintType:
 			{
 				auto Proxy = static_cast<FJointConstraintPhysicsProxy*>(Dirty.Proxy);
@@ -1209,7 +1274,7 @@ namespace Chaos
 		};
 
 		//need to create new particle handles
-		DirtyProxiesData->ForEachProxy([this,&ProcessProxyPT](int32 DataIdx,FDirtyProxy& Dirty)
+		DirtyProxiesData->ForEachProxy([this, &ProcessProxyPT, Manager](int32 DataIdx,FDirtyProxy& Dirty)
 		{
 			if(Dirty.Proxy->GetIgnoreDataOnStep_Internal() != CurrentFrame)
 			{
@@ -1246,6 +1311,18 @@ namespace Chaos
 						Dirty.Proxy->ResetDirtyIdx();
 						break;
 					}
+					case EPhysicsProxyType::ClusterUnionProxy:
+					{
+						FClusterUnionPhysicsProxy* Proxy = static_cast<FClusterUnionPhysicsProxy*>(Dirty.Proxy);
+						if (!Proxy->IsInitializedOnPhysicsThread())
+						{
+							Proxy->Initialize_Internal(this, GetParticles());
+							ClusterUnionPhysicsProxies_Internal.Add(Proxy);
+						}
+						Proxy->PushToPhysicsState(*Manager, DataIdx, Dirty);
+						Dirty.Proxy->ResetDirtyIdx();
+						break;
+					}
 					case EPhysicsProxyType::JointConstraintType:
 					case EPhysicsProxyType::SuspensionConstraintType:
 					case EPhysicsProxyType::CharacterGroundConstraintType:
@@ -1497,19 +1574,26 @@ namespace Chaos
 	void FPBDRigidsSolver::BufferPhysicsResults()
 	{
 		//ensure(IsInPhysicsThread());
+
 		TArray<FGeometryCollectionPhysicsProxy*> ActiveGC;
 		ActiveGC.Reserve(GeometryCollectionPhysicsProxies_Internal.Num());
 
 		FPullPhysicsData* PullData = MarshallingManager.GetCurrentPullData_Internal();
 
 		TParticleView<FPBDRigidParticles>& DirtyParticles = GetParticles().GetDirtyParticlesView();
+
+		TArray<FSingleParticlePhysicsProxy*> ActiveRigid;
+		ActiveRigid.Reserve(SingleParticlePhysicsProxies_PT.Num());
+
+		TArray<FClusterUnionPhysicsProxy*> ActiveClusterUnions;
+		ActiveClusterUnions.Reserve(ClusterUnionPhysicsProxies_Internal.Num());
+
 		const bool bIsResim = GetEvolution()->IsResimming();
 
 		//todo: should be able to go wide just add defaulted etc...
 		{
 			ensure(PullData->DirtyRigids.Num() == 0);	//we only fill this once per frame
 			int32 BufferIdx = 0;
-			PullData->DirtyRigids.Reserve(DirtyParticles.Num());
 
 			for (Chaos::TPBDRigidParticleHandleImp<FReal, 3, false>& DirtyParticle : DirtyParticles)
 			{
@@ -1521,8 +1605,7 @@ namespace Chaos
 						{
 							if(!bIsResim || DirtyParticle.SyncState() == ESyncState::HardDesync)
 							{
-								PullData->DirtyRigids.AddDefaulted();
-								((FSingleParticlePhysicsProxy*)(Proxy))->BufferPhysicsResults(PullData->DirtyRigids.Last());
+								ActiveRigid.AddUnique((FSingleParticlePhysicsProxy*)Proxy);
 							}
 							break;
 						}
@@ -1538,10 +1621,31 @@ namespace Chaos
 							{
 								if (ClusterParticle->InternalCluster())
 								{
+									if (Proxy->GetType() == EPhysicsProxyType::ClusterUnionProxy)
+									{
+										ActiveClusterUnions.AddUnique((FClusterUnionPhysicsProxy*)(Proxy));
+									}
+
 									const TSet<IPhysicsProxyBase*> Proxies = ClusterParticle->PhysicsProxies();
 									for (IPhysicsProxyBase* ClusterProxy : Proxies)
 									{
-										ActiveGC.AddUnique((FGeometryCollectionPhysicsProxy*)(ClusterProxy));
+										if (!ClusterProxy)
+										{
+											continue;
+										}
+
+										switch (ClusterProxy->GetType())
+										{
+										case EPhysicsProxyType::SingleParticleProxy:
+											ActiveRigid.AddUnique((FSingleParticlePhysicsProxy*)ClusterProxy);
+											break;
+										case EPhysicsProxyType::GeometryCollectionType:
+											ActiveGC.AddUnique((FGeometryCollectionPhysicsProxy*)(ClusterProxy));
+											break;
+										default:
+											ensure(false);
+											break;
+										}
 									}
 								}
 								else
@@ -1557,6 +1661,18 @@ namespace Chaos
 			}
 		}
 
+		{
+			//we only fill this once per frame
+			ensure(PullData->DirtyRigids.Num() == 0);
+			PullData->DirtyRigids.Reserve(ActiveRigid.Num());
+
+			for (int32 Idx = 0; Idx < ActiveRigid.Num(); ++Idx)
+			{
+				PullData->DirtyRigids.AddDefaulted();
+				ActiveRigid[Idx]->BufferPhysicsResults(PullData->DirtyRigids.Last());
+			}
+		}
+
 		{
 			ensure(PullData->DirtyGeometryCollections.Num() == 0);	//we only fill this once per frame
 			PullData->DirtyGeometryCollections.Reserve(ActiveGC.Num());
@@ -1568,6 +1684,17 @@ namespace Chaos
 			}
 		}
 
+		{
+			ensure(PullData->DirtyClusterUnions.IsEmpty());
+			PullData->DirtyClusterUnions.Reserve(ActiveClusterUnions.Num());
+
+			for (int32 Idx = 0; Idx < ActiveClusterUnions.Num(); ++Idx)
+			{
+				PullData->DirtyClusterUnions.AddDefaulted();
+				ActiveClusterUnions[Idx]->BufferPhysicsResults_Internal(PullData->DirtyClusterUnions.Last());
+			}
+		}
+
 		{
 			ensure(PullData->DirtyJointConstraints.Num() == 0);	//we only fill this once per frame
 			PullData->DirtyJointConstraints.Reserve(JointConstraintPhysicsProxies_Internal.Num());
@@ -1610,7 +1737,8 @@ namespace Chaos
 	// however.
 	void FPBDRigidsSolver::UpdateGameThreadStructures()
 	{
-		PullPhysicsStateForEachDirtyProxy_External([](auto){}, [](auto) {}, [](auto) {});
+		struct FDispatcher {} Dispatcher;
+		PullPhysicsStateForEachDirtyProxy_External(Dispatcher);
 	}
 
 	int32 FPBDRigidsSolver::NumJointConstraints() const
diff --git a/Engine/Source/Runtime/Experimental/Chaos/Private/PhysicsProxy/ClusterUnionPhysicsProxy.cpp b/Engine/Source/Runtime/Experimental/Chaos/Private/PhysicsProxy/ClusterUnionPhysicsProxy.cpp
new file mode 100644
--- /dev/null
+++ b/Engine/Source/Runtime/Experimental/Chaos/Private/PhysicsProxy/ClusterUnionPhysicsProxy.cpp
@@ -0,0 +1,438 @@
+// Copyright Epic Games, Inc. All Rights Reserved.
+
+#include "PhysicsProxy/ClusterUnionPhysicsProxy.h"
+
+#include "Chaos/ClusterUnionManager.h"
+#include "Chaos/ParticleHandle.h"
+#include "Chaos/PhysicsObjectInternal.h"
+#include "Chaos/PhysicsObjectInternalInterface.h"
+#include "Chaos/PullPhysicsDataImp.h"
+#include "Math/UnrealMathUtility.h"
+#include "PBDRigidsSolver.h"
+
+namespace Chaos
+{
+	namespace
+	{
+		FPBDRigidsEvolutionGBF* GetEvolution(FClusterUnionPhysicsProxy* Proxy)
+		{
+			if (!Proxy)
+			{
+				return nullptr;
+			}
+
+			FPBDRigidsSolver* RigidsSolver = Proxy->GetSolver<FPBDRigidsSolver>();
+			if (!RigidsSolver)
+			{
+				return nullptr;
+			}
+
+			return RigidsSolver->GetEvolution();
+		}
+
+		template<typename TParticle>
+		void BufferPhysicsResultsImp(FClusterUnionPhysicsProxy* Proxy, TParticle* Particle, FDirtyClusterUnionData& BufferData)
+		{
+			if (!Particle || !Proxy)
+			{
+				return;
+			}
+
+			BufferData.SetProxy(*Proxy);
+			BufferData.X = Particle->X();
+			BufferData.R = Particle->R();
+			BufferData.V = Particle->V();
+			BufferData.W = Particle->W();
+
+			if constexpr (std::is_base_of_v<FClusterUnionPhysicsProxy::FInternalParticle, TParticle>)
+			{
+				BufferData.bIsAnchored = Particle->IsAnchored();
+			}
+		}
+	}
+
+	FClusterUnionPhysicsProxy::FClusterUnionPhysicsProxy(UObject* InOwner, const FClusterCreationParameters& InParameters, void* InUserData, EThreadContext InAuthoritativeContext)
+		: Base(InOwner)
+		, ClusterParameters(InParameters)
+		, UserData(InUserData)
+		, AuthoritativeContext(InAuthoritativeContext)
+	{
+	}
+
+	void FClusterUnionPhysicsProxy::Initialize_External()
+	{
+		// Create game thread particle as well as the physics object.
+		Particle_External = FExternalParticle::CreateParticle();
+		check(Particle_External != nullptr);
+
+		Particle_External->SetProxy(this);
+		Particle_External->SetUserData(UserData);
+
+		// NO DIRTY FLAGS ALLOWED. We must strictly manage the dirty flags on the particle.
+		// Setting the particle's XR on the particle will set the XR dirty flag but that isn't
+		// used for the cluster union (there is no functionality in Chaos to let the particle
+		// be easily managed by a proxy that isn't the single particle physics proxy). And if the
+		// XR dirty flag is set, we'll try to access buffers that don't exist for cluster union proxies.
+		Particle_External->ClearDirtyFlags();
+		PhysicsObject = FPhysicsObjectFactory::CreatePhysicsObject(this);
+	}
+
+	void FClusterUnionPhysicsProxy::Initialize_Internal(FPBDRigidsSolver* RigidsSolver, FPBDRigidsSolver::FParticlesType& Particles)
+	{
+		if (!ensure(RigidsSolver) || !ensure(Particle_External != nullptr))
+		{
+			return;
+		}
+
+		bIsInitializedOnPhysicsThread = true;
+
+		FPBDRigidsEvolutionGBF* Evolution = RigidsSolver->GetEvolution();
+		if (!ensure(Evolution))
+		{
+			return;
+		}
+
+		FUniqueIdx UniqueIndex = Particle_External->UniqueIdx();
+		FClusterUnionManager& ClusterUnionManager = Evolution->GetRigidClustering().GetClusterUnionManager();
+		ClusterUnionIndex = ClusterUnionManager.CreateNewClusterUnion(ClusterParameters, INDEX_NONE, &UniqueIndex);
+		if (FClusterUnion* ClusterUnion = ClusterUnionManager.FindClusterUnion(ClusterUnionIndex); ensure(ClusterUnion != nullptr))
+		{
+			Particle_Internal = ClusterUnion->InternalCluster;
+			Particle_Internal->SetPhysicsProxy(this);
+		}
+	}
+
+	bool FClusterUnionPhysicsProxy::HasChildren_Internal() const
+	{
+		if (!Particle_Internal)
+		{
+			return false;
+		}
+
+		return Particle_Internal->ClusterIds().NumChildren > 0;
+	}
+
+	void FClusterUnionPhysicsProxy::AddPhysicsObjects_External(const TArray<FPhysicsObjectHandle>& Objects)
+	{
+		if (!Solver)
+		{
+			return;
+		}
+
+		int32 OldNum = ChildPhysicsObjects.Num();
+		ChildPhysicsObjects.Append(Objects);
+		for (FPhysicsObjectHandle Handle : Objects)
+		{
+			ChildPhysicsObjectIndexMap.Add(Handle, OldNum++);
+			SyncedData_External.ChildToParent.Add(FTransform::Identity);
+		}
+		
+		Solver->EnqueueCommandImmediate(
+			[this, Objects=Objects]() mutable
+			{
+				FReadPhysicsObjectInterface_Internal Interface = FPhysicsObjectInternalInterface::GetRead();
+				if (FPBDRigidsEvolutionGBF* Evolution = GetEvolution(this))
+				{
+					Evolution->GetRigidClustering().GetClusterUnionManager().AddPendingClusterIndexOperation(ClusterUnionIndex, EClusterUnionOperation::Add, Interface.GetAllRigidParticles(Objects));
+				}
+			}
+		);
+	}
+
+	void FClusterUnionPhysicsProxy::SetIsAnchored_External(bool bIsAnchored)
+	{
+		if (!Solver || !ensure(Particle_External))
+		{
+			return;
+		}
+
+		SyncedData_External.bIsAnchored = bIsAnchored;
+		Particle_External->ForceDirty(EChaosPropertyFlags::ClusterIsAnchored);
+		Solver->AddDirtyProxy(this);
+
+		FClusterUnionProxyTimestamp& SyncTS = GetSyncTimestampAs<FClusterUnionProxyTimestamp>();
+		SyncTS.OverWriteAnchored.Set(GetSolverSyncTimestamp_External(), bIsAnchored);
+	}
+
+	void FClusterUnionPhysicsProxy::PushToPhysicsState(const FDirtyPropertiesManager& Manager, int32 DataIdx, const FDirtyProxy& Dirty)
+	{
+		if (!ensure(Solver) || !ensure(Particle_Internal))
+		{
+			return;
+		}
+
+		const FDirtyChaosProperties& ParticleData = Dirty.PropertyData;
+		FPBDRigidsSolver& RigidsSolver = *static_cast<FPBDRigidsSolver*>(Solver);
+		if (!ensure(RigidsSolver.GetEvolution()))
+		{
+			return;
+		}
+
+		FPBDRigidsEvolutionGBF& Evolution = *RigidsSolver.GetEvolution();
+		FClusterUnionManager& ClusterUnionManager = Evolution.GetRigidClustering().GetClusterUnionManager();
+
+		if (const FParticlePositionRotation* NewXR = ParticleData.FindClusterXR(Manager, DataIdx))
+		{
+			Evolution.SetParticleTransform(Particle_Internal, NewXR->X(), NewXR->R(), true);
+		}
+
+		if (const FParticleVelocities* NewVelocities = ParticleData.FindClusterVelocities(Manager, DataIdx))
+		{
+			Particle_Internal->SetVelocities(*NewVelocities);
+		}
+
+		if (const bool* bNewIsAnchored = ParticleData.FindClusterIsAnchored(Manager, DataIdx); bNewIsAnchored && *bNewIsAnchored != Particle_Internal->IsAnchored())
+		{
+			Particle_Internal->SetIsAnchored(*bNewIsAnchored);
+			if (!*bNewIsAnchored && !Particle_Internal->IsDynamic())
+			{
+				FKinematicTarget NoKinematicTarget;
+				Evolution.SetParticleObjectState(Particle_Internal, EObjectStateType::Dynamic);
+				Evolution.SetParticleKinematicTarget(Particle_Internal, NoKinematicTarget);
+			}
+		}
+
+		if (AuthoritativeContext == EThreadContext::External)
+		{
+			if (const TArray<FTransform>* NewChildToParent = ParticleData.FindClusterChildToParent(Manager, DataIdx))
+			{
+				if (FClusterUnion* ClusterUnion = ClusterUnionManager.FindClusterUnion(ClusterUnionIndex))
+				{
+					if (NewChildToParent->Num() == ClusterUnion->ChildParticles.Num())
+					{
+						for (int32 Index = 0; Index < NewChildToParent->Num(); ++Index)
+						{
+							if (FPBDRigidClusteredParticleHandle* ChildHandle = ClusterUnion->ChildParticles[Index]->CastToClustered())
+							{
+								if (ChildHandle->ChildToParent().Equals((*NewChildToParent)[Index]))
+								{
+									continue;
+								}
+								ChildHandle->SetChildToParent((*NewChildToParent)[Index]);
+								RigidsSolver.GetParticles().MarkTransientDirtyParticle(ChildHandle);
+							}
+						}
+					}
+				}
+			}
+		}
+	}
+
+	bool FClusterUnionPhysicsProxy::PullFromPhysicsState(const FDirtyClusterUnionData& PullData, int32 SolverSyncTimestamp, const FDirtyClusterUnionData* NextPullData, const FRealSingle* Alpha)
+	{
+		if (!ensure(Particle_External))
+		{
+			return false;
+		}
+
+		const FClusterUnionProxyTimestamp* ProxyTimestamp = PullData.GetTimestamp();
+		if (!ProxyTimestamp)
+		{
+			return false;
+		}
+
+		if (SolverSyncTimestamp >= ProxyTimestamp->OverWriteAnchored.Timestamp)
+		{
+			SyncedData_External.bIsAnchored = PullData.bIsAnchored;
+		}
+
+		if (NextPullData)
+		{
+			// This is the same as in the SingleParticlePhysicsProxy.
+			auto LerpHelper = [SolverSyncTimestamp](const auto& Prev, const auto& OverwriteProperty) -> const auto*
+			{
+				//if overwrite is in the future, do nothing
+				//if overwrite is on this step, we want to interpolate from overwrite to the result of the frame that consumed the overwrite
+				//if overwrite is in the past, just do normal interpolation
+
+				//this is nested because otherwise compiler can't figure out the type of nullptr with an auto return type
+				return OverwriteProperty.Timestamp <= SolverSyncTimestamp ? (OverwriteProperty.Timestamp < SolverSyncTimestamp ? &Prev : &OverwriteProperty.Value) : nullptr;
+			};
+
+			if (const FVec3* Prev = LerpHelper(PullData.X, ProxyTimestamp->OverWriteX))
+			{
+				const FVec3 NewX = FMath::Lerp(*Prev, NextPullData->X, *Alpha);
+				Particle_External->SetX(NewX, false);
+			}
+
+			if (const FQuat* Prev = LerpHelper(PullData.R, ProxyTimestamp->OverWriteR))
+			{
+				const FQuat NewR = FMath::Lerp(*Prev, NextPullData->R, *Alpha);
+				Particle_External->SetR(NewR, false);
+			}
+
+			if (const FVec3* Prev = LerpHelper(PullData.V, ProxyTimestamp->OverWriteV))
+			{
+				const FVec3 NewV = FMath::Lerp(*Prev, NextPullData->V, *Alpha);
+				Particle_External->SetV(NewV, false);
+			}
+
+			if (const FVec3* Prev = LerpHelper(PullData.W, ProxyTimestamp->OverWriteW))
+			{
+				const FVec3 NewW = FMath::Lerp(*Prev, NextPullData->W, *Alpha);
+				Particle_External->SetW(NewW, false);
+			}
+		}
+		else
+		{
+			if (SolverSyncTimestamp >= ProxyTimestamp->OverWriteX.Timestamp)
+			{
+				Particle_External->SetX(PullData.X, false);
+			}
+
+			if (SolverSyncTimestamp >= ProxyTimestamp->OverWriteR.Timestamp)
+			{
+				Particle_External->SetR(PullData.R, false);
+			}
+
+			if (SolverSyncTimestamp >= ProxyTimestamp->OverWriteV.Timestamp)
+			{
+				Particle_External->SetV(PullData.V, false);
+			}
+
+			if (SolverSyncTimestamp >= ProxyTimestamp->OverWriteW.Timestamp)
+			{
+				Particle_External->SetW(PullData.W, false);
+			}
+		}
+
+		if (AuthoritativeContext == EThreadContext::Internal && SolverSyncTimestamp >= ProxyTimestamp->OverWriteChildToParent.Timestamp)
+		{
+			for (int32 Index = 0; Index < PullData.ChildToParent.Num() && Index < SyncedData_External.ChildToParent.Num(); ++Index)
+			{
+				SyncedData_External.ChildToParent[Index] = PullData.ChildToParent[Index];
+			}
+		}
+		
+		Particle_External->UpdateShapeBounds();
+		return true;
+	}
+
+	void FClusterUnionPhysicsProxy::BufferPhysicsResults_Internal(FDirtyClusterUnionData& BufferData)
+	{
+		BufferPhysicsResultsImp(this, Particle_Internal, BufferData);
+	
+		FPBDRigidsEvolutionGBF& Evolution = *static_cast<FPBDRigidsSolver*>(Solver)->GetEvolution();
+		FClusterUnionManager& ClusterUnionManager = Evolution.GetRigidClustering().GetClusterUnionManager();
+		if (FClusterUnion* ClusterUnion = ClusterUnionManager.FindClusterUnion(ClusterUnionIndex))
+		{
+			BufferData.ChildToParent.Empty(ClusterUnion->ChildParticles.Num());
+			for (FPBDRigidParticleHandle* Particle : ClusterUnion->ChildParticles)
+			{
+				if (FPBDRigidClusteredParticleHandle* ClusteredParticle = Particle->CastToClustered())
+				{
+					BufferData.ChildToParent.Add(ClusteredParticle->ChildToParent());
+				}
+				else
+				{
+					BufferData.ChildToParent.Add(FTransform::Identity);
+				}
+			}
+		}
+	}
+
+	void FClusterUnionPhysicsProxy::BufferPhysicsResults_External(FDirtyClusterUnionData& BufferData)
+	{
+		BufferPhysicsResultsImp(this, Particle_External.Get(), BufferData);
+		BufferData.bIsAnchored = SyncedData_External.bIsAnchored;
+		BufferData.ChildToParent = SyncedData_External.ChildToParent;
+	}
+
+	void FClusterUnionPhysicsProxy::SyncRemoteData(FDirtyPropertiesManager& Manager, int32 DataIdx, FDirtyChaosProperties& RemoteData) const
+	{
+		if (!ensure(Particle_External))
+		{
+			return;
+		}
+
+		// This is similar to TGeometryParticle::SyncRemoteData except it puts it into the cluster properties.
+		RemoteData.SetParticleBufferType(Particle_External->Type);
+
+		// We need to modify the dirty flags to remove the non cluster properties to be 100% safe.
+		FDirtyChaosPropertyFlags DirtyFlags = Particle_External->DirtyFlags();
+
+#define CHAOS_PROPERTY(PropName, Type, ProxyType) if constexpr (ProxyType != EPhysicsProxyType::ClusterUnionProxy) { DirtyFlags.MarkClean(EChaosPropertyFlags::PropName); }
+#include "Chaos/ParticleProperties.inl"
+#undef CHAOS_PROPERTY
+
+		RemoteData.SetFlags(DirtyFlags);
+
+		// SyncRemote will check the dirty flags and will skip the change in value if the dirty flag is not actually set.
+		RemoteData.SyncRemote<FParticlePositionRotation, EChaosProperty::ClusterXR>(Manager, DataIdx, Particle_External->XR());
+		RemoteData.SyncRemote<FParticleVelocities, EChaosProperty::ClusterVelocities>(Manager, DataIdx, Particle_External->Velocities());
+		RemoteData.SyncRemote<bool, EChaosProperty::ClusterIsAnchored>(Manager, DataIdx, SyncedData_External.bIsAnchored);
+		RemoteData.SyncRemote<TArray<FTransform>, EChaosProperty::ClusterChildToParent>(Manager, DataIdx, SyncedData_External.ChildToParent);
+	}
+
+	void FClusterUnionPhysicsProxy::ClearAccumulatedData()
+	{
+		if (!ensure(Particle_External))
+		{
+			return;
+		}
+
+		Particle_External->ClearDirtyFlags();
+	}
+
+	void FClusterUnionPhysicsProxy::SetXR_External(const FVector& X, const FQuat& R)
+	{
+		if (!ensure(Particle_External))
+		{
+			return;
+		}
+
+		Particle_External->SetX(X, false);
+		Particle_External->SetR(R, false);
+		Particle_External->ForceDirty(EChaosPropertyFlags::ClusterXR);
+
+		FClusterUnionProxyTimestamp& SyncTS = GetSyncTimestampAs<FClusterUnionProxyTimestamp>();
+		SyncTS.OverWriteX.Set(GetSolverSyncTimestamp_External(), X);
+		SyncTS.OverWriteR.Set(GetSolverSyncTimestamp_External(), R);
+	}
+
+	void FClusterUnionPhysicsProxy::SetLinearVelocity_External(const FVector& V)
+	{
+		if (!ensure(Particle_External))
+		{
+			return;
+		}
+
+		Particle_External->SetV(V, false);
+		Particle_External->ForceDirty(EChaosPropertyFlags::ClusterVelocities);
+
+		FClusterUnionProxyTimestamp& SyncTS = GetSyncTimestampAs<FClusterUnionProxyTimestamp>();
+		SyncTS.OverWriteV.Set(GetSolverSyncTimestamp_External(), V);
+	}
+
+	void FClusterUnionPhysicsProxy::SetAngularVelocity_External(const FVector& W)
+	{
+		if (!ensure(Particle_External))
+		{
+			return;
+		}
+
+		Particle_External->SetW(W, false);
+		Particle_External->ForceDirty(EChaosPropertyFlags::ClusterVelocities);
+
+		FClusterUnionProxyTimestamp& SyncTS = GetSyncTimestampAs<FClusterUnionProxyTimestamp>();
+		SyncTS.OverWriteW.Set(GetSolverSyncTimestamp_External(), W);
+	}
+
+	void FClusterUnionPhysicsProxy::SetChildToParent_External(FPhysicsObjectHandle Child, const FTransform& RelativeTransform)
+	{
+		if (!Solver || !ensure(Particle_External))
+		{
+			return;
+		}
+
+		if (int32* Index = ChildPhysicsObjectIndexMap.Find(Child); Index && *Index >=0 && *Index < ChildPhysicsObjects.Num() && ensure(ChildPhysicsObjects[*Index] == Child))
+		{
+			SyncedData_External.ChildToParent[*Index] = RelativeTransform;
+			Particle_External->ForceDirty(EChaosPropertyFlags::ClusterChildToParent);
+			Solver->AddDirtyProxy(this);
+
+			FClusterUnionProxyTimestamp& SyncTS = GetSyncTimestampAs<FClusterUnionProxyTimestamp>();
+			SyncTS.OverWriteChildToParent.Set(GetSolverSyncTimestamp_External(), SyncedData_External.ChildToParent);
+		}
+	}
+}
\ No newline at end of file
diff --git a/Engine/Source/Runtime/Experimental/Chaos/Public/Chaos/Framework/ChaosResultsManager.h b/Engine/Source/Runtime/Experimental/Chaos/Public/Chaos/Framework/ChaosResultsManager.h
--- a/Engine/Source/Runtime/Experimental/Chaos/Public/Chaos/Framework/ChaosResultsManager.h
+++ b/Engine/Source/Runtime/Experimental/Chaos/Public/Chaos/Framework/ChaosResultsManager.h
@@ -21,6 +21,12 @@ namespace Chaos
 		FDirtyGeometryCollectionData Prev;
 		FDirtyGeometryCollectionData Next;
 	};
+
+	struct CHAOS_API FChaosClusterUnionInterpolationData
+	{
+		FDirtyClusterUnionData Prev;
+		FDirtyClusterUnionData Next;
+	};
 	
 	struct CHAOS_API FChaosInterpolationResults
 	{
@@ -34,6 +40,7 @@ namespace Chaos
 		FChaosInterpolationResults(FChaosInterpolationResults&& Other)
 			: RigidInterpolations(MoveTemp(Other.RigidInterpolations))
 			, GeometryCollectionInterpolations(MoveTemp(Other.GeometryCollectionInterpolations))
+			, ClusterUnionInterpolations(MoveTemp(Other.ClusterUnionInterpolations))
 			, Prev(Other.Prev)
 			, Next(Other.Next)
 			, Alpha(Other.Alpha)
@@ -46,6 +53,7 @@ namespace Chaos
 		
 		TArray<FChaosRigidInterpolationData> RigidInterpolations;
 		TArray<FChaosGeometryCollectionInterpolationData> GeometryCollectionInterpolations;
+		TArray<FChaosClusterUnionInterpolationData> ClusterUnionInterpolations;
 		FPullPhysicsData* Prev;
 		FPullPhysicsData* Next;
 		FRealSingle Alpha;
diff --git a/Engine/Source/Runtime/Experimental/Chaos/Public/Chaos/Framework/PhysicsSolverBase.h b/Engine/Source/Runtime/Experimental/Chaos/Public/Chaos/Framework/PhysicsSolverBase.h
--- a/Engine/Source/Runtime/Experimental/Chaos/Public/Chaos/Framework/PhysicsSolverBase.h
+++ b/Engine/Source/Runtime/Experimental/Chaos/Public/Chaos/Framework/PhysicsSolverBase.h
@@ -501,12 +501,16 @@ namespace Chaos
 			bPaused_External = bShouldPause;
 		}
 
-		/** Used to update external thread data structures. RigidFunc allows per dirty rigid code to execute. Include PhysicsSolverBaseImpl.h to call this function*/
+		/** Used to update external thread data structures. Include PhysicsSolverBaseImpl.h to call this function*/
+		template<typename TDispatcher>
+		void PullPhysicsStateForEachDirtyProxy_External(TDispatcher& Dispatcher);
+
 		template <typename RigidLambda, typename ConstraintLambda, typename GeometryCollectionLambda>
+		UE_DEPRECATED(5.3, "Use PullPhysicsStateForEachDirtyProxy_External with the TDispatcher template parameter instead.")
 		void PullPhysicsStateForEachDirtyProxy_External(const RigidLambda& RigidFunc, const ConstraintLambda& ConstraintFunc, const GeometryCollectionLambda& GeometryCollectionFunc);
 
 		template <typename RigidLambda, typename ConstraintLambda>
-		UE_DEPRECATED(5.4, "Use PullPhysicsStateForEachDirtyProxy_External with the additional GeometryCollectionLambda parameter instead.")
+		UE_DEPRECATED(5.3, "Use PullPhysicsStateForEachDirtyProxy_External with the TDispatcher template parameter instead.")
 		void PullPhysicsStateForEachDirtyProxy_External(const RigidLambda& RigidFunc, const ConstraintLambda& ConstraintFunc);
 
 		bool IsUsingAsyncResults() const
diff --git a/Engine/Source/Runtime/Experimental/Chaos/Public/Chaos/ParticleProperties.inl b/Engine/Source/Runtime/Experimental/Chaos/Public/Chaos/ParticleProperties.inl
--- a/Engine/Source/Runtime/Experimental/Chaos/Public/Chaos/ParticleProperties.inl
+++ b/Engine/Source/Runtime/Experimental/Chaos/Public/Chaos/ParticleProperties.inl
@@ -30,4 +30,9 @@ CHAOS_PROPERTY(CharacterGroundConstraintDynamicData, FCharacterGroundConstraintD
 CHAOS_PROPERTY(CharacterParticleProxy, FParticleProxyProperty, EPhysicsProxyType::CharacterGroundConstraintType)
 CHAOS_PROPERTY(GroundParticleProxy, FParticleProxyProperty, EPhysicsProxyType::CharacterGroundConstraintType)
 
+CHAOS_PROPERTY(ClusterXR, FParticlePositionRotation, EPhysicsProxyType::ClusterUnionProxy)
+CHAOS_PROPERTY(ClusterVelocities, FParticleVelocities, EPhysicsProxyType::ClusterUnionProxy)
+CHAOS_PROPERTY(ClusterIsAnchored, bool, EPhysicsProxyType::ClusterUnionProxy)
+CHAOS_PROPERTY(ClusterChildToParent, TArray<FTransform>, EPhysicsProxyType::ClusterUnionProxy)
+
 #undef CHAOS_PROPERTY
\ No newline at end of file
diff --git a/Engine/Source/Runtime/Experimental/Chaos/Public/Chaos/PhysicsSolverBaseImpl.h b/Engine/Source/Runtime/Experimental/Chaos/Public/Chaos/PhysicsSolverBaseImpl.h
--- a/Engine/Source/Runtime/Experimental/Chaos/Public/Chaos/PhysicsSolverBaseImpl.h
+++ b/Engine/Source/Runtime/Experimental/Chaos/Public/Chaos/PhysicsSolverBaseImpl.h
@@ -9,16 +9,41 @@
 #include "Chaos/Framework/PhysicsSolverBase.h"
 #include "Chaos/PullPhysicsDataImp.h"
 #include "PhysicsProxy/CharacterGroundConstraintProxy.h"
+#include "PhysicsProxy/ClusterUnionPhysicsProxy.h"
 #include "PhysicsProxy/SingleParticlePhysicsProxy.h"
 #include "PhysicsProxy/GeometryCollectionPhysicsProxy.h"
 #include "PhysicsProxy/JointConstraintProxy.h"
 #include "Chaos/Framework/ChaosResultsManager.h"
 
+#include <type_traits>
+
 namespace Chaos
 {
-	// Pulls physics state for each dirty particle and allows caller to do additional work if needed
-	template <typename RigidLambda, typename ConstraintLambda, typename GeometryCollectionLambda>
-	void FPhysicsSolverBase::PullPhysicsStateForEachDirtyProxy_External(const RigidLambda& RigidFunc, const ConstraintLambda& ConstraintFunc, const GeometryCollectionLambda& GeometryCollectionFunc)
+	namespace Private
+	{
+		// This lets us add support for performing per-proxy operations in PullPhysicsStateForEachDirtyProxy_External
+		// without the need to keep adding function parameters to keep taking in lambdas. We assume that the incoming
+		// user-specified TDispatcher has various operator() overridden to take in a specific proxy pointer. We require
+		// the TPullPhysicsStateDispatchHelper then to wrap the TDispatcher so that we can check if the incoming TDispatcher
+		// has a override for the proxy we want to pass to it. If it does not, we do nothing. This lets us maintain backward
+		// compatability as we keep adding in new proxy types while not requiring any additional computation at runtime since
+		// this dispatch is all done at compile time.
+		template<typename TDispatcher>
+		struct TPullPhysicsStateDispatchHelper
+		{
+			template<typename TProxy>
+			static void Apply(TDispatcher& Dispatcher, TProxy* Proxy)
+			{
+				if constexpr (std::is_invocable_v<TDispatcher, TProxy*>)
+				{
+					Dispatcher(Proxy);
+				}
+			}
+		};
+	}
+
+	template<typename TDispatcher>
+	void FPhysicsSolverBase::PullPhysicsStateForEachDirtyProxy_External(TDispatcher& Dispatcher)
 	{
 		using namespace Chaos;
 
@@ -48,7 +73,7 @@ namespace Chaos
 					{
 						if (Proxy->PullFromPhysicsState(RigidInterp.Prev, SolverTimestamp, &RigidInterp.Next, &Results.Alpha))
 						{
-							RigidFunc(Proxy);
+							Private::TPullPhysicsStateDispatchHelper<TDispatcher>::Apply(Dispatcher, Proxy);
 						}
 					}
 				}
@@ -60,7 +85,19 @@ namespace Chaos
 					{
 						if (Proxy->PullFromPhysicsState(GCInterp.Prev, SolverTimestamp, &GCInterp.Next, &Results.Alpha))
 						{
-							GeometryCollectionFunc(Proxy);
+							Private::TPullPhysicsStateDispatchHelper<TDispatcher>::Apply(Dispatcher, Proxy);
+						}
+					}
+				}
+
+				// cluster unions
+				for (const FChaosClusterUnionInterpolationData& ClusterInterp : Results.ClusterUnionInterpolations)
+				{
+					if (FClusterUnionPhysicsProxy* Proxy = ClusterInterp.Prev.GetProxy())
+					{
+						if (Proxy->PullFromPhysicsState(ClusterInterp.Prev, SolverTimestamp, &ClusterInterp.Next, &Results.Alpha))
+						{
+							Private::TPullPhysicsStateDispatchHelper<TDispatcher>::Apply(Dispatcher, Proxy);
 						}
 					}
 				}
@@ -84,7 +121,7 @@ namespace Chaos
 				{
 					if (Proxy->PullFromPhysicsState(RigidInterp.Next, SolverTimestamp))
 					{
-						RigidFunc(Proxy);
+						Private::TPullPhysicsStateDispatchHelper<TDispatcher>::Apply(Dispatcher, Proxy);
 					}
 				}
 			}
@@ -96,11 +133,22 @@ namespace Chaos
 				{
 					if (Proxy->PullFromPhysicsState(GCInterp.Next, SolverTimestamp))
 					{
-						GeometryCollectionFunc(Proxy);
+						Private::TPullPhysicsStateDispatchHelper<TDispatcher>::Apply(Dispatcher, Proxy);
 					}
 				}
 			}
 
+			// cluster unions
+			for (const FChaosClusterUnionInterpolationData& ClusterInterp : Results.ClusterUnionInterpolations)
+			{
+				if (FClusterUnionPhysicsProxy* Proxy = ClusterInterp.Prev.GetProxy())
+				{
+					if (Proxy->PullFromPhysicsState(ClusterInterp.Prev, SolverTimestamp))
+					{
+						Private::TPullPhysicsStateDispatchHelper<TDispatcher>::Apply(Dispatcher, Proxy);
+					}
+				}
+			}
 		}
 
 		//no interpolation for GC or joints at the moment
@@ -122,7 +170,7 @@ namespace Chaos
 				{
 					if (Proxy->PullFromPhysicsState(DirtyData, SyncTimestamp))
 					{
-						ConstraintFunc(Proxy);
+						Private::TPullPhysicsStateDispatchHelper<TDispatcher>::Apply(Dispatcher, Proxy);
 					}
 
 				}
@@ -143,10 +191,45 @@ namespace Chaos
 		}
 	}
 
+	// Pulls physics state for each dirty particle and allows caller to do additional work if needed
+	template <typename RigidLambda, typename ConstraintLambda, typename GeometryCollectionLambda>
+	void FPhysicsSolverBase::PullPhysicsStateForEachDirtyProxy_External(const RigidLambda& RigidFunc, const ConstraintLambda& ConstraintFunc, const GeometryCollectionLambda& GeometryCollectionFunc)
+	{
+		struct FDispatcher
+		{
+			void operator()(FSingleParticlePhysicsProxy* Proxy)
+			{
+				RigidFunc(Proxy);
+			}
+
+			void operator()(FJointConstraintPhysicsProxy* Proxy)
+			{
+				ConstraintFunc(Proxy);
+			}
+
+			void operator()(FGeometryCollectionPhysicsProxy* Proxy)
+			{
+				GeometryCollectionFunc(Proxy);
+			}
+		} Dispatcher;
+		PullPhysicsStateForEachDirtyProxy_External(Dispatcher);
+	}
+
 	template <typename RigidLambda, typename ConstraintLambda>
 	void FPhysicsSolverBase::PullPhysicsStateForEachDirtyProxy_External(const RigidLambda& RigidFunc, const ConstraintLambda& ConstraintFunc)
 	{
-		auto GeometryCollectionFunc = [](FGeometryCollectionPhysicsProxy* Proxy) {};
-		PullPhysicsStateForEachDirtyProxy_External(RigidFunc, ConstraintFunc, GeometryCollectionFunc);
+		struct FDispatcher
+		{
+			void operator()(FSingleParticlePhysicsProxy* Proxy)
+			{
+				RigidFunc(Proxy);
+			}
+
+			void operator()(FJointConstraintPhysicsProxy* Proxy)
+			{
+				ConstraintFunc(Proxy);
+			}
+		} Dispatcher;
+		PullPhysicsStateForEachDirtyProxy_External(Dispatcher);
 	}
 }
diff --git a/Engine/Source/Runtime/Experimental/Chaos/Public/Chaos/PullPhysicsDataImp.h b/Engine/Source/Runtime/Experimental/Chaos/Public/Chaos/PullPhysicsDataImp.h
--- a/Engine/Source/Runtime/Experimental/Chaos/Public/Chaos/PullPhysicsDataImp.h
+++ b/Engine/Source/Runtime/Experimental/Chaos/Public/Chaos/PullPhysicsDataImp.h
@@ -13,6 +13,7 @@
 namespace Chaos
 {
 
+class FClusterUnionPhysicsProxy;
 class FCharacterGroundConstraintProxy;
 class FJointConstraintPhysicsProxy;
 
@@ -62,6 +63,16 @@ struct FDirtyGeometryCollectionData : public TBasePullData<FGeometryCollectionPh
 	FGeometryCollectionResults Results;
 };
 
+struct FDirtyClusterUnionData : public TBasePullData<FClusterUnionPhysicsProxy, FClusterUnionProxyTimestamp>
+{
+	FVec3 X;
+	FQuat R;
+	FVec3 V;
+	FVec3 W;
+	bool bIsAnchored = false;
+	TArray<FTransform> ChildToParent;
+};
+
 struct FJointConstraintOutputData {
 	bool bIsBreaking = false;
 	bool bIsBroken = false;
@@ -89,6 +100,7 @@ class FPullPhysicsData
 public:
 	TArray<FDirtyRigidParticleData> DirtyRigids;
 	TArray<FDirtyGeometryCollectionData> DirtyGeometryCollections;
+	TArray<FDirtyClusterUnionData> DirtyClusterUnions;
 	TArray<FDirtyJointConstraintData> DirtyJointConstraints;
 	TArray<FDirtyCharacterGroundConstraintData> DirtyCharacterGroundConstraints;
 
diff --git a/Engine/Source/Runtime/Experimental/Chaos/Public/PBDRigidsSolver.h b/Engine/Source/Runtime/Experimental/Chaos/Public/PBDRigidsSolver.h
--- a/Engine/Source/Runtime/Experimental/Chaos/Public/PBDRigidsSolver.h
+++ b/Engine/Source/Runtime/Experimental/Chaos/Public/PBDRigidsSolver.h
@@ -57,6 +57,7 @@ namespace Chaos
 	class FCharacterGroundConstraintProxy;
 	class FSingleParticleProxy;
 	class FGeometryParticleBuffer;
+	class FClusterUnionPhysicsProxy;
 
 	CHAOS_API extern int32 RewindCaptureNumFrames;
 
@@ -131,6 +132,9 @@ namespace Chaos
 		void RegisterObject(FGeometryCollectionPhysicsProxy* InProxy);
 		void UnregisterObject(FGeometryCollectionPhysicsProxy* InProxy);
 
+		void RegisterObject(FClusterUnionPhysicsProxy* Proxy);
+		void UnregisterObject(FClusterUnionPhysicsProxy* Proxy);
+
 		void RegisterObject(Chaos::FJointConstraint* GTConstraint);
 		void UnregisterObject(Chaos::FJointConstraint* GTConstraint);
 
@@ -354,6 +358,7 @@ namespace Chaos
 		TSharedPtr<FCriticalSection> MCurrentLock;
 		TSparseArray< FSingleParticlePhysicsProxy* > SingleParticlePhysicsProxies_PT;
 		TArray< FGeometryCollectionPhysicsProxy* > GeometryCollectionPhysicsProxies_Internal; // PT
+		TArray< FClusterUnionPhysicsProxy* > ClusterUnionPhysicsProxies_Internal; // PT
 		TArray< FJointConstraintPhysicsProxy* > JointConstraintPhysicsProxies_Internal; // PT
 		TArray< FCharacterGroundConstraintProxy* > CharacterGroundConstraintProxies_Internal;
 
@@ -382,6 +387,7 @@ namespace Chaos
 
 		TArray<FPendingDestroyInfo> PendingDestroyPhysicsProxy;
 		TArray<FGeometryCollectionPhysicsProxy*> PendingDestroyGeometryCollectionPhysicsProxy;
+		TArray<FClusterUnionPhysicsProxy*> PendingDestroyClusterUnionProxy;
 
 
 		void ProcessSinglePushedData_Internal(FPushPhysicsData& PushData);
diff --git a/Engine/Source/Runtime/Experimental/Chaos/Public/PhysicsProxy/ClusterUnionPhysicsProxy.h b/Engine/Source/Runtime/Experimental/Chaos/Public/PhysicsProxy/ClusterUnionPhysicsProxy.h
new file mode 100644
--- /dev/null
+++ b/Engine/Source/Runtime/Experimental/Chaos/Public/PhysicsProxy/ClusterUnionPhysicsProxy.h
@@ -0,0 +1,140 @@
+// Copyright Epic Games, Inc. All Rights Reserved.
+
+#pragma once
+
+#include "Chaos/ClusterCreationParameters.h"
+#include "Chaos/ClusterUnionManager.h"
+#include "Chaos/Framework/PhysicsProxy.h"
+#include "Chaos/ParticleHandleFwd.h"
+#include "Chaos/PhysicsObject.h"
+#include "Containers/Map.h"
+#include "Containers/Set.h"
+#include "Framework/Threading.h"
+#include "PBDRigidsSolver.h"
+#include "Templates/UniquePtr.h"
+
+namespace Chaos
+{
+	struct FDirtyClusterUnionData;
+
+	/**
+	 * Extra data that needs to be synced between the PT and GT for cluster unions.
+	 */
+	struct FClusterUnionSyncedData
+	{
+		// Whether the cluster is anchored or not.
+		bool bIsAnchored;
+
+		// For every child, we need to sync its ChildToParent transform. The order of the transforms is identical
+		// to the order of the physics objects stored in the proxy. The *direction* that we sync ChildToParent
+		// is determined by whether we're on the server or the client. On the server, the ChildToParent is physics
+		// thread authoritative so the sync goes from the PT to the GT. On the client, we want to use whatever
+		// ChildToParent was given by the server so it's GT authoritative so the sync goes the other way.
+		TArray<FTransform> ChildToParent;
+	};
+
+
+	class CHAOS_API FClusterUnionPhysicsProxy : public TPhysicsProxy<FClusterUnionPhysicsProxy, void, FClusterUnionProxyTimestamp>
+	{
+		using Base = TPhysicsProxy<FClusterUnionPhysicsProxy, void, FClusterUnionProxyTimestamp>;
+	public:
+		// The mismatch here is fine since we really only need to sync a few properties between the PT and the GT.
+		using FExternalParticle = TPBDRigidParticle<FReal, 3>;
+		using FInternalParticle = TPBDRigidClusteredParticleHandle<Chaos::FReal, 3>;
+
+		FClusterUnionPhysicsProxy() = delete;
+		FClusterUnionPhysicsProxy(UObject* InOwner, const FClusterCreationParameters& InParameters, void* InUserData, EThreadContext InAuthoritativeContext);
+
+		// Add physics objects to the cluster union. Should only be called from the game thread.
+		void AddPhysicsObjects_External(const TArray<FPhysicsObjectHandle>& Objects);
+		const TArray<FPhysicsObjectHandle>& GetChildPhysicsObjects_External() const { return ChildPhysicsObjects; }
+		const FClusterUnionSyncedData& GetSyncedData_External() const { return SyncedData_External; }
+
+		// Set/Remove any anchors on the cluster particle.
+		void SetIsAnchored_External(bool bIsAnchored);
+
+		// Cluster union proxy initialization happens in two first on the game thread (external) then on the
+		// physics thread (internal). Cluster unions are a primarily physics concept so the things exposed to
+		// an external context is just there to be able to safely query the state on the physics thread and to
+		// be able to add/remove things to the proxy itself (rather than to move it, for example).
+		void Initialize_External();
+
+		bool IsInitializedOnPhysicsThread() const { return bIsInitializedOnPhysicsThread; }
+		void Initialize_Internal(FPBDRigidsSolver* RigidsSolver, FPBDRigidsSolver::FParticlesType& Particles);
+
+		FExternalParticle* GetParticle_External() const { return Particle_External.Get(); }
+		FInternalParticle* GetParticle_Internal() const { return Particle_Internal; }
+		virtual void* GetHandleUnsafe() const override { return Particle_Internal; }
+		FPhysicsObjectHandle GetPhysicsObjectHandle() const { return PhysicsObject.Get(); }
+
+		bool HasChildren_External() const { return !ChildrenParticles_External.IsEmpty(); }
+		bool HasChildren_Internal() const;
+
+		bool IsAnchored_External() const { return SyncedData_External.bIsAnchored; }
+
+		void SetXR_External(const FVector& X, const FQuat& R);
+		void SetLinearVelocity_External(const FVector& V);
+		void SetAngularVelocity_External(const FVector& W);
+		void SetChildToParent_External(FPhysicsObjectHandle Child, const FTransform& RelativeTransform);
+
+		//
+		// These functions take care of marshaling data back and forth between the game thread
+		// and the physics thread.
+		//
+		void PushToPhysicsState(const FDirtyPropertiesManager& Manager, int32 DataIdx, const FDirtyProxy& Dirty);
+		bool PullFromPhysicsState(const FDirtyClusterUnionData& PullData, int32 SolverSyncTimestamp, const FDirtyClusterUnionData* NextPullData = nullptr, const FRealSingle* Alpha = nullptr);
+
+		void BufferPhysicsResults_Internal(FDirtyClusterUnionData& BufferData);
+		void BufferPhysicsResults_External(FDirtyClusterUnionData& BufferData);
+
+		void SyncRemoteData(FDirtyPropertiesManager& Manager, int32 DataIdx, FDirtyChaosProperties& RemoteData) const;
+		void ClearAccumulatedData();
+
+		FProxyInterpolationData& GetInterpolationData() { return InterpolationData; }
+		const FProxyInterpolationData& GetInterpolationData() const { return InterpolationData; }
+
+		FClusterUnionIndex GetClusterUnionIndex() const { return ClusterUnionIndex; }
+	private:
+		bool bIsInitializedOnPhysicsThread = false;
+		FClusterCreationParameters ClusterParameters;
+		void* UserData = nullptr;
+		EThreadContext AuthoritativeContext;
+
+		FPhysicsObjectUniquePtr PhysicsObject;
+		TUniquePtr<FExternalParticle> Particle_External;
+		FClusterUnionSyncedData SyncedData_External;
+
+		// This isn't captured in TPBDRigidParticle hence why we need to track it manually.
+		TSet<FExternalParticle*> ChildrenParticles_External;
+
+		FInternalParticle* Particle_Internal = nullptr;
+		FClusterUnionIndex ClusterUnionIndex = INDEX_NONE;
+
+		// Children. Assuming that the ChildPhysicsObjects array can get large so we also maintain
+		// the ChildPhysicsObjectIndexMap so that it's easy to find a given physics object handle's index.
+		// This is useful because other things such as the FClusterUnionSyncedData::ChildToParent array
+		// will utilize the same index.
+		TArray<FPhysicsObjectHandle> ChildPhysicsObjects;
+		TMap<FPhysicsObjectHandle, int32> ChildPhysicsObjectIndexMap;
+
+		FProxyInterpolationData InterpolationData;
+
+		//~ Begin TPhysicsProxy Interface
+	public:
+		bool IsSimulating() const { return true; }
+		void UpdateKinematicBodiesCallback(const FParticlesType& InParticles, const float InDt, const float InTime, FKinematicProxy& InKinematicProxy) {}
+		void StartFrameCallback(const float InDt, const float InTime) {}
+		void EndFrameCallback(const float InDt) {}
+		void CreateRigidBodyCallback(FParticlesType& InOutParticles) {}
+		void DisableCollisionsCallback(TSet<TTuple<int32, int32>>& InPairs) {}
+		void AddForceCallback(FParticlesType& InParticles, const float InDt, const int32 InIndex) {}
+		void BindParticleCallbackMapping(Chaos::TArrayCollectionArray<PhysicsProxyWrapper>& PhysicsProxyReverseMap, Chaos::TArrayCollectionArray<int32>& ParticleIDReverseMap) {}
+		EPhysicsProxyType ConcreteType() { return EPhysicsProxyType::ClusterUnionProxy; }
+		void SyncBeforeDestroy() {}
+		void OnRemoveFromScene() {}
+		bool IsDirty() { return false; }
+
+		//~ End TPhysicsProxy Interface
+	};
+
+}
\ No newline at end of file
diff --git a/Engine/Source/Runtime/PhysicsCore/Private/ChaosScene.cpp b/Engine/Source/Runtime/PhysicsCore/Private/ChaosScene.cpp
--- a/Engine/Source/Runtime/PhysicsCore/Private/ChaosScene.cpp
+++ b/Engine/Source/Runtime/PhysicsCore/Private/ChaosScene.cpp
@@ -379,7 +379,8 @@ void FChaosScene::StartFrame()
 
 void FChaosScene::OnSyncBodies(Chaos::FPhysicsSolverBase* Solver)
 {
-	Solver->PullPhysicsStateForEachDirtyProxy_External([](auto) {}, [](auto) {}, [](auto) {});
+	struct FDispatcher {} Dispatcher;
+	Solver->PullPhysicsStateForEachDirtyProxy_External(Dispatcher);
 }
 
 bool FChaosScene::AreAnyTasksPending() const
