Geometry collection : replication : make server authoritative particles unbreakable by the client
- added unbreakable flag to clustered particles SOA
- packed InternalCluster , Anchored and Unbrekable booleans in a flag structure using bit fields ( saved a bool per particle )
- non authoritative particles are now unbreakable on the client until they replicate

diff --git a/Engine/Source/Runtime/Experimental/Chaos/Private/Chaos/PBDRigidClustering.cpp b/Engine/Source/Runtime/Experimental/Chaos/Private/Chaos/PBDRigidClustering.cpp
--- a/Engine/Source/Runtime/Experimental/Chaos/Private/Chaos/PBDRigidClustering.cpp
+++ b/Engine/Source/Runtime/Experimental/Chaos/Private/Chaos/PBDRigidClustering.cpp
@@ -647,6 +647,11 @@ namespace Chaos
 	{	
 		TSet<FPBDRigidParticleHandle*> ActivatedChildren;
 
+		if (ClusteredParticle && ClusteredParticle->Unbreakable())
+		{
+			return ActivatedChildren;
+		}
+
 		if (!ensureMsgf(MChildren.Contains(ClusteredParticle), TEXT("Removing Cluster that does not exist!")))
 		{
 			return ActivatedChildren;
diff --git a/Engine/Source/Runtime/Experimental/Chaos/Public/Chaos/PBDRigidClusteredParticles.h b/Engine/Source/Runtime/Experimental/Chaos/Public/Chaos/PBDRigidClusteredParticles.h
--- a/Engine/Source/Runtime/Experimental/Chaos/Public/Chaos/PBDRigidClusteredParticles.h
+++ b/Engine/Source/Runtime/Experimental/Chaos/Public/Chaos/PBDRigidClusteredParticles.h
@@ -66,6 +66,47 @@ struct TConnectivityEdge
 typedef TConnectivityEdge<FReal> FConnectivityEdge;
 typedef TArray<FConnectivityEdge> FConnectivityEdgeArray;
 
+class FRigidClusteredFlags
+{
+public:
+	using FStorage = uint8;
+
+	FRigidClusteredFlags()
+		: Bits(0)
+	{
+	}
+
+	// set to true if the particle is an internal cluster formed of the subset of children from an original cluster 
+	bool GetInternalCluster() const { return Flags.bInternalCluster; }
+	void SetInternalCluster(bool bSet) { Flags.bInternalCluster = bSet; }
+
+	// set to true make the particle act as a kinematic anchor, this allows the particle to be broken off while still be anchor contributor through the connection graph   
+	bool GetAnchored() const { return Flags.bAnchored; }
+	void SetAnchored(bool bSet) { Flags.bAnchored = bSet; }
+
+	// set to true to make the particle unbreakable by destruction operations from breaking server authoritative particles
+	bool GetUnbreakable() const { return Flags.bUnbreakable; }
+	void SetUnbreakable(bool bSet) { Flags.bUnbreakable = bSet; }
+
+private:
+	struct FFlags
+	{
+		FStorage bInternalCluster : 1;
+		FStorage bAnchored : 1;
+		FStorage bUnbreakable : 1;
+
+		// Add new properties above this line
+		// Change FStorage typedef if we exceed the max bits
+	};
+	union
+	{
+		FFlags Flags;
+		FStorage Bits;
+	};
+	static_assert(sizeof(FFlags) <= sizeof(FStorage));
+};
+
+
 template<class T, int d>
 class TPBDRigidClusteredParticles : public TPBDRigidParticles<T, d>
 {
@@ -81,14 +122,13 @@ class TPBDRigidClusteredParticles : public TPBDRigidParticles<T, d>
 	, MClusterIds(MoveTemp(Other.MClusterIds))
 	, MChildToParent(MoveTemp(Other.MChildToParent))
 	, MClusterGroupIndex(MoveTemp(Other.MClusterGroupIndex))
-	, MInternalCluster(MoveTemp(Other.MInternalCluster))
 	, MChildrenSpatial(MoveTemp(Other.MChildrenSpatial))
 	, MPhysicsProxies(MoveTemp(Other.MPhysicsProxies))
 	, MCollisionImpulses(MoveTemp(Other.MCollisionImpulses))
 	, MStrains(MoveTemp(Other.MStrains))
 	, MConnectivityEdges(MoveTemp(Other.MConnectivityEdges))
 	, MExternalStrains(MoveTemp(Other.MExternalStrains))
-	, MAnchored(MoveTemp(Other.MAnchored))
+	, MRigidClusteredFlags(MoveTemp(Other.MRigidClusteredFlags))
 	{
 		InitHelper();
 	}
@@ -103,9 +143,6 @@ class TPBDRigidClusteredParticles : public TPBDRigidParticles<T, d>
 	const auto& ClusterGroupIndex(int32 Idx) const { return MClusterGroupIndex[Idx]; }
 	auto& ClusterGroupIndex(int32 Idx) { return MClusterGroupIndex[Idx]; }
 
-	const auto& InternalCluster(int32 Idx) const { return MInternalCluster[Idx]; }
-	auto& InternalCluster(int32 Idx) { return MInternalCluster[Idx]; }
-
 	const auto& ChildrenSpatial(int32 Idx) const { return MChildrenSpatial[Idx]; }
 	auto& ChildrenSpatial(int32 Idx) { return MChildrenSpatial[Idx]; }
 
@@ -126,8 +163,8 @@ class TPBDRigidClusteredParticles : public TPBDRigidParticles<T, d>
 	const auto& ConnectivityEdges(int32 Idx) const { return MConnectivityEdges[Idx]; }
 	auto& ConnectivityEdges(int32 Idx) { return MConnectivityEdges[Idx]; }
 
-	const bool& Anchored(int32 Idx) const { return MAnchored[Idx]; }
-	bool& Anchored(int32 Idx) { return MAnchored[Idx]; }
+	const FRigidClusteredFlags& RigidClusteredFlags(int32 Idx) const { return MRigidClusteredFlags[Idx]; }
+	FRigidClusteredFlags& RigidClusteredFlags(int32 Idx) { return MRigidClusteredFlags[Idx]; }
 	
 	const auto& ConnectivityEdgesArray() const { return MConnectivityEdges; }
 
@@ -143,11 +180,8 @@ class TPBDRigidClusteredParticles : public TPBDRigidParticles<T, d>
 	const auto& ClusterGroupIndexArray() const { return MClusterGroupIndex; }
 	auto& ClusterGroupIndexArray() { return MClusterGroupIndex; }
 
-	const auto& InternalClusterArray() const { return MInternalCluster; }
-	auto& InternalClusterArray() { return MInternalCluster; }
-
-	const auto& AnchoredArray() const { return MAnchored; }
-	auto& AnchoredArray() { return MAnchored; }
+	const auto& RigidClusteredFlags() const { return MRigidClusteredFlags; }
+	auto& RigidClusteredFlags() { return MRigidClusteredFlags; }
 
 	
 	typedef TPBDRigidClusteredParticleHandle<T, d> THandleType;
@@ -165,20 +199,18 @@ class TPBDRigidClusteredParticles : public TPBDRigidParticles<T, d>
 		  TArrayCollection::AddArray(&MClusterIds);
 		  TArrayCollection::AddArray(&MChildToParent);
 		  TArrayCollection::AddArray(&MClusterGroupIndex);
-		  TArrayCollection::AddArray(&MInternalCluster);
 		  TArrayCollection::AddArray(&MChildrenSpatial);
 		  TArrayCollection::AddArray(&MPhysicsProxies);
 		  TArrayCollection::AddArray(&MCollisionImpulses);
 		  TArrayCollection::AddArray(&MStrains);
 		  TArrayCollection::AddArray(&MConnectivityEdges);
 	  	  TArrayCollection::AddArray(&MExternalStrains);
-	  	  TArrayCollection::AddArray(&MAnchored);
+	  	  TArrayCollection::AddArray(&MRigidClusteredFlags);
 	  }
 
 	  TArrayCollectionArray<ClusterId> MClusterIds;
 	  TArrayCollectionArray<TRigidTransform<T, d>> MChildToParent;
 	  TArrayCollectionArray<int32> MClusterGroupIndex;
-	  TArrayCollectionArray<bool> MInternalCluster;
 	  TArrayCollectionArray<TUniquePtr<FImplicitObjectUnionClustered>> MChildrenSpatial;
 
 	  // Multiple proxy pointers required for internal clusters
@@ -195,10 +227,8 @@ class TPBDRigidClusteredParticles : public TPBDRigidParticles<T, d>
 	  TArrayCollectionArray<T> MStrains;
 
 	  TArrayCollectionArray<TArray<TConnectivityEdge<T>>> MConnectivityEdges;
-
-	  // make the particle act as a kinematic anchor,
-	  // this allows the particle to be broken off while still be anchor contributor through the connection graph   
-	  TArrayCollectionArray<bool> MAnchored;
+  
+	  TArrayCollectionArray<FRigidClusteredFlags> MRigidClusteredFlags;
 };
 
 using FPBDRigidClusteredParticles = TPBDRigidClusteredParticles<FReal, 3>;
diff --git a/Engine/Source/Runtime/Experimental/Chaos/Public/Chaos/PBDRigidClustering.h b/Engine/Source/Runtime/Experimental/Chaos/Public/Chaos/PBDRigidClustering.h
--- a/Engine/Source/Runtime/Experimental/Chaos/Public/Chaos/PBDRigidClustering.h
+++ b/Engine/Source/Runtime/Experimental/Chaos/Public/Chaos/PBDRigidClustering.h
@@ -227,11 +227,10 @@ public:
 	const TArrayCollectionArray<ClusterId>& GetClusterIdsArray() const { return MParticles.ClusterIdsArray(); }
 
 	/*
-	*  GetInternalClusterArray
-	*    The internal cluster array indicates if this cluster was generated internally
-	*    and would no be owned by an external source.
+	*  GetRigidClusteredFlagsArray
+	*    The RigidClusteredFlags array contains various flag related to clustered particles
 	*/
-	const TArrayCollectionArray<bool>& GetInternalClusterArray() const { return MParticles.InternalClusterArray(); }
+	const TArrayCollectionArray<FRigidClusteredFlags>& GetRigidClusteredFlagsArray() const { return MParticles.RigidClusteredFlags(); }
 
 	/*
 	*  GetChildToParentMap
diff --git a/Engine/Source/Runtime/Experimental/Chaos/Public/Chaos/ParticleHandle.h b/Engine/Source/Runtime/Experimental/Chaos/Public/Chaos/ParticleHandle.h
--- a/Engine/Source/Runtime/Experimental/Chaos/Public/Chaos/ParticleHandle.h
+++ b/Engine/Source/Runtime/Experimental/Chaos/Public/Chaos/ParticleHandle.h
@@ -1272,9 +1272,8 @@ public:
 	int32& ClusterGroupIndex() { return PBDRigidClusteredParticles->ClusterGroupIndex(ParticleIdx); }
 	void SetClusterGroupIndex(const int32 Idx) { PBDRigidClusteredParticles->ClusterGroupIndex(ParticleIdx) = Idx; }
 
-	const bool& InternalCluster() const { return PBDRigidClusteredParticles->InternalCluster(ParticleIdx); }
-	bool& InternalCluster() { return PBDRigidClusteredParticles->InternalCluster(ParticleIdx); }
-	void SetInternalCluster(const bool Value) { PBDRigidClusteredParticles->InternalCluster(ParticleIdx) = Value; }
+	bool InternalCluster() const { return PBDRigidClusteredParticles->RigidClusteredFlags(ParticleIdx).GetInternalCluster(); }
+	void SetInternalCluster(bool bValue) { PBDRigidClusteredParticles->RigidClusteredFlags(ParticleIdx).SetInternalCluster(bValue);	}
 
 	const TUniquePtr<FImplicitObjectUnionClustered>& ChildrenSpatial() const { return PBDRigidClusteredParticles->ChildrenSpatial(ParticleIdx); }
 	TUniquePtr<FImplicitObjectUnionClustered>& ChildrenSpatial() { return PBDRigidClusteredParticles->ChildrenSpatial(ParticleIdx); }
@@ -1311,10 +1310,12 @@ public:
 	void SetConnectivityEdges(const TArray<TConnectivityEdge<T>>& Edges) { PBDRigidClusteredParticles->ConnectivityEdges(ParticleIdx) = Edges; }
 	void SetConnectivityEdges(TArray<TConnectivityEdge<T>>&& Edges) { PBDRigidClusteredParticles->ConnectivityEdges(ParticleIdx) = MoveTemp(Edges); }
 
-	const bool& IsAnchored() const { return PBDRigidClusteredParticles->Anchored(ParticleIdx); }
-	bool& IsAnchored() { return PBDRigidClusteredParticles->Anchored(ParticleIdx); }
-	void SetIsAnchored(const bool Value) { PBDRigidClusteredParticles->Anchored(ParticleIdx) = Value; }
-	
+	bool IsAnchored() const { return PBDRigidClusteredParticles->RigidClusteredFlags(ParticleIdx).GetAnchored(); }
+	void SetIsAnchored(bool bValue) { PBDRigidClusteredParticles->RigidClusteredFlags(ParticleIdx).SetAnchored(bValue); }
+
+	bool Unbreakable() const { return PBDRigidClusteredParticles->RigidClusteredFlags(ParticleIdx).GetUnbreakable(); }
+	void SetUnbreakable(bool bValue) { PBDRigidClusteredParticles->RigidClusteredFlags(ParticleIdx).SetUnbreakable(bValue); }
+
 	const TPBDRigidClusteredParticleHandleImp<T, d, true>* Handle() const { return PBDRigidClusteredParticles->Handle(ParticleIdx); }
 	TPBDRigidClusteredParticleHandleImp<T, d, true>* Handle() { return PBDRigidClusteredParticles->Handle(ParticleIdx); }
 
diff --git a/Engine/Source/Runtime/Experimental/GeometryCollectionEngine/Private/GeometryCollection/GeometryCollectionComponent.cpp b/Engine/Source/Runtime/Experimental/GeometryCollectionEngine/Private/GeometryCollection/GeometryCollectionComponent.cpp
--- a/Engine/Source/Runtime/Experimental/GeometryCollectionEngine/Private/GeometryCollection/GeometryCollectionComponent.cpp
+++ b/Engine/Source/Runtime/Experimental/GeometryCollectionEngine/Private/GeometryCollection/GeometryCollectionComponent.cpp
@@ -1895,24 +1895,29 @@ bool UGeometryCollectionComponent::ProcessRepData(const float DeltaTime, const f
 		FPBDRigidParticleHandle* OneOff = PhysicsProxy->GetParticles()[ActivatedCluster.ActivatedIndex];
 		check(OneOff);
 
-		// If there's a parent cluster particle we need to release them first.
-		// This is generally an indication that something desynced between the client and server though...maybe something needs to be done
-		// to ensure internal clusters stay in sync.
-		if (FPBDRigidClusteredParticleHandle* ClusterParticle = OneOff->CastToClustered(); ClusterParticle && ClusterParticle->Parent())
+		if (FPBDRigidClusteredParticleHandle* ClusterParticle = OneOff->CastToClustered())
 		{
-			RigidClustering.ReleaseClusterParticles(TArray<FPBDRigidParticleHandle*>{ ClusterParticle->Parent() }, true);
-		}
+			// If there's a parent cluster particle we need to release them first.
+			// This is generally an indication that something desynced between the client and server though...maybe something needs to be done
+			// to ensure internal clusters stay in sync.
+			if (FPBDRigidClusteredParticleHandle* ParentParticle = ClusterParticle->Parent())
+			{
+				// server authoritative particles are unbreakable, we need to set them breakable again 
+				ParentParticle->SetUnbreakable(false);
+				RigidClustering.ReleaseClusterParticles(TArray<FPBDRigidParticleHandle*>{ ParentParticle }, true);
+			}
 
-		// Set initial velocities if not hard snapping
-		if (!bHardSnap)
-		{
-			// TODO: we should get an update cluster position first so that when particles break off they get the right position 
-			// TODO: should we invalidate?
-			OneOff->SetV(ActivatedCluster.InitialLinearVelocity);
-			OneOff->SetW(ActivatedCluster.InitialAngularVelocity);
-		}
+			// Set initial velocities if not hard snapping
+			if (!bHardSnap)
+			{
+				// TODO: we should get an update cluster position first so that when particles break off they get the right position 
+				// TODO: should we invalidate?
+				OneOff->SetV(ActivatedCluster.InitialLinearVelocity);
+				OneOff->SetW(ActivatedCluster.InitialAngularVelocity);
+			}
 
-		RigidClustering.ReleaseClusterParticles(TArray<FPBDRigidParticleHandle*>{ OneOff }, true);
+			RigidClustering.ReleaseClusterParticles(TArray<FPBDRigidParticleHandle*>{ ClusterParticle }, true);
+		}
 	}
 
 	// Keep track of whether we did some "work" on this frame so we can turn off the async tick after
@@ -2785,9 +2790,9 @@ void UGeometryCollectionComponent::RegisterAndInitializePhysicsProxy()
 						if (ParticleHandle)
 						{
 							const int32 Level = EnableAbandonAfterLevel ? ComputeParticleLevel(ParticleHandle) : -1;
-							if (Level <= AbandonAfterLevel + 1)	//we only replicate up until level X, but it means we should replicate the breaking event of level X+1 (but not X+1's positions)
+							if (Level <= AbandonAfterLevel)	//we only replicate up until level X, but it means we should replicate the breaking event of level X+1 (but not X+1's positions)
 							{
-								ParticleHandle->SetMaximumInternalStrain();
+								ParticleHandle->SetUnbreakable(true);
 							}
 						}
 					}
