fix false positive PVS studio warning V501 in ClusterUnionPhysicsProxy when including ParticleProperties.inl

diff --git a/Engine/Source/Runtime/Experimental/Chaos/Private/PhysicsProxy/ClusterUnionPhysicsProxy.cpp b/Engine/Source/Runtime/Experimental/Chaos/Private/PhysicsProxy/ClusterUnionPhysicsProxy.cpp
--- a/Engine/Source/Runtime/Experimental/Chaos/Private/PhysicsProxy/ClusterUnionPhysicsProxy.cpp
+++ b/Engine/Source/Runtime/Experimental/Chaos/Private/PhysicsProxy/ClusterUnionPhysicsProxy.cpp
@@ -351,7 +351,8 @@ namespace Chaos
 		// We need to modify the dirty flags to remove the non cluster properties to be 100% safe.
 		FDirtyChaosPropertyFlags DirtyFlags = Particle_External->DirtyFlags();
 
-#define CHAOS_PROPERTY(PropName, Type, ProxyType) if constexpr (ProxyType != EPhysicsProxyType::ClusterUnionProxy) { DirtyFlags.MarkClean(EChaosPropertyFlags::PropName); }
+		// We need to suppress V501 in PVS Studio since it's a false positive warning in this case. This is actually necessary for the codegen here.
+#define CHAOS_PROPERTY(PropName, Type, ProxyType) if constexpr (ProxyType != EPhysicsProxyType::ClusterUnionProxy) { DirtyFlags.MarkClean(EChaosPropertyFlags::PropName); } //-V501
 #include "Chaos/ParticleProperties.inl"
 #undef CHAOS_PROPERTY
 
