Avoid segfault when we have 0 particles.

diff --git a/Engine/Source/Runtime/Experimental/Chaos/Private/Chaos/Deformable/ChaosDeformableSolver.cpp b/Engine/Source/Runtime/Experimental/Chaos/Private/Chaos/Deformable/ChaosDeformableSolver.cpp
--- a/Engine/Source/Runtime/Experimental/Chaos/Private/Chaos/Deformable/ChaosDeformableSolver.cpp
+++ b/Engine/Source/Runtime/Experimental/Chaos/Private/Chaos/Deformable/ChaosDeformableSolver.cpp
@@ -249,12 +249,14 @@ namespace Chaos::Softs
 
 		if (FFleshThreadingProxy* Proxy = InProxy.As<FFleshThreadingProxy>())
 		{
-			InitializeDeformableParticles(*Proxy);
-			InitializeKinematicParticles(*Proxy);
-			InitializeWeakConstraint(*Proxy);
-			InitializeTetrahedralConstraint(*Proxy);
-			InitializeGidBasedConstraints(*Proxy);
-
+			if (Proxy->CanSimulate())
+			{
+				InitializeDeformableParticles(*Proxy);
+				InitializeKinematicParticles(*Proxy);
+				InitializeWeakConstraint(*Proxy);
+				InitializeTetrahedralConstraint(*Proxy);
+				InitializeGidBasedConstraints(*Proxy);
+			}
 		}
 		if (FCollisionManagerProxy* CollisionManagerProxy = InProxy.As< FCollisionManagerProxy>())
 		{
@@ -280,7 +282,6 @@ namespace Chaos::Softs
 		auto ChaosInvM = [](FSolverReal M) { return FSolverReal(FMath::IsNearlyZero(M) ? 0.0 : 1 / M); };
 		auto DoubleVert = [](FVector3f V) { return FVector3d(V.X, V.Y, V.Z); };
 		uint32 NumParticles = Rest.NumElements(FGeometryCollection::VerticesGroup);
-
 		int32 ParticleStart = Evolution->AddParticleRange(NumParticles, GroupOffset, true);
 		GroupOffset += 1;
 		for (uint32 vdx = 0; vdx < NumParticles; ++vdx)
@@ -351,7 +352,7 @@ namespace Chaos::Softs
 		}
 
 		bool ObjectEnableGravity = false;
-		if (const UObject* Owner = this->MObjects[ParticleStart]) 
+		if (const UObject* Owner = this->MObjects[ParticleStart])
 		{
 			FFleshThreadingProxy::FFleshInputBuffer* FleshInputBuffer = nullptr;
 			if (this->CurrentInputPackage->ObjectMap.Contains(Owner))
@@ -412,8 +413,6 @@ namespace Chaos::Softs
 				}
 			}
 		}
-
-		
 	}
 
 	void FDeformableSolver::InitializeWeakConstraint(FFleshThreadingProxy& Proxy)
@@ -452,7 +451,6 @@ namespace Chaos::Softs
 				PositionTargetStiffness[i] = DataPackage.Stiffness;
 			}
 
-
 			int32 InitIndex = Evolution->AddConstraintInitRange(1, true);
 			int32 ConstraintIndex = Evolution->AddConstraintRuleRange(1, true);
 
@@ -474,7 +472,6 @@ namespace Chaos::Softs
 			};
 
 			WeakConstraints.Add(TUniquePtr<FXPBDWeakConstraints<FSolverReal, FSolverParticles>>(WeakConstraint));
-
 		}
 	}
 
@@ -982,9 +979,12 @@ namespace Chaos::Softs
 			{
 				if (FFleshThreadingProxy* Proxy = BaseProxy->As<FFleshThreadingProxy>())
 				{
-					FIntVector2 Indices = Proxy->GetSolverParticleRange();
-					Proxies.FindAndRemoveChecked(MObjects[Indices[0]]);
-					Evolution->Particles().RemoveAt(Indices[0], Indices[1]);
+					if (Proxy->CanSimulate())
+					{
+						FIntVector2 Indices = Proxy->GetSolverParticleRange();
+						Proxies.FindAndRemoveChecked(MObjects[Indices[0]]);
+						Evolution->Particles().RemoveAt(Indices[0], Indices[1]);
+					}
 				}
 			}
 
diff --git a/Engine/Source/Runtime/Experimental/Chaos/Public/Chaos/Deformable/ChaosDeformableSolverProxy.h b/Engine/Source/Runtime/Experimental/Chaos/Public/Chaos/Deformable/ChaosDeformableSolverProxy.h
--- a/Engine/Source/Runtime/Experimental/Chaos/Public/Chaos/Deformable/ChaosDeformableSolverProxy.h
+++ b/Engine/Source/Runtime/Experimental/Chaos/Public/Chaos/Deformable/ChaosDeformableSolverProxy.h
@@ -115,6 +115,11 @@ namespace Chaos::Softs
 
 		static FName TypeName() { return FName("Flesh"); }
 
+		bool CanSimulate() const
+		{
+			return Rest.NumElements(FGeometryCollection::VerticesGroup) > 0;
+		}
+
 		bool IsBoneSpace() const
 		{
 			return SimSpace == ChaosDeformableSimSpace::Bone;
