Chaos - fix bug where GC particles end up in the Active array

This was causing the particle to appear in the broapdhase twice leading to a race condition as 2 threads tried to create its collision data.

diff --git a/Engine/Source/Runtime/Experimental/Chaos/Public/Chaos/PBDRigidsSOAs.h b/Engine/Source/Runtime/Experimental/Chaos/Public/Chaos/PBDRigidsSOAs.h
--- a/Engine/Source/Runtime/Experimental/Chaos/Public/Chaos/PBDRigidsSOAs.h
+++ b/Engine/Source/Runtime/Experimental/Chaos/Public/Chaos/PBDRigidsSOAs.h
@@ -671,15 +671,21 @@ public:
 		{ 
 			RemoveClusteredParticle(ClusteredParticle);
 			InsertClusteredParticle(ClusteredParticle);
-		}
 
-		if (ClusteredParticle->ObjectState() != EObjectStateType::Dynamic || ClusteredParticle->Disabled())
-		{
-			RemoveFromActiveArray(ClusteredParticle->CastToRigidParticle(), /*bStillDirty=*/true);
+			// Cluster particle must go in the active array because the Cluster arrays are not in the active view
+			if (ClusteredParticle->ObjectState() != EObjectStateType::Dynamic || ClusteredParticle->Disabled())
+			{
+				RemoveFromActiveArray(ClusteredParticle->CastToRigidParticle(), /*bStillDirty=*/true);
+			}
+			else
+			{
+				AddToActiveArray(ClusteredParticle->CastToRigidParticle());
+			}
 		}
 		else
 		{
-			AddToActiveArray(ClusteredParticle->CastToRigidParticle());
+			// Geometry collection particles have their own arrays which are also included in the active view
+			bGeometryCollectionDirty = true;
 		}
 		
 		UpdateViews();
