hopefully fix (again) the PVS studio 501 warnings in cluster union physics proxy

diff --git a/Engine/Source/Runtime/Experimental/Chaos/Private/PhysicsProxy/ClusterUnionPhysicsProxy.cpp b/Engine/Source/Runtime/Experimental/Chaos/Private/PhysicsProxy/ClusterUnionPhysicsProxy.cpp
--- a/Engine/Source/Runtime/Experimental/Chaos/Private/PhysicsProxy/ClusterUnionPhysicsProxy.cpp
+++ b/Engine/Source/Runtime/Experimental/Chaos/Private/PhysicsProxy/ClusterUnionPhysicsProxy.cpp
@@ -352,8 +352,9 @@ namespace Chaos
 		FDirtyChaosPropertyFlags DirtyFlags = Particle_External->DirtyFlags();
 
 		// We need to suppress V501 in PVS Studio since it's a false positive warning in this case. This is actually necessary for the codegen here.
-#define CHAOS_PROPERTY(PropName, Type, ProxyType) if constexpr (ProxyType != EPhysicsProxyType::ClusterUnionProxy) { DirtyFlags.MarkClean(EChaosPropertyFlags::PropName); } //-V501
-#include "Chaos/ParticleProperties.inl"
+//-V:CHAOS_PROPERTY:501
+#define CHAOS_PROPERTY(PropName, Type, ProxyType) if constexpr (ProxyType != EPhysicsProxyType::ClusterUnionProxy) { DirtyFlags.MarkClean(EChaosPropertyFlags::PropName); }
+#include "Chaos/ParticleProperties.inl" //-V:CHAOS_PROPERTY:501
 #undef CHAOS_PROPERTY
 
 		RemoteData.SetFlags(DirtyFlags);
