Fix Iris warnings that show up in ReplicationSystemTest.

diff --git a/Engine/Source/Runtime/Experimental/Iris/Core/Private/Iris/ReplicationState/ReplicationStateDescriptorBuilder.cpp b/Engine/Source/Runtime/Experimental/Iris/Core/Private/Iris/ReplicationState/ReplicationStateDescriptorBuilder.cpp
--- a/Engine/Source/Runtime/Experimental/Iris/Core/Private/Iris/ReplicationState/ReplicationStateDescriptorBuilder.cpp
+++ b/Engine/Source/Runtime/Experimental/Iris/Core/Private/Iris/ReplicationState/ReplicationStateDescriptorBuilder.cpp
@@ -2307,7 +2307,7 @@ TRefCountPtr<const FReplicationStateDescriptor> FReplicationStateDescriptorBuild
 		if (bAllMembersAreReplicated && PropertyCount == 0U)
 		{
 			// Warn if someone tries to replicate a struct with no properties for which we do not have a custom NetSerialzier.
-			UE_LOG_DESCRIPTORBUILDER_WARNING(TEXT("FReplicationStateDescriptorBuilder ::CreateDescriptorForStruct a replicated struct WITH no replicated properties must have a CustomNetSerializer. Struct: %s."), ToCStr(InStruct->GetName()));
+			UE_CLOG(bWarnAboutStructPropertiesWithSuspectedNotReplicatedProperties, LogIris, Warning, TEXT("FReplicationStateDescriptorBuilder ::CreateDescriptorForStruct a replicated struct WITH no replicated properties must have a CustomNetSerializer. Struct: %s."), ToCStr(InStruct->GetName()));
 			bAllMembersAreReplicated = false;
 		}
 	}
diff --git a/Engine/Source/Runtime/Experimental/Iris/Core/Private/Iris/ReplicationSystem/ReplicationReader.cpp b/Engine/Source/Runtime/Experimental/Iris/Core/Private/Iris/ReplicationSystem/ReplicationReader.cpp
--- a/Engine/Source/Runtime/Experimental/Iris/Core/Private/Iris/ReplicationSystem/ReplicationReader.cpp
+++ b/Engine/Source/Runtime/Experimental/Iris/Core/Private/Iris/ReplicationSystem/ReplicationReader.cpp
@@ -127,7 +127,7 @@ FReplicationReader::FReplicationReader()
 , ObjectsToDispatchCapacity(0U)
 , NetBlobHandlerManager(nullptr)
 , NetObjectBlobType(InvalidNetBlobType)
-, DelayAttachmentsWithUnresolvedReferences(IConsoleManager::Get().FindConsoleVariable(TEXT("net.DelayUnmappedRPCs")))
+, DelayAttachmentsWithUnresolvedReferences(IConsoleManager::Get().FindConsoleVariable(TEXT("net.DelayUnmappedRPCs"), false /* bTrackFrequentCalls */))
 {
 }
 
@@ -235,8 +235,8 @@ uint16 FReplicationReader::ReadObjectsPendingDestroy(FNetSerializationContext& C
 					continue;
 				}
 
-				// If we did not find the object or associated bridge something is wrong and we should disconnect.
-				UE_LOG_REPLICATIONREADER_WARNING(TEXT("FReplicationReader::Read Tried to destroy object with %s (This can occur if the server sends destroy for an object that has not yet been confirmed as created)"), *IncompleteHandle.ToString());
+				// If we did not find the object or associated bridge, the packet that would have created the object may have been lost.
+				UE_LOG(LogIris, Log, TEXT("FReplicationReader::Read Tried to destroy object with %s (This can occur if the server sends destroy for an object that has not yet been confirmed as created)"), *IncompleteHandle.ToString());
 			}
 		}
 	}
diff --git a/Engine/Source/Runtime/Experimental/Iris/Core/Private/Iris/ReplicationSystem/ReplicationWriter.cpp b/Engine/Source/Runtime/Experimental/Iris/Core/Private/Iris/ReplicationSystem/ReplicationWriter.cpp
--- a/Engine/Source/Runtime/Experimental/Iris/Core/Private/Iris/ReplicationSystem/ReplicationWriter.cpp
+++ b/Engine/Source/Runtime/Experimental/Iris/Core/Private/Iris/ReplicationSystem/ReplicationWriter.cpp
@@ -1604,7 +1604,7 @@ uint32 FReplicationWriter::WriteObjectsPendingDestroy(FNetSerializationContext&
 		// We do not support destroying an object that is currently being sent as a huge object.
 		if (IsObjectPartOfActiveHugeObject(InternalIndex, Info))
 		{
-			UE_LOG(LogIris, Warning, TEXT("Skipping writing destroy for object ( InternalIndex: %u ) which is part of active huge object."), InternalIndex);
+			UE_LOG(LogIris, Log, TEXT("Skipping writing destroy for object ( InternalIndex: %u ) which is part of active huge object."), InternalIndex);
 			bWroteAllDestroyedObjects = false;
 			continue;
 		}
