Chaos - fixed typo in UPROPERTY comments

diff --git a/Engine/Source/Runtime/Experimental/Chaos/Public/ChaosSolverConfiguration.h b/Engine/Source/Runtime/Experimental/Chaos/Public/ChaosSolverConfiguration.h
--- a/Engine/Source/Runtime/Experimental/Chaos/Public/ChaosSolverConfiguration.h
+++ b/Engine/Source/Runtime/Experimental/Chaos/Public/ChaosSolverConfiguration.h
@@ -62,7 +62,7 @@ struct FChaosSolverConfiguration
 	// happen because they spawned on top of each other, or the solver failed to fully reolve collisions last frame. A value of
 	// zero means "no limit". A non-zero value can be used to prevent explosive behaviour when bodies start deeply penetrating. 
 	// An alternative to using this approach is to increase the number of Velocity Iterations, which is more expensive but will 
-	// ensure the bdoies are depenetrated in a single frame without explosive behaviour.
+	// ensure the bodies are depenetrated in a single frame without explosive behaviour.
 	UPROPERTY(EditAnywhere, Category = "SolverConfiguration|Collision", meta = (ClampMin = "0.0"))
 	float CollisionMaxPushOutVelocity;
 
