Chaos - Add sentinels to TCArray (disabled)
- Add duplicate particle sanity check to Broadphase (disabled)

diff --git a/Engine/Source/Runtime/Experimental/Chaos/Public/Chaos/Collision/SpatialAccelerationBroadPhase.h b/Engine/Source/Runtime/Experimental/Chaos/Public/Chaos/Collision/SpatialAccelerationBroadPhase.h
--- a/Engine/Source/Runtime/Experimental/Chaos/Public/Chaos/Collision/SpatialAccelerationBroadPhase.h
+++ b/Engine/Source/Runtime/Experimental/Chaos/Public/Chaos/Collision/SpatialAccelerationBroadPhase.h
@@ -15,6 +15,11 @@
 #include "Chaos/GeometryParticlesfwd.h"
 #include "Chaos/AABBTree.h"
 
+// Set to 1 to enable some sanity chacks on the filtered broadphase results
+#ifndef CHAOS_CHECK_BROADPHASE
+#define CHAOS_CHECK_BROADPHASE 0
+#endif
+
 namespace Chaos
 {
 	template <typename TPayloadType, typename T, int d>
@@ -339,6 +344,10 @@ namespace Chaos
 				// Merge all the midphases from each worker into the primary allocator
 				Allocator->ProcessNewMidPhases();
 			}
+
+			// Run some error checks in non-shipping builds
+			// NOTE: This must come after MidPhase assignment for now because that's where the filter is applied
+			CheckOverlapResults();
 		}
 
 		/**
@@ -589,6 +598,36 @@ namespace Chaos
 			PHYSICS_CSV_CUSTOM_EXPENSIVE(PhysicsCounters, NumFromBroadphase, NumPotentials, ECsvCustomStatOp::Accumulate);
 		}
 
+		void CheckOverlapResults()
+		{
+#if CHAOS_CHECK_BROADPHASE
+			// Make sure that no particle pair is in the overlap lists twice which will cause a race condition
+			// in the narrow phase as multiple threads will attempt to build the same constraint(s) leading to
+			// buffer overruns and who knows what other craziness.
+			TSet<uint64> ParticlePairKeys;
+			for (int32 ContextIndex = 0; ContextIndex < NumActiveBroadphaseContexts; ++ContextIndex)
+			{
+				for (const Private::FBroadPhaseOverlap& Overlap : BroadphaseContexts[ContextIndex].Overlaps)
+				{
+					if (Overlap.bCollisionsEnabled)
+					{
+						const FCollisionParticlePairKey PairKey = FCollisionParticlePairKey(Overlap.Particles[0], Overlap.Particles[1]);
+						const bool bIsInSet = ParticlePairKeys.Contains(PairKey.GetKey());
+						if (ensure(!bIsInSet))
+						{
+							ParticlePairKeys.Add(PairKey.GetKey());
+						}
+						else
+						{
+							// Last time this happened it was because a particle was in multiple views that should be mutually exclusive
+							UE_LOG(LogChaos, Fatal, TEXT("Broadphase overlaps generated the same particle pair twice"));
+						}
+					}
+				}
+			}
+#endif
+		}
+
 		const FPBDRigidsSOAs& Particles;
 		const FAccelerationStructure* SpatialAcceleration;
 		TArray<Private::FBroadPhaseContext> BroadphaseContexts;
diff --git a/Engine/Source/Runtime/Experimental/Chaos/Public/Chaos/Framework/UncheckedArray.h b/Engine/Source/Runtime/Experimental/Chaos/Public/Chaos/Framework/UncheckedArray.h
--- a/Engine/Source/Runtime/Experimental/Chaos/Public/Chaos/Framework/UncheckedArray.h
+++ b/Engine/Source/Runtime/Experimental/Chaos/Public/Chaos/Framework/UncheckedArray.h
@@ -15,6 +15,22 @@
 #endif
 #endif
 
+// @macro CHAOS_CARRAY_SENTINEL: Set to 1 to add some sentinels into TCArray to trap buffer overruns
+#ifndef CHAOS_CARRAY_SENTINEL
+#define CHAOS_CARRAY_SENTINEL 0
+#endif
+
+// @macro IF_CHAOS_CARRAY_SENTINEL(X): Include some code block X only when sentinels are enabled
+// This is intended for one-line statements with no commas in them and is just shorthand for 
+// #if CHAOS_CARRAY_SENTINEL
+// X;
+// #endif
+#if CHAOS_CARRAY_SENTINEL
+#define IF_CHAOS_CARRAY_SENTINEL(X) X
+#else
+#define IF_CHAOS_CARRAY_SENTINEL(X)
+#endif
+
 namespace Chaos
 {
 	/**
@@ -43,22 +59,51 @@ namespace Chaos
 			return TCArray<T, N>(0);
 		}
 
-		inline TCArray() : NumElements(0) {}
+		inline TCArray() : NumElements(0)
+		{
+		}
 
-		inline int32 Num() const { return NumElements; }
-		inline int32 Max() const { return MaxElements; }
-		inline bool IsEmpty() const { return NumElements == 0; }
-		inline bool IsFull() const { return NumElements == MaxElements; }
+		inline int32 Num() const
+		{ 
+			CheckSentinels();
+
+			return NumElements;
+		}
+
+		inline int32 Max() const
+		{ 
+			CheckSentinels();
+		
+			return MaxElements; 
+		}
+		
+		inline bool IsEmpty() const
+		{ 
+			CheckSentinels();
+
+			return NumElements == 0;
+		}
+		
+		inline bool IsFull() const
+		{
+			CheckSentinels();
+
+			return NumElements == MaxElements;
+		}
 
 		inline ElementType& operator[](const int32 Index)
 		{
 			check(Index < NumElements);
+			CheckSentinels();
+
 			return Elements[Index];
 		}
 
 		inline const ElementType& operator[](const int32 Index) const
 		{
 			check(Index < NumElements);
+			CheckSentinels();
+
 			return Elements[Index];
 		}
 
@@ -69,6 +114,8 @@ namespace Chaos
 		inline void SetNum(const int32 InNum)
 		{
 			check(InNum <= MaxElements);
+			CheckSentinels();
+
 			NumElements = InNum;
 		}
 
@@ -78,6 +125,8 @@ namespace Chaos
 		*/
 		inline void Reset()
 		{
+			CheckSentinels();
+
 			NumElements = 0;
 		}
 
@@ -87,6 +136,8 @@ namespace Chaos
 		*/
 		inline void Empty()
 		{
+			CheckSentinels();
+
 			NumElements = 0;
 		}
 
@@ -96,6 +147,8 @@ namespace Chaos
 		inline int32 Add()
 		{
 			check(NumElements < MaxElements);
+			CheckSentinels();
+
 			return NumElements++;
 		}
 
@@ -105,7 +158,10 @@ namespace Chaos
 		inline int32 Add(const ElementType& V)
 		{
 			check(NumElements < MaxElements);
+
 			Elements[NumElements] = V;
+
+			CheckSentinels();
 			return NumElements++;
 		}
 
@@ -115,7 +171,10 @@ namespace Chaos
 		inline int32 Add(ElementType&& V)
 		{
 			check(NumElements < MaxElements);
+
 			Elements[NumElements] = MoveTemp(V);
+
+			CheckSentinels();
 			return NumElements++;
 		}
 
@@ -124,7 +183,11 @@ namespace Chaos
 		*/
 		inline int32 Emplace(ElementType&& V)
 		{
-			return Add(MoveTemp(V));
+			const int32 Index = Add(MoveTemp(V));
+
+			CheckSentinels();
+
+			return Index;
 		}
 
 		/**
@@ -134,11 +197,15 @@ namespace Chaos
 		inline void RemoveAt(const int32 Index)
 		{
 			check(Index < NumElements);
+			CheckSentinels();
+
 			for (int32 MoveIndex = Index; MoveIndex < NumElements - 1; ++MoveIndex)
 			{
 				Elements[MoveIndex] = MoveTemp(Elements[MoveIndex + 1]);
 			}
 			--NumElements;
+
+			CheckSentinels();
 		}
 
 		/**
@@ -148,40 +215,56 @@ namespace Chaos
 		inline void RemoveAtSwap(const int32 Index)
 		{
 			check(Index < NumElements);
+			CheckSentinels();
+
 			if (Index < NumElements - 1)
 			{
 				Elements[Index] = MoveTemp(Elements[NumElements - 1]);
 			}
 			--NumElements;
+
+			CheckSentinels();
 		}
 
 		inline ElementType* GetData()
 		{
+			CheckSentinels();
+
 			return &Elements[0];
 		}
 
 		inline const ElementType* GetData() const
 		{
+			CheckSentinels();
+
 			return &Elements[0];
 		}
 
 		inline ElementType* begin()
 		{
+			CheckSentinels();
+
 			return &Elements[0];
 		}
 
 		inline const ElementType* begin() const
 		{
+			CheckSentinels();
+
 			return &Elements[0];
 		}
 
 		inline ElementType* end()
 		{
+			CheckSentinels();
+
 			return &Elements[NumElements];
 		}
 
 		inline const ElementType* end() const
 		{
+			CheckSentinels();
+
 			return &Elements[NumElements];
 		}
 
@@ -192,9 +275,28 @@ namespace Chaos
 			check(Num() <= Max());
 		}
 
-		// Element count before array for better cache behaviour
+		static constexpr int32 SentinelValue = 0xA1B2C3D4;
+
+		void CheckSentinels() const
+		{
+#if CHAOS_CARRAY_SENTINEL
+			if (!ensureAlways(Sentinel0 == SentinelValue) || !ensureAlways(Sentinel1 == SentinelValue) || !ensureAlways(Sentinel2 == SentinelValue))
+			{
+				UE_LOG(LogChaos, Fatal, TEXT("TCArray Sentinel(s) Corrupted: 0x%08x, 0x%08x, 0x%08x [Expected 0x%08x]"), Sentinel0, Sentinel1, Sentinel2, SentinelValue);
+			}
+#endif
+		}
+
+		IF_CHAOS_CARRAY_SENTINEL(int32 Sentinel0 = SentinelValue);
+
+		// NOTE: Element count before array for better cache behaviour
 		int32 NumElements;
+
+		IF_CHAOS_CARRAY_SENTINEL(int32 Sentinel1 = SentinelValue);
+
 		ElementType Elements[MaxElements];
+
+		IF_CHAOS_CARRAY_SENTINEL(int32 Sentinel2 = SentinelValue);
 	};
 
 
