Chaos - sphere-mesh manifolds
- collision file cleanup

diff --git a/Engine/Source/Runtime/Experimental/Chaos/Private/Chaos/Collision/SphereTriangleContactPoint.cpp b/Engine/Source/Runtime/Experimental/Chaos/Private/Chaos/Collision/SphereTriangleContactPoint.cpp
new file mode 100644
--- /dev/null
+++ b/Engine/Source/Runtime/Experimental/Chaos/Private/Chaos/Collision/SphereTriangleContactPoint.cpp
@@ -0,0 +1,44 @@
+// Copyright Epic Games, Inc. All Rights Reserved.
+#include "Chaos/Collision/SphereTriangleContactPoint.h"
+#include "Chaos/Collision/ContactPoint.h"
+#include "Chaos/Collision/ContactPointsMiscShapes.h"
+#include "Chaos/Sphere.h"
+#include "Chaos/Triangle.h"
+#include "Chaos/Framework/UncheckedArray.h"
+
+//PRAGMA_DISABLE_OPTIMIZATION
+
+namespace Chaos
+{
+	void ConstructSphereTriangleOneShotManifold(const FImplicitSphere3& Sphere, const FTriangle& Triangle, const FReal CullDistance, FContactPointManifold& OutContactPoints)
+	{
+		// @todo(chaos): custom sphere-triangle contact point that avoids GJK
+
+		FVec3 PosSphere, PosTri, NormalTri;
+		FReal Phi;
+
+		// A point in the Minkowski Sum (A-B) in a likely direction
+		const FVec3 InitialV = Triangle.GetCentroid() - Sphere.GetCenter();
+
+		// Run GJK to find separating distance if available
+		// NOTE: Sphere is treated as a point (its core shape) for GJK to reduce iterations
+		// Also, Sphere and Triangle are in the same space
+		const EGJKDistanceResult GjkResult = GJKDistance(
+			MakeGJKShape(Triangle),
+			FGJKSphere(Sphere.GetCenter(), 0),
+			InitialV,
+			Phi, PosTri, PosSphere, NormalTri);
+
+		const FReal SphereRadius = Sphere.GetRadius();
+
+		// Build the contact point
+		if ((Phi - SphereRadius) < CullDistance)
+		{
+			const int32 Index = OutContactPoints.Emplace(FContactPoint());
+			OutContactPoints[Index].ShapeContactPoints[0] = PosSphere - SphereRadius * NormalTri;
+			OutContactPoints[Index].ShapeContactPoints[1] = PosTri;
+			OutContactPoints[Index].ShapeContactNormal = NormalTri;
+			OutContactPoints[Index].Phi = Phi - SphereRadius;
+		}
+	}
+}
diff --git a/Engine/Source/Runtime/Experimental/Chaos/Private/Chaos/CollisionOneShotManifoldsMeshShapes.cpp b/Engine/Source/Runtime/Experimental/Chaos/Private/Chaos/CollisionOneShotManifoldsMeshShapes.cpp
new file mode 100644
--- /dev/null
+++ b/Engine/Source/Runtime/Experimental/Chaos/Private/Chaos/CollisionOneShotManifoldsMeshShapes.cpp
@@ -0,0 +1,295 @@
+// Copyright Epic Games, Inc. All Rights Reserved.
+#include "Chaos/CollisionOneShotManifoldsMeshShapes.h"
+
+#include "Chaos/CollisionOneShotManifolds.h"
+#include "Chaos/CollisionResolution.h"
+#include "Chaos/Collision/CapsuleTriangleContactPoint.h"
+#include "Chaos/Collision/ContactPointsMiscShapes.h"
+#include "Chaos/Collision/ContactTriangles.h"
+#include "Chaos/Collision/ConvexTriangleContactPoint.h"
+#include "Chaos/Collision/PBDCollisionConstraint.h"
+#include "Chaos/Collision/SphereTriangleContactPoint.h"
+#include "Chaos/Convex.h"
+#include "Chaos/Defines.h"
+#include "Chaos/Framework/UncheckedArray.h"
+#include "Chaos/GJK.h"
+#include "Chaos/HeightField.h"
+#include "Chaos/ImplicitObjectScaled.h"
+#include "Chaos/Transform.h"
+#include "Chaos/Triangle.h"
+#include "Chaos/Utilities.h"
+#include "ChaosStats.h"
+
+//PRAGMA_DISABLE_OPTIMIZATION
+
+namespace Chaos
+{
+	extern bool bChaos_Collision_OneSidedTriangleMesh;
+	extern bool bChaos_Collision_OneSidedHeightField;
+
+	extern bool bChaos_Collision_UseCapsuleTriMesh2;
+	extern bool bChaos_Collision_UseConvexTriMesh2;
+
+	namespace Collisions
+	{
+
+		template <typename TriMeshType>
+		void ConstructSphereTriangleMeshOneShotManifold(const TSphere<FReal, 3>& Sphere, const FRigidTransform3& SphereWorldTransform, const TriMeshType& TriangleMesh, const FRigidTransform3& TriMeshWorldTransform, const FReal Dt, FPBDCollisionConstraint& Constraint)
+		{
+			// We only build one shot manifolds once
+			// All convexes are pre-scaled, or wrapped in TImplicitObjectScaled
+			ensure(SphereWorldTransform.GetScale3D() == FVec3(1.0f, 1.0f, 1.0f));
+			ensure(TriMeshWorldTransform.GetScale3D() == FVec3(1.0f, 1.0f, 1.0f));
+
+			// @todo(chaos): support manifold maintenance
+			Constraint.ResetActiveManifoldContacts();
+
+			FContactPoint ContactPoint = SphereTriangleMeshContactPoint(Sphere, SphereWorldTransform, TriangleMesh, TriMeshWorldTransform, Constraint.GetCullDistance());
+			if (ContactPoint.Phi < Constraint.GetCullDistance())
+			{
+				Constraint.AddOneshotManifoldContact(ContactPoint);
+			}
+		}
+
+		void ConstructSphereHeightFieldOneShotManifold(const TSphere<FReal, 3>& Sphere, const FRigidTransform3& SphereTransform, const FHeightField& Heightfield, const FRigidTransform3& HeightfieldTransform, const FReal Dt, FPBDCollisionConstraint& Constraint)
+		{
+			// We only build one shot manifolds once
+			// All convexes are pre-scaled, or wrapped in TImplicitObjectScaled
+			ensure(SphereTransform.GetScale3D() == FVec3(1.0f, 1.0f, 1.0f));
+			ensure(HeightfieldTransform.GetScale3D() == FVec3(1.0f, 1.0f, 1.0f));
+
+			FContactPoint ContactPoint = SphereHeightFieldContactPoint(Sphere, SphereTransform, Heightfield, HeightfieldTransform, Constraint.GetCullDistance());
+			if (ContactPoint.Phi < Constraint.GetCullDistance())
+			{
+				Constraint.AddOneshotManifoldContact(ContactPoint);
+			}
+		}
+
+		/**
+		 * @brief Generate a manifold between a convex shape and a single triangle
+		 * Templated so we can specialize for some shape types
+		*/
+		template<typename ConvexType>
+		void GenerateConvexTriangleOneShotManifold(const ConvexType& Convex, const FTriangle& Triangle, const FReal CullDistance, FContactPointManifold& OutContactPoints)
+		{
+			if (bChaos_Collision_UseConvexTriMesh2)
+			{
+				ConstructConvexTriangleOneShotManifold2(Convex, Triangle, CullDistance, OutContactPoints);
+			}
+			else
+			{
+				ConstructPlanarConvexTriangleOneShotManifold(Convex, Triangle, CullDistance, OutContactPoints);
+			}
+		}
+
+		template<>
+		void GenerateConvexTriangleOneShotManifold<FImplicitCapsule3>(const FImplicitCapsule3& Capsule, const FTriangle& Triangle, const FReal CullDistance, FContactPointManifold& OutContactPoints)
+		{
+			if (bChaos_Collision_UseCapsuleTriMesh2)
+			{
+				ConstructCapsuleTriangleOneShotManifold2(Capsule, Triangle, CullDistance, OutContactPoints);
+			}
+			else
+			{
+				ConstructCapsuleTriangleOneShotManifold(Capsule, Triangle, CullDistance, OutContactPoints);
+			}
+		}
+
+		template<>
+		void GenerateConvexTriangleOneShotManifold<FImplicitSphere3>(const FImplicitSphere3& Sphere, const FTriangle& Triangle, const FReal CullDistance, FContactPointManifold& OutContactPoints)
+		{
+			ConstructSphereTriangleOneShotManifold(Sphere, Triangle, CullDistance, OutContactPoints);
+		}
+
+		/**
+		* @brief Create a minimized set of contact points between a convex polyhedron (box, convex) and a non-convex mesh (trimesh, heightfield)
+		*
+		* @tparam ConvexType any convex type (Sphere, Capsule, Box, Convex, possibly with a scaled or instanced)
+		* @tparam MeshType any triangle mesh type (HeightField or TriangleMesh, without a scaled or instanced wrapper)
+		* @param MeshQueryBounds Triangles overlapping this box will be tested. Should be in the space of the mesh.
+		* @param MeshToConvexTransform The transform from Mesh space to Convex space. This low-level convex-triangle collision detection is performed in Convex space.
+		*
+		* @see ConstructPlanarConvexTriMeshOneShotManifoldImp, ConstructPlanarConvexHeightfieldOneShotManifoldImp
+		*/
+		template<typename ConvexType, typename MeshType>
+		void GenerateConvexMeshContactPoints(const ConvexType& Convex, const MeshType& Mesh, const FAABB3& MeshQueryBounds, const FRigidTransform3& MeshToConvexTransform, const FReal CullDistance, FContactTriangleCollector& MeshContacts)
+		{
+			FContactPointManifold TriangleManifoldPoints;
+
+			// Loop over all the triangles, build a manifold and add the points to the total manifold
+			// NOTE: contact points will be in the space of the convex until the end of the function when we convert into shape local space
+			Mesh.VisitTriangles(MeshQueryBounds, MeshToConvexTransform, [&](const FTriangle& Triangle, const int32 TriangleIndex, const int32 VertexIndex0, const int32 VertexIndex1, const int32 VertexIndex2)
+			{
+				// Generate the manifold for this triangle
+				TriangleManifoldPoints.Reset();
+				GenerateConvexTriangleOneShotManifold(Convex, Triangle, CullDistance, TriangleManifoldPoints);
+
+				if (TriangleManifoldPoints.Num() > 0)
+				{
+					// Add the points into the main contact array
+					// NOTE: The Contacts' FaceIndices will be an index into the ContactTriangles not the original tri mesh (this will get mapped back to the mesh index below)
+					MeshContacts.AddTriangleContacts(MakeArrayView(TriangleManifoldPoints.begin(), TriangleManifoldPoints.Num()), Triangle, TriangleIndex, VertexIndex0, VertexIndex1, VertexIndex2, CullDistance);
+				}
+			});
+
+			// Reduce contacts to a minimum manifold and transform contact data back into shape-local space
+			MeshContacts.ProcessContacts(MeshToConvexTransform);
+		}
+
+		/**
+		 * @brief Used by all the convex types to generate a manifold against any mesh type
+		*/
+		template<typename ConvexType, typename MeshType>
+		void ConstructConvexMeshOneShotManifold(const ConvexType& Convex, const FRigidTransform3& ConvexTransform, const MeshType& Mesh, const FRigidTransform3& MeshTransform, const FVec3& MeshScale, const FReal CullDistance, FContactTriangleCollector& MeshContacts)
+		{
+			FRigidTransform3 MeshToConvexTransform = MeshTransform.GetRelativeTransformNoScale(ConvexTransform);
+			MeshToConvexTransform.SetScale3D(MeshScale);
+
+			// @todo(chaos): add Convex.CalculateInverseTransformed bounds with scale support (to optimize sphere and capsule)
+			const FAABB3 MeshQueryBounds = Convex.BoundingBox().InverseTransformedAABB(MeshToConvexTransform).Thicken(CullDistance);
+
+			// Create the minimal manifold from all the overlapping triangles
+			GenerateConvexMeshContactPoints(Convex, Mesh, MeshQueryBounds, MeshToConvexTransform, CullDistance, MeshContacts);
+		}
+
+		void ConstructQuadraticConvexTriMeshOneShotManifold(const FImplicitObject& Quadratic, const FRigidTransform3& QuadraticTransform, const FImplicitObject& InMesh, const FRigidTransform3& MeshTransform, const FReal Dt, FPBDCollisionConstraint& Constraint)
+		{
+			// All convexes are pre-scaled, or wrapped in TImplicitObjectScaled
+			ensure(QuadraticTransform.GetScale3D() == FVec3(1));
+			ensure(MeshTransform.GetScale3D() == FVec3(1));
+
+			// Unwrap the tri mesh (remove Scaled or Instanced) and get the scale
+			FVec3 MeshScale;
+			FReal MeshMargin;	// Not used - will be zero for meshes
+			const FTriangleMeshImplicitObject* Mesh = UnwrapImplicit<FTriangleMeshImplicitObject>(InMesh, MeshScale, MeshMargin);
+			check(Mesh != nullptr);
+
+			FContactTriangleCollector MeshContacts(bChaos_Collision_OneSidedTriangleMesh, QuadraticTransform);
+			const FReal CullDistance = Constraint.GetCullDistance();
+
+			if (const FImplicitSphere3* Sphere = Quadratic.template GetObject<FImplicitSphere3>())
+			{
+				ConstructConvexMeshOneShotManifold(*Sphere, QuadraticTransform, *Mesh, MeshTransform, MeshScale, CullDistance, MeshContacts);
+			}
+			else if (const FImplicitCapsule3* Capsule = Quadratic.template GetObject<FImplicitCapsule3>())
+			{
+				ConstructConvexMeshOneShotManifold(*Capsule, QuadraticTransform, *Mesh, MeshTransform, MeshScale, CullDistance, MeshContacts);
+			}
+			else
+			{
+				check(false);
+			}
+
+			Constraint.SetOneShotManifoldContacts(MeshContacts.GetContactPoints());
+		}
+
+		void ConstructQuadraticConvexHeightFieldOneShotManifold(const FImplicitObject& Quadratic, const FRigidTransform3& QuadraticTransform, const FHeightField& Mesh, const FRigidTransform3& MeshTransform, const FReal Dt, FPBDCollisionConstraint& Constraint)
+		{
+			// We only build one shot manifolds once
+			// All convexes are pre-scaled, or wrapped in TImplicitObjectScaled
+			ensure(QuadraticTransform.GetScale3D() == FVec3(1));
+			ensure(MeshTransform.GetScale3D() == FVec3(1));
+
+			FContactTriangleCollector MeshContacts(bChaos_Collision_OneSidedHeightField, QuadraticTransform);
+			const FReal CullDistance = Constraint.GetCullDistance();
+
+			if (const FImplicitSphere3* Sphere = Quadratic.template GetObject<FImplicitSphere3>())
+			{
+				ConstructConvexMeshOneShotManifold(*Sphere, QuadraticTransform, Mesh, MeshTransform, FVec3(1), CullDistance, MeshContacts);
+			}
+			else if (const FImplicitCapsule3* Capsule = Quadratic.template GetObject<FImplicitCapsule3>())
+			{
+				ConstructConvexMeshOneShotManifold(*Capsule, QuadraticTransform, Mesh, MeshTransform, FVec3(1), CullDistance, MeshContacts);
+			}
+			else
+			{
+				check(false);
+			}
+
+			Constraint.SetOneShotManifoldContacts(MeshContacts.GetContactPoints());
+		}
+
+		/**
+		* @brief Populate the Constraint with a manifold of contacts between a Convex and a TriangleMesh
+		* @param Convex A convex polyhedron (Box, Convex) that may be wrapped in Scaled or Instanced
+		* @param InMesh A TriangleMesh ImplicitObject that may be wrapped in Scaled or Instaned
+		*/
+		void ConstructPlanarConvexTriMeshOneShotManifold(const FImplicitObject& Convex, const FRigidTransform3& ConvexTransform, const FImplicitObject& InMesh, const FRigidTransform3& MeshTransform, FPBDCollisionConstraint& Constraint)
+		{
+			// All convexes are pre-scaled, or wrapped in TImplicitObjectScaled
+			ensure(ConvexTransform.GetScale3D() == FVec3(1));
+			ensure(MeshTransform.GetScale3D() == FVec3(1));
+
+			// Unwrap the tri mesh (remove Scaled or Instanced) and get the scale
+			FVec3 MeshScale;
+			FReal MeshMargin;	// Not used - will be zero for meshes
+			const FTriangleMeshImplicitObject* Mesh = UnwrapImplicit<FTriangleMeshImplicitObject>(InMesh, MeshScale, MeshMargin);
+			check(Mesh != nullptr);
+
+			FContactTriangleCollector MeshContacts(bChaos_Collision_OneSidedTriangleMesh, ConvexTransform);
+			const FReal CullDistance = Constraint.GetCullDistance();
+
+			if (const FImplicitBox3* RawBox = Convex.template GetObject<FImplicitBox3>())
+			{
+				ConstructConvexMeshOneShotManifold(*RawBox, ConvexTransform, *Mesh, MeshTransform, MeshScale, CullDistance, MeshContacts);
+			}
+			else if (const TImplicitObjectScaled<FImplicitConvex3>* ScaledConvex = Convex.template GetObject<TImplicitObjectScaled<FImplicitConvex3>>())
+			{
+				ConstructConvexMeshOneShotManifold(*ScaledConvex, ConvexTransform, *Mesh, MeshTransform, MeshScale, CullDistance, MeshContacts);
+			}
+			else if (const TImplicitObjectInstanced<FImplicitConvex3>* InstancedConvex = Convex.template GetObject<TImplicitObjectInstanced<FImplicitConvex3>>())
+			{
+				ConstructConvexMeshOneShotManifold(*InstancedConvex, ConvexTransform, *Mesh, MeshTransform, MeshScale, CullDistance, MeshContacts);
+			}
+			else if (const FImplicitConvex3* RawConvex = Convex.template GetObject<FImplicitConvex3>())
+			{
+				ConstructConvexMeshOneShotManifold(*RawConvex, ConvexTransform, *Mesh, MeshTransform, MeshScale, CullDistance, MeshContacts);
+			}
+			else
+			{
+				check(false);
+			}
+
+			Constraint.SetOneShotManifoldContacts(MeshContacts.GetContactPoints());
+		}
+
+		/**
+		* @brief Populate the Constraint with a manifold of contacts between a Convex and a HeightField
+		* @param Convex A convex polyhedron (Box, Convex) that may be wrapped in Scaled or Instanced
+		*/
+		void ConstructPlanarConvexHeightFieldOneShotManifold(const FImplicitObject& Convex, const FRigidTransform3& ConvexTransform, const FHeightField& Mesh, const FRigidTransform3& MeshTransform, FPBDCollisionConstraint& Constraint)
+		{
+			// All convexes are pre-scaled, or wrapped in TImplicitObjectScaled
+			ensure(ConvexTransform.GetScale3D() == FVec3(1));
+			ensure(MeshTransform.GetScale3D() == FVec3(1));
+
+			FContactTriangleCollector MeshContacts(bChaos_Collision_OneSidedHeightField, ConvexTransform);
+			const FReal CullDistance = Constraint.GetCullDistance();
+			const FVec3 MeshScale = FVec3(1);	// Scale is built into heightfield
+
+			if (const FImplicitBox3* RawBox = Convex.template GetObject<FImplicitBox3>())
+			{
+				ConstructConvexMeshOneShotManifold(*RawBox, ConvexTransform, Mesh, MeshTransform, MeshScale, CullDistance, MeshContacts);
+			}
+			else if (const TImplicitObjectScaled<FImplicitConvex3>* ScaledConvex = Convex.template GetObject<TImplicitObjectScaled<FImplicitConvex3>>())
+			{
+				ConstructConvexMeshOneShotManifold(*ScaledConvex, ConvexTransform, Mesh, MeshTransform, MeshScale, CullDistance, MeshContacts);
+			}
+			else if (const TImplicitObjectInstanced<FImplicitConvex3>* InstancedConvex = Convex.template GetObject<TImplicitObjectInstanced<FImplicitConvex3>>())
+			{
+				ConstructConvexMeshOneShotManifold(*InstancedConvex, ConvexTransform, Mesh, MeshTransform, MeshScale, CullDistance, MeshContacts);
+			}
+			else if (const FImplicitConvex3* RawConvex = Convex.template GetObject<FImplicitConvex3>())
+			{
+				ConstructConvexMeshOneShotManifold(*RawConvex, ConvexTransform, Mesh, MeshTransform, MeshScale, CullDistance, MeshContacts);
+			}
+			else
+			{
+				check(false);
+			}
+
+			Constraint.SetOneShotManifoldContacts(MeshContacts.GetContactPoints());
+		}
+	}
+}
+
diff --git a/Engine/Source/Runtime/Experimental/Chaos/Private/Chaos/CollisionOneShotManifoldsMeshShapes.h b/Engine/Source/Runtime/Experimental/Chaos/Private/Chaos/CollisionOneShotManifoldsMeshShapes.h
new file mode 100644
--- /dev/null
+++ b/Engine/Source/Runtime/Experimental/Chaos/Private/Chaos/CollisionOneShotManifoldsMeshShapes.h
@@ -0,0 +1,65 @@
+// Copyright Epic Games, Inc. All Rights Reserved.
+
+#pragma once
+
+#include "Chaos/Core.h"
+#include "Chaos/CollisionResolutionTypes.h"
+#include "Chaos/ImplicitFwd.h"
+
+namespace Chaos
+{
+	class FHeightField;
+	class FPBDCollisionConstraint;
+
+	namespace Collisions
+	{
+		/**
+		 * @brief Build the contact manifold between a Quadratic Convex shape and a TriMesh
+		 * @param Convex an ImplicitObject that must be a sphere or capsule. Assert if not
+		 * @param TriMesh an ImplicitObject that must be FTriangleMeshImplicitObject or a wrapper around it. Asserts if not
+		*/
+		void ConstructQuadraticConvexTriMeshOneShotManifold(
+			const FImplicitObject& Quadratic,
+			const FRigidTransform3& Quadraticransform,
+			const FImplicitObject& TriangleMesh,
+			const FRigidTransform3& TriMeshTransform,
+			const FReal Dt,
+			FPBDCollisionConstraint& Constraint);
+
+		/**
+		 * @brief Build the contact manifold between a Quadratic Convex shape (Sphere, Capsule) and a TriMesh
+		 * @param Convex an ImplicitObject that must be a sphere or capsule. Assert if not
+		*/
+		void ConstructQuadraticConvexHeightFieldOneShotManifold(
+			const FImplicitObject& Quadratic,
+			const FRigidTransform3& QuadraticTransform,
+			const FHeightField& HeightField,
+			const FRigidTransform3& HeightFieldTransform,
+			const FReal Dt,
+			FPBDCollisionConstraint& Constraint);
+
+		/**
+		 * @brief Build the contact manifold between a Planar Convex shape and a TriMesh
+		 * @param Convex an ImplicitObject that must be a box, convex, or a wrapper around it. Asserts if not
+		 * @param TriMesh an ImplicitObject that must be FTriangleMeshImplicitObject or a wrapper around it. Asserts if not
+		*/
+		void ConstructPlanarConvexTriMeshOneShotManifold(
+			const FImplicitObject& Convex, 
+			const FRigidTransform3& ConvexTransform, 
+			const FImplicitObject& TriangleMesh, 
+			const FRigidTransform3& TriangleMeshTransform, 
+			FPBDCollisionConstraint& Constraint);
+
+		/**
+		 * @brief Build the contact manifold between a Planar Convex shape and a HeightField
+		 * @param Convex an ImplicitObject that must be a box, convex, or a wrapper around it. Asserts if not
+		*/
+		void ConstructPlanarConvexHeightFieldOneShotManifold(
+			const FImplicitObject& Convex, 
+			const FRigidTransform3& ConvexTransform, 
+			const FHeightField& HeightField, 
+			const FRigidTransform3& TriangleMeshTransform, 
+			FPBDCollisionConstraint& Constraint);
+
+	}
+}
\ No newline at end of file
diff --git a/Engine/Source/Runtime/Experimental/Chaos/Private/Chaos/CollisionOneShotManifoldsMiscShapes.cpp b/Engine/Source/Runtime/Experimental/Chaos/Private/Chaos/CollisionOneShotManifoldsMiscShapes.cpp
--- a/Engine/Source/Runtime/Experimental/Chaos/Private/Chaos/CollisionOneShotManifoldsMiscShapes.cpp
+++ b/Engine/Source/Runtime/Experimental/Chaos/Private/Chaos/CollisionOneShotManifoldsMiscShapes.cpp
@@ -5,45 +5,27 @@
 #include "Chaos/CollisionOneShotManifolds.h"
 #include "Chaos/CollisionResolution.h"
 #include "Chaos/Collision/CapsuleConvexContactPoint.h"
-#include "Chaos/Collision/CapsuleTriangleContactPoint.h"
 #include "Chaos/Collision/ContactPointsMiscShapes.h"
-#include "Chaos/Collision/ContactTriangles.h"
-#include "Chaos/Collision/ConvexTriangleContactPoint.h"
 #include "Chaos/Collision/SphereConvexContactPoint.h"
 #include "Chaos/Collision/PBDCollisionConstraint.h"
 #include "Chaos/Convex.h"
 #include "Chaos/Defines.h"
 #include "Chaos/Framework/UncheckedArray.h"
 #include "Chaos/GJK.h"
-#include "Chaos/HeightField.h"
 #include "Chaos/ImplicitObjectScaled.h"
 #include "Chaos/Transform.h"
-#include "Chaos/Triangle.h"
 #include "Chaos/Utilities.h"
 #include "ChaosStats.h"
 
-#include "HAL/IConsoleManager.h"
-
 //PRAGMA_DISABLE_OPTIMIZATION
 
 namespace Chaos
 {
-	extern FRealSingle Chaos_Collision_EdgePrunePlaneDistance;
 	extern FRealSingle Chaos_Collision_Manifold_SphereCapsuleSizeThreshold;
 	extern FRealSingle Chaos_Collision_Manifold_CapsuleAxisAlignedThreshold;
 	extern FRealSingle Chaos_Collision_Manifold_CapsuleDeepPenetrationFraction;
 	extern FRealSingle Chaos_Collision_Manifold_CapsuleRadialContactFraction;
 
-	extern bool bChaos_Collision_UseGJK2;
-	extern FRealSingle Chaos_Collision_GJKEpsilon;
-	extern FRealSingle Chaos_Collision_EPAEpsilon;
-
-	extern bool bChaos_Collision_OneSidedTriangleMesh;
-	extern bool bChaos_Collision_OneSidedHeightField;
-
-	extern bool bChaos_Collision_UseCapsuleTriMesh2;
-	extern bool bChaos_Collision_UseConvexTriMesh2;
-
 	namespace Collisions
 	{
 
@@ -304,7 +286,6 @@ namespace Chaos
 			}
 		}
 
-
 		void ConstructSphereConvexManifold(const TSphere<FReal, 3>& Sphere, const FRigidTransform3& SphereTransform, const FImplicitObject3& Convex, const FRigidTransform3& ConvexTransform, const FReal Dt, FPBDCollisionConstraint& Constraint)
 		{
 			// We only build one shot manifolds once
@@ -344,38 +325,6 @@ namespace Chaos
 			}
 		}
 
-		template <typename TriMeshType>
-		void ConstructSphereTriangleMeshOneShotManifold(const TSphere<FReal, 3>& Sphere, const FRigidTransform3& SphereWorldTransform, const TriMeshType& TriangleMesh, const FRigidTransform3& TriMeshWorldTransform, const FReal Dt, FPBDCollisionConstraint& Constraint)
-		{
-			// We only build one shot manifolds once
-			// All convexes are pre-scaled, or wrapped in TImplicitObjectScaled
-			ensure(SphereWorldTransform.GetScale3D() == FVec3(1.0f, 1.0f, 1.0f));
-			ensure(TriMeshWorldTransform.GetScale3D() == FVec3(1.0f, 1.0f, 1.0f));
-
-			// @todo(chaos): support manifold maintenance
-			Constraint.ResetActiveManifoldContacts();
-
-			FContactPoint ContactPoint = SphereTriangleMeshContactPoint(Sphere, SphereWorldTransform, TriangleMesh, TriMeshWorldTransform, Constraint.GetCullDistance());
-			if (ContactPoint.Phi < Constraint.GetCullDistance())
-			{
-				Constraint.AddOneshotManifoldContact(ContactPoint);
-			}
-		}
-
-		void ConstructSphereHeightFieldOneShotManifold(const TSphere<FReal, 3>& Sphere, const FRigidTransform3& SphereTransform, const FHeightField& Heightfield, const FRigidTransform3& HeightfieldTransform, const FReal Dt, FPBDCollisionConstraint& Constraint)
-		{
-			// We only build one shot manifolds once
-			// All convexes are pre-scaled, or wrapped in TImplicitObjectScaled
-			ensure(SphereTransform.GetScale3D() == FVec3(1.0f, 1.0f, 1.0f));
-			ensure(HeightfieldTransform.GetScale3D() == FVec3(1.0f, 1.0f, 1.0f));
-
-			FContactPoint ContactPoint = SphereHeightFieldContactPoint(Sphere, SphereTransform, Heightfield, HeightfieldTransform, Constraint.GetCullDistance());
-			if (ContactPoint.Phi < Constraint.GetCullDistance())
-			{
-				Constraint.AddOneshotManifoldContact(ContactPoint);
-			}
-		}
-
 		// @todo(chaos): this will be faster if we transform into the space of one of the capsules
 		void ConstructCapsuleCapsuleOneShotManifold(const FCapsule& CapsuleA, const FRigidTransform3& CapsuleATransform, const FCapsule& CapsuleB, const FRigidTransform3& CapsuleBTransform, const FReal Dt, FPBDCollisionConstraint& Constraint)
 		{
@@ -616,207 +565,6 @@ namespace Chaos
 			}
 		}
 
-		/**
-		 * @brief Generate a manifold between a convex shape and a single triangle
-		 * Templated so we can specialize for some shape types
-		*/
-		template<typename ConvexType>
-		void GenerateConvexTriangleOneShotManifold(const ConvexType& Convex, const FTriangle& Triangle, const FReal CullDistance, FContactPointManifold& OutContactPoints)
-		{
-			if (bChaos_Collision_UseConvexTriMesh2)
-			{
-				ConstructConvexTriangleOneShotManifold2(Convex, Triangle, CullDistance, OutContactPoints);
-			}
-			else
-			{
-				ConstructPlanarConvexTriangleOneShotManifold(Convex, Triangle, CullDistance, OutContactPoints);
-			}
-		}
-
-		template<>
-		void GenerateConvexTriangleOneShotManifold<FCapsule>(const FCapsule& Capsule, const FTriangle& Triangle, const FReal CullDistance, FContactPointManifold& OutContactPoints)
-		{
-			if (bChaos_Collision_UseCapsuleTriMesh2)
-			{
-				ConstructCapsuleTriangleOneShotManifold2(Capsule, Triangle, CullDistance, OutContactPoints);
-			}
-			else
-			{
-				ConstructCapsuleTriangleOneShotManifold(Capsule, Triangle, CullDistance, OutContactPoints);
-			}
-		}
-
-		/**
-		* @brief Create a minimized set of contact points between a convex polyhedron (box, convex) and a non-convex mesh (trimesh, heightfield)
-		*
-		* @tparam ConvexType any convex type (Sphere, Capsule, Box, Convex, possibly with a scaled or instanced)
-		* @tparam MeshType any triangle mesh type (HeightField or TriangleMesh, without a scaled or instanced wrapper)
-		* @param MeshQueryBounds Triangles overlapping this box will be tested. Should be in the space of the mesh.
-		* @param MeshToConvexTransform The transform from Mesh space to Convex space. This low-level convex-triangle collision detection is performed in Convex space.
-		*
-		* @see ConstructPlanarConvexTriMeshOneShotManifoldImp, ConstructPlanarConvexHeightfieldOneShotManifoldImp
-		*/
-		template<typename ConvexType, typename MeshType>
-		void GenerateConvexMeshContactPoints(const ConvexType& Convex, const MeshType& Mesh, const FAABB3& MeshQueryBounds, const FRigidTransform3& MeshToConvexTransform, const FReal CullDistance, FContactTriangleCollector& MeshContacts)
-		{
-			FContactPointManifold TriangleManifoldPoints;
-
-			// Loop over all the triangles, build a manifold and add the points to the total manifold
-			// NOTE: contact points will be in the space of the convex until the end of the function when we convert into shape local space
-			Mesh.VisitTriangles(MeshQueryBounds, MeshToConvexTransform, [&](const FTriangle& Triangle, const int32 TriangleIndex, const int32 VertexIndex0, const int32 VertexIndex1, const int32 VertexIndex2)
-			{
-				// Generate the manifold for this triangle
-				TriangleManifoldPoints.Reset();
-				GenerateConvexTriangleOneShotManifold(Convex, Triangle, CullDistance, TriangleManifoldPoints);
-
-				if (TriangleManifoldPoints.Num() > 0)
-				{
-					// Add the points into the main contact array
-					// NOTE: The Contacts' FaceIndices will be an index into the ContactTriangles not the original tri mesh (this will get mapped back to the mesh index below)
-					MeshContacts.AddTriangleContacts(MakeArrayView(TriangleManifoldPoints.begin(), TriangleManifoldPoints.Num()), Triangle, TriangleIndex, VertexIndex0, VertexIndex1, VertexIndex2, CullDistance);
-				}
-			});
-
-			// Reduce contacts to a minimum manifold and transform contact data back into shape-local space
-			MeshContacts.ProcessContacts(MeshToConvexTransform);
-		}
-
-		/**
-		 * @brief Used by all the convex types to generate a manifold against any mesh type
-		*/
-		template<typename ConvexType, typename MeshType>
-		void ConstructConvexMeshOneShotManifold(const ConvexType& Convex, const FRigidTransform3& ConvexTransform, const MeshType& Mesh, const FRigidTransform3& MeshTransform, const FVec3& MeshScale, const FReal CullDistance, FContactTriangleCollector& MeshContacts)
-		{
-			FRigidTransform3 MeshToConvexTransform = MeshTransform.GetRelativeTransformNoScale(ConvexTransform);
-			MeshToConvexTransform.SetScale3D(MeshScale);
-
-			// @todo(chaos): add Convex.CalculateInverseTransformed bounds with scale support (to optimize sphere and capsule)
-			const FAABB3 MeshQueryBounds = Convex.BoundingBox().InverseTransformedAABB(MeshToConvexTransform).Thicken(CullDistance);
-
-			// Create the minimal manifold from all the overlapping triangles
-			GenerateConvexMeshContactPoints(Convex, Mesh, MeshQueryBounds, MeshToConvexTransform, CullDistance, MeshContacts);
-		}
-
-		void ConstructCapsuleTriMeshOneShotManifold(const FCapsule& Capsule, const FRigidTransform3& CapsuleTransform, const FImplicitObject& InMesh, const FRigidTransform3& MeshTransform, const FReal Dt, FPBDCollisionConstraint& Constraint)
-		{
-			// All convexes are pre-scaled, or wrapped in TImplicitObjectScaled
-			ensure(CapsuleTransform.GetScale3D() == FVec3(1));
-			ensure(MeshTransform.GetScale3D() == FVec3(1));
-
-			// Unwrap the tri mesh (remove Scaled or Instanced) and get the scale
-			FVec3 MeshScale;
-			FReal MeshMargin;	// Not used - will be zero for meshes
-			const FTriangleMeshImplicitObject* Mesh = UnwrapImplicit<FTriangleMeshImplicitObject>(InMesh, MeshScale, MeshMargin);
-			check(Mesh != nullptr);
-
-			FContactTriangleCollector MeshContacts(bChaos_Collision_OneSidedTriangleMesh, CapsuleTransform);
-			const FReal CullDistance = Constraint.GetCullDistance();
-
-			ConstructConvexMeshOneShotManifold(Capsule, CapsuleTransform, *Mesh, MeshTransform, MeshScale, CullDistance, MeshContacts);
-
-			Constraint.SetOneShotManifoldContacts(MeshContacts.GetContactPoints());
-		}
-
-		void ConstructCapsuleHeightFieldOneShotManifold(const FCapsule& Capsule, const FRigidTransform3& CapsuleTransform, const FHeightField& Mesh, const FRigidTransform3& MeshTransform, const FReal Dt, FPBDCollisionConstraint& Constraint)
-		{
-			// We only build one shot manifolds once
-			// All convexes are pre-scaled, or wrapped in TImplicitObjectScaled
-			ensure(CapsuleTransform.GetScale3D() == FVec3(1));
-			ensure(MeshTransform.GetScale3D() == FVec3(1));
-
-			FContactTriangleCollector MeshContacts(bChaos_Collision_OneSidedHeightField, CapsuleTransform);
-			const FReal CullDistance = Constraint.GetCullDistance();
-
-			ConstructConvexMeshOneShotManifold(Capsule, CapsuleTransform, Mesh, MeshTransform, FVec3(1), CullDistance, MeshContacts);
-
-			Constraint.SetOneShotManifoldContacts(MeshContacts.GetContactPoints());
-		}
-
-		/**
-		* @brief Populate the Constraint with a manifold of contacts between a Convex and a TriangleMesh
-		* @param Convex A convex polyhedron (Box, Convex) that may be wrapped in Scaled or Instanced
-		* @param InMesh A TriangleMesh ImplicitObject that may be wrapped in Scaled or Instaned
-		*/
-		void ConstructPlanarConvexTriMeshOneShotManifold(const FImplicitObject& Convex, const FRigidTransform3& ConvexTransform, const FImplicitObject& InMesh, const FRigidTransform3& MeshTransform, FPBDCollisionConstraint& Constraint)
-		{
-			// All convexes are pre-scaled, or wrapped in TImplicitObjectScaled
-			ensure(ConvexTransform.GetScale3D() == FVec3(1));
-			ensure(MeshTransform.GetScale3D() == FVec3(1));
-
-			// Unwrap the tri mesh (remove Scaled or Instanced) and get the scale
-			FVec3 MeshScale;
-			FReal MeshMargin;	// Not used - will be zero for meshes
-			const FTriangleMeshImplicitObject* Mesh = UnwrapImplicit<FTriangleMeshImplicitObject>(InMesh, MeshScale, MeshMargin);
-			check(Mesh != nullptr);
-
-			FContactTriangleCollector MeshContacts(bChaos_Collision_OneSidedTriangleMesh, ConvexTransform);
-			const FReal CullDistance = Constraint.GetCullDistance();
-
-			if (const FImplicitBox3* RawBox = Convex.template GetObject<FImplicitBox3>())
-			{
-				ConstructConvexMeshOneShotManifold(*RawBox, ConvexTransform, *Mesh, MeshTransform, MeshScale, CullDistance, MeshContacts);
-			}
-			else if (const TImplicitObjectScaled<FImplicitConvex3>* ScaledConvex = Convex.template GetObject<TImplicitObjectScaled<FImplicitConvex3>>())
-			{
-				ConstructConvexMeshOneShotManifold(*ScaledConvex, ConvexTransform, *Mesh, MeshTransform, MeshScale, CullDistance, MeshContacts);
-			}
-			else if (const TImplicitObjectInstanced<FImplicitConvex3>* InstancedConvex = Convex.template GetObject<TImplicitObjectInstanced<FImplicitConvex3>>())
-			{
-				ConstructConvexMeshOneShotManifold(*InstancedConvex, ConvexTransform, *Mesh, MeshTransform, MeshScale, CullDistance, MeshContacts);
-			}
-			else if (const FImplicitConvex3* RawConvex = Convex.template GetObject<FImplicitConvex3>())
-			{
-				ConstructConvexMeshOneShotManifold(*RawConvex, ConvexTransform, *Mesh, MeshTransform, MeshScale, CullDistance, MeshContacts);
-			}
-			else
-			{
-				check(false);
-			}
-
-			Constraint.SetOneShotManifoldContacts(MeshContacts.GetContactPoints());
-		}
-
-		/**
-		* @brief Populate the Constraint with a manifold of contacts between a Convex and a HeightField
-		* @param Convex A convex polyhedron (Box, Convex) that may be wrapped in Scaled or Instanced
-		*/
-		void ConstructPlanarConvexHeightFieldOneShotManifold(const FImplicitObject& Convex, const FRigidTransform3& ConvexTransform, const FHeightField& Mesh, const FRigidTransform3& MeshTransform, FPBDCollisionConstraint& Constraint)
-		{
-			// All convexes are pre-scaled, or wrapped in TImplicitObjectScaled
-			ensure(ConvexTransform.GetScale3D() == FVec3(1));
-			ensure(MeshTransform.GetScale3D() == FVec3(1));
-
-			FContactTriangleCollector MeshContacts(bChaos_Collision_OneSidedHeightField, ConvexTransform);
-			const FReal CullDistance = Constraint.GetCullDistance();
-			const FVec3 MeshScale = FVec3(1);	// Scale is built into heightfield
-
-			if (const FImplicitBox3* RawBox = Convex.template GetObject<FImplicitBox3>())
-			{
-				ConstructConvexMeshOneShotManifold(*RawBox, ConvexTransform, Mesh, MeshTransform, MeshScale, CullDistance, MeshContacts);
-			}
-			else if (const TImplicitObjectScaled<FImplicitConvex3>* ScaledConvex = Convex.template GetObject<TImplicitObjectScaled<FImplicitConvex3>>())
-			{
-				ConstructConvexMeshOneShotManifold(*ScaledConvex, ConvexTransform, Mesh, MeshTransform, MeshScale, CullDistance, MeshContacts);
-			}
-			else if (const TImplicitObjectInstanced<FImplicitConvex3>* InstancedConvex = Convex.template GetObject<TImplicitObjectInstanced<FImplicitConvex3>>())
-			{
-				ConstructConvexMeshOneShotManifold(*InstancedConvex, ConvexTransform, Mesh, MeshTransform, MeshScale, CullDistance, MeshContacts);
-			}
-			else if (const FImplicitConvex3* RawConvex = Convex.template GetObject<FImplicitConvex3>())
-			{
-				ConstructConvexMeshOneShotManifold(*RawConvex, ConvexTransform, Mesh, MeshTransform, MeshScale, CullDistance, MeshContacts);
-			}
-			else
-			{
-				check(false);
-			}
-
-			Constraint.SetOneShotManifoldContacts(MeshContacts.GetContactPoints());
-		}
-
-
-		template void ConstructSphereTriangleMeshOneShotManifold<TImplicitObjectScaled<class Chaos::FTriangleMeshImplicitObject, 1>>(const TSphere<FReal, 3>& Sphere, const FRigidTransform3& SphereWorldTransform, const TImplicitObjectScaled<class Chaos::FTriangleMeshImplicitObject, 1>& TriangleMesh, const FRigidTransform3& TriMeshWorldTransform, const FReal Dt, FPBDCollisionConstraint& Constraint);
-		template void ConstructSphereTriangleMeshOneShotManifold<FTriangleMeshImplicitObject>(const TSphere<FReal, 3>& Sphere, const FRigidTransform3& SphereWorldTransform, const FTriangleMeshImplicitObject& TriangleMesh, const FRigidTransform3& TriMeshWorldTransform, const FReal Dt, FPBDCollisionConstraint& Constraint);
-	}
-}
+	} // namespace Collisions
+} // namespace Chaos
 
diff --git a/Engine/Source/Runtime/Experimental/Chaos/Private/Chaos/CollisionOneShotManifoldsMiscShapes.h b/Engine/Source/Runtime/Experimental/Chaos/Private/Chaos/CollisionOneShotManifoldsMiscShapes.h
--- a/Engine/Source/Runtime/Experimental/Chaos/Private/Chaos/CollisionOneShotManifoldsMiscShapes.h
+++ b/Engine/Source/Runtime/Experimental/Chaos/Private/Chaos/CollisionOneShotManifoldsMiscShapes.h
@@ -2,21 +2,22 @@
 
 #pragma once
 
-#include "Chaos/Collision/PBDCollisionConstraint.h"
-#include "Chaos/Convex.h"
-#include "Chaos/Defines.h"
-#include "Chaos/Transform.h"
+#include "Chaos/Core.h"
+#include "Chaos/CollisionResolutionTypes.h"
+#include "Chaos/ImplicitFwd.h"
 
 namespace Chaos
 {
 	class FHeightField;
+	class FPBDCollisionConstraint;
+
 	namespace Collisions
 	{
 		void ConstructSphereSphereOneShotManifold(
 			const TSphere<FReal, 3>& Sphere1,
-			const FRigidTransform3& Convex1Transform, //world
+			const FRigidTransform3& Convex1Transform,
 			const TSphere<FReal, 3>& Sphere2,
-			const FRigidTransform3& Convex2Transform, //world
+			const FRigidTransform3& Convex2Transform,
 			const FReal Dt,
 			FPBDCollisionConstraint& Constraint);
 
@@ -52,22 +53,6 @@ namespace Chaos
 			const FReal Dt,
 			FPBDCollisionConstraint& Constraint);
 
-		template <typename TriMeshType>
-		void ConstructSphereTriangleMeshOneShotManifold(
-			const TSphere<FReal, 3>& Sphere,
-			const FRigidTransform3& SphereWorldTransform,
-			const TriMeshType& TriangleMesh,
-			const FRigidTransform3& TriMeshWorldTransform,
-			const FReal Dt,
-			FPBDCollisionConstraint& Constraint);
-
-		void ConstructSphereHeightFieldOneShotManifold(
-			const TSphere<FReal, 3>& Sphere, 
-			const FRigidTransform3& SphereTransform, 
-			const FHeightField& Heightfield, 
-			const FRigidTransform3& HeightfieldTransform, 
-			const FReal Dt, FPBDCollisionConstraint& Constraint);
-
 		void ConstructCapsuleCapsuleOneShotManifold(
 			const FCapsule& CapsuleA,
 			const FRigidTransform3& CapsuleATransform,
@@ -75,46 +60,5 @@ namespace Chaos
 			const FRigidTransform3& CapsuleBTransform,
 			const FReal Dt,
 			FPBDCollisionConstraint& Constraint);
-
-		void ConstructCapsuleTriMeshOneShotManifold(
-			const FCapsule& Capsule, 
-			const FRigidTransform3& CapsuleWorldTransform, 
-			const FImplicitObject& TriangleMesh, 
-			const FRigidTransform3& TriMeshWorldTransform, 
-			const FReal Dt, 
-			FPBDCollisionConstraint& Constraint);
-
-		void ConstructCapsuleHeightFieldOneShotManifold(
-			const FCapsule& Capsule, 
-			const FRigidTransform3& CapsuleTransform, 
-			const FHeightField& HeightField, 
-			const FRigidTransform3& HeightFieldTransform, 
-			const FReal Dt, 
-			FPBDCollisionConstraint& Constraint);
-
-		/**
-		 * @brief Build the contact manifold between a Planar Convex shape and a TriMesh
-		 * @param Convex an ImplicitObject that must be a box, convex, or a wrapper around it. Asserts if not
-		 * @param TriMesh an ImplicitObject that must be FTriangleMeshImplicitObject or a wrapper around it. Asserts if not
-		*/
-		void ConstructPlanarConvexTriMeshOneShotManifold(
-			const FImplicitObject& Convex, 
-			const FRigidTransform3& ConvexTransform, 
-			const FImplicitObject& TriangleMesh, 
-			const FRigidTransform3& TriangleMeshTransform, 
-			FPBDCollisionConstraint& Constraint);
-
-
-		/**
-		 * @brief Build the contact manifold between a Planar Convex shape and a HeightField
-		 * @param Convex an ImplicitObject that must be a box, convex, or a wrapper around it. Asserts if not
-		*/
-		void ConstructPlanarConvexHeightFieldOneShotManifold(
-			const FImplicitObject& Convex, 
-			const FRigidTransform3& ConvexTransform, 
-			const FHeightField& HeightField, 
-			const FRigidTransform3& TriangleMeshTransform, 
-			FPBDCollisionConstraint& Constraint);
-
 	}
 }
\ No newline at end of file
diff --git a/Engine/Source/Runtime/Experimental/Chaos/Private/Chaos/CollisionResolution.cpp b/Engine/Source/Runtime/Experimental/Chaos/Private/Chaos/CollisionResolution.cpp
--- a/Engine/Source/Runtime/Experimental/Chaos/Private/Chaos/CollisionResolution.cpp
+++ b/Engine/Source/Runtime/Experimental/Chaos/Private/Chaos/CollisionResolution.cpp
@@ -4,19 +4,24 @@
 
 #include "Chaos/ChaosPerfTest.h"
 #include "Chaos/Capsule.h"
+#include "Chaos/CCDUtilities.h"
+#include "Chaos/CollisionOneShotManifolds.h"
+#include "Chaos/CollisionOneShotManifoldsMeshShapes.h"
+#include "Chaos/CollisionOneShotManifoldsMiscShapes.h"
 #include "Chaos/CollisionResolutionTypes.h"
 #include "Chaos/CollisionResolutionUtil.h"
 #include "Chaos/Collision/CapsuleConvexContactPoint.h"
 #include "Chaos/Collision/CollisionContext.h"
 #include "Chaos/Collision/CollisionConstraintAllocator.h"
 #include "Chaos/Collision/CollisionFilter.h"
-#include "Chaos/Collision/PBDCollisionConstraint.h"
-#include "Chaos/Collision/SphereConvexContactPoint.h"
 #include "Chaos/Collision/ContactPointsMiscShapes.h"
 #include "Chaos/Collision/GJKContactPointSwept.h"
 #include "Chaos/Collision/GJKContactPoint.h"
+#include "Chaos/Collision/PBDCollisionConstraint.h"
+#include "Chaos/Collision/SphereConvexContactPoint.h"
 #include "Chaos/Convex.h"
 #include "Chaos/Defines.h"
+#include "Chaos/GeometryQueries.h"
 #include "Chaos/HeightField.h"
 #include "Chaos/ImplicitFwd.h"
 #include "Chaos/ImplicitObjectScaled.h"
@@ -29,10 +34,6 @@
 #include "Chaos/Sphere.h"
 #include "Chaos/Transform.h"
 #include "Chaos/TriangleMeshImplicitObject.h"
-#include "Chaos/GeometryQueries.h"
-#include "Chaos/CollisionOneShotManifolds.h"
-#include "Chaos/CollisionOneShotManifoldsMiscShapes.h"
-#include "Chaos/CCDUtilities.h"
 
 //PRAGMA_DISABLE_OPTIMIZATION
 
@@ -654,7 +655,7 @@ namespace Chaos
 		{
 			if (Constraint.GetUseManifold())
 			{
-				ConstructSphereHeightFieldOneShotManifold(A, ATransform, B, BTransform, Dt, Constraint);
+				ConstructQuadraticConvexHeightFieldOneShotManifold(A, ATransform, B, BTransform, Dt, Constraint);
 			}
 			else
 			{
@@ -846,7 +847,7 @@ namespace Chaos
 		{
 			if (Constraint.GetUseManifold())
 			{
-				ConstructSphereTriangleMeshOneShotManifold(Sphere0, WorldTransform0, TriangleMesh1, WorldTransform1, Dt, Constraint);
+				ConstructQuadraticConvexTriMeshOneShotManifold(Sphere0, WorldTransform0, TriangleMesh1, WorldTransform1, Dt, Constraint);
 			}
 			else
 			{
@@ -1092,7 +1093,7 @@ namespace Chaos
 			CONDITIONAL_SCOPE_CYCLE_COUNTER(STAT_Collisions_UpdateCapsuleHeightFieldConstraint, ConstraintsDetailedStats);
 			if (Constraint.GetUseManifold())
 			{
-				ConstructCapsuleHeightFieldOneShotManifold(A, ATransform, B, BTransform, Dt, Constraint);
+				ConstructQuadraticConvexHeightFieldOneShotManifold(A, ATransform, B, BTransform, Dt, Constraint);
 			}
 			else
 			{
@@ -1155,7 +1156,7 @@ namespace Chaos
 			CONDITIONAL_SCOPE_CYCLE_COUNTER(STAT_Collisions_UpdateCapsuleTriangleMeshConstraint, ConstraintsDetailedStats);
 			if (Constraint.GetUseManifold())
 			{
-				ConstructCapsuleTriMeshOneShotManifold(Capsule0, WorldTransform0, TriangleMesh1, WorldTransform1, Dt, Constraint);
+				ConstructQuadraticConvexTriMeshOneShotManifold(Capsule0, WorldTransform0, TriangleMesh1, WorldTransform1, Dt, Constraint);
 			}
 			else
 			{
diff --git a/Engine/Source/Runtime/Experimental/Chaos/Public/Chaos/Collision/SphereTriangleContactPoint.h b/Engine/Source/Runtime/Experimental/Chaos/Public/Chaos/Collision/SphereTriangleContactPoint.h
new file mode 100644
--- /dev/null
+++ b/Engine/Source/Runtime/Experimental/Chaos/Public/Chaos/Collision/SphereTriangleContactPoint.h
@@ -0,0 +1,10 @@
+// Copyright Epic Games, Inc. All Rights Reserved.
+#pragma once
+#include "Chaos/Core.h"
+#include "Chaos/CollisionResolutionTypes.h"
+#include "Chaos/ImplicitFwd.h"
+
+namespace Chaos
+{
+	extern CHAOS_API void ConstructSphereTriangleOneShotManifold(const FImplicitSphere3& Sphere, const FTriangle& Triangle, const FReal CullDistance, FContactPointManifold& OutContactPoints);
+}
\ No newline at end of file
