Attempt to fix cook warning related to event dispatcher in Geometry collection
- made the event dispatcher a transient property
- allocate the object using NewObject instead of CreateDefaultSubobject since it is a runtime only object

diff --git a/Engine/Source/Runtime/Experimental/GeometryCollectionEngine/Private/GeometryCollection/GeometryCollectionComponent.cpp b/Engine/Source/Runtime/Experimental/GeometryCollectionEngine/Private/GeometryCollection/GeometryCollectionComponent.cpp
--- a/Engine/Source/Runtime/Experimental/GeometryCollectionEngine/Private/GeometryCollection/GeometryCollectionComponent.cpp
+++ b/Engine/Source/Runtime/Experimental/GeometryCollectionEngine/Private/GeometryCollection/GeometryCollectionComponent.cpp
@@ -384,7 +384,10 @@ UGeometryCollectionComponent::UGeometryCollectionComponent(const FObjectInitiali
 	// By default, we initialize immediately. If this is set false, we defer initialization.
 	BodyInstance.bSimulatePhysics = true;
 
-	EventDispatcher = ObjectInitializer.CreateDefaultSubobject<UChaosGameplayEventDispatcher>(this, TEXT("GameplayEventDispatcher"));
+	if (!HasAnyFlags(RF_ArchetypeObject | RF_ClassDefaultObject))
+	{
+		EventDispatcher = NewObject<UChaosGameplayEventDispatcher>(this, TEXT("GameplayEventDispatcher"), RF_Transient);
+	}
 
 	DynamicCollection = nullptr;
 	bHasCustomNavigableGeometry = EHasCustomNavigableGeometry::Yes;
@@ -1389,51 +1392,54 @@ void UGeometryCollectionComponent::DispatchChaosPhysicsCollisionBlueprintEvents(
 // call when first registering
 void UGeometryCollectionComponent::RegisterForEvents()
 {
-	if (BodyInstance.bNotifyRigidBodyCollision || bNotifyBreaks || bNotifyCollisions || bNotifyRemovals || bNotifyCrumblings)
+	if (EventDispatcher)
 	{
-		Chaos::FPhysicsSolver* Solver = GetWorld()->GetPhysicsScene()->GetSolver();
-		if (Solver)
+		if (BodyInstance.bNotifyRigidBodyCollision || bNotifyBreaks || bNotifyCollisions || bNotifyRemovals || bNotifyCrumblings)
 		{
-			if (bNotifyCollisions || BodyInstance.bNotifyRigidBodyCollision)
+			Chaos::FPhysicsSolver* Solver = GetWorld()->GetPhysicsScene()->GetSolver();
+			if (Solver)
 			{
-				EventDispatcher->RegisterForCollisionEvents(this, this);
+				if (bNotifyCollisions || BodyInstance.bNotifyRigidBodyCollision)
+				{
+					EventDispatcher->RegisterForCollisionEvents(this, this);
 
-				Solver->EnqueueCommandImmediate([Solver]()
-					{
-						Solver->SetGenerateCollisionData(true);
-					});
-			}
+					Solver->EnqueueCommandImmediate([Solver]()
+						{
+							Solver->SetGenerateCollisionData(true);
+						});
+				}
 
-			if (bNotifyBreaks)
-			{
-				EventDispatcher->RegisterForBreakEvents(this, &DispatchGeometryCollectionBreakEvent);
+				if (bNotifyBreaks)
+				{
+					EventDispatcher->RegisterForBreakEvents(this, &DispatchGeometryCollectionBreakEvent);
 
-				Solver->EnqueueCommandImmediate([Solver]()
-					{
-						Solver->SetGenerateBreakingData(true);
-					});
+					Solver->EnqueueCommandImmediate([Solver]()
+						{
+							Solver->SetGenerateBreakingData(true);
+						});
 
-			}
+				}
 
-			if (bNotifyRemovals)
-			{
-				EventDispatcher->RegisterForRemovalEvents(this, &DispatchGeometryCollectionRemovalEvent);
+				if (bNotifyRemovals)
+				{
+					EventDispatcher->RegisterForRemovalEvents(this, &DispatchGeometryCollectionRemovalEvent);
 
-				Solver->EnqueueCommandImmediate([Solver]()
-					{
-						Solver->SetGenerateRemovalData(true);
-					});
+					Solver->EnqueueCommandImmediate([Solver]()
+						{
+							Solver->SetGenerateRemovalData(true);
+						});
 
-			}
+				}
 
-			if (bNotifyCrumblings)
-			{
-				EventDispatcher->RegisterForCrumblingEvents(this, &DispatchGeometryCollectionCrumblingEvent);
+				if (bNotifyCrumblings)
+				{
+					EventDispatcher->RegisterForCrumblingEvents(this, &DispatchGeometryCollectionCrumblingEvent);
 
-				Solver->EnqueueCommandImmediate([Solver]()
-					{
-						Solver->SetGenerateBreakingData(true);
-					});
+					Solver->EnqueueCommandImmediate([Solver]()
+						{
+							Solver->SetGenerateBreakingData(true);
+						});
+				}
 			}
 		}
 	}
@@ -1441,49 +1447,61 @@ void UGeometryCollectionComponent::RegisterForEvents()
 
 void UGeometryCollectionComponent::UpdateRBCollisionEventRegistration()
 {
-	if (bNotifyCollisions || BodyInstance.bNotifyRigidBodyCollision)
+	if (EventDispatcher)
 	{
-		EventDispatcher->RegisterForCollisionEvents(this, this);
-	}
-	else
-	{
-		EventDispatcher->UnRegisterForCollisionEvents(this, this);
+		if (bNotifyCollisions || BodyInstance.bNotifyRigidBodyCollision)
+		{
+			EventDispatcher->RegisterForCollisionEvents(this, this);
+		}
+		else
+		{
+			EventDispatcher->UnRegisterForCollisionEvents(this, this);
+		}
 	}
 }
 
 void UGeometryCollectionComponent::UpdateBreakEventRegistration()
 {
-	if (bNotifyBreaks)
-	{
-		EventDispatcher->RegisterForBreakEvents(this, &DispatchGeometryCollectionBreakEvent);
-	}
-	else
+	if (EventDispatcher)
 	{
-		EventDispatcher->UnRegisterForBreakEvents(this);
+		if (bNotifyBreaks)
+		{
+			EventDispatcher->RegisterForBreakEvents(this, &DispatchGeometryCollectionBreakEvent);
+		}
+		else
+		{
+			EventDispatcher->UnRegisterForBreakEvents(this);
+		}
 	}
 }
 
 void UGeometryCollectionComponent::UpdateRemovalEventRegistration()
 {
-	if (bNotifyRemovals)
-	{
-		EventDispatcher->RegisterForRemovalEvents(this, &DispatchGeometryCollectionRemovalEvent);
-	}
-	else
+	if (EventDispatcher)
 	{
-		EventDispatcher->UnRegisterForRemovalEvents(this);
+		if (bNotifyRemovals)
+		{
+			EventDispatcher->RegisterForRemovalEvents(this, &DispatchGeometryCollectionRemovalEvent);
+		}
+		else
+		{
+			EventDispatcher->UnRegisterForRemovalEvents(this);
+		}
 	}
 }
 
 void UGeometryCollectionComponent::UpdateCrumblingEventRegistration()
 {
-	if (bNotifyCrumblings)
-	{
-		EventDispatcher->RegisterForCrumblingEvents(this, &DispatchGeometryCollectionCrumblingEvent);
-	}
-	else
+	if (EventDispatcher)
 	{
-		EventDispatcher->UnRegisterForCrumblingEvents(this);
+		if (bNotifyCrumblings)
+		{
+			EventDispatcher->RegisterForCrumblingEvents(this, &DispatchGeometryCollectionCrumblingEvent);
+		}
+		else
+		{
+			EventDispatcher->UnRegisterForCrumblingEvents(this);
+		}
 	}
 }
 
diff --git a/Engine/Source/Runtime/Experimental/GeometryCollectionEngine/Public/GeometryCollection/GeometryCollectionComponent.h b/Engine/Source/Runtime/Experimental/GeometryCollectionEngine/Public/GeometryCollection/GeometryCollectionComponent.h
--- a/Engine/Source/Runtime/Experimental/GeometryCollectionEngine/Public/GeometryCollection/GeometryCollectionComponent.h
+++ b/Engine/Source/Runtime/Experimental/GeometryCollectionEngine/Public/GeometryCollection/GeometryCollectionComponent.h
@@ -1212,7 +1212,9 @@ private:
 #endif
 	void SwitchRenderModels(const AActor* Actor);
 
-	UChaosGameplayEventDispatcher* EventDispatcher;
+	// Event dispatcher for break, crumble, removal and collision events
+	UPROPERTY(transient)
+	TObjectPtr<UChaosGameplayEventDispatcher> EventDispatcher;
 
 #if GEOMETRYCOLLECTION_EDITOR_SELECTION
 	bool bIsTransformSelectionModeEnabled;
