Chaos : better copntrol for debug draw
- ability to control server and/or client debug draw
- add p.chaos.debugdraw.seethrough cvar to control whether or not the lines are draw over objects

diff --git a/Engine/Source/Runtime/Experimental/Chaos/Private/PBDRigidsSolver.cpp b/Engine/Source/Runtime/Experimental/Chaos/Private/PBDRigidsSolver.cpp
--- a/Engine/Source/Runtime/Experimental/Chaos/Private/PBDRigidsSolver.cpp
+++ b/Engine/Source/Runtime/Experimental/Chaos/Private/PBDRigidsSolver.cpp
@@ -92,6 +92,8 @@ namespace Chaos
 		int32 ChaosSolverDrawClusterConstraints = 0;
 		int32 ChaosSolverDebugDrawMeshContacts = 0;
 		int32 ChaosSolverDebugDrawColorShapeByClientServer = 0;
+		int32 ChaosSolverDebugDrawShowServer = 1;
+		int32 ChaosSolverDebugDrawShowClient = 1;
 		DebugDraw::FChaosDebugDrawJointFeatures ChaosSolverDrawJointFeatures = DebugDraw::FChaosDebugDrawJointFeatures::MakeDefault();
 		FAutoConsoleVariableRef CVarChaosSolverDrawShapes(TEXT("p.Chaos.Solver.DebugDrawShapes"), ChaosSolverDebugDrawShapes, TEXT("Draw Shapes (0 = never; 1 = end of frame)."));
 		FAutoConsoleVariableRef CVarChaosSolverDrawCollisions(TEXT("p.Chaos.Solver.DebugDrawCollisions"), ChaosSolverDebugDrawCollisions, TEXT("Draw Collisions (0 = never; 1 = end of frame)."));
@@ -121,6 +123,8 @@ namespace Chaos
 		FAutoConsoleVariableRef CVarChaosSolverDrawClusterConstraints(TEXT("p.Chaos.Solver.DebugDraw.Cluster.Constraints"), ChaosSolverDrawClusterConstraints, TEXT("Draw Active Cluster Constraints (0 = never; 1 = end of frame)."));
 		FAutoConsoleVariableRef CVarChaosSolverDrawMeshContacts(TEXT("p.Chaos.Solver.DebugDrawMeshContacts"), ChaosSolverDebugDrawMeshContacts, TEXT("Draw Mesh contacts"));
 		FAutoConsoleVariableRef CVarChaosSolverDebugDrawColorShapeByClientServer(TEXT("p.Chaos.Solver.DebugDraw.ColorShapeByClientServer"), ChaosSolverDebugDrawColorShapeByClientServer, TEXT("Color shape according to client and server: red = server / blue = client "));
+		FAutoConsoleVariableRef CVarChaosSolverDebugDrawShowServer(TEXT("p.Chaos.Solver.DebugDraw.ShowServer"), ChaosSolverDebugDrawShowServer, TEXT("Draw server related debug data"));
+		FAutoConsoleVariableRef CVarChaosSolverDebugDrawShowClient(TEXT("p.Chaos.Solver.DebugDraw.ShowClient"), ChaosSolverDebugDrawShowClient, TEXT("Draw client related debug data"));
 
 
 		DebugDraw::FChaosDebugDrawSettings ChaosSolverDebugDebugDrawSettings(
@@ -1852,9 +1856,19 @@ CSV_CUSTOM_STAT(PhysicsCounters, Name, Value, ECsvCustomStatOp::Set);
 		QUICK_SCOPE_CYCLE_COUNTER(SolverDebugDraw);
 
 #if !UE_BUILD_SHIPPING && !UE_BUILD_TEST
+		const bool bIsServer = DebugName.ToString().StartsWith(TEXT("Server"));
+		if (bIsServer && !ChaosSolverDebugDrawShowServer)
+		{
+			return;
+		}
+		if (!bIsServer && !ChaosSolverDebugDrawShowClient)
+		{
+			return;
+		}
+
 		if (ChaosSolverDebugDrawColorShapeByClientServer)
 		{
-			if (DebugName.ToString().StartsWith(TEXT("Server")))
+			if (bIsServer)
 			{
 				ChaosSolverDebugDebugDrawSettings.ShapesColorsPerState = GetSolverShapesColorsByState_Server();
 			}
diff --git a/Engine/Source/Runtime/Experimental/ChaosSolverEngine/Private/Chaos/ChaosDebugDrawComponent.cpp b/Engine/Source/Runtime/Experimental/ChaosSolverEngine/Private/Chaos/ChaosDebugDrawComponent.cpp
--- a/Engine/Source/Runtime/Experimental/ChaosSolverEngine/Private/Chaos/ChaosDebugDrawComponent.cpp
+++ b/Engine/Source/Runtime/Experimental/ChaosSolverEngine/Private/Chaos/ChaosDebugDrawComponent.cpp
@@ -29,6 +29,9 @@ FAutoConsoleVariableRef CVarChaos_DebugDraw_MaxElements(TEXT("p.Chaos.DebugDraw.
 float ChaosDebugDraw_Radius = 3000.0f;
 FAutoConsoleVariableRef CVarChaos_DebugDraw_Radius(TEXT("p.Chaos.DebugDraw.Radius"), ChaosDebugDraw_Radius, TEXT("Set the radius from the camera where debug draw capture stops (0 means infinite)"));
 
+bool ChaosDebugDraw_SeeThrough = true;
+FAutoConsoleVariableRef CVarChaos_DebugDraw_SeeThrough(TEXT("p.Chaos.DebugDraw.SeeThrough"), ChaosDebugDraw_SeeThrough, TEXT("When enabled , lines will be drawn on top of all geometry"));
+
 bool bChaosDebugDraw_SingleActor = false;
 FAutoConsoleVariableRef CVarChaosDebugDraw_SingleActor(TEXT("p.Chaos.DebugDraw.SingleActor"), bChaosDebugDraw_SingleActor, TEXT("If true, then we draw for the actor the camera is looking at."));
 
@@ -90,22 +93,23 @@ void DebugDrawChaos(AActor* DebugDrawActor, const TArray<Chaos::FLatentDrawComma
 	{
 		for (const FLatentDrawCommand& Command : DrawCommands)
 		{
+			const uint8 DepthPriority = ChaosDebugDraw_SeeThrough ? Command.DepthPriority: 0;
 			switch (Command.Type)
 			{
 			case FLatentDrawCommand::EDrawType::Point:
-				DrawDebugPoint(World, Command.LineStart, Command.Thickness, Command.Color, Command.bPersistentLines, CommandLifeTime(Command, bIsPaused), Command.DepthPriority);
+				DrawDebugPoint(World, Command.LineStart, Command.Thickness, Command.Color, Command.bPersistentLines, CommandLifeTime(Command, bIsPaused), DepthPriority);
 				break;
 			case FLatentDrawCommand::EDrawType::Line:
-				DrawDebugLine(World, Command.LineStart, Command.LineEnd, Command.Color, Command.bPersistentLines, CommandLifeTime(Command, bIsPaused), Command.DepthPriority, Command.Thickness);
+				DrawDebugLine(World, Command.LineStart, Command.LineEnd, Command.Color, Command.bPersistentLines, CommandLifeTime(Command, bIsPaused), DepthPriority, Command.Thickness);
 				break;
 			case FLatentDrawCommand::EDrawType::DirectionalArrow:
-				DrawDebugDirectionalArrow(World, Command.LineStart, Command.LineEnd, Command.ArrowSize, Command.Color, Command.bPersistentLines, CommandLifeTime(Command, bIsPaused), Command.DepthPriority, Command.Thickness);
+				DrawDebugDirectionalArrow(World, Command.LineStart, Command.LineEnd, Command.ArrowSize, Command.Color, Command.bPersistentLines, CommandLifeTime(Command, bIsPaused), DepthPriority, Command.Thickness);
 				break;
 			case FLatentDrawCommand::EDrawType::Sphere:
-				DrawDebugSphere(World, Command.LineStart, Command.Radius, Command.Segments, Command.Color, Command.bPersistentLines, CommandLifeTime(Command, bIsPaused), Command.DepthPriority, Command.Thickness);
+				DrawDebugSphere(World, Command.LineStart, Command.Radius, Command.Segments, Command.Color, Command.bPersistentLines, CommandLifeTime(Command, bIsPaused), DepthPriority, Command.Thickness);
 				break;
 			case FLatentDrawCommand::EDrawType::Box:
-				DrawDebugBox(World, Command.Center, Command.Extent, Command.Rotation, Command.Color, Command.bPersistentLines, CommandLifeTime(Command, bIsPaused), Command.DepthPriority, Command.Thickness);
+				DrawDebugBox(World, Command.Center, Command.Extent, Command.Rotation, Command.Color, Command.bPersistentLines, CommandLifeTime(Command, bIsPaused), DepthPriority, Command.Thickness);
 				break;
 			case FLatentDrawCommand::EDrawType::String:
 				DrawDebugString(World, Command.TextLocation, Command.Text, Command.TestBaseActor, Command.Color, CommandLifeTime(Command, bIsPaused), Command.bDrawShadow, Command.FontScale);
@@ -114,11 +118,11 @@ void DebugDrawChaos(AActor* DebugDrawActor, const TArray<Chaos::FLatentDrawComma
 			{
 				FMatrix M = FRotationMatrix::MakeFromYZ(Command.YAxis, Command.ZAxis);
 				M.SetOrigin(Command.Center);
-				DrawDebugCircle(World, M, Command.Radius, Command.Segments, Command.Color, Command.bPersistentLines, CommandLifeTime(Command, bIsPaused), Command.DepthPriority, Command.Thickness, Command.bDrawAxis);
+				DrawDebugCircle(World, M, Command.Radius, Command.Segments, Command.Color, Command.bPersistentLines, CommandLifeTime(Command, bIsPaused), DepthPriority, Command.Thickness, Command.bDrawAxis);
 				break;
 			}
 			case FLatentDrawCommand::EDrawType::Capsule:
-				DrawDebugCapsule(World, Command.Center, Command.HalfHeight, Command.Radius, Command.Rotation, Command.Color, Command.bPersistentLines, CommandLifeTime(Command, bIsPaused), Command.DepthPriority, Command.Thickness);
+				DrawDebugCapsule(World, Command.Center, Command.HalfHeight, Command.Radius, Command.Rotation, Command.Color, Command.bPersistentLines, CommandLifeTime(Command, bIsPaused), DepthPriority, Command.Thickness);
 			default:
 				break;
 			}
