ChaosFlesh : DebugDraw Collisions Geometry

diff --git a/Engine/Source/Runtime/Experimental/Chaos/Private/Chaos/Deformable/ChaosDeformableSolver.cpp b/Engine/Source/Runtime/Experimental/Chaos/Private/Chaos/Deformable/ChaosDeformableSolver.cpp
--- a/Engine/Source/Runtime/Experimental/Chaos/Private/Chaos/Deformable/ChaosDeformableSolver.cpp
+++ b/Engine/Source/Runtime/Experimental/Chaos/Private/Chaos/Deformable/ChaosDeformableSolver.cpp
@@ -69,6 +69,7 @@ namespace Chaos::Softs
 
 	FAutoConsoleVariableRef CVarDeformableDebugParamsDrawTetrahedralParticles(TEXT("p.Chaos.DebugDraw.Deformable.TetrahedralParticle"), GDeformableDebugParams.bDoDrawTetrahedralParticles, TEXT("Debug draw the deformable solvers tetrahedron. [def: false]"));
 	FAutoConsoleVariableRef CVarDeformableDebugParamsDrawKinematicParticles(TEXT("p.Chaos.DebugDraw.Deformable.KinematicParticle"), GDeformableDebugParams.bDoDrawKinematicParticles, TEXT("Debug draw the deformables kinematic particles. [def: false]"));
+	FAutoConsoleVariableRef CVarDeformableDebugParamsDrawRigidCollisionGeometry(TEXT("p.Chaos.DebugDraw.Deformable.RigidCollisionGeometry"), GDeformableDebugParams.bDoDrawRigidCollisionGeometry, TEXT("Debug draw the deformable solvers rigid collision geometry. [def: false]"));
 
 	FDeformableXPBDCorotatedParams GDeformableXPBDCorotatedParams;
 	FAutoConsoleVariableRef CVarDeformableXPBDCorotatedBatchSize(TEXT("p.Chaos.Deformable.XPBDBatchSize"), GDeformableXPBDCorotatedParams.XPBDCorotatedBatchSize, TEXT("Batch size for physics parallel for. [def: 5]"));
@@ -170,6 +171,7 @@ namespace Chaos::Softs
 			InitializeSimulationObjects();
 			InitializeSimulationSpace();
 			AdvanceDt(DeltaTime);
+			DebugDrawSimulationData();
 		}
 	}
 
@@ -1279,6 +1281,66 @@ namespace Chaos::Softs
 		}
 	}
 
+
+	void FDeformableSolver::DebugDrawSimulationData()
+	{
+#if WITH_EDITOR
+		auto ToFVec3 = [](FVector3d V) { return Chaos::FVec3(V.X, V.Y, V.Z); };
+		auto ToFVector = [](Chaos::FVec3 V) { return FVector(V.X, V.Y, V.Z); };
+		auto ToFQuat = [](const TRotation<FSolverReal, 3>& R) { return FQuat(R.X, R.Y, R.Z, R.W); };
+
+		//debug draw
+		//p.Chaos.DebugDraw.Enabled 1
+		//p.Chaos.DebugDraw.Deformable.RigidCollisionGeometry 1
+		if (Evolution && GDeformableDebugParams.bDoDrawRigidCollisionGeometry)
+		{
+			Evolution->CollisionParticlesActiveView().RangeFor(
+				[this, ToFVec3, ToFVector, ToFQuat](FSolverRigidParticles& CollisionParticles, int32 CollisionOffset, int32 CollisionRange)
+				{
+					for (int32 Index = CollisionOffset; Index < CollisionRange; Index++)
+					{
+						if (Evolution->CollisionParticleGroupIds()[Index] != Index)
+						{
+							if (const TUniquePtr<FImplicitObject>& Geometry = CollisionParticles.DynamicGeometry(Index))
+							{
+								EImplicitObjectType GeomType = Geometry->GetCollisionType();
+								if (GeomType == ImplicitObjectType::Sphere)
+								{
+									const FSphere& SphereGeometry = Geometry->GetObjectChecked<FSphere>();
+									FVector Center = ToFVector(CollisionParticles.X(Index)) + SphereGeometry.GetCenter();
+									FReal Radius = SphereGeometry.GetRadius();
+									Chaos::FDebugDrawQueue::GetInstance().DrawDebugSphere(Center, Radius, 12, FColor::Red, false, -1.0f, 1, 5.f);
+								}
+								else if (GeomType == ImplicitObjectType::Box)
+								{
+									const TBox<FReal, 3>& BoxGeometry = Geometry->GetObjectChecked<TBox<FReal, 3>>();
+									FVector Extent = 0.5 * (BoxGeometry.Max() - BoxGeometry.Min());
+									FVector Center = ToFVector(CollisionParticles.X(Index)) + BoxGeometry.GetCenter();
+									const FQuat& Rotation = ToFQuat(CollisionParticles.R(Index));
+									Chaos::FDebugDrawQueue::GetInstance().DrawDebugBox(Center, Extent, Rotation, FColor::Red, false, -1.0f, 1, 5.f);
+								}
+								else if (GeomType == ImplicitObjectType::Convex)
+								{
+									const FConvex& ConvexGeometry = Geometry->GetObjectChecked<FConvex>();
+									FTransform M = FTransform(ToFQuat(CollisionParticles.R(Index)), ToFVector(CollisionParticles.X(Index)));
+									for (int32 EdgeIndex = 0; EdgeIndex < ConvexGeometry.NumEdges(); ++EdgeIndex)
+									{
+										int32 Index0 = ConvexGeometry.GetEdgeVertex(EdgeIndex, 0);
+										int32 Index1 = ConvexGeometry.GetEdgeVertex(EdgeIndex, 1);
+										const  TArray<FConvex::FVec3Type>& Verts = ConvexGeometry.GetVertices();
+										Chaos::FDebugDrawQueue::GetInstance().DrawDebugLine(
+											M.TransformPosition(ToFVector(Verts[Index0])), M.TransformPosition(ToFVector(Verts[Index1])), FColor::Red, false, -1.0f, 1, 5.f);
+									}
+								}
+							}
+						}
+					}
+				});
+		}
+#endif
+	}
+
+
 	void FDeformableSolver::WriteFrame(FThreadingProxy& InProxy, const FSolverReal DeltaTime)
 	{
 		TRACE_CPUPROFILER_EVENT_SCOPE(ChaosDeformableSolver_WriteFrame);
diff --git a/Engine/Source/Runtime/Experimental/Chaos/Public/Chaos/Deformable/ChaosDeformableSolver.h b/Engine/Source/Runtime/Experimental/Chaos/Public/Chaos/Deformable/ChaosDeformableSolver.h
--- a/Engine/Source/Runtime/Experimental/Chaos/Public/Chaos/Deformable/ChaosDeformableSolver.h
+++ b/Engine/Source/Runtime/Experimental/Chaos/Public/Chaos/Deformable/ChaosDeformableSolver.h
@@ -138,6 +138,7 @@ namespace Chaos::Softs
 		void WriteFrame(FThreadingProxy&, const FSolverReal DeltaTime);
 
 		/*Debug Drawing*/
+		void DebugDrawSimulationData();
 		void DebugDrawTetrahedralParticles(FFleshThreadingProxy& Proxy);
 
 		/*Game Thread API*/
diff --git a/Engine/Source/Runtime/Experimental/Chaos/Public/Chaos/Deformable/ChaosDeformableSolverTypes.h b/Engine/Source/Runtime/Experimental/Chaos/Public/Chaos/Deformable/ChaosDeformableSolverTypes.h
--- a/Engine/Source/Runtime/Experimental/Chaos/Public/Chaos/Deformable/ChaosDeformableSolverTypes.h
+++ b/Engine/Source/Runtime/Experimental/Chaos/Public/Chaos/Deformable/ChaosDeformableSolverTypes.h
@@ -143,6 +143,7 @@ namespace Chaos::Softs
 	{				
 		bool bDoDrawTetrahedralParticles = false;
 		bool bDoDrawKinematicParticles = false;
+		bool bDoDrawRigidCollisionGeometry = false;
 
 
 		bool IsDebugDrawingEnabled()
