ISM Pool : fix preallocation of instance passing the total number of instance instead of the added ones

diff --git a/Engine/Source/Runtime/Experimental/GeometryCollectionEngine/Private/GeometryCollection/GeometryCollectionISMPoolComponent.cpp b/Engine/Source/Runtime/Experimental/GeometryCollectionEngine/Private/GeometryCollection/GeometryCollectionISMPoolComponent.cpp
--- a/Engine/Source/Runtime/Experimental/GeometryCollectionEngine/Private/GeometryCollection/GeometryCollectionISMPoolComponent.cpp
+++ b/Engine/Source/Runtime/Experimental/GeometryCollectionEngine/Private/GeometryCollection/GeometryCollectionISMPoolComponent.cpp
@@ -78,7 +78,7 @@ int32 FGeometryCollectionISM::AddInstanceGroup(int32 InstanceCount)
 	const int32 InstanceGroupIndex = InstanceGroups.AddGroup(InstanceCount);
 	const FInstanceGroups::FInstanceGroupRange& NewInstanceGroup = InstanceGroups.GetGroup(InstanceGroupIndex);
 	const int32 TotalInstanceCount = NewInstanceGroup.Start + NewInstanceGroup.Count;
-	ISMComponent->PreAllocateInstancesMemory(TotalInstanceCount);
+	ISMComponent->PreAllocateInstancesMemory(InstanceCount);
 	FTransform ZeroScaleTransform;
 	ZeroScaleTransform.SetIdentityZeroScale();
 	for (int32 InstanceIndex = NewInstanceGroup.Start; InstanceIndex < TotalInstanceCount; InstanceIndex++)
