Geometry collection : GC asset   now supports AssetUserData

diff --git a/Engine/Source/Runtime/Experimental/GeometryCollectionEngine/Private/GeometryCollection/GeometryCollectionObject.cpp b/Engine/Source/Runtime/Experimental/GeometryCollectionEngine/Private/GeometryCollection/GeometryCollectionObject.cpp
--- a/Engine/Source/Runtime/Experimental/GeometryCollectionEngine/Private/GeometryCollection/GeometryCollectionObject.cpp
+++ b/Engine/Source/Runtime/Experimental/GeometryCollectionEngine/Private/GeometryCollection/GeometryCollectionObject.cpp
@@ -1531,3 +1531,47 @@ uint32 UGeometryCollection::GetNaniteHierarchyOffset(int32 GeometryIndex, bool b
 	}
 	return HierarchyOffset;
 }
+
+void UGeometryCollection::AddAssetUserData(UAssetUserData* InUserData)
+{
+	if (InUserData != nullptr)
+	{
+		UAssetUserData* ExistingData = GetAssetUserDataOfClass(InUserData->GetClass());
+		if (ExistingData != nullptr)
+		{
+			AssetUserData.Remove(ExistingData);
+		}
+		AssetUserData.Add(InUserData);
+	}
+}
+
+UAssetUserData* UGeometryCollection::GetAssetUserDataOfClass(TSubclassOf<UAssetUserData> InUserDataClass)
+{
+	for (int32 DataIdx = 0; DataIdx < AssetUserData.Num(); DataIdx++)
+	{
+		UAssetUserData* Datum = AssetUserData[DataIdx];
+		if (Datum != nullptr && Datum->IsA(InUserDataClass))
+		{
+			return Datum;
+		}
+	}
+	return nullptr;
+}
+
+void UGeometryCollection::RemoveUserDataOfClass(TSubclassOf<UAssetUserData> InUserDataClass)
+{
+	for (int32 DataIdx = 0; DataIdx < AssetUserData.Num(); DataIdx++)
+	{
+		UAssetUserData* Datum = AssetUserData[DataIdx];
+		if (Datum != nullptr && Datum->IsA(InUserDataClass))
+		{
+			AssetUserData.RemoveAt(DataIdx);
+			return;
+		}
+	}
+}
+
+const TArray<UAssetUserData*>* UGeometryCollection::GetAssetUserDataArray() const
+{
+	return &ToRawPtrTArrayUnsafe(AssetUserData);
+}
\ No newline at end of file
diff --git a/Engine/Source/Runtime/Experimental/GeometryCollectionEngine/Public/GeometryCollection/GeometryCollectionObject.h b/Engine/Source/Runtime/Experimental/GeometryCollectionEngine/Public/GeometryCollection/GeometryCollectionObject.h
--- a/Engine/Source/Runtime/Experimental/GeometryCollectionEngine/Public/GeometryCollection/GeometryCollectionObject.h
+++ b/Engine/Source/Runtime/Experimental/GeometryCollectionEngine/Public/GeometryCollection/GeometryCollectionObject.h
@@ -328,7 +328,7 @@ struct GEOMETRYCOLLECTIONENGINE_API FGeometryCollectionProxyMeshData
 *
 */
 UCLASS(BlueprintType, customconstructor)
-class GEOMETRYCOLLECTIONENGINE_API UGeometryCollection : public UObject
+class GEOMETRYCOLLECTIONENGINE_API UGeometryCollection : public UObject, public IInterface_AssetUserData
 {
 	GENERATED_UCLASS_BODY()
 
@@ -702,6 +702,13 @@ public:
 	void UpdateGeometryDependentProperties();
 
 
+	//~ Begin IInterface_AssetUserData Interface
+	virtual void AddAssetUserData(UAssetUserData* InUserData) override;
+	virtual void RemoveUserDataOfClass(TSubclassOf<UAssetUserData> InUserDataClass) override;
+	virtual UAssetUserData* GetAssetUserDataOfClass(TSubclassOf<UAssetUserData> InUserDataClass) override;
+	virtual const TArray<UAssetUserData*>* GetAssetUserDataArray() const override;
+	//~ End IInterface_AssetUserData Interface
+
 	//
 	// Dataflow
 	//
@@ -761,4 +768,8 @@ private:
 	TObjectPtr<UMaterialInterface> BoneSelectedMaterial = nullptr;
 
 	TSharedPtr<FGeometryCollection, ESPMode::ThreadSafe> GeometryCollection;
+
+	/** Array of user data stored with the asset */
+	UPROPERTY(EditAnywhere, AdvancedDisplay, Instanced, Category = AssetUserData)
+	TArray<TObjectPtr<UAssetUserData>> AssetUserData;
 };
