ChaosFlesh : Visiblity and Bounds check reduciton.
- Cleaned up perf macros
- Added a stub gameplay component.

diff --git a/Engine/Plugins/Experimental/ChaosFlesh/Source/ChaosFleshEngine/Private/ChaosFlesh/ChaosDeformableCollisionsComponent.cpp b/Engine/Plugins/Experimental/ChaosFlesh/Source/ChaosFleshEngine/Private/ChaosFlesh/ChaosDeformableCollisionsComponent.cpp
--- a/Engine/Plugins/Experimental/ChaosFlesh/Source/ChaosFleshEngine/Private/ChaosFlesh/ChaosDeformableCollisionsComponent.cpp
+++ b/Engine/Plugins/Experimental/ChaosFlesh/Source/ChaosFleshEngine/Private/ChaosFlesh/ChaosDeformableCollisionsComponent.cpp
@@ -13,6 +13,13 @@
 #include "PhysicsEngine/BodySetup.h"
 DEFINE_LOG_CATEGORY_STATIC(LogUDeformableCollisionsComponentInternal, Log, All);
 
+#define PERF_SCOPE(X) SCOPE_CYCLE_COUNTER(X); TRACE_CPUPROFILER_EVENT_SCOPE(X);
+DECLARE_CYCLE_STAT(TEXT("Chaos.Deformable.UDeformableCollisionsComponent.AddStaticMeshComponent"), STAT_ChaosDeformable_UDeformableCollisionsComponent_AddStaticMeshComponent, STATGROUP_Chaos);
+DECLARE_CYCLE_STAT(TEXT("Chaos.Deformable.UDeformableCollisionsComponent.RemoveStaticMeshComponent"), STAT_ChaosDeformable_UDeformableCollisionsComponent_RemoveStaticMeshComponent, STATGROUP_Chaos);
+DECLARE_CYCLE_STAT(TEXT("Chaos.Deformable.UDeformableCollisionsComponent.NewProxy"), STAT_ChaosDeformable_UDeformableCollisionsComponent_NewProxy, STATGROUP_Chaos);
+DECLARE_CYCLE_STAT(TEXT("Chaos.Deformable.UDeformableCollisionsComponent.NewDeformableData"), STAT_ChaosDeformable_UDeformableCollisionsComponent_NewDeformableData, STATGROUP_Chaos);
+
+
 UDeformableCollisionsComponent::UDeformableCollisionsComponent(const FObjectInitializer& ObjectInitializer)
 	: Super(ObjectInitializer)
 {
@@ -22,7 +29,9 @@ UDeformableCollisionsComponent::UDeformableCollisionsComponent(const FObjectInit
 
 void UDeformableCollisionsComponent::AddStaticMeshComponent(UStaticMeshComponent* StaticMeshComponent)
 {
-	if (StaticMeshComponent)
+	PERF_SCOPE(STAT_ChaosDeformable_UDeformableCollisionsComponent_AddStaticMeshComponent);
+
+		if (StaticMeshComponent)
 	{
 		if (!CollisionBodies.Contains(StaticMeshComponent))
 		{
@@ -34,6 +43,8 @@ void UDeformableCollisionsComponent::AddStaticMeshComponent(UStaticMeshComponent
 
 void UDeformableCollisionsComponent::RemoveStaticMeshComponent(UStaticMeshComponent* StaticMeshComponent)
 {
+	PERF_SCOPE(STAT_ChaosDeformable_UDeformableCollisionsComponent_RemoveStaticMeshComponent);
+
 	if (StaticMeshComponent)
 	{
 		if (CollisionBodies.Contains(StaticMeshComponent))
@@ -47,6 +58,8 @@ void UDeformableCollisionsComponent::RemoveStaticMeshComponent(UStaticMeshCompon
 
 UDeformablePhysicsComponent::FThreadingProxy* UDeformableCollisionsComponent::NewProxy()
 {
+	PERF_SCOPE(STAT_ChaosDeformable_UDeformableCollisionsComponent_NewProxy);
+
 	for (auto& Body : CollisionBodies)
 	{
 		if (Body)
@@ -65,6 +78,8 @@ template<> class Chaos::TSphere<float, 3>;
 UDeformableCollisionsComponent::FDataMapValue 
 UDeformableCollisionsComponent::NewDeformableData()
 {
+	PERF_SCOPE(STAT_ChaosDeformable_UDeformableCollisionsComponent_NewDeformableData);
+
 	TArray < Chaos::Softs::FCollisionObjectAddedBodies> AddedBodiesData;
 	TArray < Chaos::Softs::FCollisionObjectRemovedBodies> RemovedBodiesData;
 	TArray < Chaos::Softs::FCollisionObjectUpdatedBodies> UpdateBodiesData;
diff --git a/Engine/Plugins/Experimental/ChaosFlesh/Source/ChaosFleshEngine/Private/ChaosFlesh/ChaosDeformableGameplayComponent.cpp b/Engine/Plugins/Experimental/ChaosFlesh/Source/ChaosFleshEngine/Private/ChaosFlesh/ChaosDeformableGameplayComponent.cpp
new file mode 100644
--- /dev/null
+++ b/Engine/Plugins/Experimental/ChaosFlesh/Source/ChaosFleshEngine/Private/ChaosFlesh/ChaosDeformableGameplayComponent.cpp
@@ -0,0 +1,23 @@
+// Copyright Epic Games, Inc. All Rights Reserved.
+
+#include "ChaosFlesh/ChaosDeformableGameplayComponent.h"
+
+#include "ChaosStats.h"
+#include "Components/BoxComponent.h"
+#include "Chaos/DebugDrawQueue.h"
+
+
+#include UE_INLINE_GENERATED_CPP_BY_NAME(ChaosDeformableGameplayComponent)
+
+
+//DECLARE_CYCLE_STAT(TEXT("Chaos.Deformable.UDeformableGameplayComponent.StreamStaticMeshGeometry"), STAT_ChaosDeformable_UDeformableGameplayComponent_StreamStaticMeshGeometry, STATGROUP_Chaos);
+
+UDeformableGameplayComponent::UDeformableGameplayComponent(const FObjectInitializer& ObjectInitializer)
+	: Super(ObjectInitializer)
+{
+}
+
+
+UDeformableGameplayComponent::~UDeformableGameplayComponent()
+{
+}
diff --git a/Engine/Plugins/Experimental/ChaosFlesh/Source/ChaosFleshEngine/Private/ChaosFlesh/FleshComponent.cpp b/Engine/Plugins/Experimental/ChaosFlesh/Source/ChaosFleshEngine/Private/ChaosFlesh/FleshComponent.cpp
--- a/Engine/Plugins/Experimental/ChaosFlesh/Source/ChaosFleshEngine/Private/ChaosFlesh/FleshComponent.cpp
+++ b/Engine/Plugins/Experimental/ChaosFlesh/Source/ChaosFleshEngine/Private/ChaosFlesh/FleshComponent.cpp
@@ -28,9 +28,10 @@ FChaosEngineDeformableCVarParams CVarParams;
 FAutoConsoleVariableRef CVarDeforambleDoDrawSimulationMesh(TEXT("p.Chaos.DebugDraw.Deformable.SimulationMesh"), CVarParams.bDoDrawSimulationMesh, TEXT("Debug draw the deformable simulation resutls on the game thread. [def: true]"));
 FAutoConsoleVariableRef CVarDeforambleDoDrawSkeletalMeshBindingPositions(TEXT("p.Chaos.DebugDraw.Deformable.SkeletalMeshBindingPositions"), CVarParams.bDoDrawSkeletalMeshBindingPositions, TEXT("Debug draw the deformable simulation's SkeletalMeshBindingPositions on the game thread. [def: false]"));
 
+#define PERF_SCOPE(X) SCOPE_CYCLE_COUNTER(X); TRACE_CPUPROFILER_EVENT_SCOPE(X);
 DECLARE_CYCLE_STAT(TEXT("Chaos.Deformable.UFleshComponent.TickComponent"), STAT_ChaosDeformable_UFleshComponent_TickComponent, STATGROUP_Chaos);
 DECLARE_CYCLE_STAT(TEXT("Chaos.Deformable.UFleshComponent.NewProxy"), STAT_ChaosDeformable_UFleshComponent_NewProxy, STATGROUP_Chaos);
-DECLARE_CYCLE_STAT(TEXT("Chaos.Deformable.UFleshComponent.UpdateFromSimualtion"), STAT_ChaosDeformable_UFleshComponent_NewDeformableData, STATGROUP_Chaos);
+DECLARE_CYCLE_STAT(TEXT("Chaos.Deformable.UFleshComponent.NewDeformableData"), STAT_ChaosDeformable_UFleshComponent_NewDeformableData, STATGROUP_Chaos);
 DECLARE_CYCLE_STAT(TEXT("Chaos.Deformable.UFleshComponent.UpdateFromSimualtion"), STAT_ChaosDeformable_UFleshComponent_UpdateFromSimualtion, STATGROUP_Chaos);
 DECLARE_CYCLE_STAT(TEXT("Chaos.Deformable.UFleshComponent.RenderProceduralMesh"), STAT_ChaosDeformable_UFleshComponent_RenderProceduralMesh, STATGROUP_Chaos);
 
@@ -39,12 +40,12 @@ DEFINE_LOG_CATEGORY_STATIC(LogFleshComponentInternal, Log, All);
 UFleshComponent::UFleshComponent(const FObjectInitializer& ObjectInitializer)
 	: Super(ObjectInitializer)
 {
-	PrimaryComponentTick.bCanEverTick = true;
+	Mesh = ObjectInitializer.CreateDefaultSubobject<UProceduralMeshComponent>(this, TEXT("Flesh Visualization Component"));
 	PrimaryComponentTick.TickGroup = TG_LastDemotable;
+	PrimaryComponentTick.bCanEverTick = CVarParams.bDoDrawSimulationMesh;
+	bTickInEditor = CVarParams.bDoDrawSimulationMesh;
 
-	bTickInEditor = true;
 	DynamicCollection = ObjectInitializer.CreateDefaultSubobject<UFleshDynamicAsset>(this, TEXT("Flesh Dynamic Asset"));
-	Mesh = ObjectInitializer.CreateDefaultSubobject<UProceduralMeshComponent>(this, TEXT("Flesh Visualization Component"));
 }
 
 UFleshComponent::~UFleshComponent()
@@ -67,42 +68,6 @@ void UFleshComponent::OnRegister()
 	Super::OnRegister();
 }
 
-
-void UFleshComponent::BeginPlay()
-{
-	Super::BeginPlay();
-
-	if (ASkeletalMeshActor* SkeletalMeshActor = Cast<ASkeletalMeshActor>(GetOwner()))
-	{
-		if (USkeletalMeshComponent* SkeletalMeshComponent = SkeletalMeshActor->GetSkeletalMeshComponent())
-		{
-			PrimaryComponentTick.AddPrerequisite(SkeletalMeshComponent, SkeletalMeshComponent->PrimaryComponentTick);
-		}
-	}
-	if (PrimarySolverComponent)
-	{
-		PrimaryComponentTick.AddPrerequisite(PrimarySolverComponent, PrimarySolverComponent->PrimaryComponentTick);
-	}
-}
-
-void UFleshComponent::TickComponent(float DeltaTime, enum ELevelTick TickType, FActorComponentTickFunction *ThisTickFunction)
-{
-	SCOPE_CYCLE_COUNTER(STAT_ChaosDeformable_UFleshComponent_NewDeformableData);
-	TRACE_CPUPROFILER_EVENT_SCOPE(ChaosDeformable_UFleshComponent_NewDeformableData);
-
-	Super::TickComponent(DeltaTime, TickType, ThisTickFunction);
-
-	if (CVarParams.bDoDrawSimulationMesh)
-	{
-		RenderProceduralMesh();
-	}
-	else 
-	{
-		ResetProceduralMesh();
-	}
-	UpdateLocalBounds();
-}
-
 void UFleshComponent::EndPlay(const EEndPlayReason::Type ReasonEnd)
 {
 	if (GetDynamicCollection())
@@ -122,8 +87,7 @@ void UFleshComponent::SetRestCollection(const UFleshAsset* InRestCollection)
 
 UDeformablePhysicsComponent::FThreadingProxy* UFleshComponent::NewProxy()
 {
-	SCOPE_CYCLE_COUNTER(STAT_ChaosDeformable_UFleshComponent_NewProxy);
-	TRACE_CPUPROFILER_EVENT_SCOPE(ChaosDeformable_UFleshComponent_NewProxy);
+	PERF_SCOPE(STAT_ChaosDeformable_UFleshComponent_NewProxy);
 
 	UpdateSimSpaceTransformIndex();
 	if (const UFleshAsset* RestAsset = GetRestCollection())
@@ -157,8 +121,7 @@ UDeformablePhysicsComponent::FThreadingProxy* UFleshComponent::NewProxy()
 
 UDeformablePhysicsComponent::FDataMapValue UFleshComponent::NewDeformableData()
 {
-	SCOPE_CYCLE_COUNTER(STAT_ChaosDeformable_UFleshComponent_NewDeformableData);
-	TRACE_CPUPROFILER_EVENT_SCOPE(ChaosDeformable_UFleshComponent_NewDeformableData);
+	PERF_SCOPE(STAT_ChaosDeformable_UFleshComponent_NewDeformableData);
 
 	using namespace GeometryCollection::Facades;
 	if (GetOwner())
@@ -327,8 +290,7 @@ FTransform UFleshComponent::GetSimSpaceRestTransform() const
 
 void UFleshComponent::UpdateFromSimualtion(const FDataMapValue* SimualtionBuffer)
 {
-	SCOPE_CYCLE_COUNTER(STAT_ChaosDeformable_UFleshComponent_UpdateFromSimualtion);
-	TRACE_CPUPROFILER_EVENT_SCOPE(ChaosDeformable_UFleshComponent_UpdateFromSimualtion);
+	PERF_SCOPE(STAT_ChaosDeformable_UFleshComponent_UpdateFromSimualtion);
 
 	if (const FFleshThreadingProxy::FFleshOutputBuffer* FleshBuffer = (*SimualtionBuffer)->As<FFleshThreadingProxy::FFleshOutputBuffer>())
 	{
@@ -401,9 +363,45 @@ void UFleshComponent::ResetDynamicCollection()
 	}
 }
 
+//
+// Rendering Support
+//
+
+void UFleshComponent::BeginPlay()
+{
+	Super::BeginPlay();
+
+	if (ASkeletalMeshActor* SkeletalMeshActor = Cast<ASkeletalMeshActor>(GetOwner()))
+	{
+		if (USkeletalMeshComponent* SkeletalMeshComponent = SkeletalMeshActor->GetSkeletalMeshComponent())
+		{
+			PrimaryComponentTick.AddPrerequisite(SkeletalMeshComponent, SkeletalMeshComponent->PrimaryComponentTick);
+		}
+	}
+	if (PrimarySolverComponent)
+	{
+		PrimaryComponentTick.AddPrerequisite(PrimarySolverComponent, PrimarySolverComponent->PrimaryComponentTick);
+	}
+}
+
+void UFleshComponent::TickComponent(float DeltaTime, enum ELevelTick TickType, FActorComponentTickFunction* ThisTickFunction)
+{
+	PERF_SCOPE(STAT_ChaosDeformable_UFleshComponent_TickComponent);
+	Super::TickComponent(DeltaTime, TickType, ThisTickFunction);
+
+	if (IsVisible())
+	{
+		RenderProceduralMesh();
+	}
+	else
+	{
+		ResetProceduralMesh();
+	}
+}
+
 void UFleshComponent::ResetProceduralMesh()
 {
-	if (Mesh )
+	if (Mesh && RenderMesh)
 	{
 		Mesh->ClearAllMeshSections();
 	}
@@ -416,9 +414,6 @@ void UFleshComponent::ResetProceduralMesh()
 
 void UFleshComponent::RenderProceduralMesh()
 {
-	SCOPE_CYCLE_COUNTER(STAT_ChaosDeformable_UFleshComponent_RenderProceduralMesh);
-	TRACE_CPUPROFILER_EVENT_SCOPE(ChaosDeformable_UFleshComponent_RenderProceduralMesh);
-
 	bool bCanRender = false;
 	if (const UFleshAsset* FleshAsset = GetRestCollection())
 	{
@@ -430,6 +425,8 @@ void UFleshComponent::RenderProceduralMesh()
 			{
 				if (Mesh && CVarParams.bDoDrawSimulationMesh)
 				{
+					PERF_SCOPE(STAT_ChaosDeformable_UFleshComponent_RenderProceduralMesh);
+
 					if (const FFleshCollection* Flesh = FleshAsset->GetCollection())
 					{
 						int32 NumVertices = Flesh->NumElements(FGeometryCollection::VerticesGroup);
diff --git a/Engine/Plugins/Experimental/ChaosFlesh/Source/ChaosFleshEngine/Public/ChaosFlesh/ChaosDeformableGameplayComponent.h b/Engine/Plugins/Experimental/ChaosFlesh/Source/ChaosFleshEngine/Public/ChaosFlesh/ChaosDeformableGameplayComponent.h
new file mode 100644
--- /dev/null
+++ b/Engine/Plugins/Experimental/ChaosFlesh/Source/ChaosFleshEngine/Public/ChaosFlesh/ChaosDeformableGameplayComponent.h
@@ -0,0 +1,26 @@
+// Copyright Epic Games, Inc. All Rights Reserved.
+#pragma once
+
+#include "CoreMinimal.h"
+#include "ChaosFlesh/FleshComponent.h"
+#include "UObject/ObjectMacros.h"
+
+
+#include "ChaosDeformableGameplayComponent.generated.h"
+
+class UBoxComponent;
+class UStaticMeshComponent;
+
+/**
+*	UDeformableGameplayComponent
+*/
+UCLASS(meta = (BlueprintSpawnableComponent))
+class CHAOSFLESHENGINE_API UDeformableGameplayComponent : public UFleshComponent
+{
+	GENERATED_UCLASS_BODY()
+
+public:
+	~UDeformableGameplayComponent();
+
+};
+
diff --git a/Engine/Source/Runtime/Experimental/Chaos/Private/Chaos/Deformable/ChaosDeformableSolver.cpp b/Engine/Source/Runtime/Experimental/Chaos/Private/Chaos/Deformable/ChaosDeformableSolver.cpp
--- a/Engine/Source/Runtime/Experimental/Chaos/Private/Chaos/Deformable/ChaosDeformableSolver.cpp
+++ b/Engine/Source/Runtime/Experimental/Chaos/Private/Chaos/Deformable/ChaosDeformableSolver.cpp
@@ -34,6 +34,7 @@
 #include "HAL/IConsoleManager.h"
 #include "Misc/Paths.h"
 
+#define PERF_SCOPE(X) SCOPE_CYCLE_COUNTER(X); TRACE_CPUPROFILER_EVENT_SCOPE(X);
 DECLARE_CYCLE_STAT(TEXT("Chaos.Deformable.Solver.Constructor"), STAT_ChaosDeformableSolver_Constructor, STATGROUP_Chaos);
 DECLARE_CYCLE_STAT(TEXT("Chaos.Deformable.Solver.Destructor"), STAT_ChaosDeformableSolver_Destructor, STATGROUP_Chaos);
 DECLARE_CYCLE_STAT(TEXT("Chaos.Deformable.Solver.UpdateProxyInputPackages"), STAT_ChaosDeformableSolver_UpdateProxyInputPackages, STATGROUP_Chaos);
@@ -91,15 +92,13 @@ namespace Chaos::Softs
 		, PreviousInputPackage(TUniquePtr < FDeformablePackage >(nullptr))
 		, Property(InProp)
 	{
-		SCOPE_CYCLE_COUNTER(STAT_ChaosDeformableSolver_Constructor);
-		TRACE_CPUPROFILER_EVENT_SCOPE(ChaosDeformableSolver_Constructor);
+		PERF_SCOPE(STAT_ChaosDeformableSolver_Constructor);
 		Reset(Property);
 	}
 
 	FDeformableSolver::~FDeformableSolver()
 	{
-		SCOPE_CYCLE_COUNTER(STAT_ChaosDeformableSolver_Destructor);
-		TRACE_CPUPROFILER_EVENT_SCOPE(ChaosDeformableSolver_Destructor);
+		PERF_SCOPE(STAT_ChaosDeformableSolver_Destructor);
 
 		FScopeLock Lock(&InitializationMutex);
 		for (FThreadingProxy* Proxy : UninitializedProxys_Internal)
@@ -113,8 +112,7 @@ namespace Chaos::Softs
 
 	void FDeformableSolver::Reset(const FDeformableSolverProperties& InProps)
 	{
-		SCOPE_CYCLE_COUNTER(STAT_ChaosDeformableSolver_Reset);
-		TRACE_CPUPROFILER_EVENT_SCOPE(ChaosDeformableSolver_Reset);
+		PERF_SCOPE(STAT_ChaosDeformableSolver_Reset);
 
 		Property = InProps;
 		MObjects = TArrayCollectionArray<const UObject*>();
@@ -161,8 +159,7 @@ namespace Chaos::Softs
 
 	void FDeformableSolver::Simulate(FSolverReal DeltaTime)
 	{
-		SCOPE_CYCLE_COUNTER(STAT_ChaosDeformableSolver_Simulate);
-		TRACE_CPUPROFILER_EVENT_SCOPE(ChaosDeformableSolver_Simulate);
+		PERF_SCOPE(STAT_ChaosDeformableSolver_Simulate);
 
 		if (Property.NumSolverIterations)
 		{
@@ -199,8 +196,7 @@ namespace Chaos::Softs
 
 	void FDeformableSolver::InitializeSimulationObjects()
 	{
-		SCOPE_CYCLE_COUNTER(STAT_ChaosDeformableSolver_InitializeSimulationObjects);
-		TRACE_CPUPROFILER_EVENT_SCOPE(ChaosDeformableSolver_InitializeSimulationObjects);
+		PERF_SCOPE(STAT_ChaosDeformableSolver_InitializeSimulationObjects);
 		{
 			FScopeLock Lock(&InitializationMutex); // @todo(flesh) : change to threaded task based commands to prevent the lock. 
 			if (UninitializedProxys_Internal.Num())
@@ -232,8 +228,7 @@ namespace Chaos::Softs
 
 	void FDeformableSolver::UpdateSimulationObjects(FSolverReal DeltaTime)
 	{
-		SCOPE_CYCLE_COUNTER(STAT_ChaosDeformableSolver_InitializeSimulationObjects);
-		TRACE_CPUPROFILER_EVENT_SCOPE(ChaosDeformableSolver_InitializeSimulationObjects);
+		PERF_SCOPE(STAT_ChaosDeformableSolver_InitializeSimulationObjects);
 
 		typedef TPair< FThreadingProxy::FKey, TUniquePtr<FThreadingProxy> > FType;
 		for (const FType& Entry : Proxies)
@@ -251,8 +246,7 @@ namespace Chaos::Softs
 
 	void FDeformableSolver::InitializeSimulationObject(FThreadingProxy& InProxy)
 	{
-		SCOPE_CYCLE_COUNTER(STAT_ChaosDeformableSolver_InitializeSimulationObject);
-		TRACE_CPUPROFILER_EVENT_SCOPE(ChaosDeformableSolver_InitializeSimulationObject);
+		PERF_SCOPE(STAT_ChaosDeformableSolver_InitializeSimulationObject);
 
 		if (FFleshThreadingProxy* Proxy = InProxy.As<FFleshThreadingProxy>())
 		{
@@ -273,8 +267,7 @@ namespace Chaos::Softs
 
 	void FDeformableSolver::InitializeDeformableParticles(FFleshThreadingProxy& Proxy)
 	{
-		SCOPE_CYCLE_COUNTER(STAT_ChaosDeformableSolver_InitializeDeformableParticles);
-		TRACE_CPUPROFILER_EVENT_SCOPE(ChaosDeformableSolver_InitializeDeformableParticles);
+		PERF_SCOPE(STAT_ChaosDeformableSolver_InitializeDeformableParticles);
 
 		const FManagedArrayCollection& Dynamic = Proxy.GetDynamicCollection();
 		const FManagedArrayCollection& Rest = Proxy.GetRestCollection();
@@ -391,8 +384,7 @@ namespace Chaos::Softs
 
 	void FDeformableSolver::InitializeKinematicParticles(FFleshThreadingProxy& Proxy)
 	{
-		SCOPE_CYCLE_COUNTER(STAT_ChaosDeformableSolver_InitializeKinematicParticles);
-		TRACE_CPUPROFILER_EVENT_SCOPE(ChaosDeformableSolver_InitializeKinematicParticles);
+		PERF_SCOPE(STAT_ChaosDeformableSolver_InitializeKinematicParticles);
 
 		const FManagedArrayCollection& Rest = Proxy.GetRestCollection();
 		const FIntVector2& Range = Proxy.GetSolverParticleRange();
@@ -424,8 +416,8 @@ namespace Chaos::Softs
 
 	void FDeformableSolver::InitializeWeakConstraint(FFleshThreadingProxy& Proxy)
 	{
-		SCOPE_CYCLE_COUNTER(STAT_ChaosDeformableSolver_InitializeWeakConstraints);
-		TRACE_CPUPROFILER_EVENT_SCOPE(ChaosDeformableSolver_InitializeWeakConstraints);
+		PERF_SCOPE(STAT_ChaosDeformableSolver_InitializeWeakConstraints);
+
 		const FManagedArrayCollection& Rest = Proxy.GetRestCollection();
 		const FIntVector2& Range = Proxy.GetSolverParticleRange();
 		if (Property.bEnablePositionTargets)
@@ -484,14 +476,12 @@ namespace Chaos::Softs
 
 	void FDeformableSolver::InitializeCollisionBodies(FCollisionManagerProxy& Proxy)
 	{
-		SCOPE_CYCLE_COUNTER(STAT_ChaosDeformableSolver_InitializeCollisionBodies);
-		TRACE_CPUPROFILER_EVENT_SCOPE(ChaosDeformableSolver_InitializeCollisionBodies);
+		PERF_SCOPE(STAT_ChaosDeformableSolver_InitializeCollisionBodies);
 	}
 
 	void FDeformableSolver::UpdateCollisionBodies(FCollisionManagerProxy& Proxy, FThreadingProxy::FKey Owner, FSolverReal DeltaTime)
 	{
-		SCOPE_CYCLE_COUNTER(STAT_ChaosDeformableSolver_InitializeCollisionBodies);
-		TRACE_CPUPROFILER_EVENT_SCOPE(ChaosDeformableSolver_InitializeCollisionBodies);
+		PERF_SCOPE(STAT_ChaosDeformableSolver_InitializeCollisionBodies);
 
 		FCollisionManagerProxy::FCollisionsInputBuffer* CollisionsInputBuffer = nullptr;
 		if (this->CurrentInputPackage && this->CurrentInputPackage->ObjectMap.Contains(Owner))
@@ -505,7 +495,7 @@ namespace Chaos::Softs
 					{
 						if (AddBody.Shapes)
 						{
-							if (ensure(!Proxy.CollisionBodies.Contains(AddBody.Key)))
+							if (!Proxy.CollisionBodies.Contains(AddBody.Key))
 							{
 								int32 Index = Evolution->AddCollisionParticle(INDEX_NONE, true);
 								int32 ViewIndex = Evolution->CollisionParticlesActiveView().GetRanges().Num() - 1;
@@ -531,7 +521,7 @@ namespace Chaos::Softs
 					}
 					for (auto& KeyToRemove : KeysToRemove)
 					{
-						if (ensure(Proxy.CollisionBodies.Contains(KeyToRemove)))
+						if (Proxy.CollisionBodies.Contains(KeyToRemove))
 						{
 							int32 ParticleIndex = Proxy.CollisionBodies[KeyToRemove].ParticleIndex;
 							int32 ViewIndex = Proxy.CollisionBodies[KeyToRemove].ActiveViewIndex;
@@ -573,8 +563,7 @@ namespace Chaos::Softs
 
 	void FDeformableSolver::InitializeTetrahedralConstraint(FFleshThreadingProxy& Proxy)
 	{
-		SCOPE_CYCLE_COUNTER(STAT_ChaosDeformableSolver_InitializeTetrahedralConstraint);
-		TRACE_CPUPROFILER_EVENT_SCOPE(ChaosDeformableSolver_InitializeTetrahedralConstraint);
+		PERF_SCOPE(STAT_ChaosDeformableSolver_InitializeTetrahedralConstraint);
 
 		const FManagedArrayCollection& Rest = Proxy.GetRestCollection();
 
@@ -732,8 +721,7 @@ namespace Chaos::Softs
 
 	void FDeformableSolver::InitializeGidBasedConstraints(FFleshThreadingProxy& Proxy)
 	{
-		SCOPE_CYCLE_COUNTER(STAT_ChaosDeformableSolver_InitializeGidBasedConstraints);
-		TRACE_CPUPROFILER_EVENT_SCOPE(ChaosDeformableSolver_InitializeGidBasedConstraints);
+		PERF_SCOPE(STAT_ChaosDeformableSolver_InitializeGidBasedConstraints);
 
 		if (Property.bUseGridBasedConstraints)
 		{
@@ -761,8 +749,7 @@ namespace Chaos::Softs
 	{
 		auto MKineticUpdate = [this](FSolverParticles& MParticles, const FSolverReal Dt, const FSolverReal MTime, const int32 Index)
 		{
-			SCOPE_CYCLE_COUNTER(STAT_ChaosDeformableSolver_InitializeKinematicConstraint);
-			TRACE_CPUPROFILER_EVENT_SCOPE(ChaosDeformableSolver_InitializeKinematicConstraint);
+			PERF_SCOPE(STAT_ChaosDeformableSolver_InitializeKinematicConstraint);
 
 			if (0 <= Index && Index < this->MObjects.Num())
 			{
@@ -915,8 +902,7 @@ namespace Chaos::Softs
 
 	void FDeformableSolver::InitializeSelfCollisionVariables()
 	{
-		SCOPE_CYCLE_COUNTER(STAT_ChaosDeformableSolver_InitializeSelfCollisionVariables);
-		TRACE_CPUPROFILER_EVENT_SCOPE(ChaosDeformableSolver_InitializeSelfCollisionVariables);
+		PERF_SCOPE(STAT_ChaosDeformableSolver_InitializeSelfCollisionVariables);
 
 
 		int32 NumParticles = Evolution->Particles().Size();
@@ -948,9 +934,7 @@ namespace Chaos::Softs
 
 	void FDeformableSolver::InitializeGridBasedConstraintVariables()
 	{
-		SCOPE_CYCLE_COUNTER(STAT_ChaosDeformableSolver_InitializeGridBasedConstraintVariables);
-		TRACE_CPUPROFILER_EVENT_SCOPE(ChaosDeformableSolver_InitializeGridBasedConstraintVariables);
-
+		PERF_SCOPE(STAT_ChaosDeformableSolver_InitializeGridBasedConstraintVariables);
 
 		GridBasedCorotatedConstraint.Reset(new Chaos::Softs::FXPBDGridBasedCorotatedConstraints<FSolverReal, FSolverParticles>(
 			Evolution->Particles(), *AllElements, Property.GridDx, /*bRecordMetric = */false, (Chaos::Softs::FSolverReal).1, (Chaos::Softs::FSolverReal).01, (Chaos::Softs::FSolverReal).4, (Chaos::Softs::FSolverReal)1000.0));
@@ -979,8 +963,7 @@ namespace Chaos::Softs
 
 	void FDeformableSolver::RemoveSimulationObjects()
 	{
-		SCOPE_CYCLE_COUNTER(STAT_ChaosDeformableSolver_RemoveSimulationObjects);
-		TRACE_CPUPROFILER_EVENT_SCOPE(ChaosDeformableSolver_RemoveSimulationObjects);
+		PERF_SCOPE(STAT_ChaosDeformableSolver_RemoveSimulationObjects);
 
 		TArray< FThreadingProxy* > RemovedProxies;
 		{
@@ -1002,8 +985,11 @@ namespace Chaos::Softs
 					if (Proxy->CanSimulate())
 					{
 						FIntVector2 Indices = Proxy->GetSolverParticleRange();
-						Proxies.FindAndRemoveChecked(MObjects[Indices[0]]);
-						Evolution->Particles().RemoveAt(Indices[0], Indices[1]);
+						if (Indices[1] > 0)
+						{
+							Proxies.FindAndRemoveChecked(MObjects[Indices[0]]);
+							Evolution->Particles().RemoveAt(Indices[0], Indices[1]);
+						}
 					}
 				}
 			}
@@ -1045,8 +1031,7 @@ namespace Chaos::Softs
 
 	void FDeformableSolver::AdvanceDt(FSolverReal DeltaTime)
 	{
-		SCOPE_CYCLE_COUNTER(STAT_ChaosDeformableSolver_AdvanceDt);
-		TRACE_CPUPROFILER_EVENT_SCOPE(ChaosDeformableSolver_AdvanceDt);
+		PERF_SCOPE(STAT_ChaosDeformableSolver_AdvanceDt);
 
 		EventPreSolve.Broadcast(DeltaTime);
 
@@ -1115,8 +1100,7 @@ namespace Chaos::Softs
 
 	void FDeformableSolver::PushInputPackage(int32 InFrame, FDeformableDataMap&& InPackage)
 	{
-		SCOPE_CYCLE_COUNTER(STAT_ChaosDeformableSolver_PushInputPackage);
-		TRACE_CPUPROFILER_EVENT_SCOPE(ChaosDeformableSolver_PushInputPackage);
+		PERF_SCOPE(STAT_ChaosDeformableSolver_PushInputPackage);
 
 		FScopeLock Lock(&PackageInputMutex);
 		TRACE_CPUPROFILER_EVENT_SCOPE(ChaosDeformableSolver_PushInputPackage);
@@ -1125,8 +1109,7 @@ namespace Chaos::Softs
 
 	TUniquePtr<FDeformablePackage> FDeformableSolver::PullInputPackage()
 	{
-		SCOPE_CYCLE_COUNTER(STAT_ChaosDeformableSolver_PullInputPackage);
-		TRACE_CPUPROFILER_EVENT_SCOPE(ChaosDeformableSolver_PullInputPackage);
+		PERF_SCOPE(STAT_ChaosDeformableSolver_PullInputPackage);
 
 		FScopeLock Lock(&PackageInputMutex);
 		TRACE_CPUPROFILER_EVENT_SCOPE(ChaosDeformableSolver_PullInputPackage);
@@ -1137,8 +1120,7 @@ namespace Chaos::Softs
 
 	void FDeformableSolver::UpdateProxyInputPackages()
 	{
-		SCOPE_CYCLE_COUNTER(STAT_ChaosDeformableSolver_UpdateProxyInputPackages);
-		TRACE_CPUPROFILER_EVENT_SCOPE(ChaosDeformableSolver_UpdateProxyInputPackages);
+		PERF_SCOPE(STAT_ChaosDeformableSolver_UpdateProxyInputPackages);
 
 		if (CurrentInputPackage)
 		{
@@ -1156,9 +1138,7 @@ namespace Chaos::Softs
 
 	void FDeformableSolver::Update(FSolverReal DeltaTime)
 	{
-		SCOPE_CYCLE_COUNTER(STAT_ChaosDeformableSolver_Update);
-		TRACE_CPUPROFILER_EVENT_SCOPE(ChaosDeformableSolver_Update);
-
+		PERF_SCOPE(STAT_ChaosDeformableSolver_Update);
 
 		if (!Proxies.Num()) return;
 
@@ -1179,21 +1159,15 @@ namespace Chaos::Softs
 
 	void FDeformableSolver::PushOutputPackage(int32 InFrame, FDeformableDataMap&& InPackage)
 	{
-		SCOPE_CYCLE_COUNTER(STAT_ChaosDeformableSolver_PushOutputPackage);
-		TRACE_CPUPROFILER_EVENT_SCOPE(ChaosDeformableSolver_PushOutputPackage);
-
 		FScopeLock Lock(&PackageOutputMutex);
-		TRACE_CPUPROFILER_EVENT_SCOPE(ChaosDeformableSolver_PushPackage);
+		PERF_SCOPE(STAT_ChaosDeformableSolver_PushOutputPackage);
 		BufferedOutputPackages.Push(TUniquePtr< FDeformablePackage >(new FDeformablePackage(InFrame, MoveTemp(InPackage))));
 	}
 
 	TUniquePtr<FDeformablePackage> FDeformableSolver::PullOutputPackage()
 	{
-		SCOPE_CYCLE_COUNTER(STAT_ChaosDeformableSolver_PullOutputPackage);
-		TRACE_CPUPROFILER_EVENT_SCOPE(ChaosDeformableSolver_PullOutputPackage);
-
 		FScopeLock Lock(&PackageOutputMutex);
-		TRACE_CPUPROFILER_EVENT_SCOPE(ChaosDeformableSolver_PullPackage);
+		PERF_SCOPE(STAT_ChaosDeformableSolver_PullOutputPackage);
 		if (BufferedOutputPackages.Num())
 			return BufferedOutputPackages.Pop();
 		return TUniquePtr<FDeformablePackage>(nullptr);
@@ -1201,8 +1175,7 @@ namespace Chaos::Softs
 
 	void FDeformableSolver::AddProxy(FThreadingProxy* InProxy)
 	{
-		SCOPE_CYCLE_COUNTER(STAT_ChaosDeformableSolver_AddProxy);
-		TRACE_CPUPROFILER_EVENT_SCOPE(ChaosDeformableSolver_AddProxy);
+		PERF_SCOPE(STAT_ChaosDeformableSolver_AddProxy);
 
 		FScopeLock Lock(&InitializationMutex);
 		UninitializedProxys_Internal.Add(InProxy);
@@ -1211,8 +1184,7 @@ namespace Chaos::Softs
 
 	void FDeformableSolver::RemoveProxy(FThreadingProxy* InProxy)
 	{
-		SCOPE_CYCLE_COUNTER(STAT_ChaosDeformableSolver_RemoveProxy);
-		TRACE_CPUPROFILER_EVENT_SCOPE(ChaosDeformableSolver_RemoveProxy);
+		PERF_SCOPE(STAT_ChaosDeformableSolver_RemoveProxy);
 
 		FScopeLock LockA(&RemovalMutex);
 		FScopeLock LockB(&InitializationMutex);
@@ -1243,10 +1215,8 @@ namespace Chaos::Softs
 
 	void FDeformableSolver::UpdateOutputState(FThreadingProxy& InProxy)
 	{
-		SCOPE_CYCLE_COUNTER(STAT_ChaosDeformableSolver_UpdateOutputState);
-		TRACE_CPUPROFILER_EVENT_SCOPE(ChaosDeformableSolver_UpdateOutputState);
+		PERF_SCOPE(STAT_ChaosDeformableSolver_UpdateOutputState);
 
-		TRACE_CPUPROFILER_EVENT_SCOPE(ChaosDeformableSolver_UpdateOutputState);
 		if (FFleshThreadingProxy* Proxy = InProxy.As<FFleshThreadingProxy>())
 		{
 			const FIntVector2& Range = Proxy->GetSolverParticleRange();
@@ -1400,4 +1370,5 @@ namespace Chaos::Softs
 
 	}
 
+
 }; // Namespace Chaos::Softs
diff --git a/Engine/Source/Runtime/Experimental/Chaos/Private/Chaos/Deformable/ChaosDeformableSolverThreadingWrapper.cpp b/Engine/Source/Runtime/Experimental/Chaos/Private/Chaos/Deformable/ChaosDeformableSolverThreadingWrapper.cpp
--- a/Engine/Source/Runtime/Experimental/Chaos/Private/Chaos/Deformable/ChaosDeformableSolverThreadingWrapper.cpp
+++ b/Engine/Source/Runtime/Experimental/Chaos/Private/Chaos/Deformable/ChaosDeformableSolverThreadingWrapper.cpp
@@ -162,15 +162,4 @@ namespace Chaos::Softs
 			Solver->UpdateOutputState(InProxy);
 	}
 
-	void  FDeformableSolver::FPhysicsThreadAccess::WriteFrame(FThreadingProxy& InProxy, const FSolverReal DeltaTime)
-	{
-		if(Solver) 
-			Solver->WriteFrame(InProxy, DeltaTime);
-	}
-
-	void  FDeformableSolver::FPhysicsThreadAccess::WriteTrisGEO(const FSolverParticles& Particles, const TArray<TVec3<int32>>& Mesh)
-	{
-		if(Solver) 
-			Solver->WriteTrisGEO(Particles, Mesh);
-	}
 }; // Namespace Chaos::Softs
\ No newline at end of file
diff --git a/Engine/Source/Runtime/Experimental/Chaos/Public/Chaos/Deformable/ChaosDeformableSolver.h b/Engine/Source/Runtime/Experimental/Chaos/Public/Chaos/Deformable/ChaosDeformableSolver.h
--- a/Engine/Source/Runtime/Experimental/Chaos/Public/Chaos/Deformable/ChaosDeformableSolver.h
+++ b/Engine/Source/Runtime/Experimental/Chaos/Public/Chaos/Deformable/ChaosDeformableSolver.h
@@ -58,10 +58,6 @@ namespace Chaos::Softs
 			void InitializeSelfCollisionVariables();
 			void RemoveSimulationObjects();
 
-			/*IO Utility*/
-			void WriteTrisGEO(const Softs::FSolverParticles& Particles, const TArray<TVec3<int32>>& Mesh);
-			void WriteFrame(FThreadingProxy& , const FSolverReal DeltaTime);
-
 			const FDeformableSolverProperties& GetProperties() const { return Solver->GetProperties(); }
 
 			FPBDEvolution* GetEvolution() { return Solver->Evolution.Get(); }
@@ -132,14 +128,11 @@ namespace Chaos::Softs
 		void RemoveSimulationObjects();
 		TArray<Chaos::TVec3<FSolverReal>> ComputeParticleTargets(const TArray<TArray<int32>>& ParticleIndices);
 
-
-		/*IO Utility*/
-		void WriteTrisGEO(const FSolverParticles& Particles, const TArray<TVec3<int32>>& Mesh);
-		void WriteFrame(FThreadingProxy&, const FSolverReal DeltaTime);
-
-		/*Debug Drawing*/
+		/*Debug Output*/
 		void DebugDrawSimulationData();
 		void DebugDrawTetrahedralParticles(FFleshThreadingProxy& Proxy);
+		void WriteFrame(FThreadingProxy& InProxy, const FSolverReal DeltaTime);
+		void WriteTrisGEO(const FSolverParticles& Particles, const TArray<TVec3<int32>>& Mesh);
 
 		/*Game Thread API*/
 		bool HasObject(UObject* InObject) const { return InitializedObjects_External.Contains(InObject); }
