Fix missing headers causing non-unity build errors

diff --git a/Engine/Source/Runtime/Experimental/GeometryCollectionEngine/Private/GeometryCollection/GeometryCollectionObject.cpp b/Engine/Source/Runtime/Experimental/GeometryCollectionEngine/Private/GeometryCollection/GeometryCollectionObject.cpp
--- a/Engine/Source/Runtime/Experimental/GeometryCollectionEngine/Private/GeometryCollection/GeometryCollectionObject.cpp
+++ b/Engine/Source/Runtime/Experimental/GeometryCollectionEngine/Private/GeometryCollection/GeometryCollectionObject.cpp
@@ -21,6 +21,7 @@
 #include "PhysicsEngine/PhysicsSettings.h"
 #include "EditorFramework/AssetImportData.h"
 #include "Rendering/NaniteResources.h"
+#include "Engine/AssetUserData.h"
 
 #if WITH_EDITOR
 #include "GeometryCollection/DerivedDataGeometryCollectionCooker.h"
diff --git a/Engine/Source/Runtime/Experimental/GeometryCollectionEngine/Public/GeometryCollection/GeometryCollectionObject.h b/Engine/Source/Runtime/Experimental/GeometryCollectionEngine/Public/GeometryCollection/GeometryCollectionObject.h
--- a/Engine/Source/Runtime/Experimental/GeometryCollectionEngine/Public/GeometryCollection/GeometryCollectionObject.h
+++ b/Engine/Source/Runtime/Experimental/GeometryCollectionEngine/Public/GeometryCollection/GeometryCollectionObject.h
@@ -8,6 +8,7 @@
 #include "GeometryCollection/GeometryCollectionSimulationTypes.h"
 #include "GeometryCollection/ManagedArray.h"
 #include "InstanceUniformShaderParameters.h"
+#include "Interfaces/Interface_AssetUserData.h"
 #include "Misc/Crc.h"
 
 #include "GeometryCollectionObject.generated.h"
