split the MTD overlap computation from the axis overlap computation in the physics object interface

diff --git a/Engine/Source/Runtime/Experimental/Chaos/Private/Chaos/PhysicsObjectInterface.cpp b/Engine/Source/Runtime/Experimental/Chaos/Private/Chaos/PhysicsObjectInterface.cpp
--- a/Engine/Source/Runtime/Experimental/Chaos/Private/Chaos/PhysicsObjectInterface.cpp
+++ b/Engine/Source/Runtime/Experimental/Chaos/Private/Chaos/PhysicsObjectInterface.cpp
@@ -243,19 +243,123 @@ namespace Chaos
 	template<EThreadContext Id>
 	bool FReadPhysicsObjectInterface<Id>::GetPhysicsObjectOverlap(FPhysicsObjectHandle ObjectA, FPhysicsObjectHandle ObjectB, bool bTraceComplex, Chaos::FOverlapInfo& OutOverlap)
 	{
-		return GetPhysicsObjectOverlapWithTransform(ObjectA, FTransform::Identity, ObjectB, FTransform::Identity, bTraceComplex, OutOverlap);
+		TRACE_CPUPROFILER_EVENT_SCOPE(FReadPhysicsObjectInterface<Id>::GetPhysicsObjectOverlap);
+		// This is slow and inefficient and hence deprecated.
+		bool bRetOverlap = false;
+		if (OutOverlap.MTD)
+		{
+			bRetOverlap |= GetPhysicsObjectOverlapWithMTD(ObjectA, FTransform::Identity, ObjectB, FTransform::Identity, bTraceComplex, *OutOverlap.MTD);
+		}
+
+		if (OutOverlap.AxisOverlap)
+		{
+			bRetOverlap |= GetPhysicsObjectOverlapWithAABB(ObjectA, FTransform::Identity, ObjectB, FTransform::Identity, bTraceComplex, FVector::Zero(), *OutOverlap.AxisOverlap);
+		}
+
+		return bRetOverlap;
 	}
 
 	template<EThreadContext Id>
 	bool FReadPhysicsObjectInterface<Id>::GetPhysicsObjectOverlapWithTransform(FPhysicsObjectHandle ObjectA, const FTransform& InTransformA, FPhysicsObjectHandle ObjectB, const FTransform& InTransformB, bool bTraceComplex, Chaos::FOverlapInfo& OutOverlap)
 	{
 		TRACE_CPUPROFILER_EVENT_SCOPE(FReadPhysicsObjectInterface<Id>::GetPhysicsObjectOverlapWithTransform);
+		// This is slow and inefficient and hence deprecated.
+		bool bRetOverlap = false;
+		if (OutOverlap.MTD)
+		{
+			bRetOverlap |= GetPhysicsObjectOverlapWithMTD(ObjectA, InTransformA, ObjectB, InTransformB, bTraceComplex, *OutOverlap.MTD);
+		}
+
+		if (OutOverlap.AxisOverlap)
+		{
+			bRetOverlap |= GetPhysicsObjectOverlapWithAABB(ObjectA, InTransformA, ObjectB, InTransformB, bTraceComplex, FVector::Zero(), *OutOverlap.AxisOverlap);
+		}
+
+		return bRetOverlap;
+	}
+
+	template<EThreadContext Id>
+	bool FReadPhysicsObjectInterface<Id>::GetPhysicsObjectOverlap(FPhysicsObjectHandle ObjectA, const FTransform& InTransformA, FPhysicsObjectHandle ObjectB, const FTransform& InTransformB, bool bTraceComplex)
+	{
+		TRACE_CPUPROFILER_EVENT_SCOPE(FReadPhysicsObjectInterface<Id>::GetPhysicsObjectOverlap);
+		return PairwiseShapeOverlapHelper(ObjectA, InTransformA, ObjectB, InTransformB, bTraceComplex, false, [](const FShapeOverlapData&, const FShapeOverlapData&, const FMTDInfo&) { return false; });
+	}
+
+	template<EThreadContext Id>
+	bool FReadPhysicsObjectInterface<Id>::GetPhysicsObjectOverlapWithMTD(FPhysicsObjectHandle ObjectA, const FTransform& InTransformA, FPhysicsObjectHandle ObjectB, const FTransform& InTransformB, bool bTraceComplex, FMTDInfo& OutMTD)
+	{
+		TRACE_CPUPROFILER_EVENT_SCOPE(FReadPhysicsObjectInterface<Id>::GetPhysicsObjectOverlapWithMTD);
+		OutMTD.Penetration = 0.0;
+		return PairwiseShapeOverlapHelper(
+			ObjectA,
+			InTransformA,
+			ObjectB,
+			InTransformB,
+			bTraceComplex,
+			true,
+			[&OutMTD](const FShapeOverlapData&, const FShapeOverlapData&, const FMTDInfo& MTDInfo)
+			{
+				OutMTD.Penetration = FMath::Max(OutMTD.Penetration, MTDInfo.Penetration);
+				return true;
+			}
+		);
+	}
+
+	template<EThreadContext Id>
+	bool FReadPhysicsObjectInterface<Id>::GetPhysicsObjectOverlapWithAABB(FPhysicsObjectHandle ObjectA, const FTransform& InTransformA, FPhysicsObjectHandle ObjectB, const FTransform& InTransformB, bool bTraceComplex, const FVector& Tolerance, FBox& OutOverlap)
+	{
+		TRACE_CPUPROFILER_EVENT_SCOPE(FReadPhysicsObjectInterface<Id>::GetPhysicsObjectOverlapWithAABB);
+		OutOverlap = FBox{ EForceInit::ForceInitToZero };
+		return PairwiseShapeOverlapHelper(
+			ObjectA,
+			InTransformA,
+			ObjectB,
+			InTransformB,
+			bTraceComplex,
+			false,
+			[&OutOverlap, &Tolerance](const FShapeOverlapData& ShapeA, const FShapeOverlapData& ShapeB, const FMTDInfo&)
+			{
+				const FAABB3 BoxA = { ShapeA.BoundingBox.Min() + Tolerance, ShapeA.BoundingBox.Max() - Tolerance };
+				const FAABB3 BoxB = { ShapeB.BoundingBox.Min() + Tolerance, ShapeB.BoundingBox.Max() - Tolerance };
+				const FAABB3 Intersection = BoxA.GetIntersection(BoxB);
+				OutOverlap += FBox{ Intersection.Min(), Intersection.Max() };
+				return true;
+			}
+		);
+	}
+
+	template<EThreadContext Id>
+	bool FReadPhysicsObjectInterface<Id>::GetPhysicsObjectOverlapWithAABBSize(FPhysicsObjectHandle ObjectA, const FTransform& InTransformA, FPhysicsObjectHandle ObjectB, const FTransform& InTransformB, bool bTraceComplex, const FVector& Tolerance, FVector& OutOverlapSize)
+	{
+		TRACE_CPUPROFILER_EVENT_SCOPE(FReadPhysicsObjectInterface<Id>::GetPhysicsObjectOverlapWithAABBSize);
+		OutOverlapSize = FVector::Zero();
+		return PairwiseShapeOverlapHelper(
+			ObjectA,
+			InTransformA,
+			ObjectB,
+			InTransformB,
+			bTraceComplex,
+			false,
+			[&OutOverlapSize, &Tolerance](const FShapeOverlapData& ShapeA, const FShapeOverlapData& ShapeB, const FMTDInfo&)
+			{
+				const FAABB3 BoxA = { ShapeA.BoundingBox.Min() + Tolerance, ShapeA.BoundingBox.Max() - Tolerance };
+				const FAABB3 BoxB = { ShapeB.BoundingBox.Min() + Tolerance, ShapeB.BoundingBox.Max() - Tolerance };
+				const FAABB3 Intersection = BoxA.GetIntersection(BoxB);
+				OutOverlapSize += Intersection.Extents();
+				return true;
+			}
+		);
+	}
+
+	template<EThreadContext Id>
+	bool FReadPhysicsObjectInterface<Id>::PairwiseShapeOverlapHelper(FPhysicsObjectHandle ObjectA, const FTransform& InTransformA, FPhysicsObjectHandle ObjectB, const FTransform& InTransformB, bool bTraceComplex, bool bComputeMTD, const TFunction<bool(const FShapeOverlapData&, const FShapeOverlapData&, const FMTDInfo&)>& Lambda)
+	{
 		TArray<FPerShapeData*> ShapesA = GetAllShapes({ &ObjectA, 1 });
-		const FTransform TransformA = FTransform{ GetR(ObjectA), GetX(ObjectA) } * InTransformA;
+		const FTransform TransformA = FTransform{ GetR(ObjectA), GetX(ObjectA) } *InTransformA;
 		const FBox BoxA = GetWorldBounds({ &ObjectA, 1 }).TransformBy(InTransformA);
 
 		TArray<FPerShapeData*> ShapesB = GetAllShapes({ &ObjectB, 1 });
-		const FTransform TransformB = FTransform{ GetR(ObjectB), GetX(ObjectB) } * InTransformB;
+		const FTransform TransformB = FTransform{ GetR(ObjectB), GetX(ObjectB) } *InTransformB;
 		const FBox BoxB = GetWorldBounds({ &ObjectB, 1 }).TransformBy(InTransformB);
 
 		if (!BoxA.Intersect(BoxB))
@@ -263,17 +367,6 @@ namespace Chaos
 			return false;
 		}
 
-		if (OutOverlap.MTD)
-		{
-			OutOverlap.MTD->Penetration = 0.0;
-		}
-
-		if (OutOverlap.AxisOverlap)
-		{
-			*OutOverlap.AxisOverlap = FBox{ EForceInit::ForceInitToZero };
-		}
-		const bool bComputeMTD = OutOverlap.MTD != nullptr;
-
 		bool bFoundOverlap = false;
 		for (FPerShapeData* B : ShapesB)
 		{
@@ -327,22 +420,13 @@ namespace Chaos
 				if (bOverlap)
 				{
 					bFoundOverlap = true;
-					if (!OutOverlap.MTD && !OutOverlap.AxisOverlap)
-					{
-						// Don't care about computing extra overlap information so we can exit early.
-						return true;
-					}
 
-					if (OutOverlap.MTD && TmpMTDInfo.Penetration > OutOverlap.MTD->Penetration)
-					{
-						// If we need to find the MTD we need to find the largest overlap.
-						*OutOverlap.MTD = TmpMTDInfo;
-					}
+					FShapeOverlapData OverlapDataA = { A, BoxShapeA };
+					FShapeOverlapData OverlapDataB = { B, BoxShapeB };
 
-					if (OutOverlap.AxisOverlap)
+					if (!Lambda(OverlapDataA, OverlapDataB, TmpMTDInfo))
 					{
-						const FAABB3 Intersection = BoxShapeA.GetIntersection(BoxShapeB);
-						*OutOverlap.AxisOverlap += FBox{ Intersection.Min(), Intersection.Max() };
+						return true;
 					}
 				}
 			}
diff --git a/Engine/Source/Runtime/Experimental/Chaos/Public/Chaos/PhysicsObjectInterface.h b/Engine/Source/Runtime/Experimental/Chaos/Public/Chaos/PhysicsObjectInterface.h
--- a/Engine/Source/Runtime/Experimental/Chaos/Public/Chaos/PhysicsObjectInterface.h
+++ b/Engine/Source/Runtime/Experimental/Chaos/Public/Chaos/PhysicsObjectInterface.h
@@ -67,9 +67,23 @@ namespace Chaos
 		TArray<TThreadParticle<Id>*> GetAllParticles(TArrayView<FPhysicsObjectHandle> InObjects);
 		TArray<TThreadRigidParticle<Id>*> GetAllRigidParticles(TArrayView<FPhysicsObjectHandle> InObjects);
 		TArray<FPerShapeData*> GetAllShapes(TArrayView<FPhysicsObjectHandle> InObjects);
+
+		UE_DEPRECATED(5.3, "GetPhysicsObjectOverlap has been deprecated. Please use the function for the specific overlap metric you wish to compute instead.")
 		bool GetPhysicsObjectOverlap(FPhysicsObjectHandle ObjectA, FPhysicsObjectHandle ObjectB, bool bTraceComplex, Chaos::FOverlapInfo& OutOverlap);
+
+		UE_DEPRECATED(5.3, "GetPhysicsObjectOverlapWithTransform has been deprecated. Please use the function for the specific overlap metric you wish to compute instead.")
 		bool GetPhysicsObjectOverlapWithTransform(FPhysicsObjectHandle ObjectA, const FTransform& InTransformA, FPhysicsObjectHandle ObjectB, const FTransform& InTransformB, bool bTraceComplex, Chaos::FOverlapInfo& OutOverlap);
 
+		// This function will not compute any overlap heuristic.
+		bool GetPhysicsObjectOverlap(FPhysicsObjectHandle ObjectA, const FTransform& InTransformA, FPhysicsObjectHandle ObjectB, const FTransform& InTransformB, bool bTraceComplex);
+
+		// This function does the same as GetPhysicsObjectOverlap but also computes the MTD metric.
+		bool GetPhysicsObjectOverlapWithMTD(FPhysicsObjectHandle ObjectA, const FTransform& InTransformA, FPhysicsObjectHandle ObjectB, const FTransform& InTransformB, bool bTraceComplex, FMTDInfo& OutMTD);
+
+		// This function does the same as GetPhysicsObjectOverlap but also computes the AABB overlap metric.
+		bool GetPhysicsObjectOverlapWithAABB(FPhysicsObjectHandle ObjectA, const FTransform& InTransformA, FPhysicsObjectHandle ObjectB, const FTransform& InTransformB, bool bTraceComplex, const FVector& Tolerance, FBox& OutOverlap);
+		bool GetPhysicsObjectOverlapWithAABBSize(FPhysicsObjectHandle ObjectA, const FTransform& InTransformA, FPhysicsObjectHandle ObjectB, const FTransform& InTransformB, bool bTraceComplex, const FVector& Tolerance, FVector& OutOverlapSize);
+
 		bool AreAllValid(TArrayView<FPhysicsObjectHandle> InObjects);
 		bool AreAllKinematic(TArrayView<FPhysicsObjectHandle> InObjects);
 		bool AreAllSleeping(TArrayView<FPhysicsObjectHandle> InObjects);
@@ -85,6 +99,18 @@ namespace Chaos
 		friend class FPhysicsObjectInterface;
 	protected:
 		FReadPhysicsObjectInterface() = default;
+
+	private:
+		struct FShapeOverlapData
+		{
+			FPerShapeData* Shape;
+			FAABB3 BoundingBox;
+		};
+
+		/**
+		 * For every pair of shapes that overlap, allows the caller to perform some computation. If additional pairs of shapes need to be examined, the input TFunction should return true.
+		 */
+		bool PairwiseShapeOverlapHelper(FPhysicsObjectHandle ObjectA, const FTransform& InTransformA, FPhysicsObjectHandle ObjectB, const FTransform& InTransformB, bool bTraceComplex, bool bComputeMTD, const TFunction<bool(const FShapeOverlapData&, const FShapeOverlapData&, const FMTDInfo&)>& Lambda);
 	};
 
 	using FReadPhysicsObjectInterface_External = FReadPhysicsObjectInterface<EThreadContext::External>;
