Geometry collection : move GC to use ISM ( instead of HISm ) when using ISM Pool to temporarily fix an issue that breaks root proxy mesh visibility logic

diff --git a/Engine/Source/Runtime/Experimental/GeometryCollectionEngine/Private/GeometryCollection/GeometryCollectionComponent.cpp b/Engine/Source/Runtime/Experimental/GeometryCollectionEngine/Private/GeometryCollection/GeometryCollectionComponent.cpp
--- a/Engine/Source/Runtime/Experimental/GeometryCollectionEngine/Private/GeometryCollection/GeometryCollectionComponent.cpp
+++ b/Engine/Source/Runtime/Experimental/GeometryCollectionEngine/Private/GeometryCollection/GeometryCollectionComponent.cpp
@@ -114,10 +114,10 @@ FAutoConsoleVariableRef CVarChaosGCUseISMPoolForNonFracturedParts(TEXT("p.Chaos.
 bool bChaos_GC_ForceAutoAssignISMPool = false;
 FAutoConsoleVariableRef CVarChaosGCForceAutoAssignISMPool(TEXT("p.Chaos.GC.ForceAutoAssignISMPool"), bChaos_GC_ForceAutoAssignISMPool, TEXT("When enabled, force assignement of ISMPool regardgless of the settings of the components"));
 
-bool bChaos_GC_UseHierarchicalISMForProxyMesh = true;
+bool bChaos_GC_UseHierarchicalISMForProxyMesh = false;
 FAutoConsoleVariableRef CVarChaosGCUseHierarchicalISMForProxyMesh(TEXT("p.Chaos.GC.UseHierarchicalISMForProxyMesh"), bChaos_GC_UseHierarchicalISMForProxyMesh, TEXT("When enabled along with ISM Pool, proxy mesh will prefer using HISM vs standard ISM"));
 
-bool bChaos_GC_UseHierarchicalISMForLeafMeshes = true;
+bool bChaos_GC_UseHierarchicalISMForLeafMeshes = false;
 FAutoConsoleVariableRef CVarChaosGCUseHierarchicalISMForLeafMeshes(TEXT("p.Chaos.GC.UseHierarchicalISMForLeafMeshes"), bChaos_GC_UseHierarchicalISMForLeafMeshes, TEXT("When enabled along with ISM Pool, leaf meshes will prefer using HISM vs standard ISM"));
 
 bool bChaos_GC_InitConstantDataUseParallelFor = true;
