Chaos - fix outdated material access warning

diff --git a/Engine/Source/Runtime/Experimental/Chaos/Private/Chaos/PBDRigidClustering.cpp b/Engine/Source/Runtime/Experimental/Chaos/Private/Chaos/PBDRigidClustering.cpp
--- a/Engine/Source/Runtime/Experimental/Chaos/Private/Chaos/PBDRigidClustering.cpp
+++ b/Engine/Source/Runtime/Experimental/Chaos/Private/Chaos/PBDRigidClustering.cpp
@@ -567,10 +567,10 @@ namespace Chaos
 					// Can only take materials if the child has any - otherwise we fall back on defaults.
 					// Due to GC initialisation however, we should always have a valid material as even
 					// when one cannot be found we fall back on the default on GEngine
-					const int32 NumChildMaterials = ChildShapes[0]->GetMaterials().Num();
+					const int32 NumChildMaterials = ChildShapes[0]->NumMaterials();
 					if(NumChildMaterials > 0)
 					{
-						Chaos::FMaterialHandle ChildMat = ChildShapes[0]->GetMaterials()[0];
+						Chaos::FMaterialHandle ChildMat = ChildShapes[0]->GetMaterial(0);
 
 						for(const TUniquePtr<FPerShapeData>& PerShape : NewShapes)
 						{
