move the tolerance check into PairwiseShapeOverlapHelper to actually check if the shape boxes overlap

diff --git a/Engine/Source/Runtime/Experimental/Chaos/Private/Chaos/PhysicsObjectInterface.cpp b/Engine/Source/Runtime/Experimental/Chaos/Private/Chaos/PhysicsObjectInterface.cpp
--- a/Engine/Source/Runtime/Experimental/Chaos/Private/Chaos/PhysicsObjectInterface.cpp
+++ b/Engine/Source/Runtime/Experimental/Chaos/Private/Chaos/PhysicsObjectInterface.cpp
@@ -282,7 +282,7 @@ namespace Chaos
 	bool FReadPhysicsObjectInterface<Id>::GetPhysicsObjectOverlap(FPhysicsObjectHandle ObjectA, const FTransform& InTransformA, FPhysicsObjectHandle ObjectB, const FTransform& InTransformB, bool bTraceComplex)
 	{
 		TRACE_CPUPROFILER_EVENT_SCOPE(FReadPhysicsObjectInterface<Id>::GetPhysicsObjectOverlap);
-		return PairwiseShapeOverlapHelper(ObjectA, InTransformA, ObjectB, InTransformB, bTraceComplex, false, [](const FShapeOverlapData&, const FShapeOverlapData&, const FMTDInfo&) { return false; });
+		return PairwiseShapeOverlapHelper(ObjectA, InTransformA, ObjectB, InTransformB, bTraceComplex, false, FVector::Zero(), [](const FShapeOverlapData&, const FShapeOverlapData&, const FMTDInfo&) { return false; });
 	}
 
 	template<EThreadContext Id>
@@ -297,6 +297,7 @@ namespace Chaos
 			InTransformB,
 			bTraceComplex,
 			true,
+			FVector::Zero(),
 			[&OutMTD](const FShapeOverlapData&, const FShapeOverlapData&, const FMTDInfo& MTDInfo)
 			{
 				OutMTD.Penetration = FMath::Max(OutMTD.Penetration, MTDInfo.Penetration);
@@ -317,11 +318,10 @@ namespace Chaos
 			InTransformB,
 			bTraceComplex,
 			false,
+			Tolerance,
 			[&OutOverlap, &Tolerance](const FShapeOverlapData& ShapeA, const FShapeOverlapData& ShapeB, const FMTDInfo&)
 			{
-				const FAABB3 BoxA = { ShapeA.BoundingBox.Min() + Tolerance, ShapeA.BoundingBox.Max() - Tolerance };
-				const FAABB3 BoxB = { ShapeB.BoundingBox.Min() + Tolerance, ShapeB.BoundingBox.Max() - Tolerance };
-				const FAABB3 Intersection = BoxA.GetIntersection(BoxB);
+				const FAABB3 Intersection = ShapeA.BoundingBox.GetIntersection(ShapeB.BoundingBox);
 				OutOverlap += FBox{ Intersection.Min(), Intersection.Max() };
 				return true;
 			}
@@ -340,11 +340,10 @@ namespace Chaos
 			InTransformB,
 			bTraceComplex,
 			false,
-			[&OutOverlapSize, &Tolerance](const FShapeOverlapData& ShapeA, const FShapeOverlapData& ShapeB, const FMTDInfo&)
+			Tolerance,
+			[&OutOverlapSize](const FShapeOverlapData& ShapeA, const FShapeOverlapData& ShapeB, const FMTDInfo&)
 			{
-				const FAABB3 BoxA = { ShapeA.BoundingBox.Min() + Tolerance, ShapeA.BoundingBox.Max() - Tolerance };
-				const FAABB3 BoxB = { ShapeB.BoundingBox.Min() + Tolerance, ShapeB.BoundingBox.Max() - Tolerance };
-				const FAABB3 Intersection = BoxA.GetIntersection(BoxB);
+				const FAABB3 Intersection = ShapeA.BoundingBox.GetIntersection(ShapeB.BoundingBox);
 				OutOverlapSize += Intersection.Extents();
 				return true;
 			}
@@ -352,7 +351,7 @@ namespace Chaos
 	}
 
 	template<EThreadContext Id>
-	bool FReadPhysicsObjectInterface<Id>::PairwiseShapeOverlapHelper(FPhysicsObjectHandle ObjectA, const FTransform& InTransformA, FPhysicsObjectHandle ObjectB, const FTransform& InTransformB, bool bTraceComplex, bool bComputeMTD, const TFunction<bool(const FShapeOverlapData&, const FShapeOverlapData&, const FMTDInfo&)>& Lambda)
+	bool FReadPhysicsObjectInterface<Id>::PairwiseShapeOverlapHelper(FPhysicsObjectHandle ObjectA, const FTransform& InTransformA, FPhysicsObjectHandle ObjectB, const FTransform& InTransformB, bool bTraceComplex, bool bComputeMTD, const FVector& Tolerance, const TFunction<bool(const FShapeOverlapData&, const FShapeOverlapData&, const FMTDInfo&)>& Lambda)
 	{
 		TArray<FPerShapeData*> ShapesA = GetAllShapes({ &ObjectA, 1 });
 		const FTransform TransformA = FTransform{ GetR(ObjectA), GetX(ObjectA) } *InTransformA;
@@ -381,7 +380,7 @@ namespace Chaos
 				continue;
 			}
 
-			const FAABB3 BoxShapeB = GeomB->CalculateTransformedBounds(TransformB);
+			const FAABB3 BoxShapeB = GeomB->CalculateTransformedBounds(TransformB).ShrinkSymmetrically(Tolerance);
 			// At this point on, this function should be mirror the Overlap_GeomInternal function in PhysInterface_Chaos.cpp.
 			// ShapeA is equivalent to InInstance and GeomB is equivalent to InGeom.
 
@@ -401,7 +400,7 @@ namespace Chaos
 					continue;
 				}
 
-				const FAABB3 BoxShapeA = A->GetGeometry()->CalculateTransformedBounds(TransformA);
+				const FAABB3 BoxShapeA = A->GetGeometry()->CalculateTransformedBounds(TransformA).ShrinkSymmetrically(Tolerance);
 				if (!BoxShapeA.Intersects(BoxShapeB))
 				{
 					continue;
diff --git a/Engine/Source/Runtime/Experimental/Chaos/Public/Chaos/AABB.h b/Engine/Source/Runtime/Experimental/Chaos/Public/Chaos/AABB.h
--- a/Engine/Source/Runtime/Experimental/Chaos/Public/Chaos/AABB.h
+++ b/Engine/Source/Runtime/Experimental/Chaos/Public/Chaos/AABB.h
@@ -436,7 +436,7 @@ namespace Chaos
 			return *this;
 		}
 
-		//Grows (or shrinks) the box by this vector symmetrically - Changed name because previous Thicken had different semantics which caused several bugs
+		//Grows the box by this vector symmetrically - Changed name because previous Thicken had different semantics which caused several bugs
 		FORCEINLINE TAABB<T, d>& ThickenSymmetrically(const TVector<T, d>& Thickness)
 		{
 			const TVector<T, d> AbsThickness = TVector<T, d>(FGenericPlatformMath::Abs(Thickness.X), FGenericPlatformMath::Abs(Thickness.Y), FGenericPlatformMath::Abs(Thickness.Z));
@@ -445,6 +445,14 @@ namespace Chaos
 			return *this;
 		}
 
+		FORCEINLINE TAABB<T, d>& ShrinkSymmetrically(const TVector<T, d>& Thickness)
+		{
+			const TVector<T, d> AbsThickness = TVector<T, d>(FGenericPlatformMath::Abs(Thickness.X), FGenericPlatformMath::Abs(Thickness.Y), FGenericPlatformMath::Abs(Thickness.Z));
+			MMin += AbsThickness;
+			MMax -= AbsThickness;
+			return *this;
+		}
+
 		/** Grow along a vector (as if swept by the vector's direction and magnitude) */
 		FORCEINLINE TAABB<T, d>& GrowByVector(const TVector<T, d>& V)
 		{
diff --git a/Engine/Source/Runtime/Experimental/Chaos/Public/Chaos/PhysicsObjectInterface.h b/Engine/Source/Runtime/Experimental/Chaos/Public/Chaos/PhysicsObjectInterface.h
--- a/Engine/Source/Runtime/Experimental/Chaos/Public/Chaos/PhysicsObjectInterface.h
+++ b/Engine/Source/Runtime/Experimental/Chaos/Public/Chaos/PhysicsObjectInterface.h
@@ -110,7 +110,7 @@ namespace Chaos
 		/**
 		 * For every pair of shapes that overlap, allows the caller to perform some computation. If additional pairs of shapes need to be examined, the input TFunction should return true.
 		 */
-		bool PairwiseShapeOverlapHelper(FPhysicsObjectHandle ObjectA, const FTransform& InTransformA, FPhysicsObjectHandle ObjectB, const FTransform& InTransformB, bool bTraceComplex, bool bComputeMTD, const TFunction<bool(const FShapeOverlapData&, const FShapeOverlapData&, const FMTDInfo&)>& Lambda);
+		bool PairwiseShapeOverlapHelper(FPhysicsObjectHandle ObjectA, const FTransform& InTransformA, FPhysicsObjectHandle ObjectB, const FTransform& InTransformB, bool bTraceComplex, bool bComputeMTD, const FVector& Tolerance, const TFunction<bool(const FShapeOverlapData&, const FShapeOverlapData&, const FMTDInfo&)>& Lambda);
 	};
 
 	using FReadPhysicsObjectInterface_External = FReadPhysicsObjectInterface<EThreadContext::External>;
