Geometry collection : imporve anchoring by box by level to detect overlap of all particle at all levels and propgate anchoring state based on the max level

diff --git a/Engine/Source/Runtime/Experimental/Chaos/Private/PhysicsProxy/GeometryCollectionPhysicsProxy.cpp b/Engine/Source/Runtime/Experimental/Chaos/Private/PhysicsProxy/GeometryCollectionPhysicsProxy.cpp
--- a/Engine/Source/Runtime/Experimental/Chaos/Private/PhysicsProxy/GeometryCollectionPhysicsProxy.cpp
+++ b/Engine/Source/Runtime/Experimental/Chaos/Private/PhysicsProxy/GeometryCollectionPhysicsProxy.cpp
@@ -1844,18 +1844,31 @@ void FGeometryCollectionPhysicsProxy::SetAnchoredByTransformedBox_External(const
 				FPBDRigidsEvolution* Evolution = RBDSolver->GetEvolution();
 				for (int32 TransformIndex = 0; TransformIndex < SolverParticleHandles.Num(); TransformIndex++)
 				{
-					if (FClusterHandle * ParticleHandle = SolverParticleHandles[TransformIndex])
+					if (FClusterHandle* ParticleHandle = SolverParticleHandles[TransformIndex])
 					{
-						bool bCheckBounds = true;
-						if (MaxLevelToCheck > INDEX_NONE)
+						const FVec3 PositionInBoxSpace = BoxTransform.InverseTransformPosition(ParticleHandle->X());
+						if (BoundsToCheck.Contains(PositionInBoxSpace))
 						{
-							const int32 InitialLevel = InitialLevelAttribute.Get()[TransformIndex];
-							bCheckBounds = (InitialLevel <= MaxLevelToCheck);
-						}
-						if (bCheckBounds)
-						{
-							const FVec3 PositionInBoxSpace = BoxTransform.InverseTransformPosition(ParticleHandle->X());
-							if (BoundsToCheck.Contains(PositionInBoxSpace))
+							// if we have a max level , we make sure to anchor a parent of the right level 
+							if (MaxLevelToCheck > INDEX_NONE)
+							{
+								const int32 ParticleLevel = InitialLevelAttribute.Get()[TransformIndex];
+								if (ParticleLevel > MaxLevelToCheck)
+								{
+									Chaos::FPBDRigidClusteredParticleHandle* ParentParticle = ParticleHandle->Parent();
+									int32 ParentParticleLevel = (ParticleLevel - 1);
+									// we are above level, so we still contribute to the parent matching the level 
+									while (ParentParticle && ParentParticleLevel > MaxLevelToCheck)
+									{
+										ParentParticle = ParentParticle->Parent();
+									}
+									// if we still have a parent , this means the level requirement has been found
+									// let's use this particle to be anchored
+									// if the parent is null, then we don't want to anchor the particle 
+									ParticleHandle = ParentParticle;
+								}
+							}
+							if (ParticleHandle)
 							{
 								SetParticleAnchored_Internal(Evolution, ParticleHandle, bAnchored);
 								if (Chaos::FPBDRigidClusteredParticleHandle* TopParentHandle = GetTopActiveClusteredParent_Internal(ParticleHandle))
