geometry collection physics proxy should not keep track of internal clusters owned by the cluster union physics proxy

diff --git a/Engine/Source/Runtime/Experimental/Chaos/Private/PhysicsProxy/GeometryCollectionPhysicsProxy.cpp b/Engine/Source/Runtime/Experimental/Chaos/Private/PhysicsProxy/GeometryCollectionPhysicsProxy.cpp
--- a/Engine/Source/Runtime/Experimental/Chaos/Private/PhysicsProxy/GeometryCollectionPhysicsProxy.cpp
+++ b/Engine/Source/Runtime/Experimental/Chaos/Private/PhysicsProxy/GeometryCollectionPhysicsProxy.cpp
@@ -2747,12 +2747,18 @@ void FGeometryCollectionPhysicsProxy::BufferPhysicsResults_Internal(Chaos::FPBDR
 
 						if(ClusterParent->InternalCluster())
 						{
-							const int32 InternalClusterParentUniqueIdx = ClusterParent->UniqueIdx().Idx;
-							if (!UniqueIdxToInternalClusterHandle.Contains(InternalClusterParentUniqueIdx))
+							const bool bIsOwnedInternalCluster = ClusterParent->PhysicsProxy() == this;
+
+							// We don't want the to keep track of any external internal clusters (i.e. cluster unions).
+							if (bIsOwnedInternalCluster)
 							{
-								UniqueIdxToInternalClusterHandle.Add(InternalClusterParentUniqueIdx, ClusterParent);
+								const int32 InternalClusterParentUniqueIdx = ClusterParent->UniqueIdx().Idx;
+								if (!UniqueIdxToInternalClusterHandle.Contains(InternalClusterParentUniqueIdx))
+								{
+									UniqueIdxToInternalClusterHandle.Add(InternalClusterParentUniqueIdx, ClusterParent);
+								}
+								TargetResults.InternalClusterUniqueIdx[TransformGroupIndex] = ClusterParent->UniqueIdx().Idx;
 							}
-							TargetResults.InternalClusterUniqueIdx[TransformGroupIndex] = ClusterParent->UniqueIdx().Idx;
 							
 							const FTransform ParticleToWorld = Handle->ChildToParent() * FRigidTransform3(ClusterParent->X(), ClusterParent->R());    // aka ClusterChildToWorld
 							TargetResults.ParticleXs[TransformGroupIndex] = ParticleToWorld.GetTranslation();
