add function to get body ID on the component given a raw GT particle

diff --git a/Engine/Source/Runtime/Engine/Classes/Components/PrimitiveComponent.h b/Engine/Source/Runtime/Engine/Classes/Components/PrimitiveComponent.h
--- a/Engine/Source/Runtime/Engine/Classes/Components/PrimitiveComponent.h
+++ b/Engine/Source/Runtime/Engine/Classes/Components/PrimitiveComponent.h
@@ -2844,6 +2844,7 @@ public:
 	virtual Chaos::FPhysicsObject* GetPhysicsObjectById(int32 Id) const override;
 	virtual Chaos::FPhysicsObject* GetPhysicsObjectByName(const FName& Name) const override;
 	virtual TArray<Chaos::FPhysicsObject*> GetAllPhysicsObjects() const override;
+	virtual int32 GetIdFromGTParticle(Chaos::FGeometryParticle* Particle) const override;
 	//~ End IPhysicsComponent Interface.
 };
 
diff --git a/Engine/Source/Runtime/Engine/Classes/Interfaces/IPhysicsComponent.h b/Engine/Source/Runtime/Engine/Classes/Interfaces/IPhysicsComponent.h
--- a/Engine/Source/Runtime/Engine/Classes/Interfaces/IPhysicsComponent.h
+++ b/Engine/Source/Runtime/Engine/Classes/Interfaces/IPhysicsComponent.h
@@ -2,6 +2,7 @@
 
 #pragma once
 
+#include "Chaos/ParticleHandleFwd.h"
 #include "Chaos/PhysicsObject.h"
 #include "UObject/ObjectMacros.h"
 #include "UObject/Interface.h"
@@ -22,4 +23,6 @@ public:
 	virtual Chaos::FPhysicsObject* GetPhysicsObjectById(int32 Id) const = 0;
 	virtual Chaos::FPhysicsObject* GetPhysicsObjectByName(const FName& Name) const = 0;
 	virtual TArray<Chaos::FPhysicsObject*> GetAllPhysicsObjects() const = 0;
+	
+	virtual int32 GetIdFromGTParticle(Chaos::FGeometryParticle* Particle) const = 0;
 };
\ No newline at end of file
diff --git a/Engine/Source/Runtime/Engine/Private/PrimitiveComponentPhysics.cpp b/Engine/Source/Runtime/Engine/Private/PrimitiveComponentPhysics.cpp
--- a/Engine/Source/Runtime/Engine/Private/PrimitiveComponentPhysics.cpp
+++ b/Engine/Source/Runtime/Engine/Private/PrimitiveComponentPhysics.cpp
@@ -1206,4 +1206,9 @@ TArray<Chaos::FPhysicsObject*> UPrimitiveComponent::GetAllPhysicsObjects() const
 	return Bodies;
 }
 
+int32 UPrimitiveComponent::GetIdFromGTParticle(Chaos::FGeometryParticle* Particle) const
+{
+	return 0;
+}
+
 #undef LOCTEXT_NAMESPACE
diff --git a/Engine/Source/Runtime/Experimental/Chaos/Public/PhysicsProxy/GeometryCollectionPhysicsProxy.h b/Engine/Source/Runtime/Experimental/Chaos/Public/PhysicsProxy/GeometryCollectionPhysicsProxy.h
--- a/Engine/Source/Runtime/Experimental/Chaos/Public/PhysicsProxy/GeometryCollectionPhysicsProxy.h
+++ b/Engine/Source/Runtime/Experimental/Chaos/Public/PhysicsProxy/GeometryCollectionPhysicsProxy.h
@@ -383,6 +383,15 @@ public:
 		return FGeometryCollectionItemIndex::CreateInvalidItemIndex();
 	}
 
+	FGeometryCollectionItemIndex GetItemIndexFromGTParticleNoInternalCluster_External(const FParticle* GTPParticle) const
+	{
+		if (const int32* TransformGroupIndex = GTParticlesToTransformGroupIndex.Find(GTPParticle))
+		{
+			return FGeometryCollectionItemIndex::CreateTransformItemIndex(*TransformGroupIndex);
+		}
+		return FGeometryCollectionItemIndex::CreateInvalidItemIndex();
+	}
+
 	bool GetIsObjectDynamic() const { return IsObjectDynamic; }
 
 	void DisableParticles_External(TArray<int32>&& TransformGroupIndices);
diff --git a/Engine/Source/Runtime/Experimental/GeometryCollectionEngine/Private/GeometryCollection/GeometryCollectionComponent.cpp b/Engine/Source/Runtime/Experimental/GeometryCollectionEngine/Private/GeometryCollection/GeometryCollectionComponent.cpp
--- a/Engine/Source/Runtime/Experimental/GeometryCollectionEngine/Private/GeometryCollection/GeometryCollectionComponent.cpp
+++ b/Engine/Source/Runtime/Experimental/GeometryCollectionEngine/Private/GeometryCollection/GeometryCollectionComponent.cpp
@@ -5013,6 +5013,23 @@ TArray<Chaos::FPhysicsObject*> UGeometryCollectionComponent::GetAllPhysicsObject
 	return Objects;
 }
 
+int32 UGeometryCollectionComponent::GetIdFromGTParticle(Chaos::FGeometryParticle* Particle) const
+{
+	if (!PhysicsProxy || !Particle)
+	{
+		return INDEX_NONE;
+	}
+	FGeometryCollectionItemIndex Index = PhysicsProxy->GetItemIndexFromGTParticleNoInternalCluster_External(Particle->CastToRigidParticle());
+	if (Index.IsValid())
+	{
+		return Index.GetTransformIndex();
+	}
+	else
+	{
+		return INDEX_NONE;
+	}
+}
+
 void UGeometryCollectionComponent::SetEnableDamageFromCollision(bool bValue)
 {
 	bEnableDamageFromCollision = bValue;
diff --git a/Engine/Source/Runtime/Experimental/GeometryCollectionEngine/Public/GeometryCollection/GeometryCollectionComponent.h b/Engine/Source/Runtime/Experimental/GeometryCollectionEngine/Public/GeometryCollection/GeometryCollectionComponent.h
--- a/Engine/Source/Runtime/Experimental/GeometryCollectionEngine/Public/GeometryCollection/GeometryCollectionComponent.h
+++ b/Engine/Source/Runtime/Experimental/GeometryCollectionEngine/Public/GeometryCollection/GeometryCollectionComponent.h
@@ -1278,5 +1278,6 @@ public:
 	virtual Chaos::FPhysicsObject* GetPhysicsObjectById(int32 Id) const override;
 	virtual Chaos::FPhysicsObject* GetPhysicsObjectByName(const FName& Name) const override;
 	virtual TArray<Chaos::FPhysicsObject*> GetAllPhysicsObjects() const override;
+	virtual int32 GetIdFromGTParticle(Chaos::FGeometryParticle* Particle) const override;
 	//~ End IPhysicsComponent Interface.
 };
