geometry collections should not attempt to replicate parent particles that are a part of cluster unions

diff --git a/Engine/Source/Runtime/Experimental/GeometryCollectionEngine/Private/GeometryCollection/GeometryCollectionComponent.cpp b/Engine/Source/Runtime/Experimental/GeometryCollectionEngine/Private/GeometryCollection/GeometryCollectionComponent.cpp
--- a/Engine/Source/Runtime/Experimental/GeometryCollectionEngine/Private/GeometryCollection/GeometryCollectionComponent.cpp
+++ b/Engine/Source/Runtime/Experimental/GeometryCollectionEngine/Private/GeometryCollection/GeometryCollectionComponent.cpp
@@ -1600,7 +1600,8 @@ void UGeometryCollectionComponent::UpdateRepData()
 					}
 				}
 	
-				if (bProcess && Root->Disabled() == false && ClustersToRep->Find(Root) == nullptr)
+				// The additional physics proxy check is to make sure that we don't try to replicate a cluster union particle.
+				if (bProcess && Root->Disabled() == false && ClustersToRep->Find(Root) == nullptr && Root->PhysicsProxy() == PhysicsProxy)
 				{
 					int32 TransformGroupIdx = INDEX_NONE;
 					int32 Level = INDEX_NONE;
