Chaos
- support geometry repeats in collision detection
- fix level set collision detection

diff --git a/Engine/Source/Programs/HeadlessChaos/Private/HeadlessChaosTestShapePadding.cpp b/Engine/Source/Programs/HeadlessChaos/Private/HeadlessChaosTestShapePadding.cpp
--- a/Engine/Source/Programs/HeadlessChaos/Private/HeadlessChaosTestShapePadding.cpp
+++ b/Engine/Source/Programs/HeadlessChaos/Private/HeadlessChaosTestShapePadding.cpp
@@ -103,9 +103,13 @@ namespace ChaosTest {
 			true,
 			EContactShapesType::BoxBox);
 
+		FRigidTransform3 ShapeWorldTransform0 = Constraint->GetShapeRelativeTransform0() * Collisions::GetTransform(Constraint->GetParticle0());
+		FRigidTransform3 ShapeWorldTransform1 = Constraint->GetShapeRelativeTransform1() * Collisions::GetTransform(Constraint->GetParticle1());
+		Constraint->SetShapeWorldTransforms(ShapeWorldTransform0, ShapeWorldTransform1);
+
 		// Detect collisions
 		Constraint->ResetPhi(Constraint->GetCullDistance());
-		Collisions::UpdateConstraintFromGeometry<ECollisionUpdateType::Deepest>(*Constraint, FParticleUtilitiesPQ::GetActorWorldTransform(Box0), FParticleUtilitiesPQ::GetActorWorldTransform(Box1), 1 / 30.0f);
+		Collisions::UpdateConstraint(*Constraint, ShapeWorldTransform0, ShapeWorldTransform1, 1 / 30.0f);
 
 		EXPECT_NEAR(Constraint->GetPhi(), ExpectedPhi, Tolerance);
 		EXPECT_NEAR(Constraint->CalculateWorldContactNormal().X, ExpectedNormal.X, Tolerance);
@@ -216,8 +220,12 @@ namespace ChaosTest {
 			true,
 			EContactShapesType::GenericConvexConvex);
 
+		FRigidTransform3 ShapeWorldTransform0 = Constraint->GetShapeRelativeTransform0() * Collisions::GetTransform(Constraint->GetParticle0());
+		FRigidTransform3 ShapeWorldTransform1 = Constraint->GetShapeRelativeTransform1() * Collisions::GetTransform(Constraint->GetParticle1());
+		Constraint->SetShapeWorldTransforms(ShapeWorldTransform0, ShapeWorldTransform1);
+
 		// Detect collisions
-		Collisions::UpdateConstraintFromGeometry<ECollisionUpdateType::Deepest>(*Constraint, FParticleUtilitiesPQ::GetActorWorldTransform(Box0), FParticleUtilitiesPQ::GetActorWorldTransform(Box1), 1 / 30.0f);
+		Collisions::UpdateConstraint(*Constraint, ShapeWorldTransform0, ShapeWorldTransform1, 1 / 30.0f);
 
 		EXPECT_NEAR(Constraint->GetPhi(), ExpectedPhi, Tolerance);
 		EXPECT_NEAR(Constraint->CalculateWorldContactNormal().X, ExpectedNormal.X, Tolerance);
diff --git a/Engine/Source/Programs/HeadlessChaos/Public/HeadlessChaosCollisionConstraints.h b/Engine/Source/Programs/HeadlessChaos/Public/HeadlessChaosCollisionConstraints.h
--- a/Engine/Source/Programs/HeadlessChaos/Public/HeadlessChaosCollisionConstraints.h
+++ b/Engine/Source/Programs/HeadlessChaos/Public/HeadlessChaosCollisionConstraints.h
@@ -90,12 +90,12 @@ public:
 			// Dt is not important for the tests that use this function
 			const FReal Dt = FReal(1) / FReal(30);
 
+			FRigidTransform3 ShapeWorldTransform0 = Constraint.GetShapeRelativeTransform0() * Collisions::GetTransform(Constraint.GetParticle0());
+			FRigidTransform3 ShapeWorldTransform1 = Constraint.GetShapeRelativeTransform1() * Collisions::GetTransform(Constraint.GetParticle1());
+			Constraint.SetShapeWorldTransforms(ShapeWorldTransform0, ShapeWorldTransform1);
+
 			Constraint.ResetPhi(TNumericLimits<FReal>::Max());
-			Collisions::UpdateConstraintFromGeometry<ECollisionUpdateType::Deepest>(
-				Constraint, 
-				FParticleUtilities::GetActorWorldTransform(FGenericParticleHandle(Constraint.GetParticle0())), 
-				FParticleUtilities::GetActorWorldTransform(FGenericParticleHandle(Constraint.GetParticle1())), 
-				Dt);
+			Collisions::UpdateConstraint(Constraint, ShapeWorldTransform0, ShapeWorldTransform1,Dt);
 		}
 	}
 
@@ -109,7 +109,7 @@ public:
 		FRigidTransform3 WorldTransform1 = Constraint.GetShapeRelativeTransform1() * Collisions::GetTransform(Constraint.GetParticle1());
 
 		Constraint.ResetManifold();
-		Collisions::UpdateLevelsetLevelsetConstraint<ECollisionUpdateType::Deepest>(WorldTransform0, WorldTransform1, FReal(1 / 30.0f), Constraint);
+		Collisions::UpdateLevelsetLevelsetConstraint(WorldTransform0, WorldTransform1, FReal(1 / 30.0f), Constraint);
 	}
 
 	int32 NumConstraints() const
diff --git a/Engine/Source/Runtime/Experimental/Chaos/Private/Chaos/CCDUtilities.cpp b/Engine/Source/Runtime/Experimental/Chaos/Private/Chaos/CCDUtilities.cpp
--- a/Engine/Source/Runtime/Experimental/Chaos/Private/Chaos/CCDUtilities.cpp
+++ b/Engine/Source/Runtime/Experimental/Chaos/Private/Chaos/CCDUtilities.cpp
@@ -919,10 +919,14 @@ namespace Chaos
 			const FConstGenericParticleHandle P0 = FConstGenericParticleHandle(Collision->GetParticle0());
 			const FConstGenericParticleHandle P1 = FConstGenericParticleHandle(Collision->GetParticle1());
 
+			const FRigidTransform3 ShapeWorldTransform0 = Collision->GetShapeRelativeTransform0() * P0->GetTransformPQ();
+			const FRigidTransform3 ShapeWorldTransform1 = Collision->GetShapeRelativeTransform1() * P1->GetTransformPQ();
+			Collision->SetShapeWorldTransforms(ShapeWorldTransform0, ShapeWorldTransform1);
+
 			// NOTE: ResetManifold also reset friction anchors. If CCD sweep was run, static friction probably will not hold
 			// We could potentially call ResetActiveManifold here instead and then AssignSavedManifoldPoints if we want static friction
 			Collision->ResetManifold();
-			Collisions::UpdateConstraintFromGeometry<ECollisionUpdateType::Deepest>(*Collision, P0->GetTransformPQ(), P1->GetTransformPQ(), Dt);
+			Collisions::UpdateConstraint(*Collision, ShapeWorldTransform0, ShapeWorldTransform1, Dt);
 		}
 	}
 
diff --git a/Engine/Source/Runtime/Experimental/Chaos/Private/Chaos/Collision/ParticlePairMidPhase.cpp b/Engine/Source/Runtime/Experimental/Chaos/Private/Chaos/Collision/ParticlePairMidPhase.cpp
--- a/Engine/Source/Runtime/Experimental/Chaos/Private/Chaos/Collision/ParticlePairMidPhase.cpp
+++ b/Engine/Source/Runtime/Experimental/Chaos/Private/Chaos/Collision/ParticlePairMidPhase.cpp
@@ -397,8 +397,6 @@ namespace Chaos
 		{
 			PHYSICS_CSV_SCOPED_EXPENSIVE(PhysicsVerbose, NarrowPhase_UpdateConstraintCCD);
 
-			const FImplicitObject* Implicit0 = Shape0->GetLeafGeometry();
-			const FImplicitObject* Implicit1 = Shape1->GetLeafGeometry();
 			const FRigidTransform3& ShapeWorldTransform0 = Shape0->GetLeafWorldTransform(GetParticle0());
 			const FRigidTransform3& ShapeWorldTransform1 = Shape1->GetLeafWorldTransform(GetParticle1());
 
@@ -407,7 +405,6 @@ namespace Chaos
 			// NOTE: these are not used by CCD which continuously moves the particles
 			Constraint->SetShapeWorldTransforms(ShapeWorldTransform0, ShapeWorldTransform1);
 
-			Constraint->ResetModifications();
 			Constraint->SetCullDistance(CullDistance);
 			Constraint->ResetManifold();
 			Constraint->ResetActiveManifoldContacts();
@@ -424,8 +421,8 @@ namespace Chaos
 			// @todo(chaos): Pass both start and end transforms to the collision detector
 			const FRigidTransform3 CCDParticleWorldTransform0 = FRigidTransform3(P0->P() - P0->V() * Dt, P0->Q());
 			const FRigidTransform3 CCDParticleWorldTransform1 = FRigidTransform3(P1->P() - P1->V() * Dt, P1->Q());
-			const FRigidTransform3 CCDShapeWorldTransform0 = Constraint->ImplicitTransform[0] * CCDParticleWorldTransform0;
-			const FRigidTransform3 CCDShapeWorldTransform1 = Constraint->ImplicitTransform[1] * CCDParticleWorldTransform1;
+			const FRigidTransform3 CCDShapeWorldTransform0 = Constraint->GetShapeRelativeTransform0() * CCDParticleWorldTransform0;
+			const FRigidTransform3 CCDShapeWorldTransform1 = Constraint->GetShapeRelativeTransform1() * CCDParticleWorldTransform1;
 			const bool bDidSweep = Collisions::UpdateConstraintSwept(*Constraint.Get(), CCDShapeWorldTransform0, CCDShapeWorldTransform1, Dt);
 
 #if CHAOS_DEBUG_DRAW
@@ -433,13 +430,13 @@ namespace Chaos
 			{
 				if (FConstGenericParticleHandle(Constraint->GetParticle0())->CCDEnabled())
 				{
-					DebugDraw::DrawShape(CCDShapeWorldTransform0, Implicit0, Shape0, FColor::Black, &CVars::ChaosSolverDebugDebugDrawSettings);
-					DebugDraw::DrawShape(ShapeWorldTransform0, Implicit0, Shape0, FColor::White, &CVars::ChaosSolverDebugDebugDrawSettings);
+					DebugDraw::DrawShape(CCDShapeWorldTransform0, Constraint->GetImplicit0(), Shape0, FColor::Black, &CVars::ChaosSolverDebugDebugDrawSettings);
+					DebugDraw::DrawShape(ShapeWorldTransform0, Constraint->GetImplicit0(), Shape0, FColor::White, &CVars::ChaosSolverDebugDebugDrawSettings);
 				}
 				if (FConstGenericParticleHandle(Constraint->GetParticle1())->CCDEnabled())
 				{
-					DebugDraw::DrawShape(CCDShapeWorldTransform1, Implicit1, Shape1, FColor::Black, &CVars::ChaosSolverDebugDebugDrawSettings);
-					DebugDraw::DrawShape(ShapeWorldTransform1, Implicit1, Shape1, FColor::White, &CVars::ChaosSolverDebugDebugDrawSettings);
+					DebugDraw::DrawShape(CCDShapeWorldTransform1, Constraint->GetImplicit1(), Shape1, FColor::Black, &CVars::ChaosSolverDebugDebugDrawSettings);
+					DebugDraw::DrawShape(ShapeWorldTransform1, Constraint->GetImplicit1(), Shape1, FColor::White, &CVars::ChaosSolverDebugDebugDrawSettings);
 				}
 			}
 #endif
@@ -574,6 +571,25 @@ namespace Chaos
 		const FReal CullDistance,
 		const FReal Dt,
 		const FCollisionContext& Context)
+	{
+		return GenerateCollisionsImpl(false, CullDistance, Dt, Context);
+	}
+
+
+	int32 FMultiShapePairCollisionDetector::GenerateCollisionsCCD(
+		const bool bEnableCCDSweep,
+		const FReal CullDistance,
+		const FReal Dt,
+		const FCollisionContext& Context)
+	{
+		return GenerateCollisionsImpl(bEnableCCDSweep, CullDistance, Dt, Context);
+	}
+
+	int32 FMultiShapePairCollisionDetector::GenerateCollisionsImpl(
+		const bool bEnableCCDSweep,
+		const FReal CullDistance,
+		const FReal Dt,
+		const FCollisionContext& Context)
 	{
 		FConstGenericParticleHandle P0 = Particle0;
 		FConstGenericParticleHandle P1 = Particle1;
@@ -590,24 +606,18 @@ namespace Chaos
 		FCollisionContext LocalContext = Context;
 		LocalContext.MultiShapeCollisionDetector = this;
 
+		// This recurively processes the Implicit Hierarchy and generates a collision constraint for every
+		// overlapping shape pair on the two bodies via a callback to FindOrCreateConstraint
 		Collisions::ConstructConstraints(
-			Particle0,
-			Particle1,
-			Implicit0,
-			Shape0,
-			BVHParticles0,
-			Implicit1,
-			Shape1,
-			BVHParticles1,
-			ParticleWorldTransform0,
-			ShapeRelativeTransform0,
-			ParticleWorldTransform1,
-			ShapeRelativeTransform1,
-			CullDistance,
-			Dt,
-			LocalContext);
-
-		int32 NumActiveConstraints = ProcessNewConstraints(Context);
+			Particle0, Particle1,
+			Implicit0, Shape0, BVHParticles0,
+			Implicit1, Shape1, BVHParticles1,
+			ParticleWorldTransform0, ShapeRelativeTransform0,
+			ParticleWorldTransform1, ShapeRelativeTransform1,
+			CullDistance, Dt, bEnableCCDSweep, LocalContext);
+
+		// Generate manifolds for each constraint we created/recovered and (re)activate if necessary
+		int32 NumActiveConstraints = ProcessNewConstraints(CullDistance, Dt, Context);
 
 		// @todo(chaos): we could clean up unused collisions between this pair, but probably not worth it
 		// and we would have to prevent cleanup for sleeping particles because the collisions will still
@@ -619,18 +629,21 @@ namespace Chaos
 
 	FPBDCollisionConstraint* FMultiShapePairCollisionDetector::FindOrCreateConstraint(
 		FGeometryParticleHandle* InParticle0,
-		const FImplicitObject* Implicit0,
-		const FPerShapeData* InShape0,
-		const FBVHParticles* BVHParticles0,
-		const FRigidTransform3& ShapeRelativeTransform0,
+		const FImplicitObject* InImplicit0,
+		const int32 InImplicitId0,
+		const FShapeInstance* InShape0,
+		const FBVHParticles* InBVHParticles0,
+		const FRigidTransform3& InShapeRelativeTransform0,
 		FGeometryParticleHandle* InParticle1,
-		const FImplicitObject* Implicit1,
-		const FPerShapeData* InShape1,
-		const FBVHParticles* BVHParticles1,
-		const FRigidTransform3& ShapeRelativeTransform1,
+		const FImplicitObject* InImplicit1,
+		const int32 InImplicitId1,
+		const FShapeInstance* InShape1,
+		const FBVHParticles* InBVHParticles1,
+		const FRigidTransform3& InShapeRelativeTransform1,
 		const FReal CullDistance,
 		const EContactShapesType ShapePairType,
 		const bool bUseManifold,
+		const bool bEnableSweep,
 		const FCollisionContext& Context)
 	{
 		// This is a callback from the low-level collision function. It should always be the same two particles, though the
@@ -643,11 +656,11 @@ namespace Chaos
 			return nullptr;
 		}
 
-		const FCollisionParticlePairConstraintKey Key = FCollisionParticlePairConstraintKey(Implicit0, BVHParticles0, Implicit1, BVHParticles1);
+		const FCollisionParticlePairConstraintKey Key = FCollisionParticlePairConstraintKey(InShape0, InImplicit0, InImplicitId0, InBVHParticles0, InShape1, InImplicit1, InImplicitId1, InBVHParticles1);
 		FPBDCollisionConstraint* Constraint = FindConstraint(Key);
 
 		// @todo(chaos): fix key uniqueness guarantee.  We need a truly unique key gen function
-		const bool bIsKeyCollision = (Constraint != nullptr) && ((Constraint->GetImplicit0() != Implicit0) || (Constraint->GetImplicit1() != Implicit1) || (Constraint->GetCollisionParticles0() != BVHParticles0) || (Constraint->GetCollisionParticles1() != BVHParticles1));
+		const bool bIsKeyCollision = (Constraint != nullptr) && ((Constraint->GetImplicit0() != InImplicit0) || (Constraint->GetImplicit1() != InImplicit1) || (Constraint->GetCollisionParticles0() != InBVHParticles0) || (Constraint->GetCollisionParticles1() != InBVHParticles1));
 		if (bIsKeyCollision)
 		{
 			// If we get here, we have a key collision. The key uses a hash of pointers which is very likely to be unique for different implicit pairs, 
@@ -662,47 +675,26 @@ namespace Chaos
 		if (Constraint == nullptr)
 		{
 			// NOTE: Using InParticle0 and InParticle1 here because the order may be different to what we have stored
-			Constraint = CreateConstraint(InParticle0, Implicit0, InShape0, BVHParticles0, ShapeRelativeTransform0, InParticle1, Implicit1, InShape1, BVHParticles1, ShapeRelativeTransform1, CullDistance, ShapePairType, bUseManifold, Key, Context);
+			Constraint = CreateConstraint(InParticle0, InImplicit0, InShape0, InBVHParticles0, InShapeRelativeTransform0, InParticle1, InImplicit1, InShape1, InBVHParticles1, InShapeRelativeTransform1, CullDistance, ShapePairType, bUseManifold, Key, Context);
+
+			// Is this a CCD constraint?
+			Constraint->SetCCDEnabled(MidPhase.IsCCD());
 		}
 
-		// @todo(chaos): we already have the shape world transforms at the calling site - pass them in
+		// Do we want to sweep on this tick? 
+		// CCD may be temporarily disabled by the user or because we are moving slowly.
+		Constraint->SetCCDSweepEnabled(MidPhase.IsCCD() && bEnableSweep);
+
 		const FRigidTransform3 ParticleTransform0 = FConstGenericParticleHandle(InParticle0)->GetTransformPQ();
 		const FRigidTransform3 ParticleTransform1 = FConstGenericParticleHandle(InParticle1)->GetTransformPQ();
-		const FRigidTransform3 ShapeWorldTransform0 = ShapeRelativeTransform0 * ParticleTransform0;
-		const FRigidTransform3 ShapeWorldTransform1 = ShapeRelativeTransform1 * ParticleTransform1;
+		const FRigidTransform3 ShapeWorldTransform0 = InShapeRelativeTransform0 * ParticleTransform0;
+		const FRigidTransform3 ShapeWorldTransform1 = InShapeRelativeTransform1 * ParticleTransform1;
 		Constraint->SetShapeWorldTransforms(ShapeWorldTransform0, ShapeWorldTransform1);
 
-		Constraint->ResetModifications();
-
 		NewConstraints.Add(Constraint);
 		return Constraint;
 	}
 
-
-	FPBDCollisionConstraint* FMultiShapePairCollisionDetector::FindOrCreateSweptConstraint(
-		FGeometryParticleHandle* InParticle0,
-		const FImplicitObject* Implicit0,
-		const FPerShapeData* InShape0,
-		const FBVHParticles* BVHParticles0,
-		const FRigidTransform3& ShapeRelativeTransform0,
-		FGeometryParticleHandle* InParticle1,
-		const FImplicitObject* Implicit1,
-		const FPerShapeData* InShape1,
-		const FBVHParticles* BVHParticles1,
-		const FRigidTransform3& ShapeRelativeTransform1,
-		const FReal CullDistance,
-		const EContactShapesType ShapePairType,
-		const FCollisionContext& Context)
-	{
-		const bool bUseManifold = true;
-		FPBDCollisionConstraint* Constraint = FindOrCreateConstraint(InParticle0, Implicit0, InShape0, BVHParticles0, ShapeRelativeTransform0, InParticle1, Implicit1, InShape1, BVHParticles1, ShapeRelativeTransform1, CullDistance, ShapePairType, bUseManifold, Context);
-		if (Constraint != nullptr)
-		{
-			Constraint->SetCCDEnabled(true);
-		}
-		return Constraint;
-	}
-
 	FPBDCollisionConstraint* FMultiShapePairCollisionDetector::FindConstraint(const FCollisionParticlePairConstraintKey& Key)
 	{
 		FPBDCollisionConstraintPtr* PConstraint = Constraints.Find(Key.GetKey());
@@ -753,25 +745,172 @@ namespace Chaos
 		}
 	}
 
-	int32 FMultiShapePairCollisionDetector::ProcessNewConstraints(const FCollisionContext& Context)
+	int32 FMultiShapePairCollisionDetector::ProcessNewConstraints(
+		const FReal CullDistance,
+		const FReal Dt,
+		const FCollisionContext& Context)
 	{
 		int32 NumActiveConstraints = 0;
 		for (FPBDCollisionConstraint* Constraint : NewConstraints)
 		{
-			if (Constraint->GetPhi() < Constraint->GetCullDistance())
+			if (!Constraint->GetCCDSweepEnabled())
 			{
-				Context.GetAllocator()->ActivateConstraint(Constraint);
-				++NumActiveConstraints;
+				UpdateCollision(Constraint, CullDistance, Dt, Context);
 			}
 			else
 			{
-				Constraint->SetDisabled(true);
+				UpdateCollisionCCD(Constraint, CullDistance, Dt, Context);
+			}
+
+			if (!Constraint->GetDisabled())
+			{
+				++NumActiveConstraints;
 			}
 		}
 		NewConstraints.Reset();
 		return NumActiveConstraints;
 	}
 
+	void FMultiShapePairCollisionDetector::UpdateCollision(
+		FPBDCollisionConstraint* Constraint, 
+		const FReal CullDistance,
+		const FReal Dt,
+		const FCollisionContext& Context)
+	{
+		// @todo(chaos): share this code with FSingleShapePairCollisionDetector
+
+		const FRigidTransform3 ShapeWorldTransform0 = Constraint->GetShapeRelativeTransform0() * FConstGenericParticleHandle(Constraint->GetParticle0())->GetTransformPQ();
+		const FRigidTransform3 ShapeWorldTransform1 = Constraint->GetShapeRelativeTransform1() * FConstGenericParticleHandle(Constraint->GetParticle1())->GetTransformPQ();
+		const int32 CurrentEpoch = Context.GetAllocator()->GetCurrentEpoch();
+		const int32 LastEpoch = CurrentEpoch - 1;
+		const bool bWasUpdatedLastTick = (Constraint->GetContainerCookie().LastUsedEpoch >= LastEpoch);
+
+		// Update the world shape transforms on the constraint (we cannot just give it the PerShapeData 
+		// pointer because of Unions - see FMultiShapePairCollisionDetector)
+		Constraint->SetShapeWorldTransforms(ShapeWorldTransform0, ShapeWorldTransform1);
+
+		// If the constraint was not used last frame, it needs to be reset, otherwise we will try to reuse
+		if (!bWasUpdatedLastTick || (Constraint->GetManifoldPoints().Num() == 0))
+		{
+			// Clear all manifold data including saved contact data
+			Constraint->ResetManifold();
+		}
+
+		bool bWasManifoldRestored = false;
+		//if (Context.GetSettings().bAllowManifoldReuse && Flags.bEnableManifoldUpdate && bWasUpdatedLastTick)
+		//{
+		//	// Update the existing manifold. We can re-use as-is if none of the points have moved much and the bodies have not moved much
+		//	// NOTE: this can succeed in "restoring" even if we have no manifold points
+		//	// NOTE: this uses the transforms from SetLastShapeWorldTransforms, so we can only do this if we were updated last tick
+		//	bWasManifoldRestored = Constraint->UpdateAndTryRestoreManifold();
+		//}
+		//else
+		{
+			// We are not trying to reuse manifold points, so reset them but leave stored data intact (for friction)
+			Constraint->ResetActiveManifoldContacts();
+		}
+
+		if (!bWasManifoldRestored)
+		{
+			// We will be updating the manifold so update transforms used to check for movement in UpdateAndTryRestoreManifold on future ticks
+			Constraint->SetLastShapeWorldTransforms(ShapeWorldTransform0, ShapeWorldTransform1);
+
+			if (!Context.GetSettings().bDeferNarrowPhase)
+			{
+				Collisions::UpdateConstraint(*Constraint, Constraint->GetShapeWorldTransform0(), Constraint->GetShapeWorldTransform1(), Dt);
+			}
+		}
+
+		// If we have a valid contact, add it to the active list
+		// We also add it to the active list if collision detection is deferred because the data will be filled in later and we
+		// don't know in advance whether we will pass the Phi check (deferred narrow phase is used with RBAN)
+		if (Constraint->GetPhi() <= CullDistance || Context.GetSettings().bDeferNarrowPhase)
+		{
+			if (Context.GetAllocator()->ActivateConstraint(Constraint))
+			{
+				return;
+			}
+		}
+
+		// If we get here, we did not activate hte constraint and it should be disabled for this tick
+		Constraint->SetDisabled(true);
+	}
+
+
+	void FMultiShapePairCollisionDetector::UpdateCollisionCCD(
+		FPBDCollisionConstraint* Constraint,
+		const FReal CullDistance,
+		const FReal Dt,
+		const FCollisionContext& Context)
+	{
+		// @todo(chaos): share this code with FSingleShapePairCollisionDetector
+
+		if (!Constraint->GetCCDSweepEnabled())
+		{
+			return UpdateCollision(Constraint, CullDistance, Dt, Context);
+		}
+
+		Constraint->ResetManifold();
+		Constraint->ResetActiveManifoldContacts();
+
+		const FRigidTransform3 ShapeWorldTransform0 = Constraint->GetShapeRelativeTransform0() * FConstGenericParticleHandle(Constraint->GetParticle0())->GetTransformPQ();
+		const FRigidTransform3 ShapeWorldTransform1 = Constraint->GetShapeRelativeTransform1() * FConstGenericParticleHandle(Constraint->GetParticle1())->GetTransformPQ();
+
+		// Update the world shape transforms on the constraint (we cannot just give it the PerShapeData 
+		// pointer because of Unions - see FMultiShapePairCollisionDetector)
+		// NOTE: these are not used by CCD which continuously moves the particles
+		Constraint->SetShapeWorldTransforms(ShapeWorldTransform0, ShapeWorldTransform1);
+
+		Constraint->SetCullDistance(CullDistance);
+		Constraint->ResetManifold();
+		Constraint->ResetActiveManifoldContacts();
+
+		FConstGenericParticleHandle P0 = Constraint->GetParticle0();
+		FConstGenericParticleHandle P1 = Constraint->GetParticle1();
+
+		// We need the previous transform for the swept collision detector. It assumes that the current
+		// transform has been set on the constraint. 
+		// We assume that the particle's center of mass moved in a stright line and that it's rotation has 
+		// not changed so we calculate the previous transform from the current one and the velocity.
+		// NOTE: These are actor transforms, not CoM transforms
+		// @todo(chaos): Pass both start and end transforms to the collision detector
+		const FRigidTransform3 CCDParticleWorldTransform0 = FRigidTransform3(P0->P() - P0->V() * Dt, P0->Q());
+		const FRigidTransform3 CCDParticleWorldTransform1 = FRigidTransform3(P1->P() - P1->V() * Dt, P1->Q());
+		const FRigidTransform3 CCDShapeWorldTransform0 = Constraint->GetShapeRelativeTransform0() * CCDParticleWorldTransform0;
+		const FRigidTransform3 CCDShapeWorldTransform1 = Constraint->GetShapeRelativeTransform1() * CCDParticleWorldTransform1;
+		const bool bDidSweep = Collisions::UpdateConstraintSwept(*Constraint, CCDShapeWorldTransform0, CCDShapeWorldTransform1, Dt);
+
+#if CHAOS_DEBUG_DRAW
+		if (CVars::ChaosSolverDrawCCDInteractions)
+		{
+			if (FConstGenericParticleHandle(Constraint->GetParticle0())->CCDEnabled())
+			{
+				DebugDraw::DrawShape(CCDShapeWorldTransform0, Constraint->GetImplicit0(), Constraint->GetShape0(), FColor::Black, &CVars::ChaosSolverDebugDebugDrawSettings);
+				DebugDraw::DrawShape(ShapeWorldTransform0, Constraint->GetImplicit0(), Constraint->GetShape0(), FColor::White, &CVars::ChaosSolverDebugDebugDrawSettings);
+			}
+			if (FConstGenericParticleHandle(Constraint->GetParticle1())->CCDEnabled())
+			{
+				DebugDraw::DrawShape(CCDShapeWorldTransform1, Constraint->GetImplicit1(), Constraint->GetShape1(), FColor::Black, &CVars::ChaosSolverDebugDebugDrawSettings);
+				DebugDraw::DrawShape(ShapeWorldTransform1, Constraint->GetImplicit1(), Constraint->GetShape1(), FColor::White, &CVars::ChaosSolverDebugDebugDrawSettings);
+			}
+		}
+#endif
+
+		// If we did get a hit but it's at TOI = 1, treat this constraint as a regular non-swept constraint (skip the rewind)
+		if ((!bDidSweep) || Constraint->GetCCDTimeOfImpact() == FReal(1))
+		{
+			Collisions::UpdateConstraint(*Constraint, Constraint->GetShapeWorldTransform0(), Constraint->GetShapeWorldTransform1(), Dt);
+		}
+
+		// If the sweep did not find a hit, or we are treating it like a regular contact, we will skip the CCD rewind step for this contact
+		if (Constraint->GetCCDTimeOfImpact() >= FReal(1))
+		{
+			Constraint->SetCCDSweepEnabled(false);
+		}
+
+		Context.GetAllocator()->ActivateConstraint(Constraint);
+	}
+
 	void FMultiShapePairCollisionDetector::PruneConstraints(const int32 CurrentEpoch)
 	{
 		// We don't prune from NewCollisions - must call ProcessNewCollisions before Prune
@@ -1022,6 +1161,10 @@ namespace Chaos
 				{
 					NumActiveConstraints += ShapePair.GenerateCollisionCCD(bUseSweep, CullDistance, Dt, Context);
 				}
+				for (FMultiShapePairCollisionDetector& MultiShapePair : MultiShapePairDetectors)
+				{
+					NumActiveConstraints += MultiShapePair.GenerateCollisionsCCD(bUseSweep, CullDistance, Dt, Context);
+				}
 			}
 			else
 			{
@@ -1029,10 +1172,10 @@ namespace Chaos
 				{
 					NumActiveConstraints += ShapePair.GenerateCollision(CullDistance, Dt, Context);
 				}
-			}
-			for (FMultiShapePairCollisionDetector& MultiShapePair : MultiShapePairDetectors)
-			{
-				NumActiveConstraints += MultiShapePair.GenerateCollisions(CullDistance, Dt, Context);
+				for (FMultiShapePairCollisionDetector& MultiShapePair : MultiShapePairDetectors)
+				{
+					NumActiveConstraints += MultiShapePair.GenerateCollisions(CullDistance, Dt, Context);
+				}
 			}
 		}
 
diff --git a/Engine/Source/Runtime/Experimental/Chaos/Private/Chaos/CollisionResolution.cpp b/Engine/Source/Runtime/Experimental/Chaos/Private/Chaos/CollisionResolution.cpp
--- a/Engine/Source/Runtime/Experimental/Chaos/Private/Chaos/CollisionResolution.cpp
+++ b/Engine/Source/Runtime/Experimental/Chaos/Private/Chaos/CollisionResolution.cpp
@@ -96,9 +96,6 @@ FAutoConsoleVariableRef CVarChaosCollisionAABBBoundsCheck(TEXT("p.Chaos.Collisio
 bool bChaos_Collision_ShapesArrayMode = true;
 FAutoConsoleVariableRef CVarChaos_Collision_ShapesArrayMode(TEXT("p.Chaos.Collision.ShapesArrayMode"), bChaos_Collision_ShapesArrayMode, TEXT(""));
 
-bool bChaos_Collision_SortParticlesOnConstraintConstruct = true;
-FAutoConsoleVariableRef CVarChaos_Collision_SortParticlesOnConstraintConstruct(TEXT("p.Chaos.Collision.SortParticlesOnConstraintConstruct"), bChaos_Collision_SortParticlesOnConstraintConstruct, TEXT(""));
-
 namespace Chaos
 {
 	namespace Collisions
@@ -415,24 +412,6 @@ namespace Chaos
 			}
 		}
 
-		template<typename T_TRAITS>
-		void ConstructBoxBoxConstraints(TGeometryParticleHandle<FReal, 3>* Particle0, TGeometryParticleHandle<FReal, 3>* Particle1, const FImplicitObject* Implicit0, const FPerShapeData* Shape0, const FImplicitObject* Implicit1, const FPerShapeData* Shape1, const FRigidTransform3& ParticleWorldTransform0, const FRigidTransform3& LocalTransform0, const FRigidTransform3& ParticleWorldTransform1, const FRigidTransform3& LocalTransform1, const FReal CullDistance, const FReal Dt, const FCollisionContext& Context)
-		{
-			const TBox<FReal, 3>* Object0 = Implicit0->template GetObject<const TBox<FReal, 3> >();
-			const TBox<FReal, 3>* Object1 = Implicit1->template GetObject<const TBox<FReal, 3> >();
-			if (ensure(Object0 && Object1))
-			{
-				FPBDCollisionConstraint* Constraint = Context.MultiShapeCollisionDetector->FindOrCreateConstraint(Particle0, Implicit0, Shape0, nullptr, LocalTransform0, Particle1, Implicit1, Shape1, nullptr, LocalTransform1, CullDistance, EContactShapesType::BoxBox, Context.GetSettings().bAllowManifolds, Context);
-				if (T_TRAITS::bImmediateUpdate && (Constraint != nullptr))
-				{
-					const FRigidTransform3 WorldTransform0 = LocalTransform0 * ParticleWorldTransform0;
-					const FRigidTransform3 WorldTransform1 = LocalTransform1 * ParticleWorldTransform1;
-					UpdateBoxBoxConstraint(*Object0, WorldTransform0, *Object1, WorldTransform1, Dt, *Constraint);
-				}
-			}
-		}
-
-
 		//
 		// Box - HeightField
 		//
@@ -450,27 +429,6 @@ namespace Chaos
 			}
 		}
 
-		template<typename T_TRAITS>
-		void ConstructBoxHeightFieldConstraints(TGeometryParticleHandle<FReal, 3>* Particle0, TGeometryParticleHandle<FReal, 3>* Particle1, const FImplicitObject* Implicit0, const FPerShapeData* Shape0, const FImplicitObject* Implicit1, const FPerShapeData* Shape1, const FRigidTransform3& ParticleWorldTransform0, const FRigidTransform3& LocalTransform0, const FRigidTransform3& ParticleWorldTransform1, const FRigidTransform3& LocalTransform1, const FReal CullDistance, const FReal Dt, const FCollisionContext& Context)
-		{
-
-			const TBox<FReal, 3>* Object0 = Implicit0->template GetObject<const TBox<FReal, 3> >();
-			const FHeightField* Object1 = Implicit1->template GetObject<const FHeightField >();
-			if (ensure(Object0 && Object1))
-			{
-				// @todo(chaos): one-shot manifold
-				FPBDCollisionConstraint* Constraint = Context.MultiShapeCollisionDetector->FindOrCreateConstraint(Particle0, Implicit0, Shape0, nullptr, LocalTransform0, Particle1, Implicit1, Shape1, nullptr, LocalTransform1, CullDistance, EContactShapesType::BoxHeightField, Context.GetSettings().bAllowManifolds, Context);
-				if (T_TRAITS::bImmediateUpdate && (Constraint != nullptr))
-				{
-					FRigidTransform3 WorldTransform0 = LocalTransform0 * ParticleWorldTransform0;
-					FRigidTransform3 WorldTransform1 = LocalTransform1 * ParticleWorldTransform1;
-					UpdateBoxHeightFieldConstraint(*Object0, WorldTransform0, *Object1, WorldTransform1, Dt, *Constraint);
-				}
-			}
-		}
-
-
-
 		//
 		// Box-Plane
 		//
@@ -543,24 +501,6 @@ namespace Chaos
 			UpdateContactPoint(Constraint, BoxPlaneContactPoint(A, B, ATransform, BTransform), Dt);
 		}
 
-
-		template<typename T_TRAITS>
-		void ConstructBoxPlaneConstraints(TGeometryParticleHandle<FReal, 3>* Particle0, TGeometryParticleHandle<FReal, 3>* Particle1, const FImplicitObject* Implicit0, const FPerShapeData* Shape0, const FImplicitObject* Implicit1, const FPerShapeData* Shape1, const FRigidTransform3& ParticleWorldTransform0, const FRigidTransform3& LocalTransform0, const FRigidTransform3& ParticleWorldTransform1, const FRigidTransform3& LocalTransform1, const FReal CullDistance, const FReal Dt, const FCollisionContext& Context)
-		{
-			const TBox<FReal, 3>* Object0 = Implicit0->template GetObject<const TBox<FReal, 3> >();
-			const TPlane<FReal, 3>* Object1 = Implicit1->template GetObject<const TPlane<FReal, 3> >();
-			if (ensure(Object0 && Object1))
-			{
-				FPBDCollisionConstraint* Constraint = Context.MultiShapeCollisionDetector->FindOrCreateConstraint(Particle0, Implicit0, Shape0, nullptr, LocalTransform0, Particle1, Implicit1, Shape1, nullptr, LocalTransform1, CullDistance, EContactShapesType::BoxPlane, Context.GetSettings().bAllowManifolds, Context);
-				if (T_TRAITS::bImmediateUpdate && (Constraint != nullptr))
-				{
-					FRigidTransform3 WorldTransform0 = LocalTransform0 * ParticleWorldTransform0;
-					FRigidTransform3 WorldTransform1 = LocalTransform1 * ParticleWorldTransform1;
-					UpdateBoxPlaneConstraint(*Object0, WorldTransform0, *Object1, WorldTransform1, Dt, *Constraint);
-				}
-			}
-		}
-
 		//
 		// Box-TriangleMesh
 		//
@@ -579,41 +519,6 @@ namespace Chaos
 			}
 		}
 
-		template<typename T_TRAITS>
-		void ConstructBoxTriangleMeshConstraints(TGeometryParticleHandle<FReal, 3>* Particle0, TGeometryParticleHandle<FReal, 3>* Particle1, const FImplicitObject* Implicit0, const FPerShapeData* Shape0, const FImplicitObject* Implicit1, const FPerShapeData* Shape1, const FRigidTransform3& ParticleWorldTransform0, const FRigidTransform3& LocalTransform0, const FRigidTransform3& ParticleWorldTransform1, const FRigidTransform3& LocalTransform1, const FReal CullDistance, const FReal Dt, const FCollisionContext& Context)
-		{
-			const TBox<FReal, 3>* Object0 = Implicit0->template GetObject<const TBox<FReal, 3> >();
-			if (ensure(Object0))
-			{
-				if (const TImplicitObjectScaled<FTriangleMeshImplicitObject>* ScaledTriangleMesh = Implicit1->template GetObject<const TImplicitObjectScaled<FTriangleMeshImplicitObject>>())
-				{
-					// @todo(chaos): one-shot manifold
-					FPBDCollisionConstraint* Constraint = Context.MultiShapeCollisionDetector->FindOrCreateConstraint(Particle0, Implicit0, Shape0, nullptr, LocalTransform0, Particle1, Implicit1, Shape1, nullptr, LocalTransform1, CullDistance, EContactShapesType::BoxTriMesh, Context.GetSettings().bAllowManifolds, Context);
-					if (T_TRAITS::bImmediateUpdate && (Constraint != nullptr))
-					{
-						FRigidTransform3 WorldTransform0 = LocalTransform0 * ParticleWorldTransform0;
-						FRigidTransform3 WorldTransform1 = LocalTransform1 * ParticleWorldTransform1;
-						UpdateBoxTriangleMeshConstraint(*Object0, WorldTransform0, *ScaledTriangleMesh, WorldTransform1, Dt, *Constraint);
-					}
-				}
-				else if (const FTriangleMeshImplicitObject* TriangleMesh = Implicit1->template GetObject<const FTriangleMeshImplicitObject>())
-				{
-					// @todo(chaos): one-shot manifold
-					FPBDCollisionConstraint* Constraint = Context.MultiShapeCollisionDetector->FindOrCreateConstraint(Particle0, Implicit0, Shape0, nullptr, LocalTransform0, Particle1, Implicit1, Shape1, nullptr, LocalTransform1, CullDistance, EContactShapesType::BoxTriMesh, Context.GetSettings().bAllowManifolds, Context);
-					if (T_TRAITS::bImmediateUpdate && (Constraint != nullptr))
-					{
-						FRigidTransform3 WorldTransform0 = LocalTransform0 * ParticleWorldTransform0;
-						FRigidTransform3 WorldTransform1 = LocalTransform1 * ParticleWorldTransform1;
-						UpdateBoxTriangleMeshConstraint(*Object0, WorldTransform0, *TriangleMesh, WorldTransform1, Dt, *Constraint);
-					}
-				}
-				else
-				{
-					ensure(false);
-				}
-			}
-		}
-
 		//
 		// Sphere - Sphere
 		//
@@ -630,23 +535,6 @@ namespace Chaos
 			}
 		}
 
-		template<typename T_TRAITS>
-		void ConstructSphereSphereConstraints(TGeometryParticleHandle<FReal, 3>* Particle0, TGeometryParticleHandle<FReal, 3>* Particle1, const FImplicitObject* Implicit0, const FPerShapeData* Shape0, const FImplicitObject* Implicit1, const FPerShapeData* Shape1, const FRigidTransform3& ParticleWorldTransform0, const FRigidTransform3& LocalTransform0, const FRigidTransform3& ParticleWorldTransform1, const FRigidTransform3& LocalTransform1, const FReal CullDistance, const FReal Dt, const FCollisionContext& Context)
-		{
-			const TSphere<FReal, 3>* Object0 = Implicit0->template GetObject<const TSphere<FReal, 3> >();
-			const TSphere<FReal, 3>* Object1 = Implicit1->template GetObject<const TSphere<FReal, 3> >();
-			if (ensure(Object0 && Object1))
-			{
-				FPBDCollisionConstraint* Constraint = Context.MultiShapeCollisionDetector->FindOrCreateConstraint(Particle0, Implicit0, Shape0, nullptr, LocalTransform0, Particle1, Implicit1, Shape1, nullptr, LocalTransform1, CullDistance, EContactShapesType::SphereSphere, Context.GetSettings().bAllowManifolds, Context);
-				if (T_TRAITS::bImmediateUpdate && (Constraint != nullptr))
-				{
-					FRigidTransform3 WorldTransform0 = LocalTransform0 * ParticleWorldTransform0;
-					FRigidTransform3 WorldTransform1 = LocalTransform1 * ParticleWorldTransform1;
-					UpdateSphereSphereConstraint(*Object0, WorldTransform0, *Object1, WorldTransform1, Dt, *Constraint);
-				}
-			}
-		}
-
 		//
 		// Sphere - HeightField
 		//
@@ -670,37 +558,6 @@ namespace Chaos
 			Constraint.SetCCDTimeOfImpact(TOI);
 		}
 
-		template<typename T_TRAITS>
-		void ConstructSphereHeightFieldConstraints(TGeometryParticleHandle<FReal, 3>* Particle0, TGeometryParticleHandle<FReal, 3>* Particle1, const FImplicitObject* Implicit0, const FPerShapeData* Shape0, const FImplicitObject* Implicit1, const FPerShapeData* Shape1, const FRigidTransform3& ParticleWorldTransform0, const FRigidTransform3& LocalTransform0, const FRigidTransform3& ParticleWorldTransform1, const FRigidTransform3& LocalTransform1, const FReal CullDistance, const FReal Dt, const FCollisionContext& Context)
-		{
-
-			const TSphere<FReal, 3>* Object0 = Implicit0->template GetObject<const TSphere<FReal, 3> >();
-			const FHeightField* Object1 = Implicit1->template GetObject<const FHeightField >();
-			if (ensure(Object0 && Object1))
-			{
-				FPBDCollisionConstraint* Constraint = Context.MultiShapeCollisionDetector->FindOrCreateConstraint(Particle0, Implicit0, Shape0, nullptr, LocalTransform0, Particle1, Implicit1, Shape1, nullptr, LocalTransform1, CullDistance, EContactShapesType::SphereHeightField, Context.GetSettings().bAllowManifolds, Context);
-				if (T_TRAITS::bImmediateUpdate && (Constraint != nullptr))
-				{
-					FRigidTransform3 WorldTransform0 = LocalTransform0 * ParticleWorldTransform0;
-					FRigidTransform3 WorldTransform1 = LocalTransform1 * ParticleWorldTransform1;
-					UpdateSphereHeightFieldConstraint(*Object0, WorldTransform0, *Object1, WorldTransform1, Dt, *Constraint);
-				}
-			}
-		}
-
-		void ConstructSphereHeightFieldConstraintsSwept(TGeometryParticleHandle<FReal, 3>* Particle0, TGeometryParticleHandle<FReal, 3>* Particle1, const FImplicitObject* Implicit0, const FPerShapeData* Shape0, const FImplicitObject* Implicit1, const FPerShapeData* Shape1, const FRigidTransform3& ParticleWorldTransform0, const FRigidTransform3& LocalTransform0, const FRigidTransform3& ParticleWorldTransform1, const FRigidTransform3& LocalTransform1, const FReal CullDistance, const FReal Dt, const FVec3& Dir, FReal Length, const FCollisionContext& Context)
-		{
-			const TSphere<FReal, 3>* Object0 = Implicit0->template GetObject<const TSphere<FReal, 3>>();
-			const FHeightField* Object1 = Implicit1->template GetObject<const FHeightField>();
-			if (ensure(Object0 && Object1))
-			{
-				FPBDCollisionConstraint* Constraint = Context.MultiShapeCollisionDetector->FindOrCreateSweptConstraint(Particle0, Implicit0, Shape0, nullptr, LocalTransform0, Particle1, Implicit1, Shape1, nullptr, LocalTransform1, CullDistance, EContactShapesType::SphereHeightField, Context);
-				FRigidTransform3 WorldTransformX0 = LocalTransform0 * FRigidTransform3(Particle0->X(), Particle0->R());
-				FRigidTransform3 WorldTransform1 = LocalTransform1 * ParticleWorldTransform1;
-				UpdateSphereHeightFieldConstraintSwept(Particle0, *Object0, WorldTransformX0, *Object1, WorldTransform1, Dir, Length, Dt, *Constraint);
-			}
-		}
-
 		//
 		//  Sphere-Plane
 		//
@@ -717,24 +574,6 @@ namespace Chaos
 			}
 		}
 
-		template<typename T_TRAITS>
-		void ConstructSpherePlaneConstraints(TGeometryParticleHandle<FReal, 3>* Particle0, TGeometryParticleHandle<FReal, 3>* Particle1, const FImplicitObject* Implicit0, const FPerShapeData* Shape0, const FImplicitObject* Implicit1, const FPerShapeData* Shape1, const FRigidTransform3& ParticleWorldTransform0, const FRigidTransform3& LocalTransform0, const FRigidTransform3& ParticleWorldTransform1, const FRigidTransform3& LocalTransform1, const FReal CullDistance, const FReal Dt, const FCollisionContext& Context)
-		{
-
-			const TSphere<FReal, 3>* Object0 = Implicit0->template GetObject<const TSphere<FReal, 3> >();
-			const TPlane<FReal, 3>* Object1 = Implicit1->template GetObject<const TPlane<FReal, 3> >();
-			if (ensure(Object0 && Object1))
-			{
-				FPBDCollisionConstraint* Constraint = Context.MultiShapeCollisionDetector->FindOrCreateConstraint(Particle0, Implicit0, Shape0, nullptr, LocalTransform0, Particle1, Implicit1, Shape1, nullptr, LocalTransform1, CullDistance, EContactShapesType::SpherePlane, Context.GetSettings().bAllowManifolds, Context);
-				if (T_TRAITS::bImmediateUpdate && (Constraint != nullptr))
-				{
-					FRigidTransform3 WorldTransform0 = LocalTransform0 * ParticleWorldTransform0;
-					FRigidTransform3 WorldTransform1 = LocalTransform1 * ParticleWorldTransform1;
-					UpdateSpherePlaneConstraint(*Object0, WorldTransform0, *Object1, WorldTransform1, Dt, *Constraint);
-				}
-			}
-		}
-
 		//
 		// Sphere - Box
 		//
@@ -752,25 +591,6 @@ namespace Chaos
 			}
 		}
 
-		template<typename T_TRAITS>
-		void ConstructSphereBoxConstraints(TGeometryParticleHandle<FReal, 3>* Particle0, TGeometryParticleHandle<FReal, 3>* Particle1, const FImplicitObject* Implicit0, const FPerShapeData* Shape0, const FImplicitObject* Implicit1, const FPerShapeData* Shape1, const FRigidTransform3& ParticleWorldTransform0, const FRigidTransform3& LocalTransform0, const FRigidTransform3& ParticleWorldTransform1, const FRigidTransform3& LocalTransform1, const FReal CullDistance, const FReal Dt, const FCollisionContext& Context)
-		{
-
-			const TSphere<FReal, 3>* Object0 = Implicit0->template GetObject<const TSphere<FReal, 3> >();
-			const TBox<FReal, 3>* Object1 = Implicit1->template GetObject<const TBox<FReal, 3> >();
-			if (ensure(Object0 && Object1))
-			{
-				FPBDCollisionConstraint* Constraint = Context.MultiShapeCollisionDetector->FindOrCreateConstraint(Particle0, Implicit0, Shape0, nullptr, LocalTransform0, Particle1, Implicit1, Shape1, nullptr, LocalTransform1, CullDistance, EContactShapesType::SphereBox, Context.GetSettings().bAllowManifolds, Context);
-				if (T_TRAITS::bImmediateUpdate && (Constraint != nullptr))
-				{
-					FRigidTransform3 WorldTransform0 = LocalTransform0 * ParticleWorldTransform0;
-					FRigidTransform3 WorldTransform1 = LocalTransform1 * ParticleWorldTransform1;
-					UpdateSphereBoxConstraint(*Object0, WorldTransform0, *Object1, WorldTransform1, Dt, *Constraint);
-				}
-			}
-		}
-
-
 		//
 		// Sphere - Capsule
 		//
@@ -788,23 +608,6 @@ namespace Chaos
 			}
 		}
 
-		template<typename T_TRAITS>
-		void ConstructSphereCapsuleConstraints(TGeometryParticleHandle<FReal, 3>* Particle0, TGeometryParticleHandle<FReal, 3>* Particle1, const FImplicitObject* Implicit0, const FPerShapeData* Shape0, const FImplicitObject* Implicit1, const FPerShapeData* Shape1, const FRigidTransform3& ParticleWorldTransform0, const FRigidTransform3& LocalTransform0, const FRigidTransform3& ParticleWorldTransform1, const FRigidTransform3& LocalTransform1, const FReal CullDistance, const FReal Dt, const FCollisionContext& Context)
-		{
-			const TSphere<FReal, 3>* Object0 = Implicit0->template GetObject<const TSphere<FReal, 3> >();
-			const FCapsule* Object1 = Implicit1->template GetObject<const FCapsule >();
-			if (ensure(Object0 && Object1))
-			{
-				FPBDCollisionConstraint* Constraint = Context.MultiShapeCollisionDetector->FindOrCreateConstraint(Particle0, Implicit0, Shape0, nullptr, LocalTransform0, Particle1, Implicit1, Shape1, nullptr, LocalTransform1, CullDistance, EContactShapesType::SphereCapsule, Context.GetSettings().bAllowManifolds, Context);
-				if (T_TRAITS::bImmediateUpdate && (Constraint != nullptr))
-				{
-					FRigidTransform3 WorldTransform0 = LocalTransform0 * ParticleWorldTransform0;
-					FRigidTransform3 WorldTransform1 = LocalTransform1 * ParticleWorldTransform1;
-					UpdateSphereCapsuleConstraint(*Object0, WorldTransform0, *Object1, WorldTransform1, Dt, *Constraint);
-				}
-			}
-		}
-
 		//
 		// Sphere-Convex
 		//
@@ -821,23 +624,6 @@ namespace Chaos
 			}
 		}
 
-		template<typename T_TRAITS>
-		void ConstructSphereConvexConstraints(TGeometryParticleHandle<FReal, 3>* Particle0, TGeometryParticleHandle<FReal, 3>* Particle1, const FImplicitObject* Implicit0, const FPerShapeData* Shape0, const FImplicitObject* Implicit1, const FPerShapeData* Shape1, const FRigidTransform3& ParticleWorldTransform0, const FRigidTransform3& LocalTransform0, const FRigidTransform3& ParticleWorldTransform1, const FRigidTransform3& LocalTransform1, const FReal CullDistance, const FReal Dt, const FCollisionContext& Context)
-		{
-			const TSphere<FReal, 3>* Object0 = Implicit0->template GetObject<const TSphere<FReal, 3> >();
-			const FImplicitObject* Object1 = Implicit1;
-			if (ensure(Object0 && Object1))
-			{
-				FPBDCollisionConstraint* Constraint = Context.MultiShapeCollisionDetector->FindOrCreateConstraint(Particle0, Implicit0, Shape0, nullptr, LocalTransform0, Particle1, Implicit1, Shape1, nullptr, LocalTransform1, CullDistance, EContactShapesType::SphereConvex, Context.GetSettings().bAllowManifolds, Context);
-				if (T_TRAITS::bImmediateUpdate && (Constraint != nullptr))
-				{
-					FRigidTransform3 WorldTransform0 = LocalTransform0 * ParticleWorldTransform0;
-					FRigidTransform3 WorldTransform1 = LocalTransform1 * ParticleWorldTransform1;
-					UpdateSphereConvexConstraint(*Object0, WorldTransform0, *Object1, WorldTransform1, Dt, *Constraint);
-				}
-			}
-		}
-
 		//
 		// Sphere-TriangleMesh
 		//
@@ -864,73 +650,10 @@ namespace Chaos
 			Constraint.SetCCDTimeOfImpact(TOI);
 		}
 
-		template<typename T_TRAITS>
-		void ConstructSphereTriangleMeshConstraints(TGeometryParticleHandle<FReal, 3>* Particle0, TGeometryParticleHandle<FReal, 3>* Particle1, const FImplicitObject* Implicit0, const FPerShapeData* Shape0, const FImplicitObject* Implicit1, const FPerShapeData* Shape1, const FRigidTransform3& ParticleWorldTransform0, const FRigidTransform3& LocalTransform0, const FRigidTransform3& ParticleWorldTransform1, const FRigidTransform3& LocalTransform1, const FReal CullDistance, const FReal Dt, const FCollisionContext& Context)
-		{
-			const TSphere<FReal, 3>* Object0 = Implicit0->template GetObject<const TSphere<FReal, 3> >();
-			if (ensure(Object0))
-			{
-				if (const TImplicitObjectScaled<FTriangleMeshImplicitObject>* ScaledTriangleMesh = Implicit1->template GetObject<const TImplicitObjectScaled<FTriangleMeshImplicitObject>>())
-				{
-					FPBDCollisionConstraint* Constraint = Context.MultiShapeCollisionDetector->FindOrCreateConstraint(Particle0, Implicit0, Shape0, nullptr, LocalTransform0, Particle1, Implicit1, Shape1, nullptr, LocalTransform1, CullDistance, EContactShapesType::SphereTriMesh, Context.GetSettings().bAllowManifolds, Context);
-					if (T_TRAITS::bImmediateUpdate && (Constraint != nullptr))
-					{
-						FRigidTransform3 WorldTransform0 = LocalTransform0 * ParticleWorldTransform0;
-						FRigidTransform3 WorldTransform1 = LocalTransform1 * ParticleWorldTransform1;
-						UpdateSphereTriangleMeshConstraint(*Object0, WorldTransform0, *ScaledTriangleMesh, WorldTransform1, Dt, *Constraint);
-					}
-				}
-				else if (const FTriangleMeshImplicitObject* TriangleMesh = Implicit1->template GetObject<const FTriangleMeshImplicitObject>())
-				{
-					FPBDCollisionConstraint* Constraint = Context.MultiShapeCollisionDetector->FindOrCreateConstraint(Particle0, Implicit0, Shape0, nullptr, LocalTransform0, Particle1, Implicit1, Shape1, nullptr, LocalTransform1, CullDistance, EContactShapesType::SphereTriMesh, Context.GetSettings().bAllowManifolds, Context);
-					if (T_TRAITS::bImmediateUpdate && (Constraint != nullptr))
-					{
-						FRigidTransform3 WorldTransform0 = LocalTransform0 * ParticleWorldTransform0;
-						FRigidTransform3 WorldTransform1 = LocalTransform1 * ParticleWorldTransform1;
-						UpdateSphereTriangleMeshConstraint(*Object0, WorldTransform0, *TriangleMesh, WorldTransform1, Dt, *Constraint);
-					}
-				}
-				else
-				{
-					ensure(false);
-				}
-			}
-		}
-
-
-		void ConstructSphereTriangleMeshConstraintsSwept(TGeometryParticleHandle<FReal, 3>* Particle0, TGeometryParticleHandle<FReal, 3>* Particle1, const FImplicitObject* Implicit0, const FPerShapeData* Shape0, const FImplicitObject* Implicit1, const FPerShapeData* Shape1, const FRigidTransform3& ParticleWorldTransform0, const FRigidTransform3& LocalTransform0, const FRigidTransform3& ParticleWorldTransform1, const FRigidTransform3& LocalTransform1, const FReal CullDistance, const FReal Dt, const FVec3& Dir, FReal Length, const FCollisionContext& Context)
-		{
-			const TSphere<FReal, 3>* Object0 = Implicit0->template GetObject<const TSphere<FReal, 3>>();
-			if (ensure(Object0))
-			{
-				if (const TImplicitObjectScaled<FTriangleMeshImplicitObject>* ScaledTriangleMesh = Implicit1->template GetObject<const TImplicitObjectScaled<FTriangleMeshImplicitObject>>())
-				{
-					FPBDCollisionConstraint* Constraint = Context.MultiShapeCollisionDetector->FindOrCreateSweptConstraint(Particle0, Implicit0, Shape0, nullptr, LocalTransform0, Particle1, Implicit1, Shape1, nullptr, LocalTransform1, CullDistance, EContactShapesType::SphereTriMesh, Context);
-					FRigidTransform3 WorldTransformX0 = LocalTransform0 * FRigidTransform3(Particle0->X(), Particle0->R());
-					FRigidTransform3 WorldTransform1 = LocalTransform1 * ParticleWorldTransform1;
-					UpdateSphereTriangleMeshConstraintSwept(Particle0, *Object0, WorldTransformX0, *ScaledTriangleMesh, WorldTransform1, Dir, Length, Dt, *Constraint);
-				}
-				else if (const FTriangleMeshImplicitObject* TriangleMesh = Implicit1->template GetObject<const FTriangleMeshImplicitObject>())
-				{
-					FPBDCollisionConstraint* Constraint = Context.MultiShapeCollisionDetector->FindOrCreateSweptConstraint(Particle0, Implicit0, Shape0, nullptr, LocalTransform0, Particle1, Implicit1, Shape1, nullptr, LocalTransform1, CullDistance, EContactShapesType::SphereTriMesh, Context);
-					FRigidTransform3 WorldTransformX0 = LocalTransform0 * FRigidTransform3(Particle0->X(), Particle0->R());
-					FRigidTransform3 WorldTransform1 = LocalTransform1 * ParticleWorldTransform1;
-					UpdateSphereTriangleMeshConstraintSwept(Particle0, *Object0, WorldTransformX0, *TriangleMesh, WorldTransform1, Dir, Length, Dt, *Constraint);
-				}
-				else
-				{
-					ensure(false);
-				}
-			}
-		}
-
-
-
 		//
 		// Capsule-Capsule
 		//
 
-
 		void UpdateCapsuleCapsuleConstraint(const FCapsule& A, const FRigidTransform3& ATransform, const FCapsule& B, const FRigidTransform3& BTransform, const FReal Dt, FPBDCollisionConstraint& Constraint)
 		{
 			if (Constraint.GetUseManifold())
@@ -943,23 +666,6 @@ namespace Chaos
 			}
 		}
 
-		template<typename T_TRAITS>
-		void ConstructCapsuleCapsuleConstraints(TGeometryParticleHandle<FReal, 3>* Particle0, TGeometryParticleHandle<FReal, 3>* Particle1, const FImplicitObject* Implicit0, const FPerShapeData* Shape0, const FImplicitObject* Implicit1, const FPerShapeData* Shape1, const FRigidTransform3& ParticleWorldTransform0, const FRigidTransform3& LocalTransform0, const FRigidTransform3& ParticleWorldTransform1, const FRigidTransform3& LocalTransform1, const FReal CullDistance, const FReal Dt, const FCollisionContext& Context)
-		{
-			const FCapsule* Object0 = Implicit0->template GetObject<const FCapsule >();
-			const FCapsule* Object1 = Implicit1->template GetObject<const FCapsule >();
-			if (ensure(Object0 && Object1))
-			{
-				FPBDCollisionConstraint* Constraint = Context.MultiShapeCollisionDetector->FindOrCreateConstraint(Particle0, Implicit0, Shape0, nullptr, LocalTransform0, Particle1, Implicit1, Shape1, nullptr, LocalTransform1, CullDistance, EContactShapesType::CapsuleCapsule, Context.GetSettings().bAllowManifolds, Context);
-				if (T_TRAITS::bImmediateUpdate && (Constraint != nullptr))
-				{
-					FRigidTransform3 WorldTransform0 = LocalTransform0 * ParticleWorldTransform0;
-					FRigidTransform3 WorldTransform1 = LocalTransform1 * ParticleWorldTransform1;
-					UpdateCapsuleCapsuleConstraint(*Object0, WorldTransform0, *Object1, WorldTransform1, Dt, *Constraint);
-				}
-			}
-		}
-
 		//
 		// Capsule - Box
 		//
@@ -984,24 +690,6 @@ namespace Chaos
 			}
 		}
 
-		template<typename T_TRAITS>
-		void ConstructCapsuleBoxConstraints(TGeometryParticleHandle<FReal, 3>* Particle0, TGeometryParticleHandle<FReal, 3>* Particle1, const FImplicitObject* Implicit0, const FPerShapeData* Shape0, const FImplicitObject* Implicit1, const FPerShapeData* Shape1, const FRigidTransform3& ParticleWorldTransform0, const FRigidTransform3& LocalTransform0, const FRigidTransform3& ParticleWorldTransform1, const FRigidTransform3& LocalTransform1, const FReal CullDistance, const FReal Dt, const FCollisionContext& Context)
-		{
-			const FCapsule* Object0 = Implicit0->template GetObject<const FCapsule>();
-			const TBox<FReal, 3>* Object1 = Implicit1->template GetObject<const TBox<FReal, 3>>();
-			if (ensure(Object0 && Object1))
-			{
-				FPBDCollisionConstraint* Constraint = Context.MultiShapeCollisionDetector->FindOrCreateConstraint(Particle0, Implicit0, Shape0, nullptr, LocalTransform0, Particle1, Implicit1, Shape1, nullptr, LocalTransform1, CullDistance, EContactShapesType::CapsuleBox, Context.GetSettings().bAllowManifolds, Context);
-				if (T_TRAITS::bImmediateUpdate && (Constraint != nullptr))
-				{
-					FRigidTransform3 WorldTransform0 = LocalTransform0 * ParticleWorldTransform0;
-					FRigidTransform3 WorldTransform1 = LocalTransform1 * ParticleWorldTransform1;
-					UpdateCapsuleBoxConstraint(*Object0, WorldTransform0, *Object1, WorldTransform1, Dt, *Constraint);
-				}
-			}
-		}
-
-
 		//
 		// Capsule-Convex
 		//
@@ -1066,24 +754,6 @@ namespace Chaos
 				});
 		}
 
-		template<typename T_TRAITS>
-		void ConstructCapsuleConvexConstraints(TGeometryParticleHandle<FReal, 3>* Particle0, TGeometryParticleHandle<FReal, 3>* Particle1, const FImplicitObject* Implicit0, const FPerShapeData* Shape0, const FImplicitObject* Implicit1, const FPerShapeData* Shape1, const FRigidTransform3& ParticleWorldTransform0, const FRigidTransform3& LocalTransform0, const FRigidTransform3& ParticleWorldTransform1, const FRigidTransform3& LocalTransform1, const FReal CullDistance, const FReal Dt, const FCollisionContext& Context)
-		{
-			const FCapsule* Object0 = Implicit0->template GetObject<const FCapsule>();
-			const FImplicitObject* Object1 = Implicit1;
-			if (ensure(Object0 && Object1))
-			{
-				FPBDCollisionConstraint* Constraint = Context.MultiShapeCollisionDetector->FindOrCreateConstraint(Particle0, Implicit0, Shape0, nullptr, LocalTransform0, Particle1, Implicit1, Shape1, nullptr, LocalTransform1, CullDistance, EContactShapesType::CapsuleConvex, Context.GetSettings().bAllowManifolds, Context);
-				if (T_TRAITS::bImmediateUpdate && (Constraint != nullptr))
-				{
-					FRigidTransform3 WorldTransform0 = LocalTransform0 * ParticleWorldTransform0;
-					FRigidTransform3 WorldTransform1 = LocalTransform1 * ParticleWorldTransform1;
-					UpdateCapsuleGenericConvexConstraintHelper(*Object0, WorldTransform0, *Object1, WorldTransform1, Dt, *Constraint);
-				}
-			}
-		}
-
-
 		//
 		// Capsule-HeightField
 		//
@@ -1110,45 +780,10 @@ namespace Chaos
 			Constraint.SetCCDTimeOfImpact(TOI);
 		}
 
-		template<typename T_TRAITS>
-		void ConstructCapsuleHeightFieldConstraints(TGeometryParticleHandle<FReal, 3>* Particle0, TGeometryParticleHandle<FReal, 3>* Particle1, const FImplicitObject* Implicit0, const FPerShapeData* Shape0, const FImplicitObject* Implicit1, const FPerShapeData* Shape1, const FRigidTransform3& ParticleWorldTransform0, const FRigidTransform3& LocalTransform0, const FRigidTransform3& ParticleWorldTransform1, const FRigidTransform3& LocalTransform1, const FReal CullDistance, const FReal Dt, const FCollisionContext& Context)
-		{
-			CONDITIONAL_SCOPE_CYCLE_COUNTER(STAT_Collisions_ConstructCapsuleHeightFieldConstraints, ConstraintsDetailedStats);
-
-			const FCapsule* Object0 = Implicit0->template GetObject<const FCapsule >();
-			const FHeightField* Object1 = Implicit1->template GetObject<const FHeightField >();
-			if (ensure(Object0 && Object1))
-			{
-				FPBDCollisionConstraint* Constraint = Context.MultiShapeCollisionDetector->FindOrCreateConstraint(Particle0, Implicit0, Shape0, nullptr, LocalTransform0, Particle1, Implicit1, Shape1, nullptr, LocalTransform1, CullDistance, EContactShapesType::CapsuleHeightField, Context.GetSettings().bAllowManifolds, Context);
-				if (T_TRAITS::bImmediateUpdate && (Constraint != nullptr))
-				{
-					FRigidTransform3 WorldTransform0 = LocalTransform0 * ParticleWorldTransform0;
-					FRigidTransform3 WorldTransform1 = LocalTransform1 * ParticleWorldTransform1;
-					UpdateCapsuleHeightFieldConstraint(*Object0, WorldTransform0, *Object1, WorldTransform1, Dt, *Constraint);
-				}
-			}
-		}
-
-		void ConstructCapsuleHeightFieldConstraintsSwept(TGeometryParticleHandle<FReal, 3>* Particle0, TGeometryParticleHandle<FReal, 3>* Particle1, const FImplicitObject* Implicit0, const FPerShapeData* Shape0, const FImplicitObject* Implicit1, const FPerShapeData* Shape1, const FRigidTransform3& ParticleWorldTransform0, const FRigidTransform3& LocalTransform0, const FRigidTransform3& ParticleWorldTransform1, const FRigidTransform3& LocalTransform1, const FReal CullDistance, const FReal Dt, const FVec3& Dir, FReal Length, const FCollisionContext& Context)
-		{
-			CONDITIONAL_SCOPE_CYCLE_COUNTER(STAT_Collisions_ConstructCapsuleHeightFieldConstraintsSwept, ConstraintsDetailedStats);
-			const FCapsule* Object0 = Implicit0->template GetObject<const FCapsule >();
-			const FHeightField* Object1 = Implicit1->template GetObject<const FHeightField >();
-			if (ensure(Object0 && Object1))
-			{
-				FPBDCollisionConstraint* Constraint = Context.MultiShapeCollisionDetector->FindOrCreateSweptConstraint(Particle0, Implicit0, Shape0, nullptr, LocalTransform0, Particle1, Implicit1, Shape1, nullptr, LocalTransform1, CullDistance, EContactShapesType::CapsuleHeightField, Context);
-				FRigidTransform3 WorldTransformX0 = LocalTransform0 * FRigidTransform3(Particle0->X(), Particle0->R());
-				FRigidTransform3 WorldTransform1 = LocalTransform1 * ParticleWorldTransform1;
-				UpdateCapsuleHeightFieldConstraintSwept(Particle0, *Object0, WorldTransformX0, *Object1, WorldTransform1, Dir, Length, Dt, *Constraint);
-			}
-		}
-
-
 		//
 		// Capsule-TriangleMesh
 		//
 
-
 		template <typename TriMeshType>
 		void UpdateCapsuleTriangleMeshConstraint(const FCapsule& Capsule0, const FRigidTransform3& WorldTransform0, const TriMeshType& TriangleMesh1, const FRigidTransform3& WorldTransform1, const FReal Dt, FPBDCollisionConstraint& Constraint)
 		{
@@ -1173,69 +808,6 @@ namespace Chaos
 			Constraint.SetCCDTimeOfImpact(TOI);
 		}
 
-		template<typename T_TRAITS>
-		void ConstructCapsuleTriangleMeshConstraints(TGeometryParticleHandle<FReal, 3>* Particle0, TGeometryParticleHandle<FReal, 3>* Particle1, const FImplicitObject* Implicit0, const FPerShapeData* Shape0, const FImplicitObject* Implicit1, const FPerShapeData* Shape1, const FRigidTransform3& ParticleWorldTransform0, const FRigidTransform3& LocalTransform0, const FRigidTransform3& ParticleWorldTransform1, const FRigidTransform3& LocalTransform1, const FReal CullDistance, const FReal Dt, const FCollisionContext& Context)
-		{
-			CONDITIONAL_SCOPE_CYCLE_COUNTER(STAT_Collisions_ConstructCapsuleTriangleMeshConstraints, ConstraintsDetailedStats);
-
-			const FCapsule* Object0 = Implicit0->template GetObject<const FCapsule >();
-			if (ensure(Object0))
-			{
-				if (const TImplicitObjectScaled<FTriangleMeshImplicitObject>* ScaledTriangleMesh = Implicit1->template GetObject<const TImplicitObjectScaled<FTriangleMeshImplicitObject>>())
-				{
-					FPBDCollisionConstraint* Constraint = Context.MultiShapeCollisionDetector->FindOrCreateConstraint(Particle0, Implicit0, Shape0, nullptr, LocalTransform0, Particle1, Implicit1, Shape1,  nullptr, LocalTransform1, CullDistance, EContactShapesType::CapsuleTriMesh, Context.GetSettings().bAllowManifolds, Context);
-					if (T_TRAITS::bImmediateUpdate && (Constraint != nullptr))
-					{
-						FRigidTransform3 WorldTransform0 = LocalTransform0 * ParticleWorldTransform0;
-						FRigidTransform3 WorldTransform1 = LocalTransform1 * ParticleWorldTransform1;
-						UpdateCapsuleTriangleMeshConstraint(*Object0, WorldTransform0, *ScaledTriangleMesh, WorldTransform1, Dt, *Constraint);
-					}
-				}
-				else if (const FTriangleMeshImplicitObject* TriangleMesh = Implicit1->template GetObject<const FTriangleMeshImplicitObject>())
-				{
-					FPBDCollisionConstraint* Constraint = Context.MultiShapeCollisionDetector->FindOrCreateConstraint(Particle0, Implicit0, Shape0, nullptr, LocalTransform0, Particle1, Implicit1, Shape1, nullptr, LocalTransform1, CullDistance, EContactShapesType::CapsuleTriMesh, Context.GetSettings().bAllowManifolds, Context);
-					if (T_TRAITS::bImmediateUpdate && (Constraint != nullptr))
-					{
-						FRigidTransform3 WorldTransform0 = LocalTransform0 * ParticleWorldTransform0;
-						FRigidTransform3 WorldTransform1 = LocalTransform1 * ParticleWorldTransform1;
-						UpdateCapsuleTriangleMeshConstraint(*Object0, WorldTransform0, *TriangleMesh, WorldTransform1, Dt, *Constraint);
-					}
-				}
-				else
-				{
-					ensure(false);
-				}
-			}
-		}
-
-		void ConstructCapsuleTriangleMeshConstraintsSwept(TGeometryParticleHandle<FReal, 3>* Particle0, TGeometryParticleHandle<FReal, 3>* Particle1, const FImplicitObject* Implicit0, const FPerShapeData* Shape0, const FImplicitObject* Implicit1, const FPerShapeData* Shape1, const FRigidTransform3& ParticleWorldTransform0, const FRigidTransform3& LocalTransform0, const FRigidTransform3& ParticleWorldTransform1, const FRigidTransform3& LocalTransform1, const FReal CullDistance, const FReal Dt, const FVec3& Dir, FReal Length, const FCollisionContext& Context)
-		{
-			CONDITIONAL_SCOPE_CYCLE_COUNTER(STAT_Collisions_ConstructCapsuleTriangleMeshConstraintsSwept, ConstraintsDetailedStats);
-
-			const FCapsule* Object0 = Implicit0->template GetObject<const FCapsule >();
-			if (ensure(Object0))
-			{
-				if (const TImplicitObjectScaled<FTriangleMeshImplicitObject>* ScaledTriangleMesh = Implicit1->template GetObject<const TImplicitObjectScaled<FTriangleMeshImplicitObject>>())
-				{
-					FPBDCollisionConstraint* Constraint = Context.MultiShapeCollisionDetector->FindOrCreateSweptConstraint(Particle0, Implicit0, Shape0, nullptr, LocalTransform0, Particle1, Implicit1, Shape1, nullptr, LocalTransform1, CullDistance, EContactShapesType::CapsuleTriMesh, Context);
-					FRigidTransform3 WorldTransformX0 = LocalTransform0 * FRigidTransform3(Particle0->X(), Particle0->R());
-					FRigidTransform3 WorldTransform1 = LocalTransform1 * ParticleWorldTransform1;
-					UpdateCapsuleTriangleMeshConstraintSwept(Particle0, *Object0, WorldTransformX0, *ScaledTriangleMesh, WorldTransform1, Dir, Length, Dt, *Constraint);
-				}
-				else if (const FTriangleMeshImplicitObject* TriangleMesh = Implicit1->template GetObject<const FTriangleMeshImplicitObject>())
-				{
-					FPBDCollisionConstraint* Constraint = Context.MultiShapeCollisionDetector->FindOrCreateSweptConstraint(Particle0, Implicit0, Shape0, nullptr, LocalTransform0, Particle1, Implicit1, Shape1, nullptr, LocalTransform1, CullDistance, EContactShapesType::CapsuleTriMesh, Context);
-					FRigidTransform3 WorldTransformX0 = LocalTransform0 * FRigidTransform3(Particle0->X(), Particle0->R());
-					FRigidTransform3 WorldTransform1 = LocalTransform1 * ParticleWorldTransform1;
-					UpdateCapsuleTriangleMeshConstraintSwept(Particle0, *Object0, WorldTransformX0, *TriangleMesh, WorldTransform1, Dir, Length, Dt, *Constraint);
-				}
-				else
-				{
-					ensure(false);
-				}
-			}
-		}
-
 		//
 		// Generic Convex - Convex (actual concrete type could be anything)
 		//
@@ -1257,46 +829,10 @@ namespace Chaos
 			Constraint.SetCCDTimeOfImpact(TOI);
 		}
 
-
-		template<typename T_TRAITS>
-		void ConstructGenericConvexConvexConstraints(TGeometryParticleHandle<FReal, 3>* Particle0, TGeometryParticleHandle<FReal, 3>* Particle1, const FImplicitObject* Implicit0, const FPerShapeData* Shape0, const FImplicitObject* Implicit1, const FPerShapeData* Shape1, const FRigidTransform3& ParticleWorldTransform0, const FRigidTransform3& LocalTransform0, const FRigidTransform3& ParticleWorldTransform1, const FRigidTransform3& LocalTransform1, const FReal CullDistance, const FReal Dt, const FCollisionContext& Context)
-		{
-			CONDITIONAL_SCOPE_CYCLE_COUNTER(STAT_Collisions_ConstructGenericConvexConvexConstraints, ConstraintsDetailedStats);
-			EImplicitObjectType Implicit0Type = Particle0->Geometry()->GetType();
-			EImplicitObjectType Implicit1Type = Particle1->Geometry()->GetType();
-
-			FPBDCollisionConstraint* Constraint = Context.MultiShapeCollisionDetector->FindOrCreateConstraint(Particle0, Implicit0, Shape0, nullptr, LocalTransform0, Particle1, Implicit1, Shape1, nullptr, LocalTransform1, CullDistance, EContactShapesType::GenericConvexConvex, Context.GetSettings().bAllowManifolds, Context);
-			if (T_TRAITS::bImmediateUpdate && (Constraint != nullptr))
-			{
-				FRigidTransform3 WorldTransform0 = LocalTransform0 * ParticleWorldTransform0;
-				FRigidTransform3 WorldTransform1 = LocalTransform1 * ParticleWorldTransform1;
-				UpdateGenericConvexConvexConstraint(*Implicit0, WorldTransform0, *Implicit1, WorldTransform1, Dt, *Constraint);
-			}
-		}
-
-		void ConstructGenericConvexConvexConstraintsSwept(TGeometryParticleHandle<FReal, 3>* Particle0, TGeometryParticleHandle<FReal, 3>* Particle1, const FImplicitObject* Implicit0, const FPerShapeData* Shape0, const FImplicitObject* Implicit1, const FPerShapeData* Shape1, const FRigidTransform3& ParticleWorldTransform0, const FRigidTransform3& LocalTransform0, const FVec3& StartX0, const FRigidTransform3& ParticleWorldTransform1, const FRigidTransform3& LocalTransform1, const FVec3& StartX1, const FReal CullDistance, const FReal Dt, const FVec3& Dir, FReal Length, const FCollisionContext& Context)
-		{
-			CONDITIONAL_SCOPE_CYCLE_COUNTER(STAT_Collisions_ConstructGenericConvexConvexConstraintsSwept, ConstraintsDetailedStats);
-			FGenericParticleHandle RigidParticle0 = FGenericParticleHandle(Particle0);
-			FGenericParticleHandle RigidParticle1 = FGenericParticleHandle(Particle1);
-			FPBDCollisionConstraint* Constraint = Context.MultiShapeCollisionDetector->FindOrCreateSweptConstraint(Particle0, Implicit0, Shape0, nullptr, LocalTransform0, Particle1, Implicit1, Shape1, nullptr, LocalTransform1, CullDistance, EContactShapesType::GenericConvexConvex, Context);
-			// Rotations of WorldTransformXRx are not actually used in UpdateGenericConvexConvexConstraintSwept.
-			FRigidTransform3 WorldTransformXR0 = LocalTransform0 * FRigidTransform3(StartX0, RigidParticle0->R());
-			FRigidTransform3 WorldTransformXR1 = LocalTransform1 * FRigidTransform3(StartX1, RigidParticle1->R());
-			UpdateGenericConvexConvexConstraintSwept(Particle0, *Implicit0, WorldTransformXR0, Particle1, *Implicit1, WorldTransformXR1, Dir, Length, Dt, *Constraint);
-
-			// @todo(zhenglin): If resweep is not enabled, we don't need to process secondary collision constraints so constraints with TOI > 1.0f can be deleted. This could be a potential optimization but I have not made this work yet.
-			// if (!bChaosCollisionCCDEnableResweep && Constraint->TimeOfImpact > 1.0f)
-			// {
-			// 	Context.MultiShapeCollisionDetector->DestroyConstraint(Constraint);
-			// }
-		}
-
 		//
 		// Convex - HeightField
 		//
 
-
 		void UpdateConvexHeightFieldConstraint(const FImplicitObject& A, const FRigidTransform3& ATransform, const FHeightField& B, const FRigidTransform3& BTransform, const FReal Dt, FPBDCollisionConstraint& Constraint)
 		{
 			CONDITIONAL_SCOPE_CYCLE_COUNTER(STAT_Collisions_UpdateConvexHeightFieldConstraint, ConstraintsDetailedStats);
@@ -1319,36 +855,6 @@ namespace Chaos
 			Constraint.SetCCDTimeOfImpact(TOI);
 		}
 
-		template<typename T_TRAITS>
-		void ConstructConvexHeightFieldConstraints(TGeometryParticleHandle<FReal, 3>* Particle0, TGeometryParticleHandle<FReal, 3>* Particle1, const FImplicitObject* Implicit0, const FPerShapeData* Shape0, const FImplicitObject* Implicit1, const FPerShapeData* Shape1, const FRigidTransform3& ParticleWorldTransform0, const FRigidTransform3& LocalTransform0, const FRigidTransform3& ParticleWorldTransform1, const FRigidTransform3& LocalTransform1, const FReal CullDistance, const FReal Dt, const FCollisionContext& Context)
-		{
-			CONDITIONAL_SCOPE_CYCLE_COUNTER(STAT_Collisions_ConstructConvexHeightFieldConstraints, ConstraintsDetailedStats);
-			const FHeightField* Object1 = Implicit1->template GetObject<const FHeightField >();
-			if (ensure(Implicit0->IsConvex() && Object1))
-			{
-				FPBDCollisionConstraint* Constraint = Context.MultiShapeCollisionDetector->FindOrCreateConstraint(Particle0, Implicit0, Shape0, nullptr, LocalTransform0, Particle1, Implicit1, Shape1, nullptr, LocalTransform1, CullDistance, EContactShapesType::ConvexHeightField, Context.GetSettings().bAllowManifolds, Context);
-				if (T_TRAITS::bImmediateUpdate && (Constraint != nullptr))
-				{
-					FRigidTransform3 WorldTransform0 = LocalTransform0 * ParticleWorldTransform0;
-					FRigidTransform3 WorldTransform1 = LocalTransform1 * ParticleWorldTransform1;
-					UpdateConvexHeightFieldConstraint(*Implicit0, WorldTransform0, *Object1, WorldTransform1, Dt, *Constraint);
-				}
-			}
-		}
-
-		void ConstructConvexHeightFieldConstraintsSwept(TGeometryParticleHandle<FReal, 3>* Particle0, TGeometryParticleHandle<FReal, 3>* Particle1, const FImplicitObject* Implicit0, const FPerShapeData* Shape0, const FImplicitObject* Implicit1, const FPerShapeData* Shape1, const FRigidTransform3& ParticleWorldTransform0, const FRigidTransform3& LocalTransform0, const FRigidTransform3& ParticleWorldTransform1, const FRigidTransform3& LocalTransform1, const FReal CullDistance, const FReal Dt, const FVec3& Dir, FReal Length, const FCollisionContext& Context)
-		{
-			CONDITIONAL_SCOPE_CYCLE_COUNTER(STAT_Collisions_ConstructConvexHeightFieldConstraintsSwept, ConstraintsDetailedStats);
-			const FHeightField* Object1 = Implicit1->template GetObject<const FHeightField >();
-			if (ensure(Implicit0->IsConvex() && Object1))
-			{
-				FPBDCollisionConstraint* Constraint = Context.MultiShapeCollisionDetector->FindOrCreateSweptConstraint(Particle0, Implicit0, Shape0, nullptr, LocalTransform0, Particle1, Implicit1, Shape1, nullptr, LocalTransform1, CullDistance, EContactShapesType::ConvexHeightField, Context);
-				FRigidTransform3 WorldTransformX0 = LocalTransform0 * FRigidTransform3(Particle0->X(), Particle0->R());
-				FRigidTransform3 WorldTransform1 = LocalTransform1 * ParticleWorldTransform1;
-				UpdateConvexHeightFieldConstraintSwept(Particle0, *Implicit0, WorldTransformX0, *Object1, WorldTransform1, Dir, Length, Dt, *Constraint);
-			}
-		}
-
 		//
 		// Convex-TriangleMesh
 		//
@@ -1376,54 +882,10 @@ namespace Chaos
 			Constraint.SetCCDTimeOfImpact(TOI);
 		}
 
-		template<typename T_TRAITS>
-		void ConstructConvexTriangleMeshConstraints(TGeometryParticleHandle<FReal, 3>* Particle0, TGeometryParticleHandle<FReal, 3>* Particle1, const FImplicitObject* Implicit0, const FPerShapeData* Shape0, const FImplicitObject* Implicit1, const FPerShapeData* Shape1, const FRigidTransform3& ParticleWorldTransform0, const FRigidTransform3& LocalTransform0, const FRigidTransform3& ParticleWorldTransform1, const FRigidTransform3& LocalTransform1, const FReal CullDistance, const FReal Dt, const FCollisionContext& Context)
-		{
-			CONDITIONAL_SCOPE_CYCLE_COUNTER(STAT_Collisions_ConstructConvexTriangleMeshConstraints, ConstraintsDetailedStats);
-			if (ensure(Implicit0->IsConvex()))
-			{
-				FPBDCollisionConstraint* Constraint = Context.MultiShapeCollisionDetector->FindOrCreateConstraint(Particle0, Implicit0, Shape0, nullptr, LocalTransform0, Particle1, Implicit1, Shape1, nullptr, LocalTransform1, CullDistance, EContactShapesType::ConvexTriMesh, Context.GetSettings().bAllowManifolds, Context);
-				if (T_TRAITS::bImmediateUpdate && (Constraint != nullptr))
-				{
-					FRigidTransform3 WorldTransform0 = LocalTransform0 * ParticleWorldTransform0;
-					FRigidTransform3 WorldTransform1 = LocalTransform1 * ParticleWorldTransform1;
-					UpdateConvexTriangleMeshConstraint(*Implicit0, WorldTransform0, *Implicit1, WorldTransform1, Dt, *Constraint);
-				}
-			}
-		}
-
-		void ConstructConvexTriangleMeshConstraintsSwept(TGeometryParticleHandle<FReal, 3>* Particle0, TGeometryParticleHandle<FReal, 3>* Particle1, const FImplicitObject* Implicit0, const FPerShapeData* Shape0, const FImplicitObject* Implicit1, const FPerShapeData* Shape1, const FRigidTransform3& ParticleWorldTransform0, const FRigidTransform3& LocalTransform0, const FRigidTransform3& ParticleWorldTransform1, const FRigidTransform3& LocalTransform1, const FReal CullDistance, const FReal Dt, const FVec3& Dir, FReal Length, const FCollisionContext& Context)
-		{
-			CONDITIONAL_SCOPE_CYCLE_COUNTER(STAT_Collisions_ConstructConvexTriangleMeshConstraintsSwept, ConstraintsDetailedStats);
-			if (ensure(Implicit0->IsConvex()))
-			{
-				if (const TImplicitObjectScaled<FTriangleMeshImplicitObject>* ScaledTriangleMesh = Implicit1->template GetObject<const TImplicitObjectScaled<FTriangleMeshImplicitObject>>())
-				{
-					FPBDCollisionConstraint* Constraint = Context.MultiShapeCollisionDetector->FindOrCreateSweptConstraint(Particle0, Implicit0, Shape0, nullptr, LocalTransform0, Particle1, Implicit1, Shape1, nullptr, LocalTransform1, CullDistance, EContactShapesType::ConvexTriMesh, Context);
-					FRigidTransform3 WorldTransformX0 = LocalTransform0 * FRigidTransform3(Particle0->X(), Particle0->R());
-					FRigidTransform3 WorldTransform1 = LocalTransform1 * ParticleWorldTransform1;
-					UpdateConvexTriangleMeshConstraintSwept(Particle0, *Implicit0, WorldTransformX0, *ScaledTriangleMesh, WorldTransform1, Dir, Length, Dt, *Constraint);
-				}
-				else if (const FTriangleMeshImplicitObject* TriangleMesh = Implicit1->template GetObject<const FTriangleMeshImplicitObject>())
-				{
-					FPBDCollisionConstraint* Constraint = Context.MultiShapeCollisionDetector->FindOrCreateSweptConstraint(Particle0, Implicit0, Shape0, nullptr, LocalTransform0, Particle1, Implicit1, Shape1, nullptr, LocalTransform1, CullDistance, EContactShapesType::ConvexTriMesh, Context);
-					FRigidTransform3 WorldTransformX0 = LocalTransform0 * FRigidTransform3(Particle0->X(), Particle0->R());
-					FRigidTransform3 WorldTransform1 = LocalTransform1 * ParticleWorldTransform1;
-					UpdateConvexTriangleMeshConstraintSwept(Particle0, *Implicit0, WorldTransformX0, *TriangleMesh, WorldTransform1, Dir, Length, Dt, *Constraint);
-				}
-				else
-				{
-					ensure(false);
-				}
-			}
-		}
-
-
 		//
 		// Levelset-Levelset
 		//
 
-		template<ECollisionUpdateType UpdateType>
 		void UpdateLevelsetLevelsetConstraint(const FRigidTransform3& WorldTransform0, const FRigidTransform3& WorldTransform1, const FReal Dt, FPBDCollisionConstraint& Constraint)
 		{
 			SCOPE_CYCLE_COUNTER(STAT_UpdateLevelsetLevelsetConstraint);
@@ -1466,72 +928,20 @@ namespace Chaos
 			if (SampleParticles)
 			{
 				const FImplicitObject* Obj1 = Constraint.GetImplicit1();
-				FContactPoint ContactPoint = SampleObject<UpdateType>(*Obj1, LevelsetTM, *SampleParticles, ParticlesTM, Constraint.GetCullDistance());
+				FContactPoint ContactPoint = SampleObject<ECollisionUpdateType::Deepest>(*Obj1, LevelsetTM, *SampleParticles, ParticlesTM, Constraint.GetCullDistance());
 				UpdateContactPoint(Constraint, ContactPoint, Dt);
 			}
 		}
 
-		void UpdateLevelsetLevelsetManifold(FPBDCollisionConstraint& Constraint, const FRigidTransform3& ATM, const FRigidTransform3& BTM)
-		{
-			// @todo(chaos) : Stub Update Manifold
-			//   Stub function for updating the manifold prior to the Apply and ApplyPushOut
-		}
-
+		//
+		// Constraint API
+		//
 
-		template<typename T_TRAITS>
-		void ConstructLevelsetLevelsetConstraints(TGeometryParticleHandle<FReal, 3>* Particle0, TGeometryParticleHandle<FReal, 3>* Particle1, const FImplicitObject* Implicit0, const FPerShapeData* Shape0, const FBVHParticles* Simplicial0, const FImplicitObject* Implicit1, const FPerShapeData* Shape1, const FBVHParticles* Simplicial1, const FRigidTransform3& ParticleWorldTransform0, const FRigidTransform3& LocalTransform0, const FRigidTransform3& ParticleWorldTransform1, const FRigidTransform3& LocalTransform1, const FReal CullDistance, const FReal Dt, const FCollisionContext& Context)
+		// Run collision detection for the specified constraint to update the nearest contact point.
+		// NOTE: Transforms are world space shape transforms
+		// @todo(chaos): use a lookup table?
+		inline void UpdateConstraintFromGeometryImpl(FPBDCollisionConstraint& Constraint, const FRigidTransform3& WorldTransform0, const FRigidTransform3& WorldTransform1, const FReal Dt)
 		{
-			const bool bAllowManifolds = Context.GetSettings().bAllowManifolds && bChaos_Collision_AllowLevelsetManifolds;
-			FPBDCollisionConstraint* Constraint = nullptr;
-
-			// @todo(chaos): We do not really support a runtime decision for the particle order now that we have persistent collisions. Fix the logic below...?
-			// Also clean this logic up so it is readable. When is Simplical0 not equal to CollisionParticles0?
-			bool bIsParticleDynamic0 = Particle0->CastToRigidParticle() && ((Particle0->ObjectState() == EObjectStateType::Dynamic) || (Particle0->ObjectState() == EObjectStateType::Sleeping));
-			int32 P0NumCollisionParticles = Simplicial0 ? Simplicial0->Size() : Particle0->CastToRigidParticle()->CollisionParticlesSize();
-			if (!Particle1->Geometry() || (bIsParticleDynamic0 && !P0NumCollisionParticles && Particle0->Geometry() && !Particle0->Geometry()->IsUnderlyingUnion()))
-			{
-				Constraint = Context.MultiShapeCollisionDetector->FindOrCreateConstraint(Particle1, Implicit1, Shape1, Simplicial1, LocalTransform1, Particle0, Implicit0, Shape0, Simplicial0, LocalTransform0, CullDistance, EContactShapesType::LevelSetLevelSet, bAllowManifolds, Context);
-			}
-			else
-			{
-				Constraint = Context.MultiShapeCollisionDetector->FindOrCreateConstraint(Particle0, Implicit0, Shape0, Simplicial0, LocalTransform0, Particle1, Implicit1, Shape1, Simplicial1, LocalTransform1, CullDistance, EContactShapesType::LevelSetLevelSet, bAllowManifolds, Context);
-			}
-
-			if (Constraint != nullptr)
-			{
-				// @todo(chaos): support manifold restoration for Levelsets and remove this
-				Constraint->ResetActiveManifoldContacts();
-
-				if (T_TRAITS::bImmediateUpdate && !Constraint->WasManifoldRestored())
-				{
-					FRigidTransform3 WorldTransform0 = Constraint->GetShapeRelativeTransform0() * ParticleWorldTransform0;
-					FRigidTransform3 WorldTransform1 = Constraint->GetShapeRelativeTransform1() * ParticleWorldTransform1;
-					UpdateLevelsetLevelsetConstraint<ECollisionUpdateType::Deepest>(WorldTransform0, WorldTransform1, Dt, *Constraint);
-				}
-			}
-		}
-
-
-		//
-		// Constraint API
-		//
-
-
-		// Run collision detection for the specified constraint to update the nearest contact point.
-		// NOTE: Transforms are world space shape transforms
-		// @todo(chaos): use a lookup table?
-		// @todo(chaos): add the missing cases below
-		// @todo(chaos): see use GetInnerObject below - we should try to use the leaf types for all (currently only TriMesh needs this)
-		template<ECollisionUpdateType UpdateType>
-		inline void UpdateConstraintFromGeometryImpl(FPBDCollisionConstraint& Constraint, const FRigidTransform3& WorldTransform0, const FRigidTransform3& WorldTransform1, const FReal Dt)
-		{
-			// These add an extra 2% overhead in the simple shape cases
-			//CONDITIONAL_SCOPE_CYCLE_COUNTER(STAT_Collisions_UpdateConstraintFromGeometryInternal, ConstraintsDetailedStats);
-			//INC_DWORD_STAT(STAT_ChaosCollisionCounter_NumContactUpdates);
-
-			// @todo(chaos): remove
-			//const FVec3 OriginalContactPositionLocal0 = WorldTransform0.InverseTransformPosition(Constraint.Manifold.Location);
-			//const FVec3 OriginalContactPositionLocal1 = WorldTransform1.InverseTransformPosition(Constraint.Manifold.Location);
 			if (!Constraint.GetImplicit0() || !Constraint.GetImplicit1())
 			{
 				return;
@@ -1540,7 +950,6 @@ namespace Chaos
 			const FImplicitObject& Implicit0 = *Constraint.GetImplicit0();
 			const FImplicitObject& Implicit1 = *Constraint.GetImplicit1();
 
-
 			switch (Constraint.GetShapesType())
 			{
 			case EContactShapesType::SphereSphere:
@@ -1634,7 +1043,7 @@ namespace Chaos
 				UpdateConvexHeightFieldConstraint(Implicit0, WorldTransform0, *Implicit1.template GetObject< FHeightField >(), WorldTransform1, Dt, Constraint);
 				break;
 			case EContactShapesType::LevelSetLevelSet:
-				UpdateLevelsetLevelsetConstraint<UpdateType>(WorldTransform0, WorldTransform1, Dt, Constraint);
+				UpdateLevelsetLevelsetConstraint(WorldTransform0, WorldTransform1, Dt, Constraint);
 				break;
 			default:
 				// Switch needs updating....
@@ -1646,7 +1055,6 @@ namespace Chaos
 		// Run collision detection for the swept constraints
 		// NOTE: Transforms are world space shape transforms
 		
-		template<ECollisionUpdateType UpdateType>
 		inline bool UpdateConstraintFromGeometrySweptImpl(FPBDCollisionConstraint& Constraint, const FRigidTransform3& WorldTransform0, const FRigidTransform3& WorldTransform1, const FReal Dt)
 		{
 			INC_DWORD_STAT(STAT_ChaosCollisionCounter_NumContactUpdates);
@@ -1656,6 +1064,7 @@ namespace Chaos
 			const FImplicitObject& Implicit0 = *Constraint.GetImplicit0();
 			const FImplicitObject& Implicit1 = *Constraint.GetImplicit1();
 
+			// @todo(chaos): this work has alreday been done in the MidPhase - check and remove/test
 			FReal LengthCCD = 0.0f;
 			FVec3 DirCCD(0.0f);
 			const FVec3 DeltaX0 = Constraint.GetShapeWorldTransform0().GetTranslation() - WorldTransform0.GetLocation();
@@ -1763,311 +1172,113 @@ namespace Chaos
 			return false;
 		}
 
-		template<typename T_TRAITS>
-		void ConstructConstraintsImpl(TGeometryParticleHandle<FReal, 3>* Particle0, TGeometryParticleHandle<FReal, 3>* Particle1, const FImplicitObject* Implicit0, const FPerShapeData* Shape0, const FBVHParticles* Simplicial0, const FImplicitObject* Implicit1, const FPerShapeData* Shape1, const FBVHParticles* Simplicial1, const FRigidTransform3& ParticleWorldTransform0, const FRigidTransform3& LocalTransform0, const FRigidTransform3& ParticleWorldTransform1, const FRigidTransform3& LocalTransform1, const FReal CullDistance, const FReal Dt, const FCollisionContext& Context)
+		bool DoParticleImplicitsOverlap(
+			TGeometryParticleHandle<FReal, 3>* Particle0,
+			TGeometryParticleHandle<FReal, 3>* Particle1,
+			const FImplicitObject* Implicit0,
+			const FImplicitObject* Implicit1,
+			const FRigidTransform3& ParticleWorldTransform0,
+			const FRigidTransform3& LocalTransform0,
+			const FRigidTransform3& ParticleWorldTransform1,
+			const FRigidTransform3& LocalTransform1,
+			const FReal CullDistance,
+			const bool bEnableSweep,
+			const FCollisionContext& Context)
 		{
-			CONDITIONAL_SCOPE_CYCLE_COUNTER(STAT_Collisions_ConstructConstraintsInternal, ConstraintsDetailedStats);
-			INC_DWORD_STAT(STAT_ChaosCollisionCounter_NumShapePairs);
-
-			// @todo(chaos): We use GetInnerType here because TriMeshes are left with their "Instanced" wrapper, unlike all other instanced implicits. Should we strip the instance on Tri Mesh too?
-			EImplicitObjectType Implicit0Type = Implicit0 ? GetInnerType(Implicit0->GetCollisionType()) : ImplicitObjectType::Unknown;
-			EImplicitObjectType Implicit1Type = Implicit1 ? GetInnerType(Implicit1->GetCollisionType()) : ImplicitObjectType::Unknown;
-			bool bIsConvex0 = Implicit0 && Implicit0->IsConvex() && Implicit0Type != ImplicitObjectType::LevelSet;
-			bool bIsConvex1 = Implicit1 && Implicit1->IsConvex() && Implicit1Type != ImplicitObjectType::LevelSet;
-
-			FReal LengthCCD = 0.0f;
-			FVec3 DirCCD(0.0f);
-			FConstGenericParticleHandle ConstParticle0 = FConstGenericParticleHandle(Particle0);
-			FConstGenericParticleHandle ConstParticle1 = FConstGenericParticleHandle(Particle1);
-
-			// X is end-frame position for kinematic particles. To get start-frame position, we need to use P - V * Dt.
-			const FVec3 StartX0 = (ConstParticle0->ObjectState() == EObjectStateType::Kinematic) ? (ConstParticle0->P() - ConstParticle0->V() * Dt) : ConstParticle0->X();
-			const FVec3 StartX1 = (ConstParticle1->ObjectState() == EObjectStateType::Kinematic) ? (ConstParticle1->P() - ConstParticle1->V() * Dt) : ConstParticle1->X();
-			const FVec3 EndX0 = ConstParticle0->P();
-			const FVec3 EndX1 = ConstParticle1->P();
-			bool bUseCCD = Context.GetSettings().bAllowCCD && ShouldUseCCD(Particle0, EndX0 - StartX0, Particle1, EndX1 - StartX1, DirCCD, LengthCCD);
-			const bool bUseGenericSweptConstraints = bUseCCD && (CCDUseGenericSweptConvexConstraints > 0) && bIsConvex0 && bIsConvex1;
-#if CHAOS_COLLISION_CREATE_BOUNDSCHECK
-			if ((Implicit0 != nullptr) && (Implicit1 != nullptr))
+			if ((Implicit0 != nullptr) && (Implicit1 != nullptr) && !bEnableSweep)
 			{
-				if (!bUseCCD)
+				if (Implicit0->HasBoundingBox() && Implicit1->HasBoundingBox())
 				{
-					if (Implicit0->HasBoundingBox() && Implicit1->HasBoundingBox())
+					const FRigidTransform3 WorldTransform0 = LocalTransform0 * ParticleWorldTransform0;
+					const FRigidTransform3 WorldTransform1 = LocalTransform1 * ParticleWorldTransform1;
+					if (Chaos_Collision_NarrowPhase_SphereBoundsCheck)
 					{
-						const FRigidTransform3 WorldTransform0 = LocalTransform0 * ParticleWorldTransform0;
-						const FRigidTransform3 WorldTransform1 = LocalTransform1 * ParticleWorldTransform1;
-						if (Chaos_Collision_NarrowPhase_SphereBoundsCheck)
+						const FReal R1 = Implicit0->BoundingBox().OriginRadius();
+						const FReal R2 = Implicit1->BoundingBox().OriginRadius();
+						const FReal SeparationSq = (WorldTransform1.GetTranslation() - WorldTransform0.GetTranslation()).SizeSquared();
+						if (SeparationSq > FMath::Square(R1 + R2 + CullDistance))
 						{
-							const FReal R1 = Implicit0->BoundingBox().OriginRadius();
-							const FReal R2 = Implicit1->BoundingBox().OriginRadius();
-							const FReal SeparationSq = (WorldTransform1.GetTranslation() - WorldTransform0.GetTranslation()).SizeSquared();
-							if (SeparationSq > FMath::Square(R1 + R2 + CullDistance))
-							{
-								return;
-							}
+							return false;
 						}
+					}
 
-						if (Chaos_Collision_NarrowPhase_AABBBoundsCheck)
+					if (Chaos_Collision_NarrowPhase_AABBBoundsCheck)
+					{
+						const FAABB3 Box1 = Implicit0->BoundingBox();
+						const FAABB3 Box2 = Implicit1->BoundingBox();
+						//if (Box1.GetVolume() >= Box2.GetVolume())
 						{
-							const FAABB3 Box1 = Implicit0->BoundingBox();
-							const FAABB3 Box2 = Implicit1->BoundingBox();
-							//if (Box1.GetVolume() >= Box2.GetVolume())
+							const FRigidTransform3 Box2ToBox1TM = WorldTransform1.GetRelativeTransform(WorldTransform0);
+							const FAABB3 Box2In1 = Box2.TransformedAABB(Box2ToBox1TM).Thicken(CullDistance);
+							if (!Box1.Intersects(Box2In1))
 							{
-								const FRigidTransform3 Box2ToBox1TM = WorldTransform1.GetRelativeTransform(WorldTransform0);
-								const FAABB3 Box2In1 = Box2.TransformedAABB(Box2ToBox1TM).Thicken(CullDistance);
-								if (!Box1.Intersects(Box2In1))
-								{
-									return;
-								}
+								return false;
 							}
-							//else
+						}
+						//else
+						{
+							const FRigidTransform3 Box1ToBox2TM = WorldTransform0.GetRelativeTransform(WorldTransform1);
+							const FAABB3 Box1In2 = Box1.TransformedAABB(Box1ToBox2TM).Thicken(CullDistance);
+							if (!Box2.Intersects(Box1In2))
 							{
-								const FRigidTransform3 Box1ToBox2TM = WorldTransform0.GetRelativeTransform(WorldTransform1);
-								const FAABB3 Box1In2 = Box1.TransformedAABB(Box1ToBox2TM).Thicken(CullDistance);
-								if (!Box2.Intersects(Box1In2))
-								{
-									return;
-								}
+								return false;
 							}
 						}
 					}
 				}
-				else
+			}
+			else
+			{
+				if (Particle0->HasBounds() && Particle1->HasBounds())
 				{
 					const FAABB3 Box0 = Particle0->WorldSpaceInflatedBounds();
 					const FAABB3 Box1 = Particle1->WorldSpaceInflatedBounds();
 					if (!Box0.Intersects(Box1))
 					{
-						return;
+						return false;
 					}
 				}
 			}
-#endif
 
-			INC_DWORD_STAT(STAT_ChaosCollisionCounter_NumContactsCreated);
-
-			if (Implicit0Type == TBox<FReal, 3>::StaticType() && Implicit1Type == TBox<FReal, 3>::StaticType() && !bUseGenericSweptConstraints)
-			{
-				ConstructBoxBoxConstraints<T_TRAITS>(Particle0, Particle1, Implicit0, Shape0, Implicit1, Shape1, ParticleWorldTransform0, LocalTransform0, ParticleWorldTransform1, LocalTransform1, CullDistance, Dt, Context);
-			}
-			else if (Implicit0Type == TBox<FReal, 3>::StaticType() && Implicit1Type == FHeightField::StaticType())
-			{
-				ConstructBoxHeightFieldConstraints<T_TRAITS>(Particle0, Particle1, Implicit0, Shape0, Implicit1, Shape1, ParticleWorldTransform0, LocalTransform0, ParticleWorldTransform1, LocalTransform1, CullDistance, Dt, Context);
-			}
-			else if (Implicit0Type == FHeightField::StaticType() && Implicit1Type == TBox<FReal, 3>::StaticType())
-			{
-				ConstructBoxHeightFieldConstraints<T_TRAITS>(Particle1, Particle0, Implicit1, Shape1, Implicit0, Shape0, ParticleWorldTransform1, LocalTransform1, ParticleWorldTransform0, LocalTransform0, CullDistance, Dt, Context);
-			}
-			else if (Implicit0Type == TBox<FReal, 3>::StaticType() && Implicit1Type == TPlane<FReal, 3>::StaticType() && !bUseGenericSweptConstraints)
-			{
-				ConstructBoxPlaneConstraints<T_TRAITS>(Particle0, Particle1, Implicit0, Shape0, Implicit1, Shape1, ParticleWorldTransform0, LocalTransform0, ParticleWorldTransform1, LocalTransform1, CullDistance, Dt, Context);
-			}
-			else if (Implicit0Type == TPlane<FReal, 3>::StaticType() && Implicit1Type == TBox<FReal, 3>::StaticType() && !bUseGenericSweptConstraints)
-			{
-				ConstructBoxPlaneConstraints<T_TRAITS>(Particle1, Particle0, Implicit1, Shape1, Implicit0, Shape0, ParticleWorldTransform1, LocalTransform1, ParticleWorldTransform0, LocalTransform0, CullDistance, Dt, Context);
-			}
-			else if (Implicit0Type == TSphere<FReal, 3>::StaticType() && Implicit1Type == TSphere<FReal, 3>::StaticType() && !bUseGenericSweptConstraints)
-			{
-				ConstructSphereSphereConstraints<T_TRAITS>(Particle0, Particle1, Implicit0, Shape0, Implicit1, Shape1, ParticleWorldTransform0, LocalTransform0, ParticleWorldTransform1, LocalTransform1, CullDistance, Dt, Context);
-			}
-			else if (Implicit0Type == TSphere<FReal, 3>::StaticType() && Implicit1Type == FHeightField::StaticType())
-			{
-				if (bUseCCD)
-				{
-					ConstructSphereHeightFieldConstraintsSwept(Particle0, Particle1, Implicit0, Shape0, Implicit1, Shape1, ParticleWorldTransform0, LocalTransform0, ParticleWorldTransform1, LocalTransform1, CullDistance, Dt, DirCCD, LengthCCD, Context);
-					return;
-				}
-				ConstructSphereHeightFieldConstraints<T_TRAITS>(Particle0, Particle1, Implicit0, Shape0, Implicit1, Shape1, ParticleWorldTransform0, LocalTransform0, ParticleWorldTransform1, LocalTransform1, CullDistance, Dt, Context);
-			}
-			else if (Implicit0Type == FHeightField::StaticType() && Implicit1Type == TSphere<FReal, 3>::StaticType())
-			{
-				if (bUseCCD)
-				{
-					ConstructSphereHeightFieldConstraintsSwept(Particle1, Particle0, Implicit1, Shape1, Implicit0, Shape0, ParticleWorldTransform1, LocalTransform1, ParticleWorldTransform0, LocalTransform0, CullDistance, Dt, DirCCD, LengthCCD, Context);
-					return;
-				}
-				ConstructSphereHeightFieldConstraints<T_TRAITS>(Particle1, Particle0, Implicit1, Shape1, Implicit0, Shape0, ParticleWorldTransform1, LocalTransform1, ParticleWorldTransform0, LocalTransform0, CullDistance, Dt, Context);
-			}
-			else if (Implicit0Type == TSphere<FReal, 3>::StaticType() && Implicit1Type == TPlane<FReal, 3>::StaticType() && !bUseGenericSweptConstraints)
-			{
-				ConstructSpherePlaneConstraints<T_TRAITS>(Particle0, Particle1, Implicit0, Shape0, Implicit1, Shape1, ParticleWorldTransform0, LocalTransform0, ParticleWorldTransform1, LocalTransform1, CullDistance, Dt, Context);
-			}
-			else if (Implicit0Type == TPlane<FReal, 3>::StaticType() && Implicit1Type == TSphere<FReal, 3>::StaticType() && !bUseGenericSweptConstraints)
-			{
-				ConstructSpherePlaneConstraints<T_TRAITS>(Particle1, Particle0, Implicit1, Shape1, Implicit0, Shape0, ParticleWorldTransform1, LocalTransform1, ParticleWorldTransform0, LocalTransform0, CullDistance, Dt, Context);
-			}
-			else if (Implicit0Type == TSphere<FReal, 3>::StaticType() && Implicit1Type == TBox<FReal, 3>::StaticType() && !bUseGenericSweptConstraints)
-			{
-				ConstructSphereBoxConstraints<T_TRAITS>(Particle0, Particle1, Implicit0, Shape0, Implicit1, Shape1, ParticleWorldTransform0, LocalTransform0, ParticleWorldTransform1, LocalTransform1, CullDistance, Dt, Context);
-			}
-			else if (Implicit0Type == TBox<FReal, 3>::StaticType() && Implicit1Type == TSphere<FReal, 3>::StaticType() && !bUseGenericSweptConstraints)
-			{
-				ConstructSphereBoxConstraints<T_TRAITS>(Particle1, Particle0, Implicit1, Shape1, Implicit0, Shape0, ParticleWorldTransform1, LocalTransform1, ParticleWorldTransform0, LocalTransform0, CullDistance, Dt, Context);
-			}
-			else if (Implicit0Type == TSphere<FReal, 3>::StaticType() && Implicit1Type == FCapsule::StaticType() && !bUseGenericSweptConstraints)
-			{
-				ConstructSphereCapsuleConstraints<T_TRAITS>(Particle0, Particle1, Implicit0, Shape0, Implicit1, Shape1, ParticleWorldTransform0, LocalTransform0, ParticleWorldTransform1, LocalTransform1, CullDistance, Dt, Context);
-			}
-			else if (Implicit0Type == FCapsule::StaticType() && Implicit1Type == TSphere<FReal, 3>::StaticType() && !bUseGenericSweptConstraints)
-			{
-				ConstructSphereCapsuleConstraints<T_TRAITS>(Particle1, Particle0, Implicit1, Shape1, Implicit0, Shape0, ParticleWorldTransform1, LocalTransform1, ParticleWorldTransform0, LocalTransform0, CullDistance, Dt, Context);
-			}
-			else if (Implicit0Type == TSphere<FReal, 3>::StaticType() && Implicit1Type == FImplicitConvex3::StaticType() && !bUseGenericSweptConstraints)
-			{
-				ConstructSphereConvexConstraints<T_TRAITS>(Particle0, Particle1, Implicit0, Shape0, Implicit1, Shape1, ParticleWorldTransform0, LocalTransform0, ParticleWorldTransform1, LocalTransform1, CullDistance, Dt, Context);
-			}
-			else if (Implicit0Type == FImplicitConvex3::StaticType() && Implicit1Type == TSphere<FReal, 3>::StaticType() && !bUseGenericSweptConstraints)
-			{
-				ConstructSphereConvexConstraints<T_TRAITS>(Particle1, Particle0, Implicit1, Shape1, Implicit0, Shape0, ParticleWorldTransform1, LocalTransform1, ParticleWorldTransform0, LocalTransform0, CullDistance, Dt, Context);
-			}
-			else if (Implicit0Type == FCapsule::StaticType() && Implicit1Type == FCapsule::StaticType() && !bUseGenericSweptConstraints)
-			{
-				ConstructCapsuleCapsuleConstraints<T_TRAITS>(Particle0, Particle1, Implicit0, Shape0, Implicit1, Shape1, ParticleWorldTransform0, LocalTransform0, ParticleWorldTransform1, LocalTransform1, CullDistance, Dt, Context);
-			}
-			else if (Implicit0Type == FCapsule::StaticType() && Implicit1Type == TBox<FReal, 3>::StaticType() && !bUseGenericSweptConstraints)
-			{
-				ConstructCapsuleBoxConstraints<T_TRAITS>(Particle0, Particle1, Implicit0, Shape0, Implicit1, Shape1, ParticleWorldTransform0, LocalTransform0, ParticleWorldTransform1, LocalTransform1, CullDistance, Dt, Context);
-			}
-			else if (Implicit0Type == TBox<FReal, 3>::StaticType() && Implicit1Type == FCapsule::StaticType() && !bUseGenericSweptConstraints)
-			{
-				ConstructCapsuleBoxConstraints<T_TRAITS>(Particle1, Particle0, Implicit1, Shape1, Implicit0, Shape0, ParticleWorldTransform1, LocalTransform1, ParticleWorldTransform0, LocalTransform0, CullDistance, Dt, Context);
-			}
-			else if (Implicit0Type == FCapsule::StaticType() && Implicit1Type == FImplicitConvex3::StaticType() && !bUseGenericSweptConstraints)
-			{
-				ConstructCapsuleConvexConstraints<T_TRAITS>(Particle0, Particle1, Implicit0, Shape0, Implicit1, Shape1, ParticleWorldTransform0, LocalTransform0, ParticleWorldTransform1, LocalTransform1, CullDistance, Dt, Context);
-			}
-			else if (Implicit0Type == FImplicitConvex3::StaticType() && Implicit1Type == FCapsule::StaticType() && !bUseGenericSweptConstraints)
-			{
-				ConstructCapsuleConvexConstraints<T_TRAITS>(Particle1, Particle0, Implicit1, Shape1, Implicit0, Shape0, ParticleWorldTransform1, LocalTransform1, ParticleWorldTransform0, LocalTransform0, CullDistance, Dt, Context);
-			}
-			else if (Implicit0Type == FCapsule::StaticType() && Implicit1Type == FHeightField::StaticType())
-			{
-				if (bUseCCD)
-				{
-					ConstructCapsuleHeightFieldConstraintsSwept(Particle0, Particle1, Implicit0, Shape0, Implicit1, Shape1, ParticleWorldTransform0, LocalTransform0, ParticleWorldTransform1, LocalTransform1, CullDistance, Dt, DirCCD, LengthCCD, Context);
-					return;
-				}
-
-				ConstructCapsuleHeightFieldConstraints<T_TRAITS>(Particle0, Particle1, Implicit0, Shape0, Implicit1, Shape1, ParticleWorldTransform0, LocalTransform0, ParticleWorldTransform1, LocalTransform1, CullDistance, Dt, Context);
-			}
-			else if (Implicit0Type == FHeightField::StaticType() && Implicit1Type == FCapsule::StaticType())
-			{
-				if (bUseCCD)
-				{
-					ConstructCapsuleHeightFieldConstraintsSwept(Particle1, Particle0, Implicit1, Shape1, Implicit0, Shape0, ParticleWorldTransform1, LocalTransform1, ParticleWorldTransform0, LocalTransform0, CullDistance, Dt, DirCCD, LengthCCD, Context);
-					return;
-				}
-
-				ConstructCapsuleHeightFieldConstraints<T_TRAITS>(Particle1, Particle0, Implicit1, Shape1, Implicit0, Shape0, ParticleWorldTransform1, LocalTransform1, ParticleWorldTransform0, LocalTransform0, CullDistance, Dt, Context);
-			}
-			else if (Implicit0Type == TBox<FReal, 3>::StaticType() && Implicit1Type == FTriangleMeshImplicitObject::StaticType())
-			{
-				ConstructBoxTriangleMeshConstraints<T_TRAITS>(Particle0, Particle1, Implicit0, Shape0, Implicit1, Shape1, ParticleWorldTransform0, LocalTransform0, ParticleWorldTransform1, LocalTransform1, CullDistance, Dt, Context);
-			}
-			else if (Implicit0Type == FTriangleMeshImplicitObject::StaticType() && Implicit1Type == TBox<FReal, 3>::StaticType())
-			{
-				ConstructBoxTriangleMeshConstraints<T_TRAITS>(Particle1, Particle0, Implicit1, Shape1, Implicit0, Shape0, ParticleWorldTransform1, LocalTransform1, ParticleWorldTransform0, LocalTransform0, CullDistance, Dt, Context);
-			}
-			else if (Implicit0Type == TSphere<FReal, 3>::StaticType() && Implicit1Type == FTriangleMeshImplicitObject::StaticType())
-			{
-				if (bUseCCD)
-				{
-					ConstructSphereTriangleMeshConstraintsSwept(Particle0, Particle1, Implicit0, Shape0, Implicit1, Shape1, ParticleWorldTransform0, LocalTransform0, ParticleWorldTransform1, LocalTransform1, CullDistance, Dt, DirCCD, LengthCCD, Context);
-					return;
-				}
-				ConstructSphereTriangleMeshConstraints<T_TRAITS>(Particle0, Particle1, Implicit0, Shape0, Implicit1, Shape1, ParticleWorldTransform0, LocalTransform0, ParticleWorldTransform1, LocalTransform1, CullDistance, Dt, Context);
-			}
-			else if (Implicit0Type == FTriangleMeshImplicitObject::StaticType() && Implicit1Type == TSphere<FReal, 3>::StaticType())
-			{
-				if (bUseCCD)
-				{
-					ConstructSphereTriangleMeshConstraintsSwept(Particle1, Particle0, Implicit1, Shape1, Implicit0, Shape0, ParticleWorldTransform1, LocalTransform1, ParticleWorldTransform0, LocalTransform0, CullDistance, Dt, DirCCD, LengthCCD, Context);
-					return;
-				}
-				ConstructSphereTriangleMeshConstraints<T_TRAITS>(Particle1, Particle0, Implicit1, Shape1, Implicit0, Shape0, ParticleWorldTransform1, LocalTransform1, ParticleWorldTransform0, LocalTransform0, CullDistance, Dt, Context);
-			}
-			else if (Implicit0Type == FCapsule::StaticType() && Implicit1Type == FTriangleMeshImplicitObject::StaticType())
-			{
-				if (bUseCCD)
-				{
-					ConstructCapsuleTriangleMeshConstraintsSwept(Particle0, Particle1, Implicit0, Shape0, Implicit1, Shape1, ParticleWorldTransform0, LocalTransform0, ParticleWorldTransform1, LocalTransform1, CullDistance, Dt, DirCCD, LengthCCD, Context);
-					return;
-				}
-
-				ConstructCapsuleTriangleMeshConstraints<T_TRAITS>(Particle0, Particle1, Implicit0, Shape0, Implicit1, Shape1, ParticleWorldTransform0, LocalTransform0, ParticleWorldTransform1, LocalTransform1, CullDistance, Dt, Context);
-			}
-			else if (Implicit0Type == FTriangleMeshImplicitObject::StaticType() && Implicit1Type == FCapsule::StaticType())
-			{
-				if (bUseCCD)
-				{
-					ConstructCapsuleTriangleMeshConstraintsSwept(Particle1, Particle0, Implicit1, Shape1, Implicit0, Shape0, ParticleWorldTransform1, LocalTransform1, ParticleWorldTransform0, LocalTransform0, CullDistance, Dt, DirCCD, LengthCCD, Context);
-					return;
-				}
-
-				ConstructCapsuleTriangleMeshConstraints<T_TRAITS>(Particle1, Particle0, Implicit1, Shape1, Implicit0, Shape0, ParticleWorldTransform1, LocalTransform1, ParticleWorldTransform0, LocalTransform0, CullDistance, Dt, Context);
-			}
-			else if (bIsConvex0 && Implicit1Type == FHeightField::StaticType())
-			{
-				if (bUseCCD)
-				{
-					ConstructConvexHeightFieldConstraintsSwept(Particle0, Particle1, Implicit0, Shape0, Implicit1, Shape1, ParticleWorldTransform0, LocalTransform0, ParticleWorldTransform1, LocalTransform1, CullDistance, Dt, DirCCD, LengthCCD, Context);
-					return;
-				}
-
-				ConstructConvexHeightFieldConstraints<T_TRAITS>(Particle0, Particle1, Implicit0, Shape0, Implicit1, Shape1, ParticleWorldTransform0, LocalTransform0, ParticleWorldTransform1, LocalTransform1, CullDistance, Dt, Context);
-			}
-			else if (Implicit0Type == FHeightField::StaticType() && bIsConvex1)
-			{
-				if (bUseCCD)
-				{
-					ConstructConvexHeightFieldConstraintsSwept(Particle1, Particle0, Implicit1, Shape1, Implicit0, Shape0, ParticleWorldTransform1, LocalTransform1, ParticleWorldTransform0, LocalTransform0, CullDistance, Dt, DirCCD, LengthCCD, Context);
-					return;
-				}
-
-				ConstructConvexHeightFieldConstraints<T_TRAITS>(Particle1, Particle0, Implicit1, Shape1, Implicit0, Shape0, ParticleWorldTransform1, LocalTransform1, ParticleWorldTransform0, LocalTransform0, CullDistance, Dt, Context);
-			}
-			else if (bIsConvex0 && Implicit1Type == FTriangleMeshImplicitObject::StaticType())
-			{
-				if (bUseCCD)
-				{
-					ConstructConvexTriangleMeshConstraintsSwept(Particle0, Particle1, Implicit0, Shape0, Implicit1, Shape1, ParticleWorldTransform0, LocalTransform0, ParticleWorldTransform1, LocalTransform1, CullDistance, Dt, DirCCD, LengthCCD, Context);
-					return;
-				}
+			return true;
+		}
 
-				ConstructConvexTriangleMeshConstraints<T_TRAITS>(Particle0, Particle1, Implicit0, Shape0, Implicit1, Shape1, ParticleWorldTransform0, LocalTransform0, ParticleWorldTransform1, LocalTransform1, CullDistance, Dt, Context);
-			}
-			else if (Implicit0Type == FTriangleMeshImplicitObject::StaticType() && bIsConvex1)
-			{
-				if (bUseCCD)
-				{
-					ConstructConvexTriangleMeshConstraintsSwept(Particle1, Particle0, Implicit1, Shape1, Implicit0, Shape0, ParticleWorldTransform1, LocalTransform1, ParticleWorldTransform0, LocalTransform0, CullDistance, Dt, DirCCD, LengthCCD, Context);
-					return;
-				}
+		void ConstructConstraintsImpl(
+			TGeometryParticleHandle<FReal, 3>* Particle0,
+			TGeometryParticleHandle<FReal, 3>* Particle1,
+			const FImplicitObject* Implicit0,
+			const FShapeInstance* Shape0,
+			const FBVHParticles* Simplicial0,
+			const int32 ImplicitID0,
+			const FImplicitObject* Implicit1,
+			const FShapeInstance* Shape1,
+			const FBVHParticles* Simplicial1,
+			const int32 ImplicitID1,
+			const FRigidTransform3& ParticleWorldTransform0,
+			const FRigidTransform3& LocalTransform0,
+			const FRigidTransform3& ParticleWorldTransform1,
+			const FRigidTransform3& LocalTransform1,
+			const FReal CullDistance,
+			const FReal Dt,
+			const EContactShapesType ShapePairType,
+			const bool bEnableSweep,
+			const FCollisionContext& Context)
+		{
+			CONDITIONAL_SCOPE_CYCLE_COUNTER(STAT_Collisions_ConstructConstraintsInternal, ConstraintsDetailedStats);
+			INC_DWORD_STAT(STAT_ChaosCollisionCounter_NumShapePairs);
 
-				ConstructConvexTriangleMeshConstraints<T_TRAITS>(Particle1, Particle0, Implicit1, Shape1, Implicit0, Shape0, ParticleWorldTransform1, LocalTransform1, ParticleWorldTransform0, LocalTransform0, CullDistance, Dt, Context);
-			}
-			else if (bIsConvex0 && bIsConvex1)
-			{
-				if (bUseCCD)
-				{
-					ConstructGenericConvexConvexConstraintsSwept(Particle0, Particle1, Implicit0, Shape0, Implicit1, Shape1, ParticleWorldTransform0, LocalTransform0, StartX0, ParticleWorldTransform1, LocalTransform1, StartX1, CullDistance, Dt, DirCCD, LengthCCD, Context);
-					return;
-				}
+			const bool bDoOverlap = DoParticleImplicitsOverlap(
+				Particle0, Particle1, Implicit0, Implicit1, 
+				ParticleWorldTransform0, LocalTransform0, 
+				ParticleWorldTransform1, LocalTransform1, 
+				CullDistance, bEnableSweep, Context);
 
-				ConstructGenericConvexConvexConstraints<T_TRAITS>(Particle0, Particle1, Implicit0, Shape0, Implicit1, Shape1, ParticleWorldTransform0, LocalTransform0, ParticleWorldTransform1, LocalTransform1, CullDistance, Dt, Context);
-			}
-			else
+			if (bDoOverlap)
 			{
-				// LevelSets assume that the first or both particles are dynamic
-				bool bSwap = (!FConstGenericParticleHandle(Particle0)->IsDynamic() && FConstGenericParticleHandle(Particle1)->IsDynamic());
-				if (!bSwap)
-				{
-					ConstructLevelsetLevelsetConstraints<T_TRAITS>(Particle0, Particle1, Implicit0, Shape0, Simplicial0, Implicit1, Shape1, Simplicial1, ParticleWorldTransform0, LocalTransform0, ParticleWorldTransform1, LocalTransform1, CullDistance, Dt, Context);
-				}
-				else
-				{
-					ConstructLevelsetLevelsetConstraints<T_TRAITS>(Particle1, Particle0, Implicit1, Shape1, Simplicial1, Implicit0, Shape0, Simplicial0, ParticleWorldTransform1, LocalTransform1, ParticleWorldTransform0, LocalTransform0, CullDistance, Dt, Context);
-				}
+				Context.MultiShapeCollisionDetector->FindOrCreateConstraint(
+					Particle0, Implicit0, ImplicitID0, Shape0, Simplicial0, LocalTransform0,
+					Particle1, Implicit1, ImplicitID1, Shape1, Simplicial1, LocalTransform1,
+					CullDistance, ShapePairType, Context.GetSettings().bAllowManifolds, bEnableSweep, Context);
 			}
 		}
 
@@ -2076,244 +1287,140 @@ namespace Chaos
 		// NOTE: Transforms are world space shape transforms
 		void UpdateConstraint(FPBDCollisionConstraint& Constraint, const FRigidTransform3& ShapeWorldTransform0, const FRigidTransform3& ShapeWorldTransform1, const FReal Dt)
 		{
-			UpdateConstraintFromGeometryImpl<ECollisionUpdateType::Deepest>(Constraint, ShapeWorldTransform0, ShapeWorldTransform1, Dt);
+			UpdateConstraintFromGeometryImpl(Constraint, ShapeWorldTransform0, ShapeWorldTransform1, Dt);
 		}
 
 		bool UpdateConstraintSwept(FPBDCollisionConstraint& Constraint, const FRigidTransform3& ShapeStartWorldTransform0, const FRigidTransform3& ShapeStartWorldTransform1, const FReal Dt)
 		{
-			return UpdateConstraintFromGeometrySweptImpl<ECollisionUpdateType::Deepest>(Constraint, ShapeStartWorldTransform0, ShapeStartWorldTransform1, Dt);
-		}
-
-		// Run collision detection for the specified constraint to update the nearest contact point.
-		// NOTE: Transforms are world space particle transforms
-		template<ECollisionUpdateType UpdateType>
-		void UpdateConstraintFromGeometry(FPBDCollisionConstraint& Constraint, const FRigidTransform3& ParticleWorldTransform0, const FRigidTransform3& ParticleWorldTransform1, const FReal Dt)
-		{
-			const FRigidTransform3 WorldTransform0 = Constraint.GetShapeRelativeTransform0() * ParticleWorldTransform0;
-			const FRigidTransform3 WorldTransform1 = Constraint.GetShapeRelativeTransform1() * ParticleWorldTransform1;
-			Constraint.SetShapeWorldTransforms(WorldTransform0, WorldTransform1);
-
-			UpdateConstraintFromGeometryImpl<UpdateType>(Constraint, WorldTransform0, WorldTransform1, Dt);
-		}
-
-		template<typename T_TRAITS>
-		void LegacyConstructConstraints(TGeometryParticleHandle<FReal, 3>* Particle0, TGeometryParticleHandle<FReal, 3>* Particle1, const FImplicitObject* Implicit0, const FShapeOrShapesArray& Shapes0, const FBVHParticles* Simplicial0, const FImplicitObject* Implicit1, const FShapeOrShapesArray& Shapes1, const FBVHParticles* Simplicial1, const FRigidTransform3& ParticleWorldTransform0, const FRigidTransform3& LocalTransform0, const FRigidTransform3& ParticleWorldTransform1, const FRigidTransform3& LocalTransform1, const FReal CullDistance, FReal Dt, const FCollisionContext& Context)
+			return UpdateConstraintFromGeometrySweptImpl(Constraint, ShapeStartWorldTransform0, ShapeStartWorldTransform1, Dt);
+		}
+
+		// @param ImplicitId0 An identifier which allows us to differentiate between instances when the same ImplicitObject appears multiple times in the hierarchy
+		// @param ImplicitId1 See ImplicitId0
+		void ConstructConstraintsRecursive(
+			TGeometryParticleHandle<FReal, 3>* Particle0, 
+			TGeometryParticleHandle<FReal, 3>* Particle1, 
+			const FImplicitObject* Implicit0, 
+			const FShapeInstance* Shape0, 
+			const FBVHParticles* Simplicial0, 
+			const int32 ImplicitID0,
+			const FImplicitObject* Implicit1, 
+			const FShapeInstance* Shape1, 
+			const FBVHParticles* Simplicial1, 
+			const int32 ImplicitID1,
+			const FRigidTransform3& ParticleWorldTransform0,
+			const FRigidTransform3& LocalTransform0, 
+			const FRigidTransform3& ParticleWorldTransform1, 
+			const FRigidTransform3& LocalTransform1, 
+			const FReal CullDistance, 
+			const FReal Dt, 
+			const bool bEnableSweep,
+			const FCollisionContext& Context)
 		{
 			CONDITIONAL_SCOPE_CYCLE_COUNTER(STAT_Collisions_ConstructConstraints, ConstraintsDetailedStats);
 
-			EImplicitObjectType Implicit0Type = Implicit0 ? GetInnerType(Implicit0->GetType()) : ImplicitObjectType::Unknown;
-			EImplicitObjectType Implicit1Type = Implicit1 ? GetInnerType(Implicit1->GetType()) : ImplicitObjectType::Unknown;
-
-
-			if (!Implicit0 || !Implicit1)
-			{
-				const FPerShapeData* Shape0	= nullptr;
-				if (ensure(Shapes0.IsSingleShape()))
-				{
-					Shape0 = Shapes0.GetShape();
-				}
-				const FPerShapeData* Shape1 = nullptr;
-				if (ensure(Shapes1.IsSingleShape()))
-				{
-					Shape1 = Shapes1.GetShape();
-				}
-
-				ConstructLevelsetLevelsetConstraints<T_TRAITS>(Particle0, Particle1, Implicit0, Shape0, Simplicial0, Implicit1, Shape1, Simplicial1, ParticleWorldTransform0, LocalTransform0, ParticleWorldTransform1, LocalTransform1, CullDistance, Dt, Context);
-				return;
-			}
+			const EImplicitObjectType Implicit0OuterType = Implicit0 ? Implicit0->GetType() : ImplicitObjectType::Unknown;
+			const EImplicitObjectType Implicit1OuterType = Implicit1 ? Implicit1->GetType() : ImplicitObjectType::Unknown;
+			const EImplicitObjectType Implicit0Type = GetInnerType(Implicit0OuterType);
+			const EImplicitObjectType Implicit1Type = GetInnerType(Implicit1OuterType);
 
-			EImplicitObjectType Implicit0OuterType = Implicit0->GetType();
-			EImplicitObjectType Implicit1OuterType = Implicit1->GetType();
-
-			// Handle transform wrapper shape
-			if ((Implicit0OuterType == TImplicitObjectTransformed<FReal, 3>::StaticType()) && (Implicit1OuterType == TImplicitObjectTransformed<FReal, 3>::StaticType()))
-			{
-				const TImplicitObjectTransformed<FReal, 3>* TransformedImplicit0 = Implicit0->template GetObject<const TImplicitObjectTransformed<FReal, 3>>();
-				const TImplicitObjectTransformed<FReal, 3>* TransformedImplicit1 = Implicit1->template GetObject<const TImplicitObjectTransformed<FReal, 3>>();
-				FRigidTransform3 TransformedTransform0 = TransformedImplicit0->GetTransform() * LocalTransform0;
-				FRigidTransform3 TransformedTransform1 = TransformedImplicit1->GetTransform() * LocalTransform1;
-				LegacyConstructConstraints<T_TRAITS>(Particle0, Particle1, TransformedImplicit0->GetTransformedObject(), Shapes0, Simplicial0, TransformedImplicit1->GetTransformedObject(), Shapes1, Simplicial1, ParticleWorldTransform0, TransformedTransform0, ParticleWorldTransform1, TransformedTransform1, CullDistance, Dt, Context);
-				return;
-			}
-			else if (Implicit0OuterType == TImplicitObjectTransformed<FReal, 3>::StaticType())
-			{
-				const TImplicitObjectTransformed<FReal, 3>* TransformedImplicit0 = Implicit0->template GetObject<const TImplicitObjectTransformed<FReal, 3>>();
-				FRigidTransform3 TransformedTransform0 = TransformedImplicit0->GetTransform() * LocalTransform0;
-				LegacyConstructConstraints<T_TRAITS>(Particle0, Particle1, TransformedImplicit0->GetTransformedObject(), Shapes0, Simplicial0, Implicit1, Shapes1, Simplicial1, ParticleWorldTransform0, TransformedTransform0, ParticleWorldTransform1, LocalTransform1, CullDistance, Dt, Context);
-				return;
-			}
-			else if (Implicit1OuterType == TImplicitObjectTransformed<FReal, 3>::StaticType())
-			{
-				const TImplicitObjectTransformed<FReal, 3>* TransformedImplicit1 = Implicit1->template GetObject<const TImplicitObjectTransformed<FReal, 3>>();
-				FRigidTransform3 TransformedTransform1 = TransformedImplicit1->GetTransform() * LocalTransform1;
-				LegacyConstructConstraints<T_TRAITS>(Particle0, Particle1, Implicit0, Shapes0, Simplicial0, TransformedImplicit1->GetTransformedObject(), Shapes1, Simplicial1, ParticleWorldTransform0, LocalTransform0, ParticleWorldTransform1, TransformedTransform1, CullDistance, Dt, Context);
-				return;
-			}
-
-			// Strip the Instanced wrapper from most shapes, but not Convex or TriMesh.
-			// Convex collision requires the wrapper because it holds the margin.
-			// @todo(chaos): this collision logic is getting out of hand - can we make a better shape class hierarchy?
-			if (((uint32)Implicit0OuterType & ImplicitObjectType::IsInstanced) || ((uint32)Implicit1OuterType & ImplicitObjectType::IsInstanced))
+			if ((Implicit0 != nullptr) && (Implicit1 != nullptr))
 			{
-				const FImplicitObject* InnerImplicit0 = nullptr;
-				const FImplicitObject* InnerImplicit1 = nullptr;
-				if (((uint32)Implicit0OuterType & ImplicitObjectType::IsInstanced) && (Implicit0Type != FImplicitConvex3::StaticType()))
+				// Handle transform wrapper shape
+				if ((Implicit0OuterType == TImplicitObjectTransformed<FReal, 3>::StaticType()) && (Implicit1OuterType == TImplicitObjectTransformed<FReal, 3>::StaticType()))
 				{
-					InnerImplicit0 = GetInstancedImplicit(Implicit0);
-				}
-				if (((uint32)Implicit1OuterType & ImplicitObjectType::IsInstanced) && (Implicit1Type != FImplicitConvex3::StaticType()))
-				{
-					InnerImplicit1 = GetInstancedImplicit(Implicit1);
-				}
-				if (InnerImplicit0 && InnerImplicit1)
-				{
-					LegacyConstructConstraints<T_TRAITS>(Particle0, Particle1, InnerImplicit0, Shapes0, Simplicial0, InnerImplicit1, Shapes1, Simplicial1, ParticleWorldTransform0, LocalTransform0, ParticleWorldTransform1, LocalTransform1, CullDistance, Dt, Context);
+					const TImplicitObjectTransformed<FReal, 3>* TransformedImplicit0 = Implicit0->template GetObject<const TImplicitObjectTransformed<FReal, 3>>();
+					const TImplicitObjectTransformed<FReal, 3>* TransformedImplicit1 = Implicit1->template GetObject<const TImplicitObjectTransformed<FReal, 3>>();
+					FRigidTransform3 TransformedTransform0 = TransformedImplicit0->GetTransform() * LocalTransform0;
+					FRigidTransform3 TransformedTransform1 = TransformedImplicit1->GetTransform() * LocalTransform1;
+					ConstructConstraintsRecursive(Particle0, Particle1, TransformedImplicit0->GetTransformedObject(), Shape0, Simplicial0, ImplicitID0, TransformedImplicit1->GetTransformedObject(), Shape1, Simplicial1, ImplicitID1, ParticleWorldTransform0, TransformedTransform0, ParticleWorldTransform1, TransformedTransform1, CullDistance, Dt, bEnableSweep, Context);
 					return;
 				}
-				else if (InnerImplicit0 && !InnerImplicit1)
+				else if (Implicit0OuterType == TImplicitObjectTransformed<FReal, 3>::StaticType())
 				{
-					LegacyConstructConstraints<T_TRAITS>(Particle0, Particle1, InnerImplicit0, Shapes0, Simplicial0, Implicit1, Shapes1, Simplicial1, ParticleWorldTransform0, LocalTransform0, ParticleWorldTransform1, LocalTransform1, CullDistance, Dt, Context);
+					const TImplicitObjectTransformed<FReal, 3>* TransformedImplicit0 = Implicit0->template GetObject<const TImplicitObjectTransformed<FReal, 3>>();
+					FRigidTransform3 TransformedTransform0 = TransformedImplicit0->GetTransform() * LocalTransform0;
+					ConstructConstraintsRecursive(Particle0, Particle1, TransformedImplicit0->GetTransformedObject(), Shape0, Simplicial0, ImplicitID0, Implicit1, Shape1, Simplicial1, ImplicitID1, ParticleWorldTransform0, TransformedTransform0, ParticleWorldTransform1, LocalTransform1, CullDistance, Dt, bEnableSweep, Context);
 					return;
 				}
-				else if (!InnerImplicit0 && InnerImplicit1)
+				else if (Implicit1OuterType == TImplicitObjectTransformed<FReal, 3>::StaticType())
 				{
-					LegacyConstructConstraints<T_TRAITS>(Particle0, Particle1, Implicit0, Shapes0, Simplicial0, InnerImplicit1, Shapes1, Simplicial1, ParticleWorldTransform0, LocalTransform0, ParticleWorldTransform1, LocalTransform1, CullDistance, Dt, Context);
+					const TImplicitObjectTransformed<FReal, 3>* TransformedImplicit1 = Implicit1->template GetObject<const TImplicitObjectTransformed<FReal, 3>>();
+					FRigidTransform3 TransformedTransform1 = TransformedImplicit1->GetTransform() * LocalTransform1;
+					ConstructConstraintsRecursive(Particle0, Particle1, Implicit0, Shape0, Simplicial0, ImplicitID0, TransformedImplicit1->GetTransformedObject(), Shape1, Simplicial1, ImplicitID1, ParticleWorldTransform0, LocalTransform0, ParticleWorldTransform1, TransformedTransform1, CullDistance, Dt, bEnableSweep, Context);
 					return;
 				}
 
-			}
-
-			// Handle Unions
-			if (Implicit0OuterType == FImplicitObjectUnion::StaticType())
-			{
-				const FImplicitObjectUnion* Union0 = Implicit0->template GetObject<FImplicitObjectUnion>();
-				int32 ShapeIdx = 0;
-				for (const auto& Child0 : Union0->GetObjects())
+				// Strip the Instanced wrapper from most shapes, but not Convex or TriMesh.
+				// Convex collision requires the wrapper because it holds the margin.
+				if (((uint32)Implicit0OuterType & ImplicitObjectType::IsInstanced) || ((uint32)Implicit1OuterType & ImplicitObjectType::IsInstanced))
 				{
-					if (ensure(!Shapes0.IsSingleShape()))
+					const FImplicitObject* InnerImplicit0 = nullptr;
+					const FImplicitObject* InnerImplicit1 = nullptr;
+					if (((uint32)Implicit0OuterType & ImplicitObjectType::IsInstanced) && (Implicit0Type != FImplicitConvex3::StaticType()))
 					{
-						const FShapesArray& ShapesArray = *Shapes0.GetShapesArray();
-						const FPerShapeData* ChildShape0 = ShapesArray[ShapeIdx].Get();
-					
-						// If shape is not sim'd, we may end up iterating over a lot of shapes on particle1's union and wasting time filtering.
-						if (Context.GetSettings().bFilteringEnabled == false || FilterHasSimEnabled(ChildShape0))
-						{
-							LegacyConstructConstraints<T_TRAITS>(Particle0, Particle1, Child0.Get(), FShapeOrShapesArray(ChildShape0), Simplicial0, Implicit1, Shapes1, Simplicial1, ParticleWorldTransform0, LocalTransform0, ParticleWorldTransform1, LocalTransform1, CullDistance, Dt, Context);
-						}
+						InnerImplicit0 = GetInstancedImplicit(Implicit0);
 					}
-
-					ShapeIdx++;
-				}
-				return;
-			}
-
-			if (Implicit0OuterType == FImplicitObjectUnionClustered::StaticType())
-			{
-				const FImplicitObjectUnionClustered* Union0 = Implicit0->template GetObject<FImplicitObjectUnionClustered>();
-				if (Implicit1->HasBoundingBox())
-				{
-					TArray<FLargeUnionClusteredImplicitInfo> Children;
-
-					// Need to get transformed bounds of 1 in the space of 0
-					FRigidTransform3 TM0 = LocalTransform0 * ParticleWorldTransform0;
-					FRigidTransform3 TM1 = LocalTransform1 * ParticleWorldTransform1;
-					FRigidTransform3 TM1ToTM0 = TM1.GetRelativeTransform(TM0);
-					FAABB3 QueryBounds = Implicit1->BoundingBox().TransformedAABB(TM1ToTM0);
-
-					Union0->FindAllIntersectingClusteredObjects(Children, QueryBounds);
-
-					for (const FLargeUnionClusteredImplicitInfo& Child0 : Children)
+					if (((uint32)Implicit1OuterType & ImplicitObjectType::IsInstanced) && (Implicit1Type != FImplicitConvex3::StaticType()))
 					{
-						// Provide no shape for clustered child implicits, these use their parent's filter data, which was already tested when generating shape pair in midphase.
-						LegacyConstructConstraints<T_TRAITS>(Particle0, Particle1, Child0.Implicit, FShapeOrShapesArray(), Child0.BVHParticles, Implicit1, Shapes1, Simplicial1, ParticleWorldTransform0, Child0.Transform * LocalTransform0, ParticleWorldTransform1, LocalTransform1, CullDistance, Dt, Context);
+						InnerImplicit1 = GetInstancedImplicit(Implicit1);
 					}
-				}
-				else
-				{
-					for (const TUniquePtr<FImplicitObject>& Child0 : Union0->GetObjects())
+					if (InnerImplicit0 && InnerImplicit1)
 					{
-						const TPBDRigidParticleHandle<FReal, 3>* OriginalParticle = Union0->FindParticleForImplicitObject(Child0.Get());
-
-						// Provide no shape for clustered child implicits, these use their parent's filter data, which was already tested when generating shape pair in midphase.
-						LegacyConstructConstraints<T_TRAITS>(Particle0, Particle1, Child0.Get(), FShapeOrShapesArray(), OriginalParticle ? OriginalParticle->CollisionParticles().Get() : Simplicial0, Implicit1, Shapes1, Simplicial1, ParticleWorldTransform0, LocalTransform0, ParticleWorldTransform1, LocalTransform1, CullDistance, Dt, Context);
+						ConstructConstraintsRecursive(Particle0, Particle1, InnerImplicit0, Shape0, Simplicial0, ImplicitID0, InnerImplicit1, Shape1, Simplicial1, ImplicitID1, ParticleWorldTransform0, LocalTransform0, ParticleWorldTransform1, LocalTransform1, CullDistance, Dt, bEnableSweep, Context);
+						return;
 					}
-				}
-				return;
-			}
-
-			if (Implicit1OuterType == FImplicitObjectUnion::StaticType())
-			{
-				const FImplicitObjectUnion* Union1 = Implicit1->template GetObject<FImplicitObjectUnion>();
-				int32 ShapeIdx = 0;
-				for (const auto& Child1 : Union1->GetObjects())
-				{
-					if (ensure(!Shapes1.IsSingleShape()))
+					else if (InnerImplicit0 && !InnerImplicit1)
 					{
-						const FShapesArray& ShapesArray = *Shapes1.GetShapesArray();
-						const FPerShapeData* ChildShape1 = ShapesArray[ShapeIdx].Get();
-
-						// If shape is not sim'd, we may end up iterating over a lot of shapes on particle1's union and wasting time filtering.
-						if (Context.GetSettings().bFilteringEnabled == false || FilterHasSimEnabled(ChildShape1))
-						{
-							LegacyConstructConstraints<T_TRAITS>(Particle0, Particle1, Implicit0, Shapes0, Simplicial0, Child1.Get(), FShapeOrShapesArray(ChildShape1), Simplicial1, ParticleWorldTransform0, LocalTransform0, ParticleWorldTransform1, LocalTransform1, CullDistance, Dt, Context);
-						}
+						ConstructConstraintsRecursive(Particle0, Particle1, InnerImplicit0, Shape0, Simplicial0, ImplicitID0, Implicit1, Shape1, Simplicial1, ImplicitID1, ParticleWorldTransform0, LocalTransform0, ParticleWorldTransform1, LocalTransform1, CullDistance, Dt, bEnableSweep, Context);
+						return;
+					}
+					else if (!InnerImplicit0 && InnerImplicit1)
+					{
+						ConstructConstraintsRecursive(Particle0, Particle1, Implicit0, Shape0, Simplicial0, ImplicitID0, InnerImplicit1, Shape1, Simplicial1, ImplicitID1, ParticleWorldTransform0, LocalTransform0, ParticleWorldTransform1, LocalTransform1, CullDistance, Dt, bEnableSweep, Context);
+						return;
 					}
-
-					ShapeIdx++;
 				}
-				return;
-			}
 
-			if (Implicit1OuterType == FImplicitObjectUnionClustered::StaticType())
-			{
-				const FImplicitObjectUnionClustered* Union1 = Implicit1->template GetObject<FImplicitObjectUnionClustered>();
-				if (Implicit0->HasBoundingBox())
+				// Handle Unions
+				// NOTE: the Union visit function called below are recursive and visit all children in the hierarchy below the union.
+				if ((Implicit0OuterType == FImplicitObjectUnion::StaticType()) || (Implicit0OuterType == FImplicitObjectUnionClustered::StaticType()))
 				{
-					TArray<FLargeUnionClusteredImplicitInfo> Children;
-
-					// Need to get transformed bounds of 0 in the space of 1
-					FRigidTransform3 TM0 = LocalTransform0 * ParticleWorldTransform0;
-					FRigidTransform3 TM1 = LocalTransform1 * ParticleWorldTransform1;
-					FRigidTransform3 TM0ToTM1 = TM0.GetRelativeTransform(TM1);
-					FAABB3 QueryBounds = Implicit0->BoundingBox().TransformedAABB(TM0ToTM1);
+					const FRigidTransform3 TM0 = LocalTransform0 * ParticleWorldTransform0;
+					const FRigidTransform3 TM1 = LocalTransform1 * ParticleWorldTransform1;
+					const FRigidTransform3 TM1ToTM0 = TM1.GetRelativeTransform(TM0);
+					const FAABB3 QueryBounds0 = Implicit1->HasBoundingBox() ? Implicit1->BoundingBox().TransformedAABB(TM1ToTM0) : FAABB3::FullAABB();
 
-					{
-						CONDITIONAL_SCOPE_CYCLE_COUNTER(STAT_Collisions_FindAllIntersectingClusteredObjects, ConstraintsDetailedStats);
-						Union1->FindAllIntersectingClusteredObjects(Children, QueryBounds);
-					}
-
-					for (const FLargeUnionClusteredImplicitInfo& Child1 : Children)
-					{
-						// Provide no shape for clustered child implicits, these use their parent's filter data, which was already tested when generating shape pair in midphase.
-						LegacyConstructConstraints<T_TRAITS>(Particle0, Particle1, Implicit0, Shapes0, Simplicial0, Child1.Implicit, FShapeOrShapesArray(), Child1.BVHParticles, ParticleWorldTransform0, LocalTransform0, ParticleWorldTransform1, Child1.Transform * LocalTransform1, CullDistance, Dt, Context);
-					}
+					// @todo(chaos): we shouldn't need to treat Union and ClusterUnion differentlty here - change how we get the simplicial to be more generic
+					const FImplicitObjectUnionClustered* ClusterUnion0 = Implicit0->template GetObject<FImplicitObjectUnionClustered>();
+					Implicit0->VisitOverlappingObjects(QueryBounds0,
+						[&](const FImplicitObject* ChildImplicit0, const FRigidTransform3& ChildTransform0, const int32 ChildId0)
+						{
+							const FBVHParticles* ChildSimplicial0 = (ClusterUnion0 != nullptr) ? ClusterUnion0->GetChildSimplicial(ChildId0) : Simplicial0;
+							ConstructConstraintsRecursive(Particle0, Particle1, ChildImplicit0, Shape0, ChildSimplicial0, ChildId0, Implicit1, Shape1, Simplicial1, ImplicitID1, ParticleWorldTransform0, ChildTransform0 * LocalTransform0, ParticleWorldTransform1, LocalTransform1, CullDistance, Dt, bEnableSweep, Context);
+						});
+					return;
 				}
-				else
+				if ((Implicit1OuterType == FImplicitObjectUnion::StaticType()) || (Implicit1OuterType == FImplicitObjectUnionClustered::StaticType()))
 				{
-					for (const TUniquePtr<FImplicitObject>& Child1 : Union1->GetObjects())
-					{
-						const TPBDRigidParticleHandle<FReal, 3>* OriginalParticle = Union1->FindParticleForImplicitObject(Child1.Get());
+					const FRigidTransform3 TM0 = LocalTransform0 * ParticleWorldTransform0;
+					const FRigidTransform3 TM1 = LocalTransform1 * ParticleWorldTransform1;
+					const FRigidTransform3 TM0ToTM1 = TM0.GetRelativeTransform(TM1);
+					const FAABB3 QueryBounds1 = Implicit0->HasBoundingBox() ? Implicit0->BoundingBox().TransformedAABB(TM0ToTM1) : FAABB3::FullAABB();
 
-						// Provide no shape for clustered child implicits, these use their parent's filter data, which was already tested when generating shape pair in midphase.
-						LegacyConstructConstraints<T_TRAITS>(Particle0, Particle1, Implicit0, Shapes0, Simplicial0, Child1.Get(), FShapeOrShapesArray(), OriginalParticle ? OriginalParticle->CollisionParticles().Get() : Simplicial1, ParticleWorldTransform0, LocalTransform0, ParticleWorldTransform1, LocalTransform1, CullDistance, Dt, Context);
-					}
+					const FImplicitObjectUnionClustered* ClusterUnion1 = Implicit1->template GetObject<FImplicitObjectUnionClustered>();
+					Implicit1->VisitOverlappingObjects(QueryBounds1,
+						[&](const FImplicitObject* ChildImplicit1, const FRigidTransform3& ChildTransform1, const int32 ChildId1)
+						{
+							const FBVHParticles* ChildSimplicial1 = (ClusterUnion1 != nullptr) ? ClusterUnion1->GetChildSimplicial(ChildId1) : Simplicial1;
+							ConstructConstraintsRecursive(Particle0, Particle1, Implicit0, Shape0, Simplicial0, ImplicitID0, ChildImplicit1, Shape1, Simplicial1, ChildId1, ParticleWorldTransform0, LocalTransform0, ParticleWorldTransform1, ChildTransform1 * LocalTransform1, CullDistance, Dt, bEnableSweep, Context);
+						});
+					return;
 				}
-				return;
-			}
-
-			// Should not have unions at this point, shouldn't be shapes array
-			const FPerShapeData* Shape0 = nullptr;
-			if (ensure(Shapes0.IsSingleShape()))
-			{
-				Shape0 = Shapes0.GetShape();
-			}
-			const FPerShapeData* Shape1 = nullptr;
-			if (ensure(Shapes1.IsSingleShape()))
-			{
-				Shape1 = Shapes1.GetShape();
 			}
 
 			// Check shape pair filtering if enable
+			// @todo(chaos): Make sure we are setting this to false in the midphase (it has already performed filtering)
 			if (Context.GetSettings().bFilteringEnabled && !DoCollide(Implicit0Type, Shape0, Implicit1Type, Shape1))
 			{
 				return;
@@ -2321,27 +1428,20 @@ namespace Chaos
 
 			// If we get here, we have a pair of concrete shapes (i.e., no wrappers or containers)
 			// Create a constraint for the shape pair
+			bool bShouldSwapParticles = false;
+			const EContactShapesType ShapePairType = CalculateShapePairType(Implicit0, Simplicial0, Implicit1, Simplicial1, bShouldSwapParticles);
+
+			if (bShouldSwapParticles)
 			{
-				// let's make sure the order of particle is correct
-				// this needs to be done at this level because we have the concrete type available here ( and no union or other wrappers )
-				// @todo(chaos) : with this order enforced, ConstructConstraintsImpl can be optimized by remove reverse test cases ( keeping sphere|convex but removing convex|sphere for example )  
-				bool bShouldSwapParticles = false;
-				if (bChaos_Collision_SortParticlesOnConstraintConstruct)
-				{
-					CalculateShapePairType(Implicit0, Simplicial0, Implicit1, Simplicial1, bShouldSwapParticles);
-				}
-				if (bShouldSwapParticles)
-				{
-					ConstructConstraintsImpl<T_TRAITS>(Particle1, Particle0, Implicit1, Shape1, Simplicial1, Implicit0, Shape0, Simplicial0, ParticleWorldTransform1, LocalTransform1, ParticleWorldTransform0, LocalTransform0, CullDistance, Dt, Context);	
-				}
-				else
-				{
-					ConstructConstraintsImpl<T_TRAITS>(Particle0, Particle1, Implicit0, Shape0, Simplicial0, Implicit1, Shape1, Simplicial1, ParticleWorldTransform0, LocalTransform0, ParticleWorldTransform1, LocalTransform1, CullDistance, Dt, Context);
-				}
+				ConstructConstraintsImpl(Particle1, Particle0, Implicit1, Shape1, Simplicial1, ImplicitID1, Implicit0, Shape0, Simplicial0, ImplicitID0, ParticleWorldTransform1, LocalTransform1, ParticleWorldTransform0, LocalTransform0, CullDistance, Dt, ShapePairType, bEnableSweep, Context);
+			}
+			else
+			{
+				ConstructConstraintsImpl(Particle0, Particle1, Implicit0, Shape0, Simplicial0, ImplicitID0, Implicit1, Shape1, Simplicial1, ImplicitID1, ParticleWorldTransform0, LocalTransform0, ParticleWorldTransform1, LocalTransform1, CullDistance, Dt, ShapePairType, bEnableSweep, Context);
 			}
 		}
 
-		EContactShapesType CalculateShapePairType(const EImplicitObjectType Implicit0Type, const EImplicitObjectType Implicit1Type, const bool bIsConvex0, const bool bIsConvex1, const bool bIsBVH0, const bool bIsBVH1,  bool &bOutSwap)
+		EContactShapesType CalculateShapePairType(const EImplicitObjectType Implicit0Type, const EImplicitObjectType Implicit1Type, const bool bIsConvex0, const bool bIsConvex1, const bool bHasSimplicial0, const bool bHasSimplicial1,  bool &bOutSwap)
 		{
 			bOutSwap = false;
 			if (Implicit0Type == TBox<FReal, 3>::StaticType() && Implicit1Type == TBox<FReal, 3>::StaticType())
@@ -2495,14 +1595,22 @@ namespace Chaos
 			{
 				return EContactShapesType::GenericConvexConvex;
 			}
-			else if ((Implicit0Type == FLevelSet::StaticType()) || (Implicit1Type == FLevelSet::StaticType()))
+			else if ((Implicit0Type == FImplicitObjectUnion::StaticType()) || (Implicit1Type == FImplicitObjectUnion::StaticType()))
+			{
+				// Unknown here is used to mean "run the recursive collision detection path" in the midphase
+				// @todo(chaos): support FImplicitObjectUnion and FImplicitObjectUnionClustered more explicitly here
+				return EContactShapesType::Unknown;
+			}
+			else if (bHasSimplicial0 || bHasSimplicial1)
 			{
-				bOutSwap = !bIsBVH0 && bIsBVH1;
+				// If only one of the shapes has no implicit, it goes first
+				const bool bSwapBasedOnImplicit = ((Implicit1Type == ImplicitObjectType::Unknown) && (Implicit0Type != ImplicitObjectType::Unknown));
+				// If only one shape has a Simplicial, it goes first
+				const bool bSwapBasedOnSimplicial = bHasSimplicial1 && !bHasSimplicial0;
+				bOutSwap = bSwapBasedOnImplicit || bSwapBasedOnSimplicial;
 				return EContactShapesType::LevelSetLevelSet;
 			}
 
-			// @todo(chaos): support FImplicitObjectUnion
-			// @todo(chaos): support FImplicitObjectUnionClustered
 			return EContactShapesType::Unknown;
 		}
 
@@ -2512,32 +1620,44 @@ namespace Chaos
 			const EImplicitObjectType Implicit1Type = (Implicit1 != nullptr) ? GetInnerType(Implicit1->GetCollisionType()) : ImplicitObjectType::Unknown;
 			const bool bIsConvex0 = (Implicit0 != nullptr) && Implicit0->IsConvex() && (Implicit0Type != ImplicitObjectType::LevelSet);
 			const bool bIsConvex1 = (Implicit1 != nullptr) && Implicit1->IsConvex() && (Implicit1Type != ImplicitObjectType::LevelSet);
-			const bool bIsBVH0 = (BVHParticles0 != nullptr);
-			const bool bIsBVH1 = (BVHParticles1 != nullptr);
-			return CalculateShapePairType(Implicit0Type, Implicit1Type, bIsConvex0, bIsConvex1, bIsBVH0, bIsBVH1, bOutSwap);
-		}
-
-		void ConstructConstraints(TGeometryParticleHandle<FReal, 3>* Particle0, TGeometryParticleHandle<FReal, 3>* Particle1, const FImplicitObject* Implicit0, const FPerShapeData* Shape0, const FBVHParticles* Simplicial0, const FImplicitObject* Implicit1, const FPerShapeData* Shape1, const FBVHParticles* Simplicial1, const FRigidTransform3& ParticleWorldTransform0, const FRigidTransform3& LocalTransform0, const FRigidTransform3& ParticleWorldTransform1, const FRigidTransform3& LocalTransform1, const FReal CullDistance, const FReal dT, const FCollisionContext& Context)
+			const bool bHasSimplicial0 = (BVHParticles0 != nullptr) && (BVHParticles0->Size() > 0);
+			const bool bHasSimplicial1 = (BVHParticles1 != nullptr) && (BVHParticles1->Size() > 0);
+			return CalculateShapePairType(Implicit0Type, Implicit1Type, bIsConvex0, bIsConvex1, bHasSimplicial0, bHasSimplicial1, bOutSwap);
+		}
+
+		void ConstructConstraints(
+			FGeometryParticleHandle* Particle0, 
+			FGeometryParticleHandle* Particle1,
+			const FImplicitObject* Implicit0, 
+			const FPerShapeData* Shape0, 
+			const FBVHParticles* Simplicial0, 
+			const FImplicitObject* Implicit1, 
+			const FPerShapeData* Shape1, 
+			const FBVHParticles* Simplicial1, 
+			const FRigidTransform3& ParticleWorldTransform0, 
+			const FRigidTransform3& LocalTransform0, 
+			const FRigidTransform3& ParticleWorldTransform1, 
+			const FRigidTransform3& LocalTransform1, 
+			const FReal CullDistance, 
+			const FReal dT, 
+			const bool bEnableSweep, 
+			const FCollisionContext& Context)
 		{
 			INC_DWORD_STAT(STAT_ChaosCollisionCounter_NumParticlePairs);
 
-			// @todo(chaos): remove defer update - it is not a valid option with the new solver
-			bool bDeferUpdate = Context.GetSettings().bDeferNarrowPhase;
-			if (bDeferUpdate)
-			{
-				using TTraits = TConstructCollisionTraits<false>;
-				LegacyConstructConstraints<TTraits>(Particle0, Particle1, Implicit0, FShapeOrShapesArray(Shape0), Simplicial0, Implicit1, FShapeOrShapesArray(Shape1), Simplicial1, ParticleWorldTransform0, LocalTransform0, ParticleWorldTransform1, LocalTransform1, CullDistance, dT, Context);
-			}
-			else
-			{
-				using TTraits = TConstructCollisionTraits<true>;
-				LegacyConstructConstraints<TTraits>(Particle0, Particle1, Implicit0, FShapeOrShapesArray(Shape0), Simplicial0, Implicit1, FShapeOrShapesArray(Shape1), Simplicial1, ParticleWorldTransform0, LocalTransform0, ParticleWorldTransform1, LocalTransform1, CullDistance, dT, Context);
-			}
+			ConstructConstraintsRecursive(Particle0, Particle1, Implicit0, Shape0->AsShapeInstance(), Simplicial0, 0, Implicit1, Shape1->AsShapeInstance(), Simplicial1, 0, ParticleWorldTransform0, LocalTransform0, ParticleWorldTransform1, LocalTransform1, CullDistance, dT, bEnableSweep, Context);
 		}
 
-		template void UpdateLevelsetLevelsetConstraint<ECollisionUpdateType::Any>(const FRigidTransform3& WorldTransform0, const FRigidTransform3& WorldTransform1, const FReal Dt, FPBDCollisionConstraint& Constraint);
-		template void UpdateLevelsetLevelsetConstraint<ECollisionUpdateType::Deepest>(const FRigidTransform3& WorldTransform0, const FRigidTransform3& WorldTransform1, const FReal Dt, FPBDCollisionConstraint& Constraint);
+		// This is deprecated
+		template<ECollisionUpdateType UpdateType>
+		void UpdateConstraintFromGeometry(FPBDCollisionConstraint& Constraint, const FRigidTransform3& ParticleWorldTransform0, const FRigidTransform3& ParticleWorldTransform1, const FReal Dt)
+		{
+			const FRigidTransform3 WorldTransform0 = Constraint.GetShapeRelativeTransform0() * ParticleWorldTransform0;
+			const FRigidTransform3 WorldTransform1 = Constraint.GetShapeRelativeTransform1() * ParticleWorldTransform1;
+			Constraint.SetShapeWorldTransforms(WorldTransform0, WorldTransform1);
 
+			UpdateConstraintFromGeometryImpl(Constraint, WorldTransform0, WorldTransform1, Dt);
+		}
 		template void UpdateConstraintFromGeometry<ECollisionUpdateType::Any>(FPBDCollisionConstraint& ConstraintBase, const FRigidTransform3& Transform0, const FRigidTransform3& Transform1, const FReal Dt);
 		template void UpdateConstraintFromGeometry<ECollisionUpdateType::Deepest>(FPBDCollisionConstraint& ConstraintBase, const FRigidTransform3& Transform0, const FRigidTransform3& Transform1, const FReal Dt);
 
diff --git a/Engine/Source/Runtime/Experimental/Chaos/Private/Chaos/ImplicitObjectUnion.cpp b/Engine/Source/Runtime/Experimental/Chaos/Private/Chaos/ImplicitObjectUnion.cpp
--- a/Engine/Source/Runtime/Experimental/Chaos/Private/Chaos/ImplicitObjectUnion.cpp
+++ b/Engine/Source/Runtime/Experimental/Chaos/Private/Chaos/ImplicitObjectUnion.cpp
@@ -6,6 +6,7 @@ namespace Chaos
 {
 struct FLargeImplicitObjectUnionData
 {
+	// @todo(chaos): FGeometryParticles is way too much data for this. Use a custom struct
 	FGeometryParticles GeomParticles;
 	TBoundingVolumeHierarchy<FGeometryParticles, TArray<int32>> Hierarchy;
 
@@ -152,6 +153,44 @@ void FImplicitObjectUnion::CacheAllImplicitObjects()
 		}
 	}
 }
+
+void FImplicitObjectUnion::VisitOverlappingObjectsImpl(const FAABB3& LocalBounds, const FRigidTransform3& ParentTM, int32& VisitId, const TFunctionRef<void(const FImplicitObject*, const FRigidTransform3&, const int32)>& VisitorFunc) const
+{
+	if (LargeUnionData.IsValid())
+	{
+		// @todo(chaos): can a child be in multiple leafs? If not we can ditch this visited array, but
+		// see TBoundingVolumeHierarchy::FindAllIntersectionsHelper which implies that dupes are possible
+		TArray<bool> LeafVisited;
+		LeafVisited.AddZeroed(LargeUnionData->GeomParticles.Size());
+
+		// A visitor for the bounding volume.
+		// NOTE: The leaf data is an array of indices into the Hierarchy's flattened list of implicit objects
+		const auto& LeafVisitor = [this, &LeafVisited, &VisitorFunc](const TArray<int32>& Indices)
+		{
+			for (const int32 Index : Indices)
+			{
+				if (!LeafVisited[Index])
+				{
+					LeafVisited[Index] = true;
+					const FImplicitObject* Implicit = LargeUnionData->GeomParticles.Geometry(Index).Get();
+					const FRigidTransform3 Transform = FRigidTransform3(LargeUnionData->GeomParticles.X(Index), LargeUnionData->GeomParticles.R(Index));
+					VisitorFunc(Implicit, Transform, Index);
+				}
+			}
+		};
+
+		LargeUnionData->Hierarchy.VisitAllIntersections(LocalBounds, LeafVisitor);
+	}
+	else
+	{
+		for (int32 ObjectIndex = 0; ObjectIndex < MObjects.Num(); ++ObjectIndex)
+		{
+			MObjects[ObjectIndex]->VisitOverlappingObjectsImpl(LocalBounds, ParentTM, VisitId, VisitorFunc);
+			++VisitId;
+		}
+	}
+}
+
 TUniquePtr<FImplicitObject> FImplicitObjectUnion::Copy() const
 {
 	TArray<TUniquePtr<FImplicitObject>> CopyOfObjects;
@@ -344,4 +383,13 @@ const FPBDRigidParticleHandle* FImplicitObjectUnionClustered::FindParticleForImp
 	return Handle ? *Handle : nullptr;
 }
 
+const FBVHParticles* FImplicitObjectUnionClustered::GetChildSimplicial(const int32 ChildIndex) const
+{
+	if (MOriginalParticleLookupHack.IsValidIndex(ChildIndex))
+	{
+		return MOriginalParticleLookupHack[ChildIndex]->CollisionParticles().Get();
+	}
+	return nullptr;
+}
+
 } // namespace Chaos
diff --git a/Engine/Source/Runtime/Experimental/Chaos/Public/Chaos/BoundingVolumeHierarchy.h b/Engine/Source/Runtime/Experimental/Chaos/Public/Chaos/BoundingVolumeHierarchy.h
--- a/Engine/Source/Runtime/Experimental/Chaos/Public/Chaos/BoundingVolumeHierarchy.h
+++ b/Engine/Source/Runtime/Experimental/Chaos/Public/Chaos/BoundingVolumeHierarchy.h
@@ -137,6 +137,17 @@ class TBoundingVolumeHierarchy final : public ISpatialAcceleration<int32, T,d>
 		}
 	}
 
+	// Calls the visitor for every overlapping leaf.
+	// @tparam TVisitor void(const LEAF_TYPE& Leaf)
+	template<typename TVisitor>
+	void VisitAllIntersections(const FAABB3& LocalBounds, const TVisitor& Visitor) const
+	{
+		if (Elements.Num())
+		{
+			VisitAllIntersectionsRecursive(Elements[0], LocalBounds, Visitor);
+		}
+	}
+
 	// Begin ISpatialAcceleration interface
 	CHAOS_API TArray<int32> FindAllIntersections(const FAABB3& Box) const { return FindAllIntersectionsImp(Box); }
 	CHAOS_API TArray<int32> FindAllIntersections(const TSpatialRay<T,d>& Ray) const { return FindAllIntersectionsImp(Ray); }
@@ -173,6 +184,29 @@ class TBoundingVolumeHierarchy final : public ISpatialAcceleration<int32, T,d>
 		}
 	}
 
+	template<typename TVisitor>
+	void VisitAllIntersectionsRecursive(const TBVHNode<T, d>& MyNode, const FAABB3& ObjectBox, const TVisitor& Visitor) const
+	{
+		TAABB<T, d> Box(MyNode.MMin, MyNode.MMax);
+		if (!Box.Intersects(ObjectBox))
+		{
+			return;
+		}
+
+		if (MyNode.MChildren.Num() == 0)
+		{
+			Visitor(Leafs[MyNode.LeafIndex]);
+			return;
+		}
+
+		int32 NumChildren = MyNode.MChildren.Num();
+		for (int32 Child = 0; Child < NumChildren; ++Child)
+		{
+			VisitAllIntersectionsRecursive(Elements[MyNode.MChildren[Child]], ObjectBox, Visitor);
+		}
+	}
+
+
 	CHAOS_API TArray<int32> FindAllIntersectionsHelper(const TBVHNode<T,d>& MyNode, const TVector<T, d>& Point) const;
 	CHAOS_API TArray<int32> FindAllIntersectionsHelper(const TBVHNode<T,d>& MyNode, const TAABB<T, d>& ObjectBox) const;
 	CHAOS_API TArray<int32> FindAllIntersectionsHelper(const TBVHNode<T, d>& MyNode, const TSpatialRay<T, d>& Ray) const;
diff --git a/Engine/Source/Runtime/Experimental/Chaos/Public/Chaos/Collision/CollisionKeys.h b/Engine/Source/Runtime/Experimental/Chaos/Public/Chaos/Collision/CollisionKeys.h
--- a/Engine/Source/Runtime/Experimental/Chaos/Public/Chaos/Collision/CollisionKeys.h
+++ b/Engine/Source/Runtime/Experimental/Chaos/Public/Chaos/Collision/CollisionKeys.h
@@ -143,12 +143,29 @@ namespace Chaos
 		{
 		}
 
+		UE_DEPRECATED(5.3, "Replaced with version that takes Shape and ImplicitID")
 		FCollisionParticlePairConstraintKey(const FImplicitObject* Implicit0, const FBVHParticles* Simplicial0, const FImplicitObject* Implicit1, const FBVHParticles* Simplicial1)
 			: Key(0)
 		{
 			check((Implicit0 != nullptr) || (Simplicial0 != nullptr));
 			check((Implicit1 != nullptr) || (Simplicial1 != nullptr));
-			GenerateHash(Implicit0, Simplicial0, Implicit1, Simplicial1);
+			GenerateHash(nullptr, Implicit0, 0, Simplicial0, nullptr, Implicit1, 0, Simplicial1);
+		}
+
+		FCollisionParticlePairConstraintKey(
+			const FShapeInstance* Shape0, 
+			const FImplicitObject* Implicit0, 
+			const int32 ImplicitId0, 
+			const FBVHParticles* Simplicial0, 
+			const FShapeInstance* Shape1,
+			const FImplicitObject* Implicit1,
+			const int32 ImplicitId1,
+			const FBVHParticles* Simplicial1)
+			: Key(0)
+		{
+			check((Implicit0 != nullptr) || (Simplicial0 != nullptr));
+			check((Implicit1 != nullptr) || (Simplicial1 != nullptr));
+			GenerateHash(Shape0, Implicit0, ImplicitId0, Simplicial0, Shape1, Implicit1, ImplicitId1, Simplicial1);
 		}
 
 		uint32 GetKey() const
@@ -172,11 +189,21 @@ namespace Chaos
 		}
 
 	private:
-		void GenerateHash(const FImplicitObject* Implicit0, const FBVHParticles* Simplicial0, const FImplicitObject* Implicit1, const FBVHParticles* Simplicial1)
-		{
-			const uint32 Hash0 = (Implicit0 != nullptr) ? ::GetTypeHash(Implicit0) : ::GetTypeHash(Simplicial0);
-			const uint32 Hash1 = (Implicit1 != nullptr) ? ::GetTypeHash(Implicit1) : ::GetTypeHash(Simplicial1);
-			Key = OrderIndependentHashCombine(Hash0, Hash1);
+		void GenerateHash(
+			const FShapeInstance* Shape0,
+			const FImplicitObject* Implicit0,
+			const int32 ImplicitId0,
+			const FBVHParticles* Simplicial0,
+			const FShapeInstance* Shape1,
+			const FImplicitObject* Implicit1,
+			const int32 ImplicitId1,
+			const FBVHParticles* Simplicial1)
+		{
+			uint32 Hash0 = (Implicit0 != nullptr) ? ::GetTypeHash(Implicit0) : ::GetTypeHash(Simplicial0);
+			Hash0 = HashCombine(Hash0, ImplicitId0);
+			uint32 Hash1 = (Implicit1 != nullptr) ? ::GetTypeHash(Implicit1) : ::GetTypeHash(Simplicial1);
+			Hash1 = HashCombine(Hash1, ImplicitId1);
+			Key = HashCombine(Hash0, Hash1);
 		}
 
 		uint32 Key;
diff --git a/Engine/Source/Runtime/Experimental/Chaos/Public/Chaos/Collision/ParticlePairMidPhase.h b/Engine/Source/Runtime/Experimental/Chaos/Public/Chaos/Collision/ParticlePairMidPhase.h
--- a/Engine/Source/Runtime/Experimental/Chaos/Public/Chaos/Collision/ParticlePairMidPhase.h
+++ b/Engine/Source/Runtime/Experimental/Chaos/Public/Chaos/Collision/ParticlePairMidPhase.h
@@ -193,6 +193,12 @@ namespace Chaos
 			const FReal Dt,
 			const FCollisionContext& Context);
 
+		int32 GenerateCollisionsCCD(
+			const bool bEnableCCDSweep,
+			const FReal CullDistance,
+			const FReal Dt,
+			const FCollisionContext& Context);
+
 		/**
 		 * @brief Callback from the narrow phase to create a collision constraint for this particle pair.
 		 * We should never be asked for a collision for a different particle pair, but the 
@@ -200,37 +206,67 @@ namespace Chaos
 		*/
 		FPBDCollisionConstraint* FindOrCreateConstraint(
 			FGeometryParticleHandle* InParticle0,
-			const FImplicitObject* Implicit0,
-			const FPerShapeData* Shape0,
-			const FBVHParticles* BVHParticles0,
-			const FRigidTransform3& ShapeRelativeTransform0,
+			const FImplicitObject* InImplicit0,
+			const int32 InImplicitId0,
+			const FShapeInstance* InShape0,
+			const FBVHParticles* InBVHParticles0,
+			const FRigidTransform3& InShapeRelativeTransform0,
 			FGeometryParticleHandle* InParticle1,
 			const FImplicitObject* Implicit1,
-			const FPerShapeData* Shape1,
-			const FBVHParticles* BVHParticles1,
-			const FRigidTransform3& ShapeRelativeTransform1,
+			const int32 InImplicitId1,
+			const FShapeInstance* InShape1,
+			const FBVHParticles* InBVHParticles1,
+			const FRigidTransform3& InhapeRelativeTransform1,
 			const FReal CullDistance,
 			const EContactShapesType ShapePairType,
 			const bool bUseManifold,
+			const bool bEnableSweep,
 			const FCollisionContext& Context);
 
-		/**
-		 * @brief FindOrCreateConstraint for swept constraints 
-		*/
+		UE_DEPRECATED(5.3, "Replaced with version taking ImplicitIDs")
+		FPBDCollisionConstraint* FindOrCreateConstraint(
+			FGeometryParticleHandle* InParticle0,
+			const FImplicitObject* InImplicit0,
+			const FPerShapeData* InShape0,
+			const FBVHParticles* InBVHParticles0,
+			const FRigidTransform3& InShapeRelativeTransform0,
+			FGeometryParticleHandle* InParticle1,
+			const FImplicitObject* InImplicit1,
+			const FPerShapeData* InShape1,
+			const FBVHParticles* InBVHParticles1,
+			const FRigidTransform3& InShapeRelativeTransform1,
+			const FReal CullDistance,
+			const EContactShapesType ShapePairType,
+			const bool bUseManifold,
+			const FCollisionContext& Context)
+		{
+			return FindOrCreateConstraint(
+				InParticle0, InImplicit0, 0, InShape0->AsShapeInstance(), InBVHParticles0, InShapeRelativeTransform0,
+				InParticle1, InImplicit1, 0, InShape1->AsShapeInstance(), InBVHParticles1, InShapeRelativeTransform1,
+				CullDistance, ShapePairType, bUseManifold, false, Context);
+		}
+
+		UE_DEPRECATED(5.3, "Replaced with version taking ImplicitIDs")
 		FPBDCollisionConstraint* FindOrCreateSweptConstraint(
 			FGeometryParticleHandle* InParticle0,
-			const FImplicitObject* Implicit0,
-			const FPerShapeData* Shape0,
-			const FBVHParticles* BVHParticles0,
-			const FRigidTransform3& ShapeRelativeTransform0,
+			const FImplicitObject* InImplicit0,
+			const FPerShapeData* InShape0,
+			const FBVHParticles* InBVHParticles0,
+			const FRigidTransform3& InShapeRelativeTransform0,
 			FGeometryParticleHandle* InParticle1,
-			const FImplicitObject* Implicit1,
-			const FPerShapeData* Shape1,
-			const FBVHParticles* BVHParticles1,
-			const FRigidTransform3& ShapeRelativeTransform1,
+			const FImplicitObject* InImplicit1,
+			const FPerShapeData* InShape1,
+			const FBVHParticles* InBVHParticles1,
+			const FRigidTransform3& InShapeRelativeTransform1,
 			const FReal CullDistance,
 			const EContactShapesType ShapePairType,
-			const FCollisionContext& Context);
+			const FCollisionContext& Context)
+		{
+			return FindOrCreateConstraint(
+				InParticle0, InImplicit0, 0, InShape0->AsShapeInstance(), InBVHParticles0, InShapeRelativeTransform0,
+				InParticle1, InImplicit1, 0, InShape1->AsShapeInstance(), InBVHParticles1, InShapeRelativeTransform1,
+				CullDistance, ShapePairType, true, true, Context);
+		}
 
 		/**
 		 * @brief Reactivate the constraint
@@ -254,6 +290,12 @@ namespace Chaos
 		ECollisionVisitorResult VisitConstCollisions(const TLambda& Visitor, const bool bOnlyActive = true) const;
 
 	private:
+		int32 GenerateCollisionsImpl(
+			const bool bEnableCCDSweep,
+			const FReal CullDistance,
+			const FReal Dt,
+			const FCollisionContext& Context);
+
 		FPBDCollisionConstraint* FindConstraint(const FCollisionParticlePairConstraintKey& Key);
 
 		FPBDCollisionConstraint* CreateConstraint(
@@ -273,9 +315,25 @@ namespace Chaos
 			const FCollisionParticlePairConstraintKey& Key,
 			const FCollisionContext& Context);
 
-		int32 ProcessNewConstraints(const FCollisionContext& Context);
+		int32 ProcessNewConstraints(
+			const FReal CullDistance,
+			const FReal Dt,
+			const FCollisionContext& Context);
+
 		void PruneConstraints(const int32 CurrentEpoch);
 
+		void UpdateCollision(
+			FPBDCollisionConstraint* Constraint,
+			const FReal CullDistance,
+			const FReal Dt,
+			const FCollisionContext& Context);
+
+		void UpdateCollisionCCD(
+			FPBDCollisionConstraint* Constraint,
+			const FReal CullDistance,
+			const FReal Dt,
+			const FCollisionContext& Context);
+
 		FParticlePairMidPhase& MidPhase;
 		TMap<uint32, FPBDCollisionConstraintPtr> Constraints;
 		TArray<FPBDCollisionConstraint*> NewConstraints;
@@ -460,6 +518,22 @@ namespace Chaos
 			Flags.bIsCCDActive = bCCDIsActive;
 		}
 
+		/*
+		 * True if CCD is supported by either particle
+		 */
+		bool IsCCD() const
+		{
+			return Flags.bIsCCD;
+		}
+
+		/*
+		 * True if CCD is active for this midphase on this frame.This can be changed by modifiers and resets to bIsCCD each frame.
+		 */
+		bool IsCCDActive() const
+		{
+			return Flags.bIsCCDActive;
+		}
+
 	private:
 		/**
 		 * @brief Build the list of potentially colliding shape pairs.
@@ -495,7 +569,7 @@ namespace Chaos
 			{
 				uint32 bIsActive : 1;    // True if this midphase should generate a narrow phase at all
 				uint32 bIsCCD : 1;       // True if CCD is supported by either particle
-				uint32 bIsCCDActive : 1; // True if CCD is active for this midphase on this frame. This can be changes by modifiers and resets to bIsCCD each frame.
+				uint32 bIsCCDActive : 1; // True if CCD is active for this midphase on this frame. This can be changed by modifiers and resets to bIsCCD each frame.
 				uint32 bIsSleeping : 1;
 				uint32 bIsModified : 1;  // True if a modifier applied any changes to this midphase
 			};
diff --git a/Engine/Source/Runtime/Experimental/Chaos/Public/Chaos/CollisionResolution.h b/Engine/Source/Runtime/Experimental/Chaos/Public/Chaos/CollisionResolution.h
--- a/Engine/Source/Runtime/Experimental/Chaos/Public/Chaos/CollisionResolution.h
+++ b/Engine/Source/Runtime/Experimental/Chaos/Public/Chaos/CollisionResolution.h
@@ -57,18 +57,24 @@ namespace Chaos
 		*/
 		bool CHAOS_API ShouldUseCCD(const FGeometryParticleHandle* Particle0, const FVec3& DeltaX0, const FGeometryParticleHandle* Particle1, const FVec3& DeltaX1, FVec3& Dir, FReal& Length);
 
-		// Update the constraint by re-running collision detection on the shape pair.
-		// @todo(chaos): remove this and use UpdateConstraint instead
-		template<ECollisionUpdateType UpdateType>
-		void CHAOS_API UpdateConstraintFromGeometry(FPBDCollisionConstraint& Constraint, const FRigidTransform3& ParticleTransform0, const FRigidTransform3& ParticleTransform1, const FReal Dt);
 		// Create constraints for the particle pair. This could create multiple constraints: one for each potentially colliding shape pair in multi-shape particles.
-		void CHAOS_API ConstructConstraints(TGeometryParticleHandle<FReal, 3>* Particle0, TGeometryParticleHandle<FReal, 3>* Particle1, const FImplicitObject* Implicit0, const FPerShapeData* Shape0, const FBVHParticles* Simplicial0, const FImplicitObject* Implicit1, const FPerShapeData* Shape1, const FBVHParticles* Simplicial1, const FRigidTransform3& ParticleWorldTransform0, const FRigidTransform3& Transform0, const FRigidTransform3& ParticleWorldTransform1, const FRigidTransform3& Transform1, const FReal CullDistance, const FReal Dt,const FCollisionContext& Context);
+		void CHAOS_API ConstructConstraints(TGeometryParticleHandle<FReal, 3>* Particle0, TGeometryParticleHandle<FReal, 3>* Particle1, const FImplicitObject* Implicit0, const FPerShapeData* Shape0, const FBVHParticles* Simplicial0, const FImplicitObject* Implicit1, const FPerShapeData* Shape1, const FBVHParticles* Simplicial1, const FRigidTransform3& ParticleWorldTransform0, const FRigidTransform3& Transform0, const FRigidTransform3& ParticleWorldTransform1, const FRigidTransform3& Transform1, const FReal CullDistance, const FReal Dt, const bool bEnableSweep, const FCollisionContext& Context);
 
-		// @todo(chaos): this is only called in tests - should it really be exposed?
-		template<ECollisionUpdateType UpdateType>
+		// @todo(chaos): this should not be exposed but is currently used in tests
 		void UpdateLevelsetLevelsetConstraint(const FRigidTransform3& WorldTransform0, const FRigidTransform3& WorldTransform1, const FReal Dt, FPBDCollisionConstraint& Constraint);
 
 		// Reset per-frame collision stat counters
 		void CHAOS_API ResetChaosCollisionCounters();
+
+		// See Above
+		UE_DEPRECATED(5.3, "Replaced with a version that takes a bEnableSweep flag")
+		inline void CHAOS_API ConstructConstraints(TGeometryParticleHandle<FReal, 3>* Particle0, TGeometryParticleHandle<FReal, 3>* Particle1, const FImplicitObject* Implicit0, const FPerShapeData* Shape0, const FBVHParticles* Simplicial0, const FImplicitObject* Implicit1, const FPerShapeData* Shape1, const FBVHParticles* Simplicial1, const FRigidTransform3& ParticleWorldTransform0, const FRigidTransform3& Transform0, const FRigidTransform3& ParticleWorldTransform1, const FRigidTransform3& Transform1, const FReal CullDistance, const FReal Dt, const FCollisionContext& Context)
+		{
+			ConstructConstraints(Particle0, Particle1, Implicit0, Shape0, Simplicial0, Implicit1, Shape1, Simplicial1, ParticleWorldTransform0, Transform0, ParticleWorldTransform1, Transform1, CullDistance, Dt, false, Context);
+		}
+
+		template<ECollisionUpdateType UpdateType>
+		UE_DEPRECATED(5.3, "Use UpdateConstraint, but call Constraint.SetShapeWorldTransform first (see implementation)")
+		void CHAOS_API UpdateConstraintFromGeometry(FPBDCollisionConstraint& Constraint, const FRigidTransform3& ParticleTransform0, const FRigidTransform3& ParticleTransform1, const FReal Dt);
 	}
 }
diff --git a/Engine/Source/Runtime/Experimental/Chaos/Public/Chaos/ImplicitObject.h b/Engine/Source/Runtime/Experimental/Chaos/Public/Chaos/ImplicitObject.h
--- a/Engine/Source/Runtime/Experimental/Chaos/Public/Chaos/ImplicitObject.h
+++ b/Engine/Source/Runtime/Experimental/Chaos/Public/Chaos/ImplicitObject.h
@@ -314,6 +314,24 @@ public:
 
 	virtual uint16 GetMaterialIndex(uint32 HintIndex) const { return 0; }
 
+	// Visit all the objects in the hierarchy that overlap the specified local-space bounds.
+	// The Visitor ChildId is an unique index into the hierarchy that can be used to distinguish ImplicitObjects when
+	// they are reused in multiple locations in the hierarchy.
+	void VisitOverlappingObjects(const FAABB3& LocalBounds, const TFunctionRef<void(const FImplicitObject* Implicit, const FRigidTransform3& Transform, const int32 ChildId)>& Visitor) const
+	{
+		int32 VisitId = 0;
+		VisitOverlappingObjectsImpl(LocalBounds, FRigidTransform3::Identity, VisitId, Visitor);
+	}
+
+	// This should not be public, but it needs to be callable by derived classes on another instance (e.g., see FImplicitObjectUnion::VisitOverlappingObjectsImpl)
+	virtual void VisitOverlappingObjectsImpl(const FAABB3& LocalBounds, const FRigidTransform3& ParentTM, int32& VisitId, const TFunctionRef<void(const FImplicitObject*, const FRigidTransform3&, const int32)>& VisitorFunc) const
+	{
+		if (!HasBoundingBox() || LocalBounds.Intersects(BoundingBox()))
+		{
+			VisitorFunc(this, ParentTM, VisitId);
+		}
+	}
+
 protected:
 
 	// Safely scale a normalized Vec3 - used with both scale and inverse scale
diff --git a/Engine/Source/Runtime/Experimental/Chaos/Public/Chaos/ImplicitObjectTransformed.h b/Engine/Source/Runtime/Experimental/Chaos/Public/Chaos/ImplicitObjectTransformed.h
--- a/Engine/Source/Runtime/Experimental/Chaos/Public/Chaos/ImplicitObjectTransformed.h
+++ b/Engine/Source/Runtime/Experimental/Chaos/Public/Chaos/ImplicitObjectTransformed.h
@@ -353,6 +353,13 @@ public:
 		return MObject->GetMaterialIndex(HintIndex);
 	}
 
+protected:
+	virtual void VisitOverlappingObjectsImpl(const FAABB3& InLocalBounds, const FRigidTransform3& ParentTM, int32& VisitId, const TFunctionRef<void(const FImplicitObject*, const FRigidTransform3&, const int32)>& VisitorFunc) const
+	{
+		const FAABB3 LocalBounds = InLocalBounds.InverseTransformedAABB(MTransform);
+		MObject->VisitOverlappingObjectsImpl(LocalBounds, MTransform * ParentTM, VisitId, VisitorFunc);
+	}
+
 private:
 	ObjectType MObject;
 	TUniquePtr<Chaos::FImplicitObject> MObjectOwner;
diff --git a/Engine/Source/Runtime/Experimental/Chaos/Public/Chaos/ImplicitObjectUnion.h b/Engine/Source/Runtime/Experimental/Chaos/Public/Chaos/ImplicitObjectUnion.h
--- a/Engine/Source/Runtime/Experimental/Chaos/Public/Chaos/ImplicitObjectUnion.h
+++ b/Engine/Source/Runtime/Experimental/Chaos/Public/Chaos/ImplicitObjectUnion.h
@@ -211,6 +211,8 @@ protected:
 		return ClosestIntersection;
 	}
 
+	virtual void VisitOverlappingObjectsImpl(const FAABB3& LocalBounds, const FRigidTransform3& ParentTM, int32& VisitId, const TFunctionRef<void(const FImplicitObject*, const FRigidTransform3&, const int32)>& VisitorFunc) const override final;
+
   protected:
 	//needed for serialization
 	FImplicitObjectUnion();
@@ -252,41 +254,14 @@ public:
 	void FindAllIntersectingClusteredObjects(TArray<FLargeUnionClusteredImplicitInfo>& Out, const FAABB3& LocalBounds) const;
 	TArray<FPBDRigidParticleHandle*> FindAllIntersectingChildren(const FAABB3& LocalBounds) const;
 
-#if CHAOS_PARTICLEHANDLE_TODO
-	TArray<int32> FindAllIntersectingChildren(const TSpatialRay<FReal,3>& LocalRay) const;
-	{
-		TArray<int32> IntersectingChildren;
-		if (LargeUnionData) //todo: make this work when hierarchy is not built
-		{
-			IntersectingChildren = LargeUnionData->Hierarchy.FindAllIntersections(LocalRay);
-			for (int32 i = IntersectingChildren.Num() - 1; i >= 0; --i)
-			{
-				const int32 Idx = IntersectingChildren[i];
-				if (Idx < MOriginalParticleLookupHack.Num())
-				{
-					IntersectingChildren[i] = MOriginalParticleLookupHack[Idx];
-				}
-				else
-				{
-					IntersectingChildren.RemoveAtSwap(i);
-				}
-			}
-			/*for (int32& Idx : IntersectingChildren)
-			{
-				Idx = MOriginalParticleLookupHack[Idx];
-			}*/
-		}
-		else
-		{
-			IntersectingChildren = MOriginalParticleLookupHack;
-		}
-
-		return IntersectingChildren;
-	}
-#endif
-
+	// DO NOT USE!!
+	// @todo(chaos): we should get rid of this. Instead we should hold the map/whatever on the GeometryParticle that currently owns the geom
 	const FPBDRigidParticleHandle* FindParticleForImplicitObject(const FImplicitObject* Object) const;
 
+	// DO NOT USE!!
+	// @todo(chaos): move this fucntionality to the geometry particle?
+	const FBVHParticles* GetChildSimplicial(const int32 ChildIndex) const;
+
 private:
 	// Temp hack for finding original particles
 	TArray<FPBDRigidParticleHandle*> MOriginalParticleLookupHack;	
