- Add GC Item Name access to proxy
- Store GC Item Name in query result
Fixes issue with player movement based on GC primitive component not being able to receive updated transforms

diff --git a/Engine/Source/Runtime/Engine/Private/Collision/CollisionConversions.cpp b/Engine/Source/Runtime/Engine/Private/Collision/CollisionConversions.cpp
--- a/Engine/Source/Runtime/Engine/Private/Collision/CollisionConversions.cpp
+++ b/Engine/Source/Runtime/Engine/Private/Collision/CollisionConversions.cpp
@@ -159,7 +159,9 @@ static void SetHitResultFromShapeAndFaceIndex(const FPhysicsShape& Shape,  const
 				if (ActorProxy->GetType() == EPhysicsProxyType::GeometryCollectionType)
 				{
 					const FGeometryCollectionPhysicsProxy* ConcreteProxy = static_cast<const FGeometryCollectionPhysicsProxy*>(ActorProxy);
-					OutResult.Item = ConcreteProxy->GetItemIndexFromGTParticle_External(Actor.CastToRigidParticle()).GetItemIndex();
+					const FGeometryCollectionItemIndex ItemIndex = ConcreteProxy->GetItemIndexFromGTParticle_External(Actor.CastToRigidParticle());
+					OutResult.Item = ItemIndex.GetItemIndex();
+					OutResult.BoneName = ConcreteProxy->GetTransformName_External(ItemIndex);
 				}
 			}
 		}
diff --git a/Engine/Source/Runtime/Experimental/Chaos/Private/PhysicsProxy/GeometryCollectionPhysicsProxy.cpp b/Engine/Source/Runtime/Experimental/Chaos/Private/PhysicsProxy/GeometryCollectionPhysicsProxy.cpp
--- a/Engine/Source/Runtime/Experimental/Chaos/Private/PhysicsProxy/GeometryCollectionPhysicsProxy.cpp
+++ b/Engine/Source/Runtime/Experimental/Chaos/Private/PhysicsProxy/GeometryCollectionPhysicsProxy.cpp
@@ -1644,6 +1644,21 @@ void FGeometryCollectionPhysicsProxy::GetRelevantParticleHandles(
 #endif
 }
 
+FName FGeometryCollectionPhysicsProxy::GetTransformName_External(const FGeometryCollectionItemIndex ItemIndex) const
+{
+	if (Parameters.RestCollection)
+	{
+		if (!ItemIndex.IsInternalCluster())
+		{
+			const Chaos::FPhysicsObjectHandle PhysicsObject = GetPhysicsObjectByIndex(ItemIndex.GetTransformIndex());
+			return Chaos::FPhysicsObjectInterface::GetName(PhysicsObject);
+		}
+	}
+
+	// NOTE: Internal clusters don't have names
+	return NAME_None;
+}
+
 void FGeometryCollectionPhysicsProxy::DisableParticles_External(TArray<int32>&& TransformGroupIndices)
 {
 	check(IsInGameThread());
diff --git a/Engine/Source/Runtime/Experimental/Chaos/Public/PhysicsProxy/GeometryCollectionPhysicsProxy.h b/Engine/Source/Runtime/Experimental/Chaos/Public/PhysicsProxy/GeometryCollectionPhysicsProxy.h
--- a/Engine/Source/Runtime/Experimental/Chaos/Public/PhysicsProxy/GeometryCollectionPhysicsProxy.h
+++ b/Engine/Source/Runtime/Experimental/Chaos/Public/PhysicsProxy/GeometryCollectionPhysicsProxy.h
@@ -391,6 +391,8 @@ public:
 		}
 		return FGeometryCollectionItemIndex::CreateInvalidItemIndex();
 	}
+	
+	FName GetTransformName_External(FGeometryCollectionItemIndex ItemIndex) const;
 
 	bool GetIsObjectDynamic() const { return IsObjectDynamic; }
 
