Geometry collection : move rendering and nav updates from TickCOmponent to PostPhysicsUpdate
- also cache RootIndex in Geometry collection asset to avoid scanning through all the transforms each time to find it

diff --git a/Engine/Source/Runtime/Experimental/GeometryCollectionEngine/Private/GeometryCollection/GeometryCollectionComponent.cpp b/Engine/Source/Runtime/Experimental/GeometryCollectionEngine/Private/GeometryCollection/GeometryCollectionComponent.cpp
--- a/Engine/Source/Runtime/Experimental/GeometryCollectionEngine/Private/GeometryCollection/GeometryCollectionComponent.cpp
+++ b/Engine/Source/Runtime/Experimental/GeometryCollectionEngine/Private/GeometryCollection/GeometryCollectionComponent.cpp
@@ -2324,85 +2324,22 @@ void UGeometryCollectionComponent::UpdateAttachedChildrenTransform() const
 	}
 }
 
+bool UGeometryCollectionComponent::HasVisibleGeometry() const
+{
+	return AssignedISMPool || RestCollection->HasVisibleGeometry();
+}
+
 void UGeometryCollectionComponent::TickComponent(float DeltaTime, enum ELevelTick TickType, FActorComponentTickFunction *ThisTickFunction)
 {
 	//UE_LOG(UGCC_LOG, Log, TEXT("GeometryCollectionComponent[%p]::TickComponent()"), this);
 	Super::TickComponent(DeltaTime, TickType, ThisTickFunction);
 
-#if WITH_EDITOR
-	if (IsRegistered() && SceneProxy && RestCollection)
-	{
-		const bool bWantNanite = RestCollection->EnableNanite && GGeometryCollectionNanite != 0;
-		const bool bHaveNanite = SceneProxy->IsNaniteMesh();
-		bool bRecreateProxy = bWantNanite != bHaveNanite;
-		if (bRecreateProxy)
-		{
-			// Wait until resources are released
-			FlushRenderingCommands();
-
-			FComponentReregisterContext ReregisterContext(this);
-			UpdateAllPrimitiveSceneInfosForSingleComponent(this);
-		}
-	}
-#endif
-
-	//if (bRenderStateDirty && DynamicCollection)	//todo: always send for now
-	if (RestCollection)
+	// todo(chaos) : cache root broken state ? 
+	if (IsRootBroken())
 	{
-		// In editor mode we have no DynamicCollection so this test is necessary
-		if(DynamicCollection) //, TEXT("No dynamic collection available for component %s during tick."), *GetName()))
-		{
-			if (IsRootBroken())
-			{
-				IncrementSleepTimer(DeltaTime);
-				IncrementBreakTimer(DeltaTime);
-			}
-
-			// todo(chaos) : find a way to only update that of transform have changed
-			// right now this does not work properly because the dirty flags may not be updated at the right time
-			//if (PhysicsProxy && PhysicsProxy->IsGTCollectionDirty())
-			// {
-			// 	for (const TObjectPtr<USceneComponent>& AttachedChild: this->GetAttachChildren())
-			// 	{
-			// 		AttachedChild->UpdateComponentToWorld();
-			// 	}
-			// }
-			
-			const bool bHasVisibleGeometry = AssignedISMPool || RestCollection->HasVisibleGeometry();
-			if (bHasVisibleGeometry && DynamicCollection->IsDirty())
-			{
-				// #todo review: When we've made changes to ISMC, we need to move this function call to SetRenderDynamicData_Concurrent
-				RefreshEmbeddedGeometry();
-				
-				// we may want to call this when the geometry collection updates ( notified by the proxy buffer updates ) 
-				// otherwise we are getting a frame delay 
-				RefreshISMPoolInstances();
-
-				if (SceneProxy && SceneProxy->IsNaniteMesh())
-				{
-					FNaniteGeometryCollectionSceneProxy* NaniteProxy = static_cast<FNaniteGeometryCollectionSceneProxy*>(SceneProxy);
-					NaniteProxy->FlushGPUSceneUpdate_GameThread();
-				}
-				
-				MarkRenderTransformDirty();
-				MarkRenderDynamicDataDirty();
-				bRenderStateDirty = false;
-
-				if (bUpdateNavigationInTick)
-				{
-					const UWorld* MyWorld = GetWorld();
-					if (MyWorld && MyWorld->IsGameWorld())
-					{
-						//cycle every 0xff frames
-						//@todo - Need way of seeing if the collection is actually changing
-						if (bNavigationRelevant && bRegistered && (((GFrameCounter + NavmeshInvalidationTimeSliceIndex) & 0xff) == 0))
-						{
-							UpdateNavigationData();
-						}
-					}
-				}
-			}
-		}
+		// todo(chaos) : move removal logic on the physics thread
+		IncrementSleepTimer(DeltaTime);
+		IncrementBreakTimer(DeltaTime);
 	}
 }
 
@@ -2843,6 +2780,67 @@ void UGeometryCollectionComponent::OnPostPhysicsSync()
 		// So as long as the physics sync happens, check for potential replication updates.
 		RequestUpdateRepData();
 	}
+
+	const bool bDynamicDataIsDirty = (DynamicCollection && DynamicCollection->IsDirty() && HasVisibleGeometry());
+	UpdateRenderSystemsIfNeeded(bDynamicDataIsDirty);
+	UpdateNavigationDataIfNeeded(bDynamicDataIsDirty);
+}
+
+void UGeometryCollectionComponent::UpdateRenderSystemsIfNeeded(bool bDynamicCollectionDirty)
+{
+#if WITH_EDITOR
+	if (IsRegistered() && SceneProxy && RestCollection)
+	{
+		const bool bWantNanite = RestCollection->EnableNanite && GGeometryCollectionNanite != 0;
+		const bool bHaveNanite = SceneProxy->IsNaniteMesh();
+		bool bRecreateProxy = bWantNanite != bHaveNanite;
+		if (bRecreateProxy)
+		{
+			// Wait until resources are released
+			FlushRenderingCommands();
+
+			FComponentReregisterContext ReregisterContext(this);
+			UpdateAllPrimitiveSceneInfosForSingleComponent(this);
+		}
+	}
+#endif
+
+	if (bDynamicCollectionDirty)
+	{
+		// #todo review: When we've made changes to ISMC, we need to move this function call to SetRenderDynamicData_Concurrent
+		RefreshEmbeddedGeometry();
+
+		// we may want to call this when the geometry collection updates ( notified by the proxy buffer updates ) 
+		// otherwise we are getting a frame delay 
+		RefreshISMPoolInstances();
+
+		if (SceneProxy && SceneProxy->IsNaniteMesh())
+		{
+			FNaniteGeometryCollectionSceneProxy* NaniteProxy = static_cast<FNaniteGeometryCollectionSceneProxy*>(SceneProxy);
+			NaniteProxy->FlushGPUSceneUpdate_GameThread();
+		}
+
+		MarkRenderTransformDirty();
+		MarkRenderDynamicDataDirty();
+		bRenderStateDirty = false;
+	}
+}
+
+void UGeometryCollectionComponent::UpdateNavigationDataIfNeeded(bool bDynamicCollectionDirty)
+{
+	if (bUpdateNavigationInTick && bDynamicCollectionDirty)
+	{
+		const UWorld* MyWorld = GetWorld();
+		if (MyWorld && MyWorld->IsGameWorld())
+		{
+			//cycle every 0xff frames
+			//@todo - Need way of seeing if the collection is actually changing
+			if (bNavigationRelevant && bRegistered && (((GFrameCounter + NavmeshInvalidationTimeSliceIndex) & 0xff) == 0))
+			{
+				UpdateNavigationData();
+			}
+		}
+	}
 }
 
 void UGeometryCollectionComponent::RequestUpdateRepData()
@@ -4556,7 +4554,7 @@ void UGeometryCollectionComponent::UpdateDecay(int32 TransformIdx, float Updated
 
 void UGeometryCollectionComponent::IncrementSleepTimer(float DeltaTime)
 {
-	if (DeltaTime <= 0 || !RestCollection->bRemoveOnMaxSleep || !bAllowRemovalOnSleep)
+	if (DeltaTime <= 0 || !RestCollection || !RestCollection->bRemoveOnMaxSleep || !bAllowRemovalOnSleep)
 	{
 		return;
 	}
@@ -4839,10 +4837,9 @@ int32 UGeometryCollectionComponent::GetInitialLevel(int32 ItemIndex)
 
 int32 UGeometryCollectionComponent::GetRootIndex() const
 {
-	if (RestCollection && RestCollection->GetGeometryCollection())
+	if (RestCollection)
 	{
-		Chaos::Facades::FCollectionHierarchyFacade HierarchyFacade(*RestCollection->GetGeometryCollection());
-		return HierarchyFacade.GetRootIndex();
+		return RestCollection->GetRootIndex();
 	}
 	return INDEX_NONE;
 }
diff --git a/Engine/Source/Runtime/Experimental/GeometryCollectionEngine/Private/GeometryCollection/GeometryCollectionObject.cpp b/Engine/Source/Runtime/Experimental/GeometryCollectionEngine/Private/GeometryCollection/GeometryCollectionObject.cpp
--- a/Engine/Source/Runtime/Experimental/GeometryCollectionEngine/Private/GeometryCollection/GeometryCollectionObject.cpp
+++ b/Engine/Source/Runtime/Experimental/GeometryCollectionEngine/Private/GeometryCollection/GeometryCollectionObject.cpp
@@ -39,6 +39,7 @@
 #include "GeometryCollection/GeometryCollectionSimulationCoreTypes.h"
 #include "Chaos/ChaosArchive.h"
 #include "GeometryCollectionProxyData.h"
+#include "GeometryCollection/Facades/CollectionHierarchyFacade.h"
 
 #include UE_INLINE_GENERATED_CPP_BY_NAME(GeometryCollectionObject)
 
@@ -295,6 +296,18 @@ void UGeometryCollection::ValidateSizeSpecificDataDefaults()
 	check(SizeSpecificData.Num());
 }
 
+
+// update cachedroot index using the current hierarchy setup
+void UGeometryCollection::UpdateRootIndex()
+{
+	RootIndex = INDEX_NONE;
+	if (GeometryCollection)
+	{
+		Chaos::Facades::FCollectionHierarchyFacade HierarchyFacade(*GeometryCollection);
+		RootIndex = HierarchyFacade.GetRootIndex();
+	}
+}
+
 void UGeometryCollection::UpdateGeometryDependentProperties()
 {
 #if WITH_EDITOR
@@ -978,6 +991,12 @@ void UGeometryCollection::Serialize(FArchive& Ar)
 		BoneSelectedMaterialIndex = INDEX_NONE;
 	}
 
+	// Make sure the root index is properly set 
+	if (RootIndex == INDEX_NONE)
+	{
+		UpdateRootIndex();
+	}
+
 #if WITH_EDITORONLY_DATA
 	if (bCreateSimulationData)
 	{
@@ -1261,6 +1280,7 @@ void UGeometryCollection::ReleaseResources()
 void UGeometryCollection::InvalidateCollection()
 {
 	StateGuid = FGuid::NewGuid();
+	UpdateRootIndex();
 }
 
 #if WITH_EDITOR
diff --git a/Engine/Source/Runtime/Experimental/GeometryCollectionEngine/Public/GeometryCollection/GeometryCollectionComponent.h b/Engine/Source/Runtime/Experimental/GeometryCollectionEngine/Public/GeometryCollection/GeometryCollectionComponent.h
--- a/Engine/Source/Runtime/Experimental/GeometryCollectionEngine/Public/GeometryCollection/GeometryCollectionComponent.h
+++ b/Engine/Source/Runtime/Experimental/GeometryCollectionEngine/Public/GeometryCollection/GeometryCollectionComponent.h
@@ -1255,6 +1255,9 @@ private:
 
 	void UpdateAttachedChildrenTransform() const;
 	
+	void UpdateRenderSystemsIfNeeded(bool bDynamicCollectionDirty);
+	void UpdateNavigationDataIfNeeded(bool bDynamicCollectionDirty);
+
 	void BuildInitialFilterData();
 
 	void LoadCollisionProfiles();
@@ -1263,6 +1266,8 @@ private:
 
 	void AddAutoInstancesToISMPool();
 
+	bool HasVisibleGeometry() const;
+
 	/** The clusters we need to replicate */
 	TUniquePtr<TSet<Chaos::FPBDRigidClusteredParticleHandle*>> ClustersToRep;
 
diff --git a/Engine/Source/Runtime/Experimental/GeometryCollectionEngine/Public/GeometryCollection/GeometryCollectionObject.h b/Engine/Source/Runtime/Experimental/GeometryCollectionEngine/Public/GeometryCollection/GeometryCollectionObject.h
--- a/Engine/Source/Runtime/Experimental/GeometryCollectionEngine/Public/GeometryCollection/GeometryCollectionObject.h
+++ b/Engine/Source/Runtime/Experimental/GeometryCollectionEngine/Public/GeometryCollection/GeometryCollectionObject.h
@@ -453,6 +453,9 @@ public:
 	FGuid GetIdGuid() const;
 	FGuid GetStateGuid() const;
 
+	/** Get the cached root index */
+	int32 GetRootIndex() const { return RootIndex; }
+
 	/** Pointer to the data used to render this geometry collection. */
 	TUniquePtr<FGeometryCollectionRenderData> RenderData;
 
@@ -732,6 +735,8 @@ private:
 	*/
 	void ValidateSizeSpecificDataDefaults();
 
+	// update cachedroot index using the current hierarchy setup
+	void UpdateRootIndex();
 
 private:
 	/** Guid created on construction of this collection. It should be used to uniquely identify this collection */
@@ -758,6 +763,10 @@ private:
 	FGuid RenderDataGuid;
 #endif
 
+	// cached root index for faster queries
+	UPROPERTY(VisibleAnywhere, Category = "Clustering")
+	int32 RootIndex = INDEX_NONE;
+
 	// #todo(dmp): rename to be consistent BoneSelectedMaterialID?
 	// Legacy index of the bone selected material in the object's Materials array, or INDEX_NONE if it is not stored there.
 	// Note for new objects the bone selected material should not be stored in the Materials array, so this should be INDEX_NONE
