Chaos: Removed unnecessary tick function from event dispatcher base class which could lead to accumulating no-op ticks from geometry collection components in TG_PostPhysics

diff --git a/Engine/Source/Runtime/Experimental/ChaosSolverEngine/Private/Chaos/ChaosGameplayEventDispatcher.cpp b/Engine/Source/Runtime/Experimental/ChaosSolverEngine/Private/Chaos/ChaosGameplayEventDispatcher.cpp
--- a/Engine/Source/Runtime/Experimental/ChaosSolverEngine/Private/Chaos/ChaosGameplayEventDispatcher.cpp
+++ b/Engine/Source/Runtime/Experimental/ChaosSolverEngine/Private/Chaos/ChaosGameplayEventDispatcher.cpp
@@ -50,6 +50,14 @@ FChaosRemovalEvent::FChaosRemovalEvent()
 
 }
 
+UChaosGameplayEventDispatcher::UChaosGameplayEventDispatcher()
+{
+	PrimaryComponentTick.bCanEverTick = false;
+	PrimaryComponentTick.bStartWithTickEnabled = false;
+	PrimaryComponentTick.bAllowTickOnDedicatedServer = false;
+	PrimaryComponentTick.SetTickFunctionEnable(false);
+}
+
 void UChaosGameplayEventDispatcher::OnRegister()
 {
 	Super::OnRegister();
diff --git a/Engine/Source/Runtime/Experimental/ChaosSolverEngine/Public/Chaos/ChaosGameplayEventDispatcher.h b/Engine/Source/Runtime/Experimental/ChaosSolverEngine/Public/Chaos/ChaosGameplayEventDispatcher.h
--- a/Engine/Source/Runtime/Experimental/ChaosSolverEngine/Public/Chaos/ChaosGameplayEventDispatcher.h
+++ b/Engine/Source/Runtime/Experimental/ChaosSolverEngine/Public/Chaos/ChaosGameplayEventDispatcher.h
@@ -194,6 +194,8 @@ class CHAOSSOLVERENGINE_API UChaosGameplayEventDispatcher : public UChaosEventLi
 
 public:
 
+	UChaosGameplayEventDispatcher();
+
 	virtual void OnRegister() override;
 	virtual void OnUnregister() override;
 
