Iris - Packed 64-bit integer serialization support.

diff --git a/Engine/Plugins/Runtime/ReplicationSystemTestPlugin/Source/Private/Tests/Serialization/TestPackedIntNetSerializer.cpp b/Engine/Plugins/Runtime/ReplicationSystemTestPlugin/Source/Private/Tests/Serialization/TestPackedIntNetSerializer.cpp
--- a/Engine/Plugins/Runtime/ReplicationSystemTestPlugin/Source/Private/Tests/Serialization/TestPackedIntNetSerializer.cpp
+++ b/Engine/Plugins/Runtime/ReplicationSystemTestPlugin/Source/Private/Tests/Serialization/TestPackedIntNetSerializer.cpp
@@ -30,41 +30,47 @@ protected:
 	const FPackedInt32NetSerializerConfig Config;
 };
 
-class FTestPackedInt32NetSerializer : public TTestPackedIntNetSerializer<int32>
-{
-public:
-	FTestPackedInt32NetSerializer() : TTestPackedIntNetSerializer<int32>(UE_NET_GET_SERIALIZER(FPackedInt32NetSerializer)) {}
-};
-
-UE_NET_TEST_FIXTURE(FTestPackedInt32NetSerializer, HasTestValues)
-{
-	UE_NET_ASSERT_GT_MSG(ValueCount, SIZE_T(0), "No test values found");
-}
-
-UE_NET_TEST_FIXTURE(FTestPackedInt32NetSerializer, TestIsEqual)
-{
-	TestIsEqual();
+#define UE_NET_IMPLEMENT_PACKEDINT_NETSERIALIZER_TEST(TestClassName, NetSerializerType, SourceType) \
+class TestClassName : public TTestPackedIntNetSerializer<SourceType> \
+{ \
+public: \
+	TestClassName() : TTestPackedIntNetSerializer<SourceType>(UE_NET_GET_SERIALIZER(NetSerializerType)) {} \
+}; \
+\
+UE_NET_TEST_FIXTURE(TestClassName, HasTestValues) \
+{ \
+	UE_NET_ASSERT_GT_MSG(ValueCount, SIZE_T(0), "No test values found"); \
+} \
+\
+UE_NET_TEST_FIXTURE(TestClassName, TestIsEqual) \
+{ \
+	TestIsEqual(); \
+} \
+\
+UE_NET_TEST_FIXTURE(TestClassName, TestSerialize) \
+{ \
+	TestSerialize(); \
+} \
+\
+UE_NET_TEST_FIXTURE(TestClassName, TestSerializeDelta) \
+{ \
+	TestSerializeDelta(); \
+} \
+\
+UE_NET_TEST_FIXTURE(TestClassName, TestSmallValuesArePacked) \
+{ \
+	TestSmallValuesArePacked(); \
 }
 
-UE_NET_TEST_FIXTURE(FTestPackedInt32NetSerializer, TestSerialize)
-{
-	TestSerialize();
-}
-
-UE_NET_TEST_FIXTURE(FTestPackedInt32NetSerializer, TestSerializeDelta)
-{
-	TestSerializeDelta();
-}
+UE_NET_IMPLEMENT_PACKEDINT_NETSERIALIZER_TEST(FTestPackedInt64NetSerializer, FPackedInt64NetSerializer, int64);
+UE_NET_IMPLEMENT_PACKEDINT_NETSERIALIZER_TEST(FTestPackedInt32NetSerializer, FPackedInt32NetSerializer, int32);
 
-UE_NET_TEST_FIXTURE(FTestPackedInt32NetSerializer, TestSmallValuesArePacked)
-{
-	TestSmallValuesArePacked();
-}
+#undef UE_NET_IMPLEMENT_PACKEDINT_NETSERIALIZER_TEST
 
 //
 template<typename SourceType> const SourceType TTestPackedIntNetSerializer<SourceType>::Values[] =
 {
-	TNumericLimits<SourceType>::Lowest(), TNumericLimits<SourceType>::Max(), SourceType(0), (TNumericLimits<SourceType>::Max() + TNumericLimits<SourceType>::Lowest())/SourceType(2) - SourceType(15), SourceType(2048458 % TNumericLimits<SourceType>::Max())
+	TNumericLimits<SourceType>::Lowest(), TNumericLimits<SourceType>::Max(), SourceType(0), SourceType(-1), (TNumericLimits<SourceType>::Max() + TNumericLimits<SourceType>::Lowest())/SourceType(2) - SourceType(15), SourceType(2048458 % TNumericLimits<SourceType>::Max())
 };
 template<typename SourceType> const SIZE_T TTestPackedIntNetSerializer<SourceType>::ValueCount = sizeof(Values)/sizeof(Values[0]);
 
@@ -140,4 +146,4 @@ void TTestPackedIntNetSerializer<SourceType>::TestSmallValuesArePacked()
 	UE_NET_ASSERT_LT_MSG(SmallValueBits, LargeValueBits, "Bits serialized for small value (" << SmallValue << ") should've been less than for large value (" << LargeValue << ")");
 }
 
-}
\ No newline at end of file
+}
diff --git a/Engine/Plugins/Runtime/ReplicationSystemTestPlugin/Source/Private/Tests/Serialization/TestPackedUintNetSerializer.cpp b/Engine/Plugins/Runtime/ReplicationSystemTestPlugin/Source/Private/Tests/Serialization/TestPackedUintNetSerializer.cpp
new file mode 100644
--- /dev/null
+++ b/Engine/Plugins/Runtime/ReplicationSystemTestPlugin/Source/Private/Tests/Serialization/TestPackedUintNetSerializer.cpp
@@ -0,0 +1,149 @@
+// Copyright Epic Games, Inc. All Rights Reserved.
+
+#include "TestNetSerializerFixture.h"
+#include "Iris/Serialization/PackedIntNetSerializers.h"
+
+namespace UE::Net::Private
+{
+
+static FTestMessage& PrintPackedUintNetSerializerConfig(FTestMessage& Message, const FNetSerializerConfig& InConfig)
+{
+	return Message;
+}
+
+template<typename SourceType>
+class TTestPackedUintNetSerializer : public TTestNetSerializerFixture<PrintPackedUintNetSerializerConfig, SourceType>
+{
+	typedef TTestNetSerializerFixture<PrintPackedUintNetSerializerConfig, SourceType> Super;
+
+public:
+	TTestPackedUintNetSerializer(const FNetSerializer& Serializer) : Super(Serializer) {}
+
+	void TestIsEqual();
+	void TestSerialize();
+	void TestSerializeDelta();
+	void TestSmallValuesArePacked();
+
+protected:
+	static const SourceType Values[];
+	static const SIZE_T ValueCount;
+	const FPackedUint32NetSerializerConfig Config;
+};
+
+#define UE_NET_IMPLEMENT_PACKEDUINT_NETSERIALIZER_TEST(TestClassName, NetSerializerType, SourceType) \
+class TestClassName : public TTestPackedUintNetSerializer<SourceType> \
+{ \
+public: \
+	TestClassName() : TTestPackedUintNetSerializer<SourceType>(UE_NET_GET_SERIALIZER(NetSerializerType)) {} \
+}; \
+\
+UE_NET_TEST_FIXTURE(TestClassName, HasTestValues) \
+{ \
+	UE_NET_ASSERT_GT_MSG(ValueCount, SIZE_T(0), "No test values found"); \
+} \
+\
+UE_NET_TEST_FIXTURE(TestClassName, TestIsEqual) \
+{ \
+	TestIsEqual(); \
+} \
+\
+UE_NET_TEST_FIXTURE(TestClassName, TestSerialize) \
+{ \
+	TestSerialize(); \
+} \
+\
+UE_NET_TEST_FIXTURE(TestClassName, TestSerializeDelta) \
+{ \
+	TestSerializeDelta(); \
+} \
+\
+UE_NET_TEST_FIXTURE(TestClassName, TestSmallValuesArePacked) \
+{ \
+	TestSmallValuesArePacked(); \
+}
+
+UE_NET_IMPLEMENT_PACKEDUINT_NETSERIALIZER_TEST(FTestPackedUint64NetSerializer, FPackedUint64NetSerializer, uint64);
+UE_NET_IMPLEMENT_PACKEDUINT_NETSERIALIZER_TEST(FTestPackedUint32NetSerializer, FPackedUint32NetSerializer, uint32);
+
+#undef UE_NET_IMPLEMENT_PACKEDUINT_NETSERIALIZER_TEST
+
+//
+template<typename SourceType> const SourceType TTestPackedUintNetSerializer<SourceType>::Values[] =
+{
+	TNumericLimits<SourceType>::Lowest(), TNumericLimits<SourceType>::Max(), SourceType(0), (TNumericLimits<SourceType>::Max() + TNumericLimits<SourceType>::Lowest())/SourceType(2) - SourceType(15), SourceType(2048458 % TNumericLimits<SourceType>::Max())
+};
+template<typename SourceType> const SIZE_T TTestPackedUintNetSerializer<SourceType>::ValueCount = sizeof(Values)/sizeof(Values[0]);
+
+template<typename SourceType>
+void TTestPackedUintNetSerializer<SourceType>::TestIsEqual()
+{
+	SourceType CompareValues[2][sizeof(Values)/sizeof(Values[0])];
+	bool ExpectedResults[2][sizeof(Values)/sizeof(Values[0])];
+	for (SIZE_T ValueIt = 0; ValueIt < ValueCount; ++ValueIt)
+	{
+		CompareValues[0][ValueIt] = Values[ValueIt];
+		ExpectedResults[0][ValueIt] = true;
+	}
+	for (SIZE_T ValueIt = 0; ValueIt < ValueCount; ++ValueIt)
+	{
+		CompareValues[1][ValueIt] = ~Values[ValueIt];
+		ExpectedResults[1][ValueIt] = false;
+	}
+
+	// Do two rounds of testing, one where we compare each value with itself and one where we compare against a guaranteed non-equal value.
+	// We will only do an unquantized compare since we know these serializers uses a default quantize implementation.
+	for (SIZE_T TestRoundIt = 0, TestRoundEndIt = 2; TestRoundIt != TestRoundEndIt; ++TestRoundIt)
+	{
+		constexpr bool bQuantizedCompare = false;
+		const bool bSuccess = Super::TestIsEqual(Values, CompareValues[TestRoundIt], ExpectedResults[TestRoundIt], ValueCount, Config, bQuantizedCompare);
+		if (!bSuccess)
+		{
+			return;
+		}
+	}
+}
+
+template<typename SourceType>
+void TTestPackedUintNetSerializer<SourceType>::TestSerialize()
+{
+	SourceType ExpectedValues[sizeof(Values) / sizeof(Values[0])];
+	for (SIZE_T ValueIt = 0; ValueIt < ValueCount; ++ValueIt)
+	{
+		ExpectedValues[ValueIt] = Values[ValueIt];
+	}
+
+	for (SIZE_T ValueIt = 0; ValueIt < ValueCount; ++ValueIt)
+	{
+		constexpr bool bQuantizedCompare = false;
+		const bool bSuccess = Super::TestSerialize(Values, ExpectedValues, ValueCount, Config, bQuantizedCompare);
+		if (!bSuccess)
+		{
+			return;
+		}
+	}
+}
+
+template<typename SourceType>
+void TTestPackedUintNetSerializer<SourceType>::TestSerializeDelta()
+{
+	Super::TestSerializeDelta(Values, ValueCount, Config);
+}
+
+template<typename SourceType>
+void TTestPackedUintNetSerializer<SourceType>::TestSmallValuesArePacked()
+{
+	constexpr SourceType SmallValue = SourceType(0);
+	constexpr SourceType LargeValue = TNumericLimits<SourceType>::Max();
+
+	const bool bSerializedSmallValue = this->Serialize(Config, NetSerializerValuePointer(&SmallValue));
+	UE_NET_ASSERT_TRUE_MSG(bSerializedSmallValue, "Failed serializing " << SmallValue);
+	const uint32 SmallValueBits = this->Writer.GetPosBits();
+
+	const bool bSerializedLargeValue = this->Serialize(Config, NetSerializerValuePointer(&LargeValue));
+	UE_NET_ASSERT_TRUE_MSG(bSerializedLargeValue, "Failed serializing " << LargeValue);
+	const uint32 LargeValueBits = this->Writer.GetPosBits();
+
+	UE_NET_ASSERT_LT_MSG(SmallValueBits, LargeValueBits, "Bits serialized for small value (" << SmallValue << ") should've been less than for large value (" << LargeValue << ")");
+}
+
+}
diff --git a/Engine/Source/Runtime/Experimental/Iris/Core/Private/Iris/Serialization/NetBitStreamUtil.cpp b/Engine/Source/Runtime/Experimental/Iris/Core/Private/Iris/Serialization/NetBitStreamUtil.cpp
--- a/Engine/Source/Runtime/Experimental/Iris/Core/Private/Iris/Serialization/NetBitStreamUtil.cpp
+++ b/Engine/Source/Runtime/Experimental/Iris/Core/Private/Iris/Serialization/NetBitStreamUtil.cpp
@@ -10,6 +10,85 @@
 namespace UE::Net
 {
 
+void WritePackedUint64(FNetBitStreamWriter* Writer, uint64 Value)
+{
+	// As we represent the number of bytes to write with two bits we want bits needed to be >= 1
+	const uint32 BitCountNeeded = GetBitsNeeded(Value | 1U);
+	const uint32 ByteCountNeeded = (BitCountNeeded + 7U) / 8U;
+	const uint32 BitCountToWrite = ByteCountNeeded * 8U;
+
+	Writer->WriteBits(ByteCountNeeded - 1U, 3U);
+	if (BitCountToWrite <= 32U)
+	{
+		Writer->WriteBits(Value & 0xFFFFFFFFU, BitCountToWrite);
+	}
+	else
+	{
+		Writer->WriteBits(Value & 0xFFFFFFFFU, 32U);
+		Writer->WriteBits(static_cast<uint32>(Value >> 32U), BitCountToWrite - 32U);
+	}
+}
+
+uint64 ReadPackedUint64(FNetBitStreamReader* Reader)
+{
+	const uint32 ByteCountToRead = Reader->ReadBits(3U) + 1U;
+	const uint32 BitCountToRead = ByteCountToRead * 8U;
+
+	if (BitCountToRead <= 32)
+	{
+		const uint64 Value = Reader->ReadBits(BitCountToRead);
+		return Value;
+	}
+	else
+	{
+		uint64 Value = Reader->ReadBits(32U);
+		Value |= (static_cast<uint64>(Reader->ReadBits(BitCountToRead - 32U)) << 32U);
+		return Value;
+	}
+}
+
+void WritePackedInt64(FNetBitStreamWriter* Writer, int64 Value)
+{
+	const uint32 BitCountNeeded = GetBitsNeeded(Value | 1U);
+	const uint32 ByteCountNeeded = (BitCountNeeded + 7U) / 8U;
+	const uint32 BitCountToWrite = ByteCountNeeded * 8U;
+
+	Writer->WriteBits(ByteCountNeeded - 1U, 3U);
+	if (BitCountToWrite <= 32U)
+	{
+		Writer->WriteBits(static_cast<uint64>(Value) & 0xFFFFFFFFU, BitCountToWrite);
+	}
+	else
+	{
+		const uint64 UnsignedValue = static_cast<uint64>(Value);
+		Writer->WriteBits(UnsignedValue & 0xFFFFFFFFU, 32U);
+		Writer->WriteBits(static_cast<uint32>(UnsignedValue >> 32U), BitCountToWrite - 32U);
+	}
+}
+
+int64 ReadPackedInt64(FNetBitStreamReader* Reader)
+{
+	const uint32 ByteCountToRead = Reader->ReadBits(3U) + 1U;
+	const uint32 BitCountToRead = ByteCountToRead * 8U;
+
+	uint64 UnsignedValue;
+	if (BitCountToRead <= 32U)
+	{
+		UnsignedValue = Reader->ReadBits(BitCountToRead);
+	}
+	else
+	{
+		UnsignedValue = Reader->ReadBits(32U);
+		UnsignedValue |= (static_cast<uint64>(Reader->ReadBits(BitCountToRead - 32U)) << 32U);
+	}
+
+	// Sign-extend the value
+	const uint64 Mask = 1ULL << (BitCountToRead - 1U);
+	UnsignedValue = (UnsignedValue ^ Mask) - Mask;
+	const int64 Value = static_cast<int64>(UnsignedValue);
+	return Value;
+}
+
 void WritePackedUint32(FNetBitStreamWriter* Writer, uint32 Value)
 {
 	// As we represent the number of bytes to write with two bits we want bits needed to be >= 1
diff --git a/Engine/Source/Runtime/Experimental/Iris/Core/Private/Iris/Serialization/PackedIntNetSerializers.cpp b/Engine/Source/Runtime/Experimental/Iris/Core/Private/Iris/Serialization/PackedIntNetSerializers.cpp
--- a/Engine/Source/Runtime/Experimental/Iris/Core/Private/Iris/Serialization/PackedIntNetSerializers.cpp
+++ b/Engine/Source/Runtime/Experimental/Iris/Core/Private/Iris/Serialization/PackedIntNetSerializers.cpp
@@ -9,6 +9,47 @@
 namespace UE::Net
 {
 
+struct FPackedInt64NetSerializerBase
+{
+	static constexpr SIZE_T DeltaBitCountTableEntryCount = 3;
+	// Bit counts aiming to have small value changes use few bits.
+	static constexpr uint8 DeltaBitCountTable[] = {0, 14, 32};
+};
+
+struct FPackedInt64NetSerializer : public FPackedInt64NetSerializerBase
+{
+	static const uint32 Version = 0;
+
+	typedef int64 SourceType;
+	typedef FPackedInt64NetSerializerConfig ConfigType;
+
+	inline static const ConfigType DefaultConfig;
+
+	static void Serialize(FNetSerializationContext&, const FNetSerializeArgs&);
+	static void Deserialize(FNetSerializationContext&, const FNetDeserializeArgs&);
+
+	static void SerializeDelta(FNetSerializationContext&, const FNetSerializeDeltaArgs&);
+	static void DeserializeDelta(FNetSerializationContext&, const FNetDeserializeDeltaArgs&);
+};
+UE_NET_IMPLEMENT_SERIALIZER(FPackedInt64NetSerializer);
+
+struct FPackedUint64NetSerializer : public FPackedInt64NetSerializerBase
+{
+	static const uint32 Version = 0;
+
+	typedef uint64 SourceType;
+	typedef FPackedUint64NetSerializerConfig ConfigType;
+
+	inline static const ConfigType DefaultConfig;
+
+	static void Serialize(FNetSerializationContext&, const FNetSerializeArgs&);
+	static void Deserialize(FNetSerializationContext&, const FNetDeserializeArgs&);
+
+	static void SerializeDelta(FNetSerializationContext&, const FNetSerializeDeltaArgs&);
+	static void DeserializeDelta(FNetSerializationContext&, const FNetDeserializeDeltaArgs&);
+};
+UE_NET_IMPLEMENT_SERIALIZER(FPackedUint64NetSerializer);
+
 struct FPackedInt32NetSerializerBase
 {
 	static constexpr SIZE_T DeltaBitCountTableEntryCount = 3;
@@ -23,7 +64,7 @@ struct FPackedInt32NetSerializer : public FPackedInt32NetSerializerBase
 	typedef int32 SourceType;
 	typedef FPackedInt32NetSerializerConfig ConfigType;
 
-	static const ConfigType DefaultConfig;
+	inline static const ConfigType DefaultConfig;
 
 	static void Serialize(FNetSerializationContext&, const FNetSerializeArgs&);
 	static void Deserialize(FNetSerializationContext&, const FNetDeserializeArgs&);
@@ -33,8 +74,6 @@ struct FPackedInt32NetSerializer : public FPackedInt32NetSerializerBase
 };
 UE_NET_IMPLEMENT_SERIALIZER(FPackedInt32NetSerializer);
 
-const FPackedInt32NetSerializer::ConfigType FPackedInt32NetSerializer::DefaultConfig;
-
 struct FPackedUint32NetSerializer : public FPackedInt32NetSerializerBase
 {
 	static const uint32 Version = 0;
@@ -42,7 +81,7 @@ struct FPackedUint32NetSerializer : public FPackedInt32NetSerializerBase
 	typedef uint32 SourceType;
 	typedef FPackedUint32NetSerializerConfig ConfigType;
 
-	static const ConfigType DefaultConfig;
+	inline static const ConfigType DefaultConfig;
 
 	static void Serialize(FNetSerializationContext&, const FNetSerializeArgs&);
 	static void Deserialize(FNetSerializationContext&, const FNetDeserializeArgs&);
@@ -52,8 +91,135 @@ struct FPackedUint32NetSerializer : public FPackedInt32NetSerializerBase
 };
 UE_NET_IMPLEMENT_SERIALIZER(FPackedUint32NetSerializer);
 
-const FPackedUint32NetSerializer::ConfigType FPackedUint32NetSerializer::DefaultConfig;
+// FPackedInt64NetSerializer implementation
+void FPackedInt64NetSerializer::Serialize(FNetSerializationContext& Context, const FNetSerializeArgs& Args)
+{
+	const int64 Value = *reinterpret_cast<const int64*>(Args.Source);
+
+	const uint32 BitCountNeeded = GetBitsNeeded(Value);
+	const uint32 ByteCountNeeded = (BitCountNeeded + 7U) / 8U;
+	const uint32 BitCountToWrite = ByteCountNeeded * 8U;
+
+	FNetBitStreamWriter* Writer = Context.GetBitStreamWriter();
+	Writer->WriteBits(ByteCountNeeded - 1U, 3U);
+	if (BitCountToWrite <= 32U)
+	{
+		Writer->WriteBits(static_cast<uint64>(Value) & 0xFFFFFFFFU, BitCountToWrite);
+	}
+	else
+	{
+		const uint64 UnsignedValue = static_cast<uint64>(Value);
+		Writer->WriteBits(UnsignedValue & 0xFFFFFFFFU, 32U);
+		Writer->WriteBits(static_cast<uint32>(UnsignedValue >> 32U), BitCountToWrite - 32U);
+	}
+}
+
+void FPackedInt64NetSerializer::Deserialize(FNetSerializationContext& Context, const FNetDeserializeArgs& Args)
+{
+	FNetBitStreamReader* Reader = Context.GetBitStreamReader();
+	const uint32 ByteCountToRead = Reader->ReadBits(3U) + 1U;
+	const uint32 BitCountToRead = ByteCountToRead * 8U;
+
+	uint64 UnsignedValue;
+	if (BitCountToRead <= 32U)
+	{
+		UnsignedValue = Reader->ReadBits(BitCountToRead);
+	}
+	else
+	{
+		UnsignedValue = Reader->ReadBits(32U);
+		UnsignedValue |= (static_cast<uint64>(Reader->ReadBits(BitCountToRead - 32U)) << 32U);
+	}
+
+	const uint64 Mask = 1ULL << (BitCountToRead - 1U);
+	UnsignedValue = (UnsignedValue ^ Mask) - Mask;
+	const int64 Value = static_cast<int64>(UnsignedValue);
+
+	int64* Target = reinterpret_cast<int64*>(Args.Target);
+	*Target = Value;
+}
+
+void FPackedInt64NetSerializer::SerializeDelta(FNetSerializationContext& Context, const FNetSerializeDeltaArgs& Args)
+{
+	const int64 Value = *reinterpret_cast<const int64*>(Args.Source);
+	const int64 PrevValue = *reinterpret_cast<const int64*>(Args.Prev);
+
+	FNetBitStreamWriter* Writer = Context.GetBitStreamWriter();
+	SerializeIntDelta(*Writer, Value, PrevValue, DeltaBitCountTable, DeltaBitCountTableEntryCount, 64U);
+}
+
+void FPackedInt64NetSerializer::DeserializeDelta(FNetSerializationContext& Context, const FNetDeserializeDeltaArgs& Args)
+{
+	int64& Target = *reinterpret_cast<int64*>(Args.Target);
+	const int64 PrevValue = *reinterpret_cast<const int64*>(Args.Prev);
+
+	FNetBitStreamReader* Reader = Context.GetBitStreamReader();
+	DeserializeIntDelta(*Reader, Target, PrevValue, DeltaBitCountTable, DeltaBitCountTableEntryCount, 64U);
+}
+
+// FPackedUint64NetSerializer implementation
+void FPackedUint64NetSerializer::Serialize(FNetSerializationContext& Context, const FNetSerializeArgs& Args)
+{
+	const uint64 Value = *reinterpret_cast<const uint64*>(Args.Source);
+
+	// As we represent the number of bytes to write with two bits we want bits needed to be >= 1
+	const uint32 BitCountNeeded = GetBitsNeeded(Value | 1U);
+	const uint32 ByteCountNeeded = (BitCountNeeded + 7U) / 8U;
+	const uint32 BitCountToWrite = ByteCountNeeded * 8U;
+
+	FNetBitStreamWriter* Writer = Context.GetBitStreamWriter();
+	Writer->WriteBits(ByteCountNeeded - 1U, 3U);
+	if (BitCountToWrite <= 32U)
+	{
+		Writer->WriteBits(Value & 0xFFFFFFFFU, BitCountToWrite);
+	}
+	else
+	{
+		Writer->WriteBits(Value & 0xFFFFFFFFU, 32U);
+		Writer->WriteBits(static_cast<uint32>(Value >> 32U), BitCountToWrite - 32U);
+	}
+}
+
+void FPackedUint64NetSerializer::Deserialize(FNetSerializationContext& Context, const FNetDeserializeArgs& Args)
+{
+	FNetBitStreamReader* Reader = Context.GetBitStreamReader();
+	const uint32 ByteCountToRead = Reader->ReadBits(3U) + 1U;
+	const uint32 BitCountToRead = ByteCountToRead * 8U;
+
+	uint64 Value;
+	if (BitCountToRead <= 32)
+	{
+		Value = Reader->ReadBits(BitCountToRead);
+	}
+	else
+	{
+		Value = Reader->ReadBits(32U);
+		Value |= (static_cast<uint64>(Reader->ReadBits(BitCountToRead - 32U)) << 32U);
+	}
+
+	uint64* Target = reinterpret_cast<uint64*>(Args.Target);
+	*Target = Value;
+}
+
+void FPackedUint64NetSerializer::SerializeDelta(FNetSerializationContext& Context, const FNetSerializeDeltaArgs& Args)
+{
+	const uint64 Value = *reinterpret_cast<const uint64*>(Args.Source);
+	const uint64 PrevValue = *reinterpret_cast<const uint64*>(Args.Prev);
+
+	FNetBitStreamWriter* Writer = Context.GetBitStreamWriter();
+	SerializeUintDelta(*Writer, Value, PrevValue, DeltaBitCountTable, DeltaBitCountTableEntryCount, 64U);
+}
+
+void FPackedUint64NetSerializer::DeserializeDelta(FNetSerializationContext& Context, const FNetDeserializeDeltaArgs& Args)
+{
+	uint64& Target = *reinterpret_cast<uint64*>(Args.Target);
+	const uint64 PrevValue = *reinterpret_cast<const uint64*>(Args.Prev);
+
+	FNetBitStreamReader* Reader = Context.GetBitStreamReader();
+	DeserializeUintDelta(*Reader, Target, PrevValue, DeltaBitCountTable, DeltaBitCountTableEntryCount, 64U);
+}
 
+// FPackedInt32NetSerializer implementation
 void FPackedInt32NetSerializer::Serialize(FNetSerializationContext& Context, const FNetSerializeArgs& Args)
 {
 	const int32 Value = *reinterpret_cast<const int32*>(Args.Source);
@@ -100,7 +266,7 @@ void FPackedInt32NetSerializer::DeserializeDelta(FNetSerializationContext& Conte
 	DeserializeIntDelta(*Reader, Target, PrevValue, DeltaBitCountTable, DeltaBitCountTableEntryCount, 32U);
 }
 
-// FPackedUint32NetSerializer
+// FPackedUint32NetSerializer implementation
 void FPackedUint32NetSerializer::Serialize(FNetSerializationContext& Context, const FNetSerializeArgs& Args)
 {
 	const uint32 Value = *reinterpret_cast<const uint32*>(Args.Source);
diff --git a/Engine/Source/Runtime/Experimental/Iris/Core/Public/Iris/Serialization/NetBitStreamUtil.h b/Engine/Source/Runtime/Experimental/Iris/Core/Public/Iris/Serialization/NetBitStreamUtil.h
--- a/Engine/Source/Runtime/Experimental/Iris/Core/Public/Iris/Serialization/NetBitStreamUtil.h
+++ b/Engine/Source/Runtime/Experimental/Iris/Core/Public/Iris/Serialization/NetBitStreamUtil.h
@@ -109,6 +109,18 @@ inline int64 ReadInt64(FNetBitStreamReader* Reader)
 	return static_cast<int64>(Value);
 }
 
+/** Write a uint64 using as few bytes as possible */
+IRISCORE_API void WritePackedUint64(FNetBitStreamWriter* Writer, uint64 Value);
+
+/** Read a uint64 that was written using WritePackedUint64 */
+IRISCORE_API uint64 ReadPackedUint64(FNetBitStreamReader* Reader);
+
+/** Write an int64 using as few bytes as possible */
+IRISCORE_API void WritePackedInt64(FNetBitStreamWriter* Writer, int64 Value);
+
+/** Read an int64 that was written using WritePackedInt64 */
+IRISCORE_API int64 ReadPackedInt64(FNetBitStreamReader* Reader);
+
 /** Write a uint32 using as few bytes as possible */
 IRISCORE_API void WritePackedUint32(FNetBitStreamWriter* Writer, uint32 Value);
 
diff --git a/Engine/Source/Runtime/Experimental/Iris/Core/Public/Iris/Serialization/PackedIntNetSerializers.h b/Engine/Source/Runtime/Experimental/Iris/Core/Public/Iris/Serialization/PackedIntNetSerializers.h
--- a/Engine/Source/Runtime/Experimental/Iris/Core/Public/Iris/Serialization/PackedIntNetSerializers.h
+++ b/Engine/Source/Runtime/Experimental/Iris/Core/Public/Iris/Serialization/PackedIntNetSerializers.h
@@ -5,6 +5,18 @@
 #include "NetSerializer.h"
 #include "PackedIntNetSerializers.generated.h"
 
+USTRUCT()
+struct FPackedInt64NetSerializerConfig : public FNetSerializerConfig
+{
+	GENERATED_BODY()
+};
+
+USTRUCT()
+struct FPackedUint64NetSerializerConfig : public FNetSerializerConfig
+{
+	GENERATED_BODY()
+};
+
 USTRUCT()
 struct FPackedInt32NetSerializerConfig : public FNetSerializerConfig
 {
@@ -20,6 +32,8 @@ struct FPackedUint32NetSerializerConfig : public FNetSerializerConfig
 namespace UE::Net
 {
 
+UE_NET_DECLARE_SERIALIZER(FPackedInt64NetSerializer, IRISCORE_API);
+UE_NET_DECLARE_SERIALIZER(FPackedUint64NetSerializer, IRISCORE_API);
 UE_NET_DECLARE_SERIALIZER(FPackedInt32NetSerializer, IRISCORE_API);
 UE_NET_DECLARE_SERIALIZER(FPackedUint32NetSerializer, IRISCORE_API);
 
