Geometry collection : Only enable primary tick function when removal features are need ( after the root breaks )
- Set the Tick function to be disabled at start
- Disabled the auto activate that was automatically enabling the primary tick function
- Enable the tick function post physics update after the root is broken

diff --git a/Engine/Source/Runtime/Experimental/GeometryCollectionEngine/Private/GeometryCollection/GeometryCollectionComponent.cpp b/Engine/Source/Runtime/Experimental/GeometryCollectionEngine/Private/GeometryCollection/GeometryCollectionComponent.cpp
--- a/Engine/Source/Runtime/Experimental/GeometryCollectionEngine/Private/GeometryCollection/GeometryCollectionComponent.cpp
+++ b/Engine/Source/Runtime/Experimental/GeometryCollectionEngine/Private/GeometryCollection/GeometryCollectionComponent.cpp
@@ -364,9 +364,11 @@ UGeometryCollectionComponent::UGeometryCollectionComponent(const FObjectInitiali
 #endif  // #if GEOMETRYCOLLECTION_EDITOR_SELECTION
 	, bIsMoving(false)
 {
+	// by default tick is registered but disabled, we only need it when we need to update the removal timers
+	// tick will be then enabled only when the root is broken from OnPostPhysicsSync callback
 	PrimaryComponentTick.bCanEverTick = true;
+	PrimaryComponentTick.bStartWithTickEnabled = false;
 	bTickInEditor = true;
-	bAutoActivate = true;
 
 	static uint32 GlobalNavMeshInvalidationCounter = 0;
 	//space these out over several frames (3 is arbitrary)
@@ -2781,6 +2783,16 @@ void UGeometryCollectionComponent::OnPostPhysicsSync()
 		RequestUpdateRepData();
 	}
 
+	// Onc ethe GC is broken, removal feature will need the tick to properly update the timers 
+	// even if the physics does not get any updates
+	if (IsRootBroken())
+	{
+		if (!PrimaryComponentTick.IsTickFunctionEnabled())
+		{ 
+			PrimaryComponentTick.SetTickFunctionEnable(true);
+		}
+	}
+
 	const bool bDynamicDataIsDirty = (DynamicCollection && DynamicCollection->IsDirty() && HasVisibleGeometry());
 	UpdateRenderSystemsIfNeeded(bDynamicDataIsDirty);
 	UpdateNavigationDataIfNeeded(bDynamicDataIsDirty);
@@ -4322,6 +4334,7 @@ void UGeometryCollectionComponent::UnregisterFromISMPool()
 			ISMPoolComp->DestroyMeshGroup(ISMPoolMeshGroupIndex);
 			ISMPoolMeshGroupIndex = INDEX_NONE;
 			ISMPoolRootProxyMeshIds.Empty();
+			ISMPoolAutoInstancesMeshIds.Empty();
 		}
 		SetVisibility(true);
 	}
