Use AddInstances() from ISM Pool.
This gives a single HISM rebuild event instead of one for each instance added.

diff --git a/Engine/Source/Runtime/Experimental/GeometryCollectionEngine/Private/GeometryCollection/GeometryCollectionISMPoolComponent.cpp b/Engine/Source/Runtime/Experimental/GeometryCollectionEngine/Private/GeometryCollection/GeometryCollectionISMPoolComponent.cpp
--- a/Engine/Source/Runtime/Experimental/GeometryCollectionEngine/Private/GeometryCollection/GeometryCollectionISMPoolComponent.cpp
+++ b/Engine/Source/Runtime/Experimental/GeometryCollectionEngine/Private/GeometryCollection/GeometryCollectionISMPoolComponent.cpp
@@ -76,15 +76,16 @@ FGeometryCollectionISM::FGeometryCollectionISM(AActor* OwmingActor, const FGeome
 int32 FGeometryCollectionISM::AddInstanceGroup(int32 InstanceCount)
 {
 	const int32 InstanceGroupIndex = InstanceGroups.AddGroup(InstanceCount);
-	const FInstanceGroups::FInstanceGroupRange& NewInstanceGroup = InstanceGroups.GetGroup(InstanceGroupIndex);
-	const int32 TotalInstanceCount = NewInstanceGroup.Start + NewInstanceGroup.Count;
+
 	ISMComponent->PreAllocateInstancesMemory(InstanceCount);
+
 	FTransform ZeroScaleTransform;
 	ZeroScaleTransform.SetIdentityZeroScale();
-	for (int32 InstanceIndex = NewInstanceGroup.Start; InstanceIndex < TotalInstanceCount; InstanceIndex++)
-	{
-		ISMComponent->AddInstance(ZeroScaleTransform, true);
-	}
+	TArray<FTransform> ZeroScaleTransforms;
+	ZeroScaleTransforms.Init(ZeroScaleTransform, InstanceCount);
+
+	ISMComponent->AddInstances(ZeroScaleTransforms, false, true);
+
 	return InstanceGroupIndex;
 }
 
