Add MaxSimulatedLevel in Geometry Collection that allows not to create particles from the input level and the beyond.

diff --git a/Engine/Source/Programs/HeadlessChaos/Private/GeometryCollection/GeometryCollectionTestFramework.cpp b/Engine/Source/Programs/HeadlessChaos/Private/GeometryCollection/GeometryCollectionTestFramework.cpp
--- a/Engine/Source/Programs/HeadlessChaos/Private/GeometryCollection/GeometryCollectionTestFramework.cpp
+++ b/Engine/Source/Programs/HeadlessChaos/Private/GeometryCollection/GeometryCollectionTestFramework.cpp
@@ -169,6 +169,7 @@ namespace GeometryCollectionTest
 			SimulationParams.InitialVelocityType = Params.InitialVelocityType;
 			SimulationParams.DamageThreshold = Params.DamageThreshold;
 			SimulationParams.MaxClusterLevel = Params.MaxClusterLevel;
+			SimulationParams.MaxSimulatedLevel = Params.MaxSimulatedLevel;
 			SimulationParams.ClusterConnectionMethod = Params.ClusterConnectionMethod;
 			SimulationParams.CollisionGroup = Params.CollisionGroup;
 			SimulationParams.ClusterGroupIndex = Params.ClusterGroupIndex;
diff --git a/Engine/Source/Programs/HeadlessChaos/Public/GeometryCollection/GeometryCollectionTestFramework.h b/Engine/Source/Programs/HeadlessChaos/Public/GeometryCollection/GeometryCollectionTestFramework.h
--- a/Engine/Source/Programs/HeadlessChaos/Public/GeometryCollection/GeometryCollectionTestFramework.h
+++ b/Engine/Source/Programs/HeadlessChaos/Public/GeometryCollection/GeometryCollectionTestFramework.h
@@ -91,6 +91,7 @@ namespace GeometryCollectionTest
 		FTransform GeomTransform = FTransform::Identity;
 		TSharedPtr<FGeometryCollection> RestCollection = nullptr;
 		int32 MaxClusterLevel = 100;
+		int32 MaxSimulatedLevel = 100;
 		TArray<float> DamageThreshold = { 1000.0 };
 		Chaos::FClusterCreationParameters::EConnectionMethod ClusterConnectionMethod = Chaos::FClusterCreationParameters::EConnectionMethod::PointImplicit;
 		int32 CollisionGroup = 0;
diff --git a/Engine/Source/Runtime/Engine/Private/PhysicsEngine/Experimental/PhysScene_Chaos.cpp b/Engine/Source/Runtime/Engine/Private/PhysicsEngine/Experimental/PhysScene_Chaos.cpp
--- a/Engine/Source/Runtime/Engine/Private/PhysicsEngine/Experimental/PhysScene_Chaos.cpp
+++ b/Engine/Source/Runtime/Engine/Private/PhysicsEngine/Experimental/PhysScene_Chaos.cpp
@@ -695,10 +695,13 @@ void FPhysScene_Chaos::RemoveObject(FGeometryCollectionPhysicsProxy* InObject)
 	for (TUniquePtr<Chaos::TPBDRigidParticle<Chaos::FReal, 3>>& GTParticleUnique : InObject->GetExternalParticles())
 	{
 		Chaos::TPBDRigidParticle<Chaos::FReal, 3>* GTParticle = GTParticleUnique.Get();
-		RemoveActorFromAccelerationStructureImp(GTParticle);
-		if(Solver)
+		if (GTParticle != nullptr)
 		{
-			Solver->UpdateParticleInAccelerationStructure_External(GTParticle, /*bDelete=*/true); // ensures deletion will be applied to structures passed from physics thread.
+			RemoveActorFromAccelerationStructureImp(GTParticle);
+			if (Solver)
+			{
+				Solver->UpdateParticleInAccelerationStructure_External(GTParticle, /*bDelete=*/true); // ensures deletion will be applied to structures passed from physics thread.
+			}
 		}
 	}
 
diff --git a/Engine/Source/Runtime/Experimental/Chaos/Private/PhysicsProxy/GeometryCollectionPhysicsProxy.cpp b/Engine/Source/Runtime/Experimental/Chaos/Private/PhysicsProxy/GeometryCollectionPhysicsProxy.cpp
--- a/Engine/Source/Runtime/Experimental/Chaos/Private/PhysicsProxy/GeometryCollectionPhysicsProxy.cpp
+++ b/Engine/Source/Runtime/Experimental/Chaos/Private/PhysicsProxy/GeometryCollectionPhysicsProxy.cpp
@@ -391,6 +391,7 @@ FGeometryCollectionPhysicsProxy::FGeometryCollectionPhysicsProxy(
 	: Base(InOwner)
 	, Parameters(SimulationParameters)
 	, NumParticles(INDEX_NONE)
+	, NumEffectiveParticles(INDEX_NONE)
 	, BaseParticleIndex(INDEX_NONE)
 	, IsObjectDynamic(false)
 	, IsObjectLoading(true)
@@ -445,11 +446,21 @@ void FGeometryCollectionPhysicsProxy::Initialize(Chaos::FPBDRigidsEvolutionBase
 	SolverClusterHandles.Init(nullptr, NumParticles);
 	SolverParticleHandles.Init(nullptr, NumParticles);
 	
+	TBitArray<> EffectiveParticles;
+	NumEffectiveParticles = CalculateEffectiveParticles(DynamicCollection, NumParticles, EffectiveParticles);
+
 	PhysicsObjects.Empty();
 	PhysicsObjects.Reserve(NumParticles);
 	for (int32 Index = 0; Index < NumParticles; ++Index)
 	{
-		PhysicsObjects.Emplace(Chaos::FPhysicsObjectFactory::CreatePhysicsObject(this, Index, FName(Parameters.RestCollection->BoneName[Index])));
+		if (EffectiveParticles[Index])
+		{
+			PhysicsObjects.Emplace(Chaos::FPhysicsObjectFactory::CreatePhysicsObject(this, Index, FName(Parameters.RestCollection->BoneName[Index])));
+		}
+		else
+		{
+			PhysicsObjects.Emplace(nullptr);
+		}
 	}
 
 	// compatibility requirement to make sure we at least initialize GameThreadPerFrameData properly
@@ -495,50 +506,53 @@ void FGeometryCollectionPhysicsProxy::Initialize(Chaos::FPBDRigidsEvolutionBase
 	{
 		for(int32 Index = 0; Index < NumTransforms; ++Index)
 		{
-			GTParticles[Index] = FParticle::CreateParticle();
-			FParticle* P = GTParticles[Index].Get();
-			GTParticlesToTransformGroupIndex.Add(P, Index);
-
-			GTParticles[Index]->SetUniqueIdx(Evolution->GenerateUniqueIdx());
-
-			// Note that this transform must match the physics thread transform computation for initialization.
-			// Take for example, a geometry collection in the editor that is linked with a joint constraint. The joint
-			// constraint will query the game thread particle position/rotation for the geometry collection to compute its
-			// reference frame. If that position/rotation does not match up with the physics thread's position/rotation,
-			// the geometry collection particle will have an added velocity computed by the joint constraint solver.
-			const FTransform& T = GameThreadCollection.MassToLocal[Index] * GameThreadCollection.Transform[Index] * Parameters.WorldTransform;
-			P->SetX(T.GetTranslation(), false);
-			P->SetR(T.GetRotation(), false);
-			P->SetM(Mass[Index] * MassScale);
-			P->SetUserData(Parameters.UserData);
-			P->SetProxy(this);
-
-			FGeometryDynamicCollection::FSharedImplicit ImplicitGeometry = GameThreadCollection.Implicits[Index];
-			if (ImplicitGeometry && !Scale.Equals(FVector::OneVector))
+			if (EffectiveParticles[Index])
 			{
-				TUniquePtr<Chaos::FImplicitObject> ScaledImplicit = ImplicitGeometry->CopyWithScale(Scale);
-				ImplicitGeometry = FGeometryDynamicCollection::FSharedImplicit(ScaledImplicit.Release());
-			}
-			P->SetGeometry(ImplicitGeometry);
+				GTParticles[Index] = FParticle::CreateParticle();
+				FParticle* P = GTParticles[Index].Get();
+				GTParticlesToTransformGroupIndex.Add(P, Index);
+
+				GTParticles[Index]->SetUniqueIdx(Evolution->GenerateUniqueIdx());
+
+				// Note that this transform must match the physics thread transform computation for initialization.
+				// Take for example, a geometry collection in the editor that is linked with a joint constraint. The joint
+				// constraint will query the game thread particle position/rotation for the geometry collection to compute its
+				// reference frame. If that position/rotation does not match up with the physics thread's position/rotation,
+				// the geometry collection particle will have an added velocity computed by the joint constraint solver.
+				const FTransform& T = GameThreadCollection.MassToLocal[Index] * GameThreadCollection.Transform[Index] * Parameters.WorldTransform;
+				P->SetX(T.GetTranslation(), false);
+				P->SetR(T.GetRotation(), false);
+				P->SetM(Mass[Index] * MassScale);
+				P->SetUserData(Parameters.UserData);
+				P->SetProxy(this);
+
+				FGeometryDynamicCollection::FSharedImplicit ImplicitGeometry = GameThreadCollection.Implicits[Index];
+				if (ImplicitGeometry && !Scale.Equals(FVector::OneVector))
+				{
+					TUniquePtr<Chaos::FImplicitObject> ScaledImplicit = ImplicitGeometry->CopyWithScale(Scale);
+					ImplicitGeometry = FGeometryDynamicCollection::FSharedImplicit(ScaledImplicit.Release());
+				}
+				P->SetGeometry(ImplicitGeometry);
 
-			// this step is necessary for Phase 2 where we need to walk back the hierarchy from children to parent 
-			if (bGeometryCollectionAlwaysGenerateGTCollisionForClusters && GameThreadCollection.Children[Index].Num() == 0)
-			{
-				ChildrenToCheckForParentFix.Add(Index);
-			}
-			
-			// IMPORTANT: we need to set the right spatial index because GT particle is static and PT particle is rigid
-			// this is causing a mismatch when using the separate acceleration structures optimization which can cause crashes when destroying the particle while async tracing 
-			// todo(chaos) we should eventually refactor this code to use rigid particles on the GT side for geometry collection  
-			P->SetSpatialIdx(Chaos::FSpatialAccelerationIdx{ 0,1 });
 
-			if (Chaos::AccelerationStructureSplitStaticAndDynamic == 1)
-			{
+				// IMPORTANT: we need to set the right spatial index because GT particle is static and PT particle is rigid
+				// this is causing a mismatch when using the separate acceleration structures optimization which can cause crashes when destroying the particle while async tracing 
+				// todo(chaos) we should eventually refactor this code to use rigid particles on the GT side for geometry collection  
 				P->SetSpatialIdx(Chaos::FSpatialAccelerationIdx{ 0,1 });
+
+				if (Chaos::AccelerationStructureSplitStaticAndDynamic == 1)
+				{
+					P->SetSpatialIdx(Chaos::FSpatialAccelerationIdx{ 0,1 });
+				}
+				else
+				{
+					P->SetSpatialIdx(Chaos::FSpatialAccelerationIdx{ 0,0 });
+				}
 			}
-			else
+			// this step is necessary for Phase 2 where we need to walk back the hierarchy from children to parent 
+			if (bGeometryCollectionAlwaysGenerateGTCollisionForClusters && GameThreadCollection.Children[Index].Num() == 0)
 			{
-				P->SetSpatialIdx(Chaos::FSpatialAccelerationIdx{ 0,0 });
+				ChildrenToCheckForParentFix.Add(Index);
 			}
 		}
 
@@ -610,7 +624,10 @@ void FGeometryCollectionPhysicsProxy::Initialize(Chaos::FPBDRigidsEvolutionBase
 								Chaos::FImplicitObject* UnionImplicit = new Chaos::FImplicitObjectUnion(MoveTemp(ChildImplicits));
 								GameThreadCollection.Implicits[ParentToFixIndex] = FGeometryDynamicCollection::FSharedImplicit(UnionImplicit);
 							}
-							GTParticles[ParentToFixIndex]->SetGeometry(GameThreadCollection.Implicits[ParentToFixIndex]);
+							if (GTParticles[ParentToFixIndex] != nullptr)
+							{
+								GTParticles[ParentToFixIndex]->SetGeometry(GameThreadCollection.Implicits[ParentToFixIndex]);
+							}
 						}
 					}
 				}
@@ -625,15 +642,18 @@ void FGeometryCollectionPhysicsProxy::Initialize(Chaos::FPBDRigidsEvolutionBase
 		for(int32 Index = 0; Index < NumTransforms; ++Index)
 		{
 			FParticle* P = GTParticles[Index].Get();
-			const Chaos::FShapesArray& Shapes = P->ShapesArray();
-			const int32 NumShapes = Shapes.Num();
-			for(int32 ShapeIndex = 0; ShapeIndex < NumShapes; ++ShapeIndex)
+			if (P != nullptr)
 			{
-				Chaos::FPerShapeData* Shape = Shapes[ShapeIndex].Get();
-				Shape->SetSimData(SimFilter);
-				Shape->SetQueryData(QueryFilter);
-				Shape->SetProxy(this);
-				Shape->SetMaterial(Parameters.PhysicalMaterialHandle);
+				const Chaos::FShapesArray& Shapes = P->ShapesArray();
+				const int32 NumShapes = Shapes.Num();
+				for (int32 ShapeIndex = 0; ShapeIndex < NumShapes; ++ShapeIndex)
+				{
+					Chaos::FPerShapeData* Shape = Shapes[ShapeIndex].Get();
+					Shape->SetSimData(SimFilter);
+					Shape->SetQueryData(QueryFilter);
+					Shape->SetProxy(this);
+					Shape->SetMaterial(Parameters.PhysicalMaterialHandle);
+				}
 			}
 		}
 	}
@@ -688,7 +708,10 @@ void FGeometryCollectionPhysicsProxy::Initialize(Chaos::FPBDRigidsEvolutionBase
 
 			if (FParticle* P = GTParticles[TransformIndex].Get())
 			{
-				P->SetDisabled(!bIsRoot);
+				if (P != nullptr)
+				{
+					P->SetDisabled(!bIsRoot);
+				}
 			}
 		}
 	}
@@ -779,7 +802,7 @@ void FGeometryCollectionPhysicsProxy::InitializeDynamicCollection(FGeometryDynam
 int32 ReportTooManyChildrenNum = -1;
 FAutoConsoleVariableRef CVarReportTooManyChildrenNum(TEXT("p.ReportTooManyChildrenNum"), ReportTooManyChildrenNum, TEXT("Issue warning if more than this many children exist in a single cluster"));
 
-void FGeometryCollectionPhysicsProxy::CreateNonClusteredParticles(Chaos::FPBDRigidsSolver* RigidsSolver, const FGeometryCollection& RestCollection, const FGeometryDynamicCollection& DynamicCollection)
+void FGeometryCollectionPhysicsProxy::CreateNonClusteredParticles(Chaos::FPBDRigidsSolver* RigidsSolver, const FGeometryCollection& RestCollection, const FGeometryDynamicCollection& DynamicCollection, const TBitArray<>& EffectiveParticles)
 {
 	const TManagedArray<bool>& SimulatableParticles = DynamicCollection.SimulatableParticles;
 	
@@ -795,11 +818,14 @@ void FGeometryCollectionPhysicsProxy::CreateNonClusteredParticles(Chaos::FPBDRig
 	int NumSimulatedParticles = 0;
 	for (int32 Idx = 0; Idx < SimulatableParticles.Num(); ++Idx)
 	{
-		NumSimulatedParticles += SimulatableParticles[Idx];
-		if (SimulatableParticles[Idx] && !RestCollection.IsClustered(Idx) && RestCollection.IsGeometry(Idx))
+		if (EffectiveParticles[Idx])
 		{
-			NumRigids++;
-			UniqueIndices.Add(GTParticles[Idx]->UniqueIdx());
+			NumSimulatedParticles += SimulatableParticles[Idx];
+			if (SimulatableParticles[Idx] && !RestCollection.IsClustered(Idx) && RestCollection.IsGeometry(Idx))
+			{
+				NumRigids++;
+				UniqueIndices.Add(GTParticles[Idx]->UniqueIdx());
+			}
 		}
 	}
 
@@ -811,7 +837,7 @@ void FGeometryCollectionPhysicsProxy::CreateNonClusteredParticles(Chaos::FPBDRig
 	for (int32 Idx = 0; Idx < SimulatableParticles.Num(); ++Idx)
 	{
 		SolverParticleHandles[Idx] = nullptr;
-		if (SimulatableParticles[Idx] && !RestCollection.IsClustered(Idx))
+		if (EffectiveParticles[Idx] && SimulatableParticles[Idx] && !RestCollection.IsClustered(Idx))
 		{
 			// todo: Unblocked read access of game thread data on the physics thread.
 
@@ -891,7 +917,6 @@ void FGeometryCollectionPhysicsProxy::InitializeBodiesPT(Chaos::FPBDRigidsSolver
 
 	if (Parameters.Simulating && RestCollection)
 	{
-		const TManagedArray<int32>& TransformIndex = RestCollection->TransformIndex;
 		const TManagedArray<int32>& BoneMap = RestCollection->BoneMap;
 		const TManagedArray<int32>& SimulationType = RestCollection->SimulationType;
 		const TManagedArray<FVector3f>& Vertex = RestCollection->Vertex;
@@ -921,7 +946,11 @@ void FGeometryCollectionPhysicsProxy::InitializeBodiesPT(Chaos::FPBDRigidsSolver
 		TArray<FTransform> Transform;
 		GeometryCollectionAlgo::GlobalMatrices(DynamicCollection.Transform, Parent, Transform);
 
-		CreateNonClusteredParticles(RigidsSolver, *RestCollection, DynamicCollection);
+		// Here Clean up Additional particles
+		TBitArray<> EffectiveParticles;
+		NumEffectiveParticles = CalculateEffectiveParticles(DynamicCollection, DynamicCollection.NumElements(FGeometryCollection::TransformGroup), EffectiveParticles);
+
+		CreateNonClusteredParticles(RigidsSolver, *RestCollection, DynamicCollection, EffectiveParticles);
 
 		const float StrainDefault = Parameters.DamageThreshold.Num() ? Parameters.DamageThreshold[0] : 0;
 		// Add the rigid bodies
@@ -930,7 +959,21 @@ void FGeometryCollectionPhysicsProxy::InitializeBodiesPT(Chaos::FPBDRigidsSolver
 		const FVector::FReal MassScale = WorldScale.X * WorldScale.Y * WorldScale.Z;
 
 		const Chaos::Facades::FCollectionAnchoringFacade AnchoringFacade(PhysicsThreadCollection);
-		
+
+
+		TArray<float> Masses;
+		TArray<Chaos::FVec3f> Inertias;
+		TArray < FTransform> Transforms;
+
+		for (int32 Index = 0; Index < NumTransforms; Index++)
+		{
+			const Chaos::FVec3f ScaledInertia = Chaos::Utilities::ScaleInertia<float>((Chaos::FVec3f)InertiaTensor[Index], (Chaos::FVec3f)(WorldScale), true);
+			Masses.Add(Mass[Index] * float(MassScale));
+			Inertias.Add(ScaledInertia);
+			const FTransform WorldTransform = MassToLocal[Index] * Transform[Index] * Parameters.WorldTransform;
+			Transforms.Add(WorldTransform);
+		}
+
 		// Iterating over the geometry group is a fast way of skipping everything that's
 		// not a leaf node, as each geometry has a transform index, which is a shortcut
 		// for the case when there's a 1-to-1 mapping between transforms and geometries.
@@ -941,11 +984,8 @@ void FGeometryCollectionPhysicsProxy::InitializeBodiesPT(Chaos::FPBDRigidsSolver
 			if (FClusterHandle* Handle = SolverParticleHandles[TransformGroupIndex])
 			{
 				// Mass space -> Composed parent space -> world
-				const FTransform WorldTransform = 
-					MassToLocal[TransformGroupIndex] * Transform[TransformGroupIndex] * Parameters.WorldTransform;
-
 				const Chaos::FVec3f ScaledInertia = Chaos::Utilities::ScaleInertia<float>((Chaos::FVec3f)InertiaTensor[TransformGroupIndex], (Chaos::FVec3f)(WorldScale), true);
-				
+
 				PopulateSimulatedParticle(
 					Handle,
 					Parameters.Shared,
@@ -955,38 +995,38 @@ void FGeometryCollectionPhysicsProxy::InitializeBodiesPT(Chaos::FPBDRigidsSolver
 					QueryFilter,
 					Mass[TransformGroupIndex] * MassScale,
 					ScaledInertia,
-					WorldTransform,
+					Transforms[TransformGroupIndex],
 					static_cast<uint8>(DynamicState[TransformGroupIndex]),
 					static_cast<int16>(CollisionGroup[TransformGroupIndex]),
 					CollisionParticlesPerObjectFraction);
 
 				// initialize anchoring information if available 
 				Handle->SetIsAnchored(AnchoringFacade.IsAnchored(TransformGroupIndex));
-				
+
 				if (Parameters.EnableClustering)
 				{
 					Handle->SetClusterGroupIndex(Parameters.ClusterGroupIndex);
 
-					float DamageThreshold = StrainDefault; 
+					float DamageThreshold = StrainDefault;
 					if (!Parameters.bUsePerClusterOnlyDamageThreshold)
 					{
 						DamageThreshold = ComputeDamageThreshold(DynamicCollection, TransformGroupIndex);
 					}
 					Chaos::FPhysicsSolver* RBDSolver = GetSolver<Chaos::FPhysicsSolver>();
 					RBDSolver->GetEvolution()->GetRigidClustering().SetInternalStrain(Handle, DamageThreshold);
-					
+
 				}
 
 				// #BGTODO - non-updating parameters - remove lin/ang drag arrays and always query material if this stays a material parameter
 				Chaos::FChaosPhysicsMaterial* SolverMaterial = RigidsSolver->GetSimMaterials().Get(Parameters.PhysicalMaterialHandle.InnerHandle);
-				if(SolverMaterial)
+				if (SolverMaterial)
 				{
 					Handle->SetLinearEtherDrag(SolverMaterial->LinearEtherDrag);
 					Handle->SetAngularEtherDrag(SolverMaterial->AngularEtherDrag);
 				}
 
 				const Chaos::FShapesArray& Shapes = Handle->ShapesArray();
-				for(const TUniquePtr<Chaos::FPerShapeData>& Shape : Shapes)
+				for (const TUniquePtr<Chaos::FPerShapeData>& Shape : Shapes)
 				{
 					Shape->SetMaterial(Parameters.PhysicalMaterialHandle);
 				}
@@ -1060,8 +1100,7 @@ void FGeometryCollectionPhysicsProxy::InitializeBodiesPT(Chaos::FPBDRigidsSolver
 
 				{
 					// Rigid node
-					SubTreeContainsSimulatableParticle[TransformGroupIndex] =
-						SolverParticleHandles[TransformGroupIndex] != nullptr;
+					SubTreeContainsSimulatableParticle[TransformGroupIndex] = SolverParticleHandles[TransformGroupIndex] != nullptr || !EffectiveParticles[TransformGroupIndex];
 				}
 				else
 				{
@@ -1092,23 +1131,42 @@ void FGeometryCollectionPhysicsProxy::InitializeBodiesPT(Chaos::FPBDRigidsSolver
 			{
 				// Don't construct particles for branches of the hierarchy that  
 				// don't contain any simulated particles.
-				if (!SubTreeContainsSimulatableParticle[TransformGroupIndex])
+				if (!SubTreeContainsSimulatableParticle[TransformGroupIndex] || !EffectiveParticles[TransformGroupIndex])
 				{
 					continue;
 				}
 
 				RigidChildren.Reset(Children.Num());
 				RigidChildrenTransformGroupIndex.Reset(Children.Num());
+
+				float ParentMass = 0.0f;
+				Chaos::FMatrix33 ParentInertia = Chaos::FMatrix33(0);
+
+				if (!Children[TransformGroupIndex].IsEmpty())
+				{
+					Masses[TransformGroupIndex] = 0;
+					Inertias[TransformGroupIndex] = Chaos::FVec3(0.0);
+				}
+
+				Chaos::FMatrix33 FullInertia = Chaos::FMatrix33(0);
 				for (const int32 ChildIndex : Children[TransformGroupIndex])
 				{
-					if (Chaos::TPBDRigidClusteredParticleHandle<Chaos::FReal, 3>* Handle = SolverParticleHandles[ChildIndex])
+					Masses[TransformGroupIndex] += Masses[ChildIndex];
+					const Chaos::FMatrix33 ChildWorldSpaceI = Chaos::Utilities::ComputeWorldSpaceInertia(Transforms[ChildIndex].GetRotation(), Inertias[ChildIndex]);
+					FullInertia += ChildWorldSpaceI;
+					Inertias[TransformGroupIndex] = FullInertia.GetDiagonal();
+
+					if (EffectiveParticles[ChildIndex])
 					{
-						RigidChildren.Add(Handle);
-						RigidChildrenTransformGroupIndex.Add(ChildIndex);
+						if (Chaos::TPBDRigidClusteredParticleHandle<Chaos::FReal, 3>*Handle = SolverParticleHandles[ChildIndex])
+						{
+							RigidChildren.Add(Handle);
+							RigidChildrenTransformGroupIndex.Add(ChildIndex);
+						}
 					}
 				}
 
-				if (RigidChildren.Num())
+				if (SolverParticleHandles[TransformGroupIndex] == nullptr)
 				{
 					if (ReportTooManyChildrenNum >= 0 && RigidChildren.Num() > ReportTooManyChildrenNum)
 					{
@@ -1124,7 +1182,15 @@ void FGeometryCollectionPhysicsProxy::InitializeBodiesPT(Chaos::FPBDRigidsSolver
 					FParticle* GTParticle = GTParticles[TransformGroupIndex].Get();
 
 					Chaos::FUniqueIdx ExistingIndex = GTParticle->UniqueIdx();
-					Chaos::FPBDRigidClusteredParticleHandle* Handle = BuildClusters_Internal(TransformGroupIndex, RigidChildren, RigidChildrenTransformGroupIndex, CreationParameters, &ExistingIndex);
+					Chaos::FPBDRigidClusteredParticleHandle* Handle = nullptr;
+					if (RigidChildren.Num())
+					{
+						Handle = BuildClusters_Internal(TransformGroupIndex, RigidChildren, RigidChildrenTransformGroupIndex, CreationParameters, &ExistingIndex);
+					}
+					else
+					{
+						Handle = BuildNonClusters_Internal(TransformGroupIndex, RigidsSolver, Masses[TransformGroupIndex], Inertias[TransformGroupIndex]);
+					}
 					Handle->GTGeometryParticle() = GTParticle;
 
 					int32 RigidChildrenIdx = 0;
@@ -1132,7 +1198,7 @@ void FGeometryCollectionPhysicsProxy::InitializeBodiesPT(Chaos::FPBDRigidsSolver
 					{
 						SolverClusterID[ChildTransformIndex] = RigidChildren[RigidChildrenIdx++]->CastToClustered()->ClusterIds().Id;;
 					}
-					SolverClusterID[TransformGroupIndex] = Handle->ClusterIds().Id;					
+					SolverClusterID[TransformGroupIndex] = Handle->ClusterIds().Id;
 
 					SolverClusterHandles[TransformGroupIndex] = Handle;
 					SolverParticleHandles[TransformGroupIndex] = Handle;
@@ -1249,7 +1315,6 @@ void FGeometryCollectionPhysicsProxy::InitializeBodiesPT(Chaos::FPBDRigidsSolver
 				Handle->SetAngularEtherDrag(Parameters.AngularDamping);
 			}
 		}
-		
 		// call DirtyParticle to make sure the acceleration structure is up to date with all the changes happening here
 		DirtyAllParticles(*RigidsSolver);
 
@@ -1269,7 +1334,7 @@ float FGeometryCollectionPhysicsProxy::ComputeDamageThreshold(const FGeometryDyn
 
 	const int32 LevelOffset = Parameters.bUsePerClusterOnlyDamageThreshold ? 0 : -1;
 	const int32 Level = FMath::Clamp(CalculateHierarchyLevel(DynamicCollection, TransformIndex) + LevelOffset, 0, INT_MAX);
-	if (Level >= Parameters.MaxClusterLevel)
+	if (Level >= FMath::Min(Parameters.MaxClusterLevel, Parameters.MaxSimulatedLevel))
 	{
 		DamageThreshold = TNumericLimits<float>::Max();
 	}
@@ -1301,9 +1366,129 @@ float FGeometryCollectionPhysicsProxy::ComputeDamageThreshold(const FGeometryDyn
 	return DamageThreshold;
 }
 
+Chaos::TPBDGeometryCollectionParticleHandle<Chaos::FReal, 3>* FGeometryCollectionPhysicsProxy::BuildNonClusters_Internal(const uint32 CollectionClusterIndex, Chaos::FPBDRigidsSolver* RigidsSolver, float Mass, Chaos::FVec3f Inertia)
+{
+	FGeometryDynamicCollection& DynamicCollection = PhysicsThreadCollection;
+	TManagedArray<int32>& DynamicState = DynamicCollection.DynamicState;
+	const TManagedArray<int32>& CollisionGroup = DynamicCollection.CollisionGroup;
+	TManagedArray<int32>& ParentIndex = DynamicCollection.Parent;
+	TManagedArray<TSet<int32>>& Children = DynamicCollection.Children;
+	TManagedArray<FTransform>& Transform = DynamicCollection.Transform;
+	TManagedArray<FTransform>& MassToLocal = DynamicCollection.MassToLocal;
+	const TManagedArray<FGeometryDynamicCollection::FSharedImplicit>& Implicits = DynamicCollection.Implicits;
+	const TManagedArray<TUniquePtr<FCollisionStructureManager::FSimplicial>>& Simplicials = DynamicCollection.Simplicials;
+	Chaos::Facades::FCollectionAnchoringFacade AnchoringFacade(DynamicCollection);
+
+	//If we are a root particle use the world transform, otherwise set the relative transform
+	const FTransform CollectionSpaceTransform = GeometryCollectionAlgo::GlobalMatrix(Transform, ParentIndex, CollectionClusterIndex);
+	const Chaos::TRigidTransform<Chaos::FReal, 3> ParticleTM = MassToLocal[CollectionClusterIndex] * CollectionSpaceTransform * Parameters.WorldTransform;
+
+	// Gather unique indices from GT to pass into PT handle creation
+	TArray<Chaos::FUniqueIdx> UniqueIndices;
+	UniqueIndices.Add(GTParticles[CollectionClusterIndex]->UniqueIdx());
+
+	// Add entries into simulation array
+	RigidsSolver->GetEvolution()->ReserveParticles(1);
+	TArray<Chaos::TPBDGeometryCollectionParticleHandle<Chaos::FReal, 3>*> Handles = RigidsSolver->GetEvolution()->CreateGeometryCollectionParticles(1, UniqueIndices.GetData());
+
+	Chaos::TPBDGeometryCollectionParticleHandle<Chaos::FReal, 3>* Handle = Handles[0];
+
+	Handle->SetPhysicsProxy(this);
+
+	SolverParticleHandles[CollectionClusterIndex] = Handle;
+	HandleToTransformGroupIndex.Add(Handle, CollectionClusterIndex);
+
+	// We're on the physics thread here but we've already set up the GT particles and we're just linking here
+	Handle->GTGeometryParticle() = GTParticles[CollectionClusterIndex].Get();
+
+	check(SolverParticleHandles[CollectionClusterIndex]->GetParticleType() == Handle->GetParticleType());
+	RigidsSolver->GetEvolution()->RegisterParticle(Handle);
+
+
+	if (AnchoringFacade.IsAnchored(CollectionClusterIndex))
+	{
+		if (Handle->ObjectState() == Chaos::EObjectStateType::Static)
+		{
+			DynamicState[CollectionClusterIndex] = (uint8)EObjectStateTypeEnum::Chaos_Object_Static;
+		}
+		else
+		{
+			DynamicState[CollectionClusterIndex] = (uint8)EObjectStateTypeEnum::Chaos_Object_Kinematic;
+		}
+	}
+
+	check(Parameters.RestCollection);
+
+	// NOTE: The particle creation call above (CreateClusterParticle) activates the particle, which does various things including adding
+	// it to the constraint graph. It seems a bit dodgy to be completely changing what the particle is after it has been enabled. Maybe we should fix this.
+	PopulateSimulatedParticle(
+		Handle,
+		Parameters.Shared,
+		Simplicials[CollectionClusterIndex].Get(),
+		Implicits[CollectionClusterIndex],
+		SimFilter,
+		QueryFilter,
+		Mass,
+		Inertia,
+		ParticleTM,
+		static_cast<uint8>(DynamicState[CollectionClusterIndex]),
+		static_cast<int16>(CollisionGroup[CollectionClusterIndex]),
+		CollisionParticlesPerObjectFraction); // CollisionGroup
+
+
+	// initialize anchoring information if available 
+	Handle->SetIsAnchored(AnchoringFacade.IsAnchored(CollectionClusterIndex));
+
+	if (Parameters.EnableClustering)
+	{
+		Handle->SetClusterGroupIndex(Parameters.ClusterGroupIndex);
+
+		const float StrainDefault = Parameters.DamageThreshold.Num() ? Parameters.DamageThreshold[0] : 0;
+		float DamageThreshold = StrainDefault;
+		if (!Parameters.bUsePerClusterOnlyDamageThreshold)
+		{
+			DamageThreshold = ComputeDamageThreshold(DynamicCollection, CollectionClusterIndex);
+		}
+		Chaos::FPhysicsSolver* RBDSolver = GetSolver<Chaos::FPhysicsSolver>();
+		RBDSolver->GetEvolution()->GetRigidClustering().SetInternalStrain(Handle, DamageThreshold);
+
+	}
+
+	// #BGTODO This will not automatically update - material properties should only ever exist in the material, not in other arrays
+	const Chaos::FChaosPhysicsMaterial* CurMaterial = static_cast<Chaos::FPBDRigidsSolver*>(Solver)->GetSimMaterials().Get(Parameters.PhysicalMaterialHandle.InnerHandle);
+	if (CurMaterial)
+	{
+		Handle->SetLinearEtherDrag(CurMaterial->LinearEtherDrag);
+		Handle->SetAngularEtherDrag(CurMaterial->AngularEtherDrag);
+	}
+
+	const Chaos::FShapesArray& Shapes = Handle->ShapesArray();
+	for (const TUniquePtr<Chaos::FPerShapeData>& Shape : Shapes)
+	{
+		Shape->SetMaterial(Parameters.PhysicalMaterialHandle);
+	}
+
+	const FTransform ParentTransform = GeometryCollectionAlgo::GlobalMatrix(DynamicCollection.Transform, DynamicCollection.Parent, CollectionClusterIndex);
+
+	// Populate bounds as we didn't pass a shared implicit to PopulateSimulatedParticle this will have been skipped, now that we have the full cluster we can build it
+	if (Handle->Geometry() && Handle->Geometry()->HasBoundingBox())
+	{
+		Handle->SetHasBounds(true);
+		Handle->SetLocalBounds(Handle->Geometry()->BoundingBox());
+		const Chaos::FRigidTransform3 Xf(Handle->X(), Handle->R());
+		Handle->UpdateWorldSpaceState(Xf, Chaos::FVec3(0));
+
+		static_cast<Chaos::FPBDRigidsSolver*>(Solver)->GetEvolution()->DirtyParticle(*Handle);
+	}
+
+	return Handle;
+}
+
+
+
 Chaos::FPBDRigidClusteredParticleHandle*
 FGeometryCollectionPhysicsProxy::BuildClusters_Internal(
-	const uint32 CollectionClusterIndex, // TransformGroupIndex
+	const uint32 CollectionClusterIndex,
 	TArray<Chaos::FPBDRigidParticleHandle*>& ChildHandles,
 	const TArray<int32>& ChildTransformGroupIndices,
 	const Chaos::FClusterCreationParameters & ClusterParameters,
@@ -1403,7 +1588,7 @@ FGeometryCollectionPhysicsProxy::BuildClusters_Internal(
 		Parent->I() != Chaos::FVec3f(0.0) ? Parent->I() : ScaledInertia,
 		ParticleTM, 
 		(uint8)DynamicState[CollectionClusterIndex], 
-		0,
+		0, // static_cast<int16>(CollisionGroup[TransformGroupIndex])
 		CollisionParticlesPerObjectFraction); // CollisionGroup
 
 	// two-way mapping
@@ -2173,6 +2358,22 @@ int32 FGeometryCollectionPhysicsProxy::CalculateHierarchyLevel(const FGeometryDy
 }
 
 
+int32 FGeometryCollectionPhysicsProxy::CalculateEffectiveParticles(const FGeometryDynamicCollection& DynamicCollection, int32 NumTransform, TBitArray<>& EffectiveParticles) const
+{
+	int32 NumEffectiveParticlesFound = 0;
+	EffectiveParticles.Init(false, NumTransform);
+	for (int32 TransformIndex = 0; TransformIndex < NumTransform; ++TransformIndex)
+	{
+		const int32 Level = FMath::Clamp(CalculateHierarchyLevel(DynamicCollection, TransformIndex), 0, INT_MAX);
+		if (Level <= Parameters.MaxSimulatedLevel || !Parameters.EnableClustering)
+		{
+			EffectiveParticles[TransformIndex] = true;
+			NumEffectiveParticlesFound++;
+		}
+	}
+	return NumEffectiveParticlesFound;
+}
+
 
 int32 FGeometryCollectionPhysicsProxy::CalculateAndSetLevel(int32 TransformGroupIndex, const TManagedArray<int32>& Parent, TManagedArray<int32>& Levels)
 {
@@ -2541,6 +2742,11 @@ void FGeometryCollectionPhysicsProxy::BufferPhysicsResults_External(Chaos::FDirt
 	
 	for (int32 TransformGroupIndex = 0; TransformGroupIndex < TargetResults.Transforms.Num(); ++TransformGroupIndex)
 	{
+		if (GTParticles[TransformGroupIndex] == nullptr)
+		{
+			continue;
+		}
+
 		FParticle& GTParticle = *GTParticles[TransformGroupIndex];
 		TargetResults.ParticleXs[TransformGroupIndex] = GTParticle.X();
 		GC_PHYSICSPROXY_CHECK_FOR_NAN(TargetResults.ParticleXs[TransformGroupIndex]);
@@ -2887,6 +3093,10 @@ bool FGeometryCollectionPhysicsProxy::PullFromPhysicsState(const Chaos::FDirtyGe
 		// first : copy the non interpolate-able values
 		for (int32 TransformGroupIndex = 0; TransformGroupIndex < NumTransforms; ++TransformGroupIndex)
 		{
+			if (GTParticles[TransformGroupIndex] == nullptr)
+			{
+				continue;
+			}
 			FParticle& GTParticle = *GTParticles[TransformGroupIndex];
 
 			const bool bIsActive = !TargetResults.States[TransformGroupIndex].DisabledState;
@@ -2957,6 +3167,10 @@ bool FGeometryCollectionPhysicsProxy::PullFromPhysicsState(const Chaos::FDirtyGe
 				const bool WasDisabled = Prev.States[TransformGroupIndex].DisabledState;
 				if (!WasDisabled)
 				{
+					if (GTParticles[TransformGroupIndex] == nullptr)
+					{
+						continue;
+					}
 					FParticle& GTParticle = *GTParticles[TransformGroupIndex];
 
 					GC_PHYSICSPROXY_CHECK_FOR_NAN(Prev.ParticleXs[TransformGroupIndex]);
@@ -3013,6 +3227,10 @@ bool FGeometryCollectionPhysicsProxy::PullFromPhysicsState(const Chaos::FDirtyGe
 
 				if (!TargetResults.States[TransformGroupIndex].DisabledState || bAnimatingWhileDisabled)
 				{
+					if (GTParticles[TransformGroupIndex] == nullptr)
+					{
+						continue;
+					}
 					FParticle& GTParticle = *GTParticles[TransformGroupIndex];
 
 					const Chaos::FVec3 OldX = GTParticle.X();
@@ -3086,18 +3304,20 @@ void FGeometryCollectionPhysicsProxy::UpdateFilterData_External(const FCollision
 	check(IsInGameThread());
 
 	// SimFilter/QueryFilter members are read on both threads, these are const after initialization and are not updated here.
-
-	int32 NumTransforms = GameThreadCollection.Transform.Num();
+	const int32 NumTransforms = GameThreadCollection.Transform.Num();
 	for (int32 Index = 0; Index < NumTransforms; ++Index)
 	{
 		FParticle* P = GTParticles[Index].Get();
-		const Chaos::FShapesArray& Shapes = P->ShapesArray();
-		const int32 NumShapes = Shapes.Num();
-		for (int32 ShapeIndex = 0; ShapeIndex < NumShapes; ++ShapeIndex)
+		if (P != nullptr)
 		{
-			Chaos::FPerShapeData* Shape = Shapes[ShapeIndex].Get();
-			Shape->SetSimData(NewSimFilter);
-			Shape->SetQueryData(NewQueryFilter);
+			const Chaos::FShapesArray& Shapes = P->ShapesArray();
+			const int32 NumShapes = Shapes.Num();
+			for (int32 ShapeIndex = 0; ShapeIndex < NumShapes; ++ShapeIndex)
+			{
+				Chaos::FPerShapeData* Shape = Shapes[ShapeIndex].Get();
+				Shape->SetSimData(NewSimFilter);
+				Shape->SetQueryData(NewQueryFilter);
+			}
 		}
 	}
 
@@ -4288,12 +4508,28 @@ void FGeometryCollectionPhysicsProxy::FieldForcesUpdateCallback(Chaos::FPBDRigid
 	}
 }
 
-TArray<Chaos::FPhysicsObjectHandle> FGeometryCollectionPhysicsProxy::GetAllPhysicsObjects()
+TArray<Chaos::FPhysicsObjectHandle> FGeometryCollectionPhysicsProxy::GetAllPhysicsObjects() const
+{
+	TArray<Chaos::FPhysicsObjectHandle> Handles;
+	Handles.Reserve(NumEffectiveParticles);
+
+	for (const Chaos::FPhysicsObjectUniquePtr& Object : PhysicsObjects)
+	{
+		if (Object.Get() != nullptr)
+		{
+			Handles.Add(Object.Get());
+		}
+	}
+
+	return Handles;
+}
+
+TArray<Chaos::FPhysicsObjectHandle> FGeometryCollectionPhysicsProxy::GetAllPhysicsObjectIncludingNulls() const
 {
 	TArray<Chaos::FPhysicsObjectHandle> Handles;
-	Handles.Reserve(PhysicsObjects.Num());
+	Handles.Reserve(NumEffectiveParticles);
 
-	for (Chaos::FPhysicsObjectUniquePtr& Object : PhysicsObjects)
+	for (const Chaos::FPhysicsObjectUniquePtr& Object : PhysicsObjects)
 	{
 		Handles.Add(Object.Get());
 	}
diff --git a/Engine/Source/Runtime/Experimental/Chaos/Public/GeometryCollection/GeometryCollectionSimulationCoreTypes.h b/Engine/Source/Runtime/Experimental/Chaos/Public/GeometryCollection/GeometryCollectionSimulationCoreTypes.h
--- a/Engine/Source/Runtime/Experimental/Chaos/Public/GeometryCollection/GeometryCollectionSimulationCoreTypes.h
+++ b/Engine/Source/Runtime/Experimental/Chaos/Public/GeometryCollection/GeometryCollectionSimulationCoreTypes.h
@@ -177,6 +177,7 @@ struct FSimulationParameters
 		, EnableClustering(true)
 		, ClusterGroupIndex(0)
 		, MaxClusterLevel(100)
+		, MaxSimulatedLevel(100)
 		, bUseSizeSpecificDamageThresholds(false)
 		, DamageThreshold({500000.f, 50000.f, 5000.f})
 		, bUsePerClusterOnlyDamageThreshold(false)
@@ -225,6 +226,7 @@ struct FSimulationParameters
 		, EnableClustering(Other.EnableClustering)
 		, ClusterGroupIndex(Other.ClusterGroupIndex)
 		, MaxClusterLevel(Other.MaxClusterLevel)
+		, MaxSimulatedLevel(Other.MaxSimulatedLevel)
 		, bUseSizeSpecificDamageThresholds(Other.bUseSizeSpecificDamageThresholds)
 		, DamageThreshold(Other.DamageThreshold)
 		, bUsePerClusterOnlyDamageThreshold(Other.bUsePerClusterOnlyDamageThreshold)
@@ -290,6 +292,7 @@ struct FSimulationParameters
 	bool EnableClustering;
 	int32 ClusterGroupIndex;
 	int32 MaxClusterLevel;
+	int32 MaxSimulatedLevel;
 	bool bUseSizeSpecificDamageThresholds;
 	TArray<float> DamageThreshold;
 	bool bUsePerClusterOnlyDamageThreshold;
diff --git a/Engine/Source/Runtime/Experimental/Chaos/Public/PhysicsProxy/GeometryCollectionPhysicsProxy.h b/Engine/Source/Runtime/Experimental/Chaos/Public/PhysicsProxy/GeometryCollectionPhysicsProxy.h
--- a/Engine/Source/Runtime/Experimental/Chaos/Public/PhysicsProxy/GeometryCollectionPhysicsProxy.h
+++ b/Engine/Source/Runtime/Experimental/Chaos/Public/PhysicsProxy/GeometryCollectionPhysicsProxy.h
@@ -447,7 +447,8 @@ public:
 		PostPhysicsSyncCallback = Callback;
 	}
 	
-	TArray<Chaos::FPhysicsObjectHandle> GetAllPhysicsObjects();
+	TArray<Chaos::FPhysicsObjectHandle> GetAllPhysicsObjects() const ;
+	TArray<Chaos::FPhysicsObjectHandle> GetAllPhysicsObjectIncludingNulls() const;
 	Chaos::FPhysicsObjectHandle GetPhysicsObjectByIndex(int32 Index) const;
 	int32 GetNumParticles() const { return NumParticles; }
 protected:
@@ -459,7 +460,10 @@ protected:
 	* #return damage threshold value
 	*/
 	float ComputeDamageThreshold(const FGeometryDynamicCollection& DynamicCollection, int32 TransformIndex) const;
-		
+
+
+	Chaos::TPBDGeometryCollectionParticleHandle<Chaos::FReal, 3>* BuildNonClusters_Internal(const uint32 CollectionClusterIndex, Chaos::FPBDRigidsSolver* RigidsSolver, float Mass, Chaos::FVec3f Inertia);
+
 	/**
 	 * Build a physics thread cluster parent particle.
 	 *	\p CollectionClusterIndex - the source geometry collection transform index.
@@ -484,15 +488,18 @@ protected:
 	 */
 	static int32 CalculateHierarchyLevel(const FGeometryDynamicCollection& DynamicCollection, int32 TransformIndex);
 
+	int32 CalculateEffectiveParticles(const FGeometryDynamicCollection& DynamicCollection, int32 NumTransform, TBitArray<>& EffectiveParticles) const;
+
 	void SetClusteredParticleKinematicTarget_Internal(Chaos::FPBDRigidClusteredParticleHandle* Handle, const FTransform& WorldTransform);
 
 	void PrepareBufferData(Chaos::FDirtyGeometryCollectionData& BufferData, const FGeometryDynamicCollection& ThreadCollection,  Chaos::FReal SolverLastDt = 0.0);
 
-	void CreateNonClusteredParticles(Chaos::FPBDRigidsSolver* RigidsSolver,	const FGeometryCollection& RestCollection, const FGeometryDynamicCollection& DynamicCollection);
+	void CreateNonClusteredParticles(Chaos::FPBDRigidsSolver* RigidsSolver,	const FGeometryCollection& RestCollection, const FGeometryDynamicCollection& DynamicCollection, const TBitArray<>& EffectiveParticles);
 
 	Chaos::FPBDRigidClusteredParticleHandle* FindClusteredParticleHandleByItemIndex_Internal(FGeometryCollectionItemIndex ItemIndex) const;
 private:
 
+
 	/* set to true once InitializeBodiesPT has been called*/
 	bool bIsInitializedOnPhysicsThread = false;
 
@@ -507,6 +514,7 @@ private:
 	//  Proxy State Information
 	//
 	int32 NumParticles;
+	int32 NumEffectiveParticles;
 	int32 BaseParticleIndex;
 	TArray<FParticleHandle*> SolverClusterID;
 	TArray<FClusterHandle*> SolverClusterHandles; // make a TArray of the base clase with type
diff --git a/Engine/Source/Runtime/Experimental/GeometryCollectionEngine/Private/GeometryCollection/GeometryCollectionComponent.cpp b/Engine/Source/Runtime/Experimental/GeometryCollectionEngine/Private/GeometryCollection/GeometryCollectionComponent.cpp
--- a/Engine/Source/Runtime/Experimental/GeometryCollectionEngine/Private/GeometryCollection/GeometryCollectionComponent.cpp
+++ b/Engine/Source/Runtime/Experimental/GeometryCollectionEngine/Private/GeometryCollection/GeometryCollectionComponent.cpp
@@ -316,6 +316,7 @@ UGeometryCollectionComponent::UGeometryCollectionComponent(const FObjectInitiali
 	, EnableClustering(true)
 	, ClusterGroupIndex(0)
 	, MaxClusterLevel(100)
+	, MaxSimulatedLevel(100)
 	, DamageThreshold({ 500000.f, 50000.f, 5000.f })
 	, bUseSizeSpecificDamageThreshold(false)
 	, bEnableDamageFromCollision(true)
@@ -2657,6 +2658,7 @@ void UGeometryCollectionComponent::RegisterAndInitializePhysicsProxy()
 		SimulationParameters.EnableClustering = EnableClustering;
 		SimulationParameters.ClusterGroupIndex = EnableClustering ? ClusterGroupIndex : 0;
 		SimulationParameters.MaxClusterLevel = MaxClusterLevel;
+		SimulationParameters.MaxSimulatedLevel = MaxSimulatedLevel;
 		SimulationParameters.bUseSizeSpecificDamageThresholds = bUseSizeSpecificDamageThreshold;
 		SimulationParameters.DamageThreshold = DamageThreshold;
 		SimulationParameters.bUsePerClusterOnlyDamageThreshold = RestCollection? RestCollection->PerClusterOnlyDamageThreshold: false; 
@@ -2943,28 +2945,32 @@ void UGeometryCollectionComponent::LoadCollisionProfiles()
 	const CollisionProfileDataCache* AbandonedData = (bEnableAbandonAfterLevel && GetIsReplicated()) ? CreateOrGetCollisionProfileData(AbandonedCollisionProfileName) : nullptr;
 
 	Chaos::Facades::FCollectionHierarchyFacade HierarchyFacade(*RestCollection->GetGeometryCollection());
-	TArray<Chaos::FPhysicsObjectHandle> PhysicsObjects = PhysicsProxy->GetAllPhysicsObjects();
+	// Use GetAllPhysicsObjectIncludingNulls instead of GetAllPhysicsObjects if you need to use Level data or any data from HierarchyFacade
+	TArray<Chaos::FPhysicsObjectHandle> PhysicsObjects = PhysicsProxy->GetAllPhysicsObjectIncludingNulls();
 	FLockedWritePhysicsObjectExternalInterface Interface = FPhysicsObjectExternalInterface::LockWrite(PhysicsObjects);
 
-	for (int32 ParticleIndex = 0; ParticleIndex < PhysicsProxy->GetNumParticles(); ++ParticleIndex)
+	for (int32 ParticleIndex = 0; ParticleIndex < PhysicsObjects.Num(); ++ParticleIndex)
 	{
-		const int32 Level = HierarchyFacade.GetInitialLevel(ParticleIndex);
-
-		TArrayView<Chaos::FPhysicsObjectHandle> ParticleView{ &PhysicsObjects[ParticleIndex], 1 };
-		const CollisionProfileDataCache* Data = nullptr;
-		if (AbandonedData && Level >= ReplicationAbandonAfterLevel + 1)
-		{
-			Data = AbandonedData;
-		}
-		else if (!CollisionProfilePerLevel.IsEmpty())
+		if (PhysicsObjects[ParticleIndex] != nullptr)
 		{
-			Data = CreateOrGetCollisionProfileData(CollisionProfilePerLevel[FMath::Min(CollisionProfilePerLevel.Num() - 1, Level)]);
-		}
+			const int32 Level = HierarchyFacade.GetInitialLevel(ParticleIndex);
 
-		if (Data)
-		{
-			Interface->UpdateShapeCollisionFlags(ParticleView, CollisionEnabledHasPhysics(Data->Template.CollisionEnabled), CollisionEnabledHasQuery(Data->Template.CollisionEnabled));
-			Interface->UpdateShapeFilterData(ParticleView, Data->QueryFilter, Data->SimFilter);
+			TArrayView<Chaos::FPhysicsObjectHandle> ParticleView{ &PhysicsObjects[ParticleIndex], 1 };
+			const CollisionProfileDataCache* Data = nullptr;
+			if (AbandonedData && Level >= ReplicationAbandonAfterLevel + 1)
+			{
+				Data = AbandonedData;
+			}
+			else if (!CollisionProfilePerLevel.IsEmpty())
+			{
+				Data = CreateOrGetCollisionProfileData(CollisionProfilePerLevel[FMath::Min(CollisionProfilePerLevel.Num() - 1, Level)]);
+			}
+
+			if (Data)
+			{
+				Interface->UpdateShapeCollisionFlags(ParticleView, CollisionEnabledHasPhysics(Data->Template.CollisionEnabled), CollisionEnabledHasQuery(Data->Template.CollisionEnabled));
+				Interface->UpdateShapeFilterData(ParticleView, Data->QueryFilter, Data->SimFilter);
+			}
 		}
 	}
 }
diff --git a/Engine/Source/Runtime/Experimental/GeometryCollectionEngine/Public/GeometryCollection/GeometryCollectionComponent.h b/Engine/Source/Runtime/Experimental/GeometryCollectionEngine/Public/GeometryCollection/GeometryCollectionComponent.h
--- a/Engine/Source/Runtime/Experimental/GeometryCollectionEngine/Public/GeometryCollection/GeometryCollectionComponent.h
+++ b/Engine/Source/Runtime/Experimental/GeometryCollectionEngine/Public/GeometryCollection/GeometryCollectionComponent.h
@@ -746,6 +746,16 @@ public:
 	UPROPERTY(EditAnywhere, BlueprintReadWrite, Category = "ChaosPhysics|Clustering")
 	int32 MaxClusterLevel;
 
+	/** The maximum level to create rigid bodies that could be simulated.
+	Example: If we have a Geometry Collection with 2 levels, where:
+	0 = Root
+	1 = Clusters
+	2 = Leaf Nodes
+	A setting of '1' would only generate a physics representation of the Root transform and Level 1 clusters. 
+	The leaf nodes on Level 2 would never be created on the solver, and would therefore never be considered as part of the simulation. */
+	UPROPERTY(EditAnywhere, BlueprintReadWrite, Category = "ChaosPhysics|Clustering")
+	int32 MaxSimulatedLevel;
+
 	UPROPERTY(EditAnywhere, BlueprintReadWrite, Category = "ChaosPhysics|Damage", meta = (EditCondition = "!bUseSizeSpecificDamageThreshold"))
 	TArray<float> DamageThreshold;
 
