Add a Cvar to control a Max Simulated Level for the collision geometry.

diff --git a/Engine/Source/Runtime/Experimental/Chaos/Private/PhysicsProxy/GeometryCollectionPhysicsProxy.cpp b/Engine/Source/Runtime/Experimental/Chaos/Private/PhysicsProxy/GeometryCollectionPhysicsProxy.cpp
--- a/Engine/Source/Runtime/Experimental/Chaos/Private/PhysicsProxy/GeometryCollectionPhysicsProxy.cpp
+++ b/Engine/Source/Runtime/Experimental/Chaos/Private/PhysicsProxy/GeometryCollectionPhysicsProxy.cpp
@@ -1325,6 +1325,10 @@ void FGeometryCollectionPhysicsProxy::InitializeBodiesPT(Chaos::FPBDRigidsSolver
 int32 ReportNoLevelsetCluster = 0;
 FAutoConsoleVariableRef CVarReportNoLevelsetCluster(TEXT("p.gc.ReportNoLevelsetCluster"), ReportNoLevelsetCluster, TEXT("Report any cluster objects without levelsets"));
 
+int32 GlobalMaxSimulatedLevel = 100;
+FAutoConsoleVariableRef CVarGlobalMaxSimulatedLevel(TEXT("p.gc.GlobalMaxSimulatedLevel"), GlobalMaxSimulatedLevel, TEXT("Allow to set the Global Maximum Simulated Level for Geoemtry Collection. The min between the MaxSimulatedLevel and the GlobalMaxSimulatedLevel will be used. "));
+
+
 DECLARE_CYCLE_STAT(TEXT("FGeometryCollectionPhysicsProxy::BuildClusters"), STAT_BuildClusters, STATGROUP_Chaos);
 DECLARE_CYCLE_STAT(TEXT("FGeometryCollectionPhysicsProxy::BuildClusters:GlobalMatrices"), STAT_BuildClustersGlobalMatrices, STATGROUP_Chaos);
 
@@ -1334,7 +1338,7 @@ float FGeometryCollectionPhysicsProxy::ComputeDamageThreshold(const FGeometryDyn
 
 	const int32 LevelOffset = Parameters.bUsePerClusterOnlyDamageThreshold ? 0 : -1;
 	const int32 Level = FMath::Clamp(CalculateHierarchyLevel(DynamicCollection, TransformIndex) + LevelOffset, 0, INT_MAX);
-	if (Level >= FMath::Min(Parameters.MaxClusterLevel, Parameters.MaxSimulatedLevel))
+	if (Level >= FMath::Min(Parameters.MaxClusterLevel, FMath::Min(GlobalMaxSimulatedLevel, Parameters.MaxSimulatedLevel)))
 	{
 		DamageThreshold = TNumericLimits<float>::Max();
 	}
@@ -2362,10 +2366,11 @@ int32 FGeometryCollectionPhysicsProxy::CalculateEffectiveParticles(const FGeomet
 {
 	int32 NumEffectiveParticlesFound = 0;
 	EffectiveParticles.Init(false, NumTransform);
+	const int32 MaxSimulatedLevel = FMath::Min(GlobalMaxSimulatedLevel, Parameters.MaxSimulatedLevel);
 	for (int32 TransformIndex = 0; TransformIndex < NumTransform; ++TransformIndex)
 	{
 		const int32 Level = FMath::Clamp(CalculateHierarchyLevel(DynamicCollection, TransformIndex), 0, INT_MAX);
-		if (Level <= Parameters.MaxSimulatedLevel || !Parameters.EnableClustering)
+		if (Level <= MaxSimulatedLevel || !Parameters.EnableClustering)
 		{
 			EffectiveParticles[TransformIndex] = true;
 			NumEffectiveParticlesFound++;
