ChaosFlesh : Handle add and removes of collision bodies on the same frame.

diff --git a/Engine/Source/Runtime/Experimental/Chaos/Private/Chaos/Deformable/ChaosDeformableSolver.cpp b/Engine/Source/Runtime/Experimental/Chaos/Private/Chaos/Deformable/ChaosDeformableSolver.cpp
--- a/Engine/Source/Runtime/Experimental/Chaos/Private/Chaos/Deformable/ChaosDeformableSolver.cpp
+++ b/Engine/Source/Runtime/Experimental/Chaos/Private/Chaos/Deformable/ChaosDeformableSolver.cpp
@@ -590,6 +590,7 @@ namespace Chaos::Softs
 				CollisionsInputBuffer = this->CurrentInputPackage->ObjectMap[Owner]->As<FCollisionManagerProxy::FCollisionsInputBuffer>();
 				if (CollisionsInputBuffer)
 				{
+					TArray<FCollisionObjectAddedBodies> IgnoredAdditions;
 					for (auto& AddBody : CollisionsInputBuffer->Added)
 					{
 						if (AddBody.Shapes)
@@ -604,6 +605,32 @@ namespace Chaos::Softs
 								Evolution->CollisionParticles().SetDynamicGeometry(Index, MoveTemp(UniquePtr));
 								Proxy.CollisionBodies.Add(AddBody.Key, FCollisionObjectParticleHandel(Index, ViewIndex, AddBody.Transform));
 							}
+							else
+							{
+								IgnoredAdditions.Add(AddBody);
+							}
+						}
+					}
+
+					// If we tried to add a body that already was added, this means their
+					// should be a matching delete as the body was removed and added
+					// back before the physics thread could evaluate. 
+					for (auto& AddedBody : IgnoredAdditions)
+					{
+						for (int i=0; i<CollisionsInputBuffer->Removed.Num();)
+						{
+							if ((void*)CollisionsInputBuffer->Removed[i].Key.Get<0>() == (void*)AddedBody.Key.Get<0>())
+							{
+								CollisionsInputBuffer->Removed.RemoveAtSwap(i);
+								if (i == CollisionsInputBuffer->Removed.Num() - 1)
+								{
+									break;
+								}
+							}
+							else
+							{
+								i++;
+							}
 						}
 					}
 
