Add InitialRootIndex to simulation parameters to make it accessible on physics thread

diff --git a/Engine/Source/Runtime/Experimental/Chaos/Public/GeometryCollection/GeometryCollectionSimulationCoreTypes.h b/Engine/Source/Runtime/Experimental/Chaos/Public/GeometryCollection/GeometryCollectionSimulationCoreTypes.h
--- a/Engine/Source/Runtime/Experimental/Chaos/Public/GeometryCollection/GeometryCollectionSimulationCoreTypes.h
+++ b/Engine/Source/Runtime/Experimental/Chaos/Public/GeometryCollection/GeometryCollectionSimulationCoreTypes.h
@@ -170,6 +170,7 @@ struct FSimulationParameters
 	FSimulationParameters()
 		: Name("")
 		, RestCollection(nullptr)
+		, InitialRootIndex(INDEX_NONE)
 		, RecordedTrack(nullptr)
 		, bOwnsTrack(false)
 		, Simulating(false)
@@ -218,6 +219,7 @@ struct FSimulationParameters
 	FSimulationParameters(const FSimulationParameters& Other)
 		: Name(Other.Name)
 		, RestCollection(Other.RestCollection)
+		, InitialRootIndex(Other.InitialRootIndex)
 		, InitializationCommands(Other.InitializationCommands)
 		, RecordedTrack(Other.RecordedTrack)
 		, bOwnsTrack(false)
@@ -280,6 +282,7 @@ struct FSimulationParameters
 
 	FString Name;
 	const FGeometryCollection* RestCollection;
+	int32 InitialRootIndex;
 	TArray<FFieldSystemCommand> InitializationCommands;
 	const FRecordedTransformTrack* RecordedTrack;
 	bool bOwnsTrack;
diff --git a/Engine/Source/Runtime/Experimental/GeometryCollectionEngine/Private/GeometryCollection/GeometryCollectionComponent.cpp b/Engine/Source/Runtime/Experimental/GeometryCollectionEngine/Private/GeometryCollection/GeometryCollectionComponent.cpp
--- a/Engine/Source/Runtime/Experimental/GeometryCollectionEngine/Private/GeometryCollection/GeometryCollectionComponent.cpp
+++ b/Engine/Source/Runtime/Experimental/GeometryCollectionEngine/Private/GeometryCollection/GeometryCollectionComponent.cpp
@@ -2651,6 +2651,7 @@ void UGeometryCollectionComponent::RegisterAndInitializePhysicsProxy()
 		{
 			RestCollection->GetSharedSimulationParams(SimulationParameters.Shared);
 			SimulationParameters.RestCollection = RestCollection->GetGeometryCollection().Get();
+			SimulationParameters.InitialRootIndex = RestCollection->GetRootIndex();
 			ClusterCollectionType = RestCollection->ClusterConnectionType;
 			ConnectionGraphBoundsFilteringMargin = RestCollection->ConnectionGraphBoundsFilteringMargin;
 		}
