dirtying a disabled particle in a cluster union should not add the particle into the SQ

diff --git a/Engine/Source/Runtime/Experimental/Chaos/Public/Chaos/PBDRigidsEvolution.h b/Engine/Source/Runtime/Experimental/Chaos/Public/Chaos/PBDRigidsEvolution.h
--- a/Engine/Source/Runtime/Experimental/Chaos/Public/Chaos/PBDRigidsEvolution.h
+++ b/Engine/Source/Runtime/Experimental/Chaos/Public/Chaos/PBDRigidsEvolution.h
@@ -517,6 +517,8 @@ public:
 			{
 				// For clustered particles, they may appear disabled but they're being driven by an internal (solver-owned) cluster parent.
 				// If this is the case we let the spatial data update with those particles, otherwise skip.
+				// Alternatively, the particle may be being driven by a cluster union. Disabled children of a cluster union should not be added
+				// to the SQ.
 				// #BGTODO consider converting MDisabled into a bitfield for multiple disable types (Disabled, DisabledDriven, etc.)
 				if(FPBDRigidParticleHandle* ClusterParentBase = AsClustered->ClusterIds().Id)
 				{
@@ -526,6 +528,14 @@ public:
 						{
 							return;
 						}
+
+						// We know we're an internal cluster now. If this was a GC internal proxy, we'd expect the 
+						// parent's proxy to be the same as the input particle's. If this is not the case, we know
+						// we're in a cluster union.
+						if (ClusterParent->PhysicsProxy() != Particle.PhysicsProxy())
+						{
+							return;
+						}
 					}
 				}
 			}
@@ -1118,4 +1128,4 @@ protected:
 };
 
 
-}
+}
\ No newline at end of file
