set spatial index on the cluster union particle

diff --git a/Engine/Source/Runtime/Experimental/Chaos/Private/PBDRigidsSolver.cpp b/Engine/Source/Runtime/Experimental/Chaos/Private/PBDRigidsSolver.cpp
--- a/Engine/Source/Runtime/Experimental/Chaos/Private/PBDRigidsSolver.cpp
+++ b/Engine/Source/Runtime/Experimental/Chaos/Private/PBDRigidsSolver.cpp
@@ -693,6 +693,7 @@ namespace Chaos
 
 		if (FPBDRigidParticle* Particle = Proxy->GetParticle_External())
 		{
+			Particle->SetSpatialIdx(FSpatialAccelerationIdx{ 0, 0 });
 			Particle->SetUniqueIdx(GetEvolution()->GenerateUniqueIdx());
 		}
 		Proxy->SetSolver(this);
