add ability to remove from and clear child physics proxies on clustered particles

diff --git a/Engine/Source/Runtime/Experimental/Chaos/Public/Chaos/ParticleHandle.h b/Engine/Source/Runtime/Experimental/Chaos/Public/Chaos/ParticleHandle.h
--- a/Engine/Source/Runtime/Experimental/Chaos/Public/Chaos/ParticleHandle.h
+++ b/Engine/Source/Runtime/Experimental/Chaos/Public/Chaos/ParticleHandle.h
@@ -1257,6 +1257,8 @@ public:
 
 	const TSet<IPhysicsProxyBase*>& PhysicsProxies() const { return PBDRigidClusteredParticles->PhysicsProxies(ParticleIdx); }
 	void AddPhysicsProxy(IPhysicsProxyBase* PhysicsProxy) { PBDRigidClusteredParticles->PhysicsProxies(ParticleIdx).Add(PhysicsProxy); }
+	void RemovePhysicsProxy(IPhysicsProxyBase* PhysicsProxy) { PBDRigidClusteredParticles->PhysicsProxies(ParticleIdx).Remove(PhysicsProxy); }
+	void ClearPhysicsProxies() { PBDRigidClusteredParticles->PhysicsProxies(ParticleIdx).Empty(); }
 	
 	void SetClusterId(const ClusterId& Id) { PBDRigidClusteredParticles->ClusterIds(ParticleIdx) = Id; }
 	const ClusterId& ClusterIds() const { return PBDRigidClusteredParticles->ClusterIds(ParticleIdx); }
