Don't shun BoneIndex 0.  Bone 0 needs love too.

diff --git a/Engine/Source/Runtime/Experimental/Chaos/Private/GeometryCollection/Facades/CollectionVertexBoneWeightsFacade.cpp b/Engine/Source/Runtime/Experimental/Chaos/Private/GeometryCollection/Facades/CollectionVertexBoneWeightsFacade.cpp
--- a/Engine/Source/Runtime/Experimental/Chaos/Private/GeometryCollection/Facades/CollectionVertexBoneWeightsFacade.cpp
+++ b/Engine/Source/Runtime/Experimental/Chaos/Private/GeometryCollection/Facades/CollectionVertexBoneWeightsFacade.cpp
@@ -65,7 +65,7 @@ namespace GeometryCollection::Facades
 		const TManagedArray<FVector3f>& Vertices = VerticesAttribute.Modify();
 		if (0 <= VertexIndex && VertexIndex < Vertices.Num())
 		{
-			if (0 < BoneIndex && BoneIndex < ParentAttribute.Num())
+			if (0 <= BoneIndex && BoneIndex < ParentAttribute.Num())
 			{
 				IndicesArray[VertexIndex].Add(BoneIndex);
 				WeightsArray[VertexIndex].Add(BoneWeight);
