Chaos - fix collision when an implicit is set to LevelSet but has no CollisionParticles

diff --git a/Engine/Source/Runtime/Experimental/Chaos/Private/Chaos/Collision/PBDCollisionConstraint.cpp b/Engine/Source/Runtime/Experimental/Chaos/Private/Chaos/Collision/PBDCollisionConstraint.cpp
--- a/Engine/Source/Runtime/Experimental/Chaos/Private/Chaos/Collision/PBDCollisionConstraint.cpp
+++ b/Engine/Source/Runtime/Experimental/Chaos/Private/Chaos/Collision/PBDCollisionConstraint.cpp
@@ -4,6 +4,7 @@
 #include "Chaos/Collision/PBDCollisionConstraintHandle.h"
 #include "Chaos/Collision/PBDCollisionContainerSolver.h"
 #include "Chaos/Collision/PBDCollisionSolver.h"
+#include "Chaos/CollisionResolution.h"
 #include "Chaos/Evolution/SolverBody.h"
 #include "Chaos/Evolution/SolverBodyContainer.h"
 #include "Chaos/Island/IslandManager.h"
@@ -289,8 +290,8 @@ namespace Chaos
 		// Initialize tolerances that depend on shape type etc, also the bIsQuadratic flags
 		const FRealSingle Margin0 = FRealSingle(GetImplicit0()->GetMargin());
 		const FRealSingle Margin1 = FRealSingle(GetImplicit1()->GetMargin());
-		const EImplicitObjectType ImplicitType0 = GetInnerType(GetImplicit0()->GetCollisionType());
-		const EImplicitObjectType ImplicitType1 = GetInnerType(GetImplicit1()->GetCollisionType());
+		const EImplicitObjectType ImplicitType0 = Collisions::GetImplicitCollisionType(GetParticle0(), GetImplicit0());
+		const EImplicitObjectType ImplicitType1 = Collisions::GetImplicitCollisionType(GetParticle1(), GetImplicit1());
 		InitMarginsAndTolerances(ImplicitType0, ImplicitType1, Margin0, Margin1);
 
 		// Are we allowing manifolds?
diff --git a/Engine/Source/Runtime/Experimental/Chaos/Private/Chaos/Collision/ParticlePairMidPhase.cpp b/Engine/Source/Runtime/Experimental/Chaos/Private/Chaos/Collision/ParticlePairMidPhase.cpp
--- a/Engine/Source/Runtime/Experimental/Chaos/Private/Chaos/Collision/ParticlePairMidPhase.cpp
+++ b/Engine/Source/Runtime/Experimental/Chaos/Private/Chaos/Collision/ParticlePairMidPhase.cpp
@@ -1049,18 +1049,16 @@ namespace Chaos
 	void FParticlePairMidPhase::TryAddShapePair(const FPerShapeData* Shape0, const FPerShapeData* Shape1)
 	{
 		const FImplicitObject* Implicit0 = Shape0->GetLeafGeometry();
-		const FBVHParticles* BVHParticles0 = FConstGenericParticleHandle(Particle0)->CollisionParticles().Get();
-		const EImplicitObjectType ImplicitType0 = (Implicit0 != nullptr) ? GetInnerType(Implicit0->GetCollisionType()) : ImplicitObjectType::Unknown;
+		const EImplicitObjectType ImplicitType0 = Collisions::GetImplicitCollisionType(Particle0, Implicit0);
 
 		const FImplicitObject* Implicit1 = Shape1->GetLeafGeometry();
-		const FBVHParticles* BVHParticles1 = FConstGenericParticleHandle(Particle1)->CollisionParticles().Get();
-		const EImplicitObjectType ImplicitType1 = (Implicit1 != nullptr) ? GetInnerType(Implicit1->GetCollisionType()) : ImplicitObjectType::Unknown;
+		const EImplicitObjectType ImplicitType1 = Collisions::GetImplicitCollisionType(Particle1, Implicit1);
 
-		const bool bDoPassFilter = DoCollide(ImplicitType0, Shape0, ImplicitType1, Shape1);
+		const bool bDoPassFilter = ShapePairNarrowPhaseFilter(ImplicitType0, Shape0, ImplicitType1, Shape1);
 		if (bDoPassFilter)
 		{
 			bool bSwap = false;
-			const EContactShapesType ShapePairType = Collisions::CalculateShapePairType(Implicit0, BVHParticles0, Implicit1, BVHParticles1, bSwap);
+			const EContactShapesType ShapePairType = Collisions::CalculateShapePairType(Particle0, Implicit0, Particle1, Implicit1, bSwap);
 
 			if (ShapePairType != EContactShapesType::Unknown)
 			{
diff --git a/Engine/Source/Runtime/Experimental/Chaos/Private/Chaos/CollisionResolution.cpp b/Engine/Source/Runtime/Experimental/Chaos/Private/Chaos/CollisionResolution.cpp
--- a/Engine/Source/Runtime/Experimental/Chaos/Private/Chaos/CollisionResolution.cpp
+++ b/Engine/Source/Runtime/Experimental/Chaos/Private/Chaos/CollisionResolution.cpp
@@ -1429,7 +1429,7 @@ namespace Chaos
 			// If we get here, we have a pair of concrete shapes (i.e., no wrappers or containers)
 			// Create a constraint for the shape pair
 			bool bShouldSwapParticles = false;
-			const EContactShapesType ShapePairType = CalculateShapePairType(Implicit0, Simplicial0, Implicit1, Simplicial1, bShouldSwapParticles);
+			const EContactShapesType ShapePairType = CalculateShapePairType(Particle0, Implicit0, Particle1, Implicit1, bShouldSwapParticles);
 
 			if (bShouldSwapParticles)
 			{
@@ -1441,8 +1441,39 @@ namespace Chaos
 			}
 		}
 
-		EContactShapesType CalculateShapePairType(const EImplicitObjectType Implicit0Type, const EImplicitObjectType Implicit1Type, const bool bIsConvex0, const bool bIsConvex1, const bool bHasSimplicial0, const bool bHasSimplicial1,  bool &bOutSwap)
+		EImplicitObjectType GetImplicitCollisionType(const FGeometryParticleHandle* Particle, const FImplicitObject* Implicit)
 		{
+			// NOTE: GetCollisionType(), not GetType()
+			// We use CollisionType on the implicit to determine how to collide. Normally this is the same as the actual ImplicitObject's
+			// type, but may be set to LevelSet in which case we will use CollisionParticles instead (if it has any).
+			EImplicitObjectType ImplicitType = (Implicit != nullptr) ? GetInnerType(Implicit->GetCollisionType()) : ImplicitObjectType::Unknown;
+
+			// If we are a levelset make sure we have CollisionParticles, otherwise go back to the builtin implicit object collision type
+			if ((ImplicitType == ImplicitObjectType::LevelSet) || (ImplicitType == ImplicitObjectType::Unknown))
+			{
+				const FBVHParticles* Simplicial = FConstGenericParticleHandle(Particle)->CollisionParticles().Get();
+				const bool bHasSimplicial = (Simplicial != nullptr) && (Simplicial->Size() > 0);
+				if (!bHasSimplicial)
+				{
+					// NOTE: GetType(), not GetCollisionType()
+					ImplicitType = (Implicit != nullptr) ? GetInnerType(Implicit->GetType()) : ImplicitObjectType::Unknown;
+				}
+			}
+
+			return ImplicitType;
+		}
+
+		EContactShapesType CalculateShapePairType(const FGeometryParticleHandle* Particle0, const FImplicitObject* Implicit0, const FGeometryParticleHandle* Particle1, const FImplicitObject* Implicit1, bool &bOutSwap)
+		{
+			// What types do the implicits collide as?
+			// NOTE: If the implicit is set to LevelSet but has no CollisionParticles we default to the implicits actual type
+			const EImplicitObjectType Implicit0Type = GetImplicitCollisionType(Particle0, Implicit0);
+			const EImplicitObjectType Implicit1Type = GetImplicitCollisionType(Particle1, Implicit1);
+
+			//
+			// Basic implicit pairs
+			//
+
 			bOutSwap = false;
 			if (Implicit0Type == TBox<FReal, 3>::StaticType() && Implicit1Type == TBox<FReal, 3>::StaticType())
 			{
@@ -1573,7 +1604,14 @@ namespace Chaos
 				bOutSwap = true;
 				return EContactShapesType::CapsuleTriMesh;
 			}
-			else if (bIsConvex0 && Implicit1Type == FHeightField::StaticType())
+
+			//
+			// Convex implicits
+			//
+
+			const bool bIsConvex0 = (Implicit0 != nullptr) && Implicit0->IsConvex() && (Implicit0Type != ImplicitObjectType::LevelSet);
+			const bool bIsConvex1 = (Implicit1 != nullptr) && Implicit1->IsConvex() && (Implicit1Type != ImplicitObjectType::LevelSet);
+			if (bIsConvex0 && Implicit1Type == FHeightField::StaticType())
 			{
 				return EContactShapesType::ConvexHeightField;
 			}
@@ -1595,36 +1633,36 @@ namespace Chaos
 			{
 				return EContactShapesType::GenericConvexConvex;
 			}
-			else if ((Implicit0Type == FImplicitObjectUnion::StaticType()) || (Implicit1Type == FImplicitObjectUnion::StaticType()))
+
+			//
+			// Unions
+			//
+			
+			if ((Implicit0Type == FImplicitObjectUnion::StaticType()) || (Implicit1Type == FImplicitObjectUnion::StaticType()))
 			{
 				// Unknown here is used to mean "run the recursive collision detection path" in the midphase
 				// @todo(chaos): support FImplicitObjectUnion and FImplicitObjectUnionClustered more explicitly here
 				return EContactShapesType::Unknown;
 			}
-			else if (bHasSimplicial0 || bHasSimplicial1)
+
+			//
+			// Levelsets
+			//
+
+			const bool bIsLevelSet0 = (Implicit0Type == ImplicitObjectType::LevelSet);
+			const bool bIsLevelSet1 = (Implicit1Type == ImplicitObjectType::LevelSet);
+			if (bIsLevelSet0 || bIsLevelSet1)
 			{
-				// If only one of the shapes has no implicit, it goes first
-				const bool bSwapBasedOnImplicit = ((Implicit1Type == ImplicitObjectType::Unknown) && (Implicit0Type != ImplicitObjectType::Unknown));
-				// If only one shape has a Simplicial, it goes first
-				const bool bSwapBasedOnSimplicial = bHasSimplicial1 && !bHasSimplicial0;
-				bOutSwap = bSwapBasedOnImplicit || bSwapBasedOnSimplicial;
+				// If only one of the shapes is a levelset, it goes first
+				bOutSwap = bIsLevelSet1 && !bIsLevelSet0;
 				return EContactShapesType::LevelSetLevelSet;
 			}
 
+			// If we get here something is probably wrong.
+			ensure(false);
 			return EContactShapesType::Unknown;
 		}
 
-		EContactShapesType CalculateShapePairType(const FImplicitObject* Implicit0, const FBVHParticles* BVHParticles0, const FImplicitObject* Implicit1, const FBVHParticles* BVHParticles1, bool& bOutSwap)
-		{
-			const EImplicitObjectType Implicit0Type = (Implicit0 != nullptr) ? GetInnerType(Implicit0->GetCollisionType()) : ImplicitObjectType::Unknown;
-			const EImplicitObjectType Implicit1Type = (Implicit1 != nullptr) ? GetInnerType(Implicit1->GetCollisionType()) : ImplicitObjectType::Unknown;
-			const bool bIsConvex0 = (Implicit0 != nullptr) && Implicit0->IsConvex() && (Implicit0Type != ImplicitObjectType::LevelSet);
-			const bool bIsConvex1 = (Implicit1 != nullptr) && Implicit1->IsConvex() && (Implicit1Type != ImplicitObjectType::LevelSet);
-			const bool bHasSimplicial0 = (BVHParticles0 != nullptr) && (BVHParticles0->Size() > 0);
-			const bool bHasSimplicial1 = (BVHParticles1 != nullptr) && (BVHParticles1->Size() > 0);
-			return CalculateShapePairType(Implicit0Type, Implicit1Type, bIsConvex0, bIsConvex1, bHasSimplicial0, bHasSimplicial1, bOutSwap);
-		}
-
 		void ConstructConstraints(
 			FGeometryParticleHandle* Particle0, 
 			FGeometryParticleHandle* Particle1,
diff --git a/Engine/Source/Runtime/Experimental/Chaos/Public/Chaos/CollisionResolution.h b/Engine/Source/Runtime/Experimental/Chaos/Public/Chaos/CollisionResolution.h
--- a/Engine/Source/Runtime/Experimental/Chaos/Public/Chaos/CollisionResolution.h
+++ b/Engine/Source/Runtime/Experimental/Chaos/Public/Chaos/CollisionResolution.h
@@ -20,6 +20,7 @@ namespace Chaos
 	template <typename T, int d>
 	class TPlane;
 
+	class FBVHParticles;
 	class FCollisionContext;
 	class FHeightField;
 	class FImplicitObject;
@@ -49,8 +50,13 @@ namespace Chaos
 		/**
 		 * @brief Determine the shape pair type for use in UpdateConstraints
 		*/
-		EContactShapesType CHAOS_API CalculateShapePairType(const FImplicitObject* Implicit0, const FBVHParticles* BVHParticles0, const FImplicitObject* Implicit1, const FBVHParticles* BVHParticles1, bool& bOutSwap);
-		EContactShapesType CHAOS_API CalculateShapePairType(const EImplicitObjectType Implicit0Type, const EImplicitObjectType Implicit1Type, const bool bIsConvex0, const bool bIsConvex1, const bool bIsBVH0, const bool bIsBVH1, bool& bOutSwap);
+		EContactShapesType CHAOS_API CalculateShapePairType(const FGeometryParticleHandle* Particle0, const FImplicitObject* Implicit0, const FGeometryParticleHandle* Particle1, const FImplicitObject* Implicit1, bool& bOutSwap);
+
+		/**
+		 * Extract the implicit type from the geometry. This just calls GetInnerType but with a check for valid levels sets if the geometry is set to collide with a levelset
+		 * and is the collision type used by CalculateShapePairType()
+		 */
+		EImplicitObjectType CHAOS_API GetImplicitCollisionType(const FGeometryParticleHandle* Particle, const FImplicitObject* Implicit);
 
 		/**
 		 * @brief Whether CCD should be enabled for a contact given the current particle velocities etc
@@ -66,7 +72,13 @@ namespace Chaos
 		// Reset per-frame collision stat counters
 		void CHAOS_API ResetChaosCollisionCounters();
 
-		// See Above
+
+		//
+		//
+		// DEPRECATED STUFF (see above)
+		//
+		//
+
 		UE_DEPRECATED(5.3, "Replaced with a version that takes a bEnableSweep flag")
 		inline void CHAOS_API ConstructConstraints(TGeometryParticleHandle<FReal, 3>* Particle0, TGeometryParticleHandle<FReal, 3>* Particle1, const FImplicitObject* Implicit0, const FPerShapeData* Shape0, const FBVHParticles* Simplicial0, const FImplicitObject* Implicit1, const FPerShapeData* Shape1, const FBVHParticles* Simplicial1, const FRigidTransform3& ParticleWorldTransform0, const FRigidTransform3& Transform0, const FRigidTransform3& ParticleWorldTransform1, const FRigidTransform3& Transform1, const FReal CullDistance, const FReal Dt, const FCollisionContext& Context)
 		{
@@ -76,5 +88,16 @@ namespace Chaos
 		template<ECollisionUpdateType UpdateType>
 		UE_DEPRECATED(5.3, "Use UpdateConstraint, but call Constraint.SetShapeWorldTransform first (see implementation)")
 		void CHAOS_API UpdateConstraintFromGeometry(FPBDCollisionConstraint& Constraint, const FRigidTransform3& ParticleTransform0, const FRigidTransform3& ParticleTransform1, const FReal Dt);
+
+		UE_DEPRECATED(5.3, "Use the version which takes a particle and an implicit")
+		inline EContactShapesType CHAOS_API CalculateShapePairType(const FImplicitObject* Implicit0, const FBVHParticles* BVHParticles0, const FImplicitObject* Implicit1, const FBVHParticles* BVHParticles1, bool& bOutSwap)
+		{
+			return EContactShapesType::Unknown;
+		}
+		UE_DEPRECATED(5.3, "Use the version which takes a particle and an implicit")
+		inline EContactShapesType CHAOS_API CalculateShapePairType(const EImplicitObjectType Implicit0Type, const EImplicitObjectType Implicit1Type, const bool bIsConvex0, const bool bIsConvex1, const bool bIsBVH0, const bool bIsBVH1, bool& bOutSwap)
+		{
+			return EContactShapesType::Unknown;
+		}
 	}
 }
