Cluster modification callback. The clustermodification accessor is intentionally left as a stub class for now.

diff --git a/Engine/Source/Runtime/Experimental/Chaos/Private/Chaos/PBDRigidClustering.cpp b/Engine/Source/Runtime/Experimental/Chaos/Private/Chaos/PBDRigidClustering.cpp
--- a/Engine/Source/Runtime/Experimental/Chaos/Private/Chaos/PBDRigidClustering.cpp
+++ b/Engine/Source/Runtime/Experimental/Chaos/Private/Chaos/PBDRigidClustering.cpp
@@ -21,6 +21,7 @@
 #include "Chaos/PerParticleEulerStepVelocity.h"
 #include "Chaos/PerParticleEtherDrag.h"
 #include "Chaos/PerParticlePBDEulerStep.h"
+#include "Chaos/StrainModification.h"
 #include "PhysicsProxy/GeometryCollectionPhysicsProxy.h"
 #include "CoreMinimal.h"
 
@@ -87,7 +88,7 @@ namespace Chaos
 	// TPBDRigidClustering
 	//==========================================================================
 
-	FRigidClustering::FRigidClustering(FPBDRigidsEvolution& InEvolution, FPBDRigidClusteredParticles& InParticles)
+	FRigidClustering::FRigidClustering(FPBDRigidsEvolution& InEvolution, FPBDRigidClusteredParticles& InParticles, const TArray<ISimCallbackObject*>* InStrainModifiers)
 		: MEvolution(InEvolution)
 		, MParticles(InParticles)
 		, ClusterUnionManager(*this, InEvolution)
@@ -95,6 +96,7 @@ namespace Chaos
 		, DoGenerateBreakingData(false)
 		, MClusterConnectionFactor(1.0)
 		, MClusterUnionConnectionType(FClusterCreationParameters::EConnectionMethod::DelaunayTriangulation)
+		, StrainModifiers(InStrainModifiers)
 	{}
 
 	FRigidClustering::~FRigidClustering()
@@ -884,6 +886,13 @@ namespace Chaos
 			}
 
 			//
+			// Modify internal strains
+			//
+			if (StrainModifiers)
+			{
+				ApplyStrainModifiers();
+			}
+
 			//  Monitor the MStrain array for 0 or less values.
 			//  That will trigger a break too.
 			//
@@ -1209,6 +1218,7 @@ namespace Chaos
 							if (TPBDRigidClusteredParticleHandle<FReal, 3>* ClusteredChild = ClosestChild->CastToClustered())
 							{
 								ClusteredChild->CollisionImpulses() += AccumulatedImpulse;
+								UpdateTopLevelParticle(ClusteredChild);
 							}
 						}
 					}
@@ -1227,6 +1237,7 @@ namespace Chaos
 								if (TPBDRigidClusteredParticleHandle<FReal, 3>*ClusteredChild = Child->CastToClustered())
 								{
 									ClusteredChild->CollisionImpulses() += AccumulatedImpulse;
+									UpdateTopLevelParticle(ClusteredChild);
 								}
 							}
 						}
@@ -1279,6 +1290,18 @@ namespace Chaos
 		ClusteredParticle->ClusterGroupIndex() = 0;
 	}
 
+	void FRigidClustering::ApplyStrainModifiers()
+	{
+		if (StrainModifiers)
+		{
+			for (ISimCallbackObject* Modifier : *StrainModifiers)
+			{
+				FStrainModifierAccessor Accessor(*this);
+				Modifier->StrainModification_Internal(Accessor);
+			}
+		}
+	}
+
 	FPBDRigidClusteredParticleHandle*
 	FRigidClustering::DestroyClusterParticle(
 		FPBDRigidClusteredParticleHandle* ClusteredParticle,
diff --git a/Engine/Source/Runtime/Experimental/Chaos/Private/Chaos/PBDRigidsEvolutionGBF.cpp b/Engine/Source/Runtime/Experimental/Chaos/Private/Chaos/PBDRigidsEvolutionGBF.cpp
--- a/Engine/Source/Runtime/Experimental/Chaos/Private/Chaos/PBDRigidsEvolutionGBF.cpp
+++ b/Engine/Source/Runtime/Experimental/Chaos/Private/Chaos/PBDRigidsEvolutionGBF.cpp
@@ -719,10 +719,11 @@ FPBDRigidsEvolutionGBF::FPBDRigidsEvolutionGBF(
 	THandleArray<FChaosPhysicsMaterial>& SolverPhysicsMaterials, 
 	const TArray<ISimCallbackObject*>* InMidPhaseModifiers,
 	const TArray<ISimCallbackObject*>* InCCDModifiers,
+	const TArray<ISimCallbackObject*>* InStrainModifiers,
 	const TArray<ISimCallbackObject*>* InCollisionModifiers,
 	bool InIsSingleThreaded)
 	: Base(InParticles, SolverPhysicsMaterials, InIsSingleThreaded)
-	, Clustering(*this, Particles.GetClusteredParticles())
+	, Clustering(*this, Particles.GetClusteredParticles(), InStrainModifiers)
 	, CollisionConstraints(InParticles, Collided, PhysicsMaterials, PerParticlePhysicsMaterials, &SolverPhysicsMaterials, CalculateNumCollisionsPerBlock(), DefaultRestitutionThreshold)
 	, BroadPhase(InParticles)
 	, CollisionDetector(BroadPhase, CollisionConstraints)
diff --git a/Engine/Source/Runtime/Experimental/Chaos/Private/Chaos/StrainModification.cpp b/Engine/Source/Runtime/Experimental/Chaos/Private/Chaos/StrainModification.cpp
new file mode 100644
--- /dev/null
+++ b/Engine/Source/Runtime/Experimental/Chaos/Private/Chaos/StrainModification.cpp
@@ -0,0 +1,166 @@
+// Copyright Epic Games, Inc. All Rights Reserved.
+
+#include "Chaos/StrainModification.h"
+#include "Chaos/PBDRigidClustering.h"
+#include "PhysicsProxy/GeometryCollectionPhysicsProxy.h"
+#include "GeometryCollection/GeometryCollectionSimulationCoreTypes.h"
+#include "GeometryCollection/GeometryCollection.h"
+
+const FGeometryCollectionPhysicsProxy* Chaos::FStrainedProxyModifier::GetProxy() const
+{
+	return Proxy;
+}
+
+int32 Chaos::FStrainedProxyModifier::GetNumRestBreakables() const
+{
+	if (RestChildren)
+	{
+		return RestChildren->Num();
+	}
+	return 0;
+}
+
+int32 Chaos::FStrainedProxyModifier::GetNumBreakingStrains() const
+{
+	// Make sure we have a proxy and rest-children
+	if (Proxy == nullptr) { return 0; }
+	if (RestChildren == nullptr) { return 0; }
+
+	// Loop over each child, checking whether or not it will have been freed
+	// by the strain that it has accumulated
+	int32 NumBreakingStrains = 0;
+	for (int32 RestChildIdx : *RestChildren)
+	{
+		Chaos::FPBDRigidClusteredParticleHandle* ChildHandle = Proxy->GetSolverParticleHandles()[RestChildIdx];
+		if (ChildHandle->Parent() == nullptr) { continue; }
+
+		// Get the applied and internal strain
+		//
+		// TODO: Make this a computation internal to the strain, or the connection
+		// graph or something... this logic is currently copied from
+		// FRigidClustering::ReleaseClusterParticlesImpl, but should be tucked behind
+		// an interface somewhere.
+		const Chaos::FReal CollisionImpulses = ChildHandle->CollisionImpulses();
+		const Chaos::FReal ExternalStrain = ChildHandle->GetExternalStrain();
+		const Chaos::FReal MaxAppliedStrain = FMath::Max(CollisionImpulses, ExternalStrain);
+		const Chaos::FReal InternalStrain = ChildHandle->GetInternalStrains();
+		if (MaxAppliedStrain >= InternalStrain)
+		{
+			++NumBreakingStrains;
+		}
+	}
+
+	// Return the number of breaking strains
+	return NumBreakingStrains;
+}
+
+void Chaos::FStrainedProxyModifier::ClearStrains()
+{
+	// Make sure we have a proxy and rest-children
+	if (Proxy == nullptr) { return; }
+	if (RestChildren == nullptr) { return; }
+
+	// Loop over each child, clearing the strain associated with it
+	for (int32 RestChildIdx : *RestChildren)
+	{
+		if (Chaos::FPBDRigidClusteredParticleHandle* ChildHandle = Proxy->GetSolverParticleHandles()[RestChildIdx])
+		{
+			// Clear accumulated collision impulses
+			ChildHandle->ClearCollisionImpulse();
+
+			// Clear accumulated external strains
+			//
+			// NOTE: The following is deprecated:
+			// ChildHandle->ClearExternalStrain();
+			// Instead we use:
+			RigidClustering.SetExternalStrain(ChildHandle, FReal(0));
+		}
+	}
+}
+
+const TSet<int32>* Chaos::FStrainedProxyModifier::InitRestChildren(FGeometryCollectionPhysicsProxy* Proxy)
+{
+	if (Proxy == nullptr) { return nullptr; }
+
+	// Get the root index of the proxy
+	FSimulationParameters& Parameters = Proxy->GetSimParameters();
+	const int32 RootIndex = Parameters.InitialRootIndex;
+	if (RootIndex == INDEX_NONE) { return nullptr; }
+
+	// Return a ptr to the set of child indices
+	return &Parameters.RestCollection->Children[RootIndex];
+}
+
+Chaos::FStrainedProxyModifier Chaos::FStrainedProxyIterator::operator*()
+{
+	return Chaos::FStrainedProxyModifier(RigidClustering, Proxies[Index]);
+}
+
+Chaos::FStrainedProxyIterator& Chaos::FStrainedProxyIterator::operator++()
+{
+	++Index;
+	return *this;
+}
+
+bool Chaos::FStrainedProxyIterator::operator==(const FStrainedProxyIterator& Other) const
+{
+	return
+		Proxies == Other.Proxies &&
+		Index == Other.Index;
+}
+
+Chaos::FStrainedProxyRange::FStrainedProxyRange(Chaos::FRigidClustering& InRigidClustering, const bool bRootLevelOnly)
+	: RigidClustering(InRigidClustering)
+{
+	const TSet<Chaos::FPBDRigidClusteredParticleHandle*>& StrainedParents = RigidClustering.GetTopLevelClusterParentsStrained();
+
+	Proxies.Reserve(StrainedParents.Num());
+	for (Chaos::FPBDRigidClusteredParticleHandle* Cluster : StrainedParents)
+	{
+		// Make sure the cluster's physics proxy exists and is the right type
+		IPhysicsProxyBase* ProxyBase = Cluster->PhysicsProxy();
+		if (ProxyBase == nullptr)
+		{
+			continue;
+		}
+
+		if (ProxyBase->GetType() != EPhysicsProxyType::GeometryCollectionType)
+		{
+			continue;
+		}
+
+		FGeometryCollectionPhysicsProxy* Proxy = static_cast<FGeometryCollectionPhysicsProxy*>(ProxyBase);
+
+		// Make sure the rest collection has a root index
+		if (bRootLevelOnly)
+		{
+			FSimulationParameters& Parameters = Proxy->GetSimParameters();
+			const int32 RootIndex = Parameters.InitialRootIndex;
+			TArray<Chaos::FPBDRigidClusteredParticleHandle*>& ParticleHandles = Proxy->GetSolverParticleHandles();
+			if (ParticleHandles.IsValidIndex(RootIndex))
+			{
+				if (Cluster != ParticleHandles[RootIndex])
+				{
+					continue;
+				}
+			}
+		}
+
+		// Only need to use AddUnique if we're not checking for root, since at most
+		// one cluster will have the rest collection's root index.
+		if (bRootLevelOnly)
+		{
+			Proxies.Add(Proxy);
+		}
+		else
+		{
+			Proxies.AddUnique(Proxy);
+		}
+	}
+}
+
+Chaos::FStrainedProxyRange Chaos::FStrainModifierAccessor::GetStrainedProxies(const bool bRootLevelOnly)
+{
+	return FStrainedProxyRange(RigidClustering, bRootLevelOnly);
+}
+
diff --git a/Engine/Source/Runtime/Experimental/Chaos/Private/PBDRigidsSolver.cpp b/Engine/Source/Runtime/Experimental/Chaos/Private/PBDRigidsSolver.cpp
--- a/Engine/Source/Runtime/Experimental/Chaos/Private/PBDRigidsSolver.cpp
+++ b/Engine/Source/Runtime/Experimental/Chaos/Private/PBDRigidsSolver.cpp
@@ -504,7 +504,7 @@ namespace Chaos
 		, FloorHeight(0.f)
 		, bIsDeterministic(false)
 		, Particles(UniqueIndices)
-		, MEvolution(new FPBDRigidsEvolution(Particles, SimMaterials, &MidPhaseModifiers, &CCDModifiers, &ContactModifiers, BufferingModeIn == Chaos::EMultiBufferMode::Single))
+		, MEvolution(new FPBDRigidsEvolution(Particles, SimMaterials, &MidPhaseModifiers, &CCDModifiers, &StrainModifiers, &ContactModifiers, BufferingModeIn == Chaos::EMultiBufferMode::Single))
 		, MEventManager(new FEventManager(BufferingModeIn))
 		, MSolverEventFilters(new FSolverEventFilters())
 		, MDirtyParticlesBuffer(new FDirtyParticlesBuffer(BufferingModeIn, BufferingModeIn == Chaos::EMultiBufferMode::Single))
@@ -875,7 +875,7 @@ namespace Chaos
 		SetMaxDeltaTime_External(1.0f);
 		SetMinDeltaTime_External(UE_SMALL_NUMBER);
 		SetMaxSubSteps_External(1);
-		MEvolution = TUniquePtr<FPBDRigidsEvolution>(new FPBDRigidsEvolution(Particles, SimMaterials, &MidPhaseModifiers, &CCDModifiers, &ContactModifiers, BufferMode == EMultiBufferMode::Single)); 
+		MEvolution = TUniquePtr<FPBDRigidsEvolution>(new FPBDRigidsEvolution(Particles, SimMaterials, &MidPhaseModifiers, &CCDModifiers, &StrainModifiers, &ContactModifiers, BufferMode == EMultiBufferMode::Single)); 
 
 		PerSolverField = MakeUnique<FPerSolverFieldSystem>();
 
@@ -1470,6 +1470,11 @@ namespace Chaos
 				CCDModifiers.Add(SimCallbackObject);
 			}
 
+			if (SimCallbackObject->HasOption(ESimCallbackOptions::StrainModification))
+			{
+				StrainModifiers.Add(SimCallbackObject);
+			}
+
 			if (SimCallbackObject->HasOption(ESimCallbackOptions::ContactModification))
 			{
 				ContactModifiers.Add(SimCallbackObject);
diff --git a/Engine/Source/Runtime/Experimental/Chaos/Public/Chaos/Framework/PhysicsSolverBase.h b/Engine/Source/Runtime/Experimental/Chaos/Public/Chaos/Framework/PhysicsSolverBase.h
--- a/Engine/Source/Runtime/Experimental/Chaos/Public/Chaos/Framework/PhysicsSolverBase.h
+++ b/Engine/Source/Runtime/Experimental/Chaos/Public/Chaos/Framework/PhysicsSolverBase.h
@@ -662,6 +662,7 @@ namespace Chaos
 	TArray<ISimCallbackObject*> SimCallbackObjects;
 	TArray<ISimCallbackObject*> MidPhaseModifiers;
 	TArray<ISimCallbackObject*> CCDModifiers;
+	TArray<ISimCallbackObject*> StrainModifiers;
 	TArray<ISimCallbackObject*> ContactModifiers;
 	TArray<ISimCallbackObject*> RegistrationWatchers;
 	TArray<ISimCallbackObject*> UnregistrationWatchers;
diff --git a/Engine/Source/Runtime/Experimental/Chaos/Public/Chaos/PBDRigidClustering.h b/Engine/Source/Runtime/Experimental/Chaos/Public/Chaos/PBDRigidClustering.h
--- a/Engine/Source/Runtime/Experimental/Chaos/Public/Chaos/PBDRigidClustering.h
+++ b/Engine/Source/Runtime/Experimental/Chaos/Public/Chaos/PBDRigidClustering.h
@@ -61,7 +61,7 @@ public:
 	typedef TMap<FClusterHandle, FRigidHandleArray>				FClusterMap;
 	typedef TFunction<void(FRigidClustering&, FRigidHandle)>	FVisitorFunction;
 
-	FRigidClustering(FRigidEvolution& InEvolution, FPBDRigidClusteredParticles& InParticles);
+	FRigidClustering(FRigidEvolution& InEvolution, FPBDRigidClusteredParticles& InParticles, const TArray<ISimCallbackObject*>* InStrainModifiers);
 	~FRigidClustering();
 
 	//
@@ -364,11 +364,16 @@ public:
 
 	FClusterUnionManager& GetClusterUnionManager() { return ClusterUnionManager; }
 	const FClusterUnionManager& GetClusterUnionManager() const { return ClusterUnionManager; }
+
+
+	const TSet<Chaos::FPBDRigidClusteredParticleHandle*>& GetTopLevelClusterParentsStrained() const { return TopLevelClusterParentsStrained; }
+
  protected:
 
 	void ComputeStrainFromCollision(const FPBDCollisionConstraints& CollisionRule);
 	void ResetCollisionImpulseArray();
 	void DisableCluster(FPBDRigidClusteredParticleHandle* ClusteredParticle);
+	void ApplyStrainModifiers();
 
 	/*
 	* Connectivity
@@ -404,6 +409,7 @@ public:
 	TArray<FPBDRigidParticleHandle*> CreateClustersFromNewIslands(TArray<FParticleIsland>& Islands, FPBDRigidClusteredParticleHandle* ClusteredParent);
 
 	void UpdateTopLevelParticle(FPBDRigidClusteredParticleHandle* Particle);
+
 private:
 
 	FRigidEvolution& MEvolution;
@@ -434,6 +440,9 @@ private:
 
 	FReal MClusterConnectionFactor;
 	FClusterCreationParameters::EConnectionMethod MClusterUnionConnectionType;
+
+	// Sim callback objects which implement cluster modification steps
+	const TArray<ISimCallbackObject*>* StrainModifiers;
 };
 
 } // namespace Chaos
diff --git a/Engine/Source/Runtime/Experimental/Chaos/Public/Chaos/PBDRigidsEvolutionGBF.h b/Engine/Source/Runtime/Experimental/Chaos/Public/Chaos/PBDRigidsEvolutionGBF.h
--- a/Engine/Source/Runtime/Experimental/Chaos/Public/Chaos/PBDRigidsEvolutionGBF.h
+++ b/Engine/Source/Runtime/Experimental/Chaos/Public/Chaos/PBDRigidsEvolutionGBF.h
@@ -71,6 +71,7 @@ namespace Chaos
 			THandleArray<FChaosPhysicsMaterial>& SolverPhysicsMaterials, 
 			const TArray<ISimCallbackObject*>* InMidPhaseModifiers = nullptr,
 			const TArray<ISimCallbackObject*>* InCCDModifiers = nullptr,
+			const TArray<ISimCallbackObject*>* InStrainModifiers = nullptr,
 			const TArray<ISimCallbackObject*>* InCollisionModifiers = nullptr,
 			bool InIsSingleThreaded = false);
 		CHAOS_API ~FPBDRigidsEvolutionGBF();
diff --git a/Engine/Source/Runtime/Experimental/Chaos/Public/Chaos/ParticleHandle.h b/Engine/Source/Runtime/Experimental/Chaos/Public/Chaos/ParticleHandle.h
--- a/Engine/Source/Runtime/Experimental/Chaos/Public/Chaos/ParticleHandle.h
+++ b/Engine/Source/Runtime/Experimental/Chaos/Public/Chaos/ParticleHandle.h
@@ -1287,6 +1287,7 @@ public:
 	const T& CollisionImpulses() const { return PBDRigidClusteredParticles->CollisionImpulses(ParticleIdx); }
 	T& CollisionImpulses() { return PBDRigidClusteredParticles->CollisionImpulses(ParticleIdx); }
 	void SetCollisionImpulses(const T Value) { PBDRigidClusteredParticles->CollisionImpulses(ParticleIdx) = Value; }
+	void ClearCollisionImpulse() { PBDRigidClusteredParticles->CollisionImpulses(ParticleIdx) = static_cast<T>(0); }
 
 	T GetExternalStrain() const { return PBDRigidClusteredParticles->ExternalStrains(ParticleIdx); }
 	UE_DEPRECATED(5.2, "This method should not be used anymore. FRigidClustering::SetExternalStrain should be used instead.")
diff --git a/Engine/Source/Runtime/Experimental/Chaos/Public/Chaos/SimCallbackObject.h b/Engine/Source/Runtime/Experimental/Chaos/Public/Chaos/SimCallbackObject.h
--- a/Engine/Source/Runtime/Experimental/Chaos/Public/Chaos/SimCallbackObject.h
+++ b/Engine/Source/Runtime/Experimental/Chaos/Public/Chaos/SimCallbackObject.h
@@ -12,6 +12,7 @@ namespace Chaos
 class FPhysicsSolverBase;
 class FMidPhaseModifierAccessor;
 class FCCDModifierAccessor;
+class FStrainModifierAccessor;
 class FCollisionContactModifier;
 class FSingleParticlePhysicsProxy;
 
@@ -20,16 +21,17 @@ namespace Utilities
 	CHAOS_API FReal GetSolverPhysicsResultsTime(FPhysicsSolverBase*);
 }
 
-enum class ESimCallbackOptions : uint8
+enum class ESimCallbackOptions : uint16
 {
 	Presimulate				= 1 << 0,
 	MidPhaseModification	= 1 << 1,
 	CCDModification			= 1 << 2,
 	ContactModification		= 1 << 3,
-	ParticleRegister		= 1 << 4,
-	ParticleUnregister		= 1 << 5,
-	RunOnFrozenGameThread	= 1 << 6,
-	Rewind					= 1 << 7
+	StrainModification		= 1 << 5,
+	ParticleRegister		= 1 << 6,
+	ParticleUnregister		= 1 << 7,
+	RunOnFrozenGameThread	= 1 << 8,
+	Rewind					= 1 << 9
 };
 ENUM_CLASS_FLAGS(ESimCallbackOptions)
 
@@ -79,6 +81,11 @@ public:
 		OnCCDModification_Internal(Modifier);
 	}
 
+	void StrainModification_Internal(FStrainModifierAccessor& Modifier)
+	{
+		OnStrainModification_Internal(Modifier);
+	}
+
 	void ContactModification_Internal(FCollisionContactModifier& Modifier)
 	{
 		OnContactModification_Internal(Modifier);
@@ -191,6 +198,11 @@ private:
 		check(false);
 	}
 
+	virtual void OnStrainModification_Internal(FStrainModifierAccessor& Modifier)
+	{
+		check(false);
+	}
+
 	/**
 	* Called once per simulation step. Allows user to modify contacts
 	*
diff --git a/Engine/Source/Runtime/Experimental/Chaos/Public/Chaos/StrainModification.h b/Engine/Source/Runtime/Experimental/Chaos/Public/Chaos/StrainModification.h
new file mode 100644
--- /dev/null
+++ b/Engine/Source/Runtime/Experimental/Chaos/Public/Chaos/StrainModification.h
@@ -0,0 +1,142 @@
+// Copyright Epic Games, Inc. All Rights Reserved.
+
+#pragma once
+
+#include "Chaos/Core.h"
+#include "ParticleHandleFwd.h"
+
+class FGeometryCollectionPhysicsProxy;
+
+namespace Chaos
+{
+	class FRigidClustering;
+	class FStrainModifierAccessor;
+	class FStrainedProxyModifier;
+	class FStrainedProxyIterator;
+	class FStrainedProxyRange;
+	class FStrainModifierAccessor;
+
+	// FStrainedProxyModifier
+	//
+	// User-facing api for accessing the proxy of a strained cluster. Provides const access to the
+	// proxy and limited read/write access to its internal strain data.
+	class CHAOS_API FStrainedProxyModifier
+	{
+	public:
+		FStrainedProxyModifier(FRigidClustering& InRigidClustering, FGeometryCollectionPhysicsProxy* InProxy)
+			: RigidClustering(InRigidClustering)
+			, Proxy(InProxy)
+			, RestChildren(InitRestChildren(InProxy))
+		{ }
+
+		FStrainedProxyModifier(const FStrainedProxyModifier& Other)
+			: RigidClustering(Other.RigidClustering)
+			, Proxy(Other.Proxy)
+			, RestChildren(Other.RestChildren)
+		{ }
+
+		// Get the proxy that owns the strained cluster or clusters
+		const FGeometryCollectionPhysicsProxy* GetProxy() const;
+
+		// Get the number of level-1 strainable entities (number of rest-children in the per-particle
+		// strain model, or number of rest-connections in the edge/area model).
+		int32 GetNumRestBreakables() const;
+
+		// Get the number of breaking strains amongst the level-1 strainables
+		int32 GetNumBreakingStrains() const;
+
+		// Clear strains for all strained cluster children
+		void ClearStrains();
+
+	private:
+
+		static const TSet<int32>* InitRestChildren(FGeometryCollectionPhysicsProxy* Proxy);
+
+		FRigidClustering& RigidClustering;
+		FGeometryCollectionPhysicsProxy* Proxy;
+		const TSet<int32>* RestChildren;
+	};
+
+	// FStrainedProxyIterator
+	//
+	// Iterator produced by FStrainedProxyRange, for looping over proxies which are
+	// associated with rigid clusters which have been strained.
+	class CHAOS_API FStrainedProxyIterator
+	{
+	public:
+		FStrainedProxyIterator(FRigidClustering& InRigidClustering, TArray<FGeometryCollectionPhysicsProxy*>& InProxies, int32 InIndex)
+			: RigidClustering(InRigidClustering)
+			, Proxies(InProxies)
+			, Index(InIndex)
+		{ }
+
+		FStrainedProxyIterator(const FStrainedProxyIterator& Other)
+			: RigidClustering(Other.RigidClustering)
+			, Proxies(Other.Proxies)
+			, Index(Other.Index)
+		{ }
+
+		FStrainedProxyModifier operator*();
+		FStrainedProxyIterator& operator++();
+		bool operator==(const FStrainedProxyIterator& Other) const;
+		bool operator!=(const FStrainedProxyIterator& Other) const
+		{
+			return !operator==(Other);
+		}
+
+	private:
+		FRigidClustering& RigidClustering;
+		TArray<FGeometryCollectionPhysicsProxy*>& Proxies;
+		int32 Index;
+	};
+
+	// FStrainedProxyRange
+	//
+	// Provides an interface for use with ranged-for, for iterating over strained proxies.
+	// Constructor produces filtered array of proxies, and begin and end functions produce
+	// iterators which can modify strain related properties of the clusters associated
+	// with each proxy.
+	class CHAOS_API FStrainedProxyRange
+	{
+	public:
+		FStrainedProxyRange(Chaos::FRigidClustering& InRigidClustering, bool bRootLevelOnly);
+
+		FStrainedProxyRange(const FStrainedProxyRange& Other)
+			: RigidClustering(Other.RigidClustering)
+			, Proxies(Other.Proxies)
+		{ }
+
+		FStrainedProxyIterator begin()
+		{
+			return FStrainedProxyIterator(RigidClustering, Proxies, 0);
+		}
+
+		FStrainedProxyIterator end()
+		{
+			return FStrainedProxyIterator(RigidClustering, Proxies, Proxies.Num());
+		}
+
+	private:
+		FRigidClustering& RigidClustering;
+		TArray<FGeometryCollectionPhysicsProxy*> Proxies;
+	};
+
+	// FStrainModifierAccessor
+	//
+	// Provides access to strained proxies and clusters
+	class CHAOS_API FStrainModifierAccessor
+	{
+	public:
+
+		FStrainModifierAccessor(FRigidClustering& InRigidClustering) : RigidClustering(InRigidClustering) { }
+
+		// Get an iterable range of unique geometry collection proxies which
+		// correspond to all strained clusters. Optionally, only include proxies
+		// for whom the strained parent is still the original root (ie, unbroken).
+		FStrainedProxyRange GetStrainedProxies(bool bRootLevelOnly = true);
+
+	private:
+
+		FRigidClustering& RigidClustering;
+	};
+}
