GeometryCollection object can store custom instance data to pass to instances spawned in ISM pools.
Also add a per component custom data which, if present, is added to all ISM pool instances that the component generates.

diff --git a/Engine/Source/Runtime/Engine/Classes/Components/HierarchicalInstancedStaticMeshComponent.h b/Engine/Source/Runtime/Engine/Classes/Components/HierarchicalInstancedStaticMeshComponent.h
--- a/Engine/Source/Runtime/Engine/Classes/Components/HierarchicalInstancedStaticMeshComponent.h
+++ b/Engine/Source/Runtime/Engine/Classes/Components/HierarchicalInstancedStaticMeshComponent.h
@@ -234,7 +234,7 @@ public:
 	virtual bool RemoveInstances(const TArray<int32>& InstancesToRemove) override;
 	virtual bool UpdateInstanceTransform(int32 InstanceIndex, const FTransform& NewInstanceTransform, bool bWorldSpace, bool bMarkRenderStateDirty = false, bool bTeleport = false) override;
 	virtual bool SetCustomDataValue(int32 InstanceIndex, int32 CustomDataIndex, float CustomDataValue, bool bMarkRenderStateDirty = false) override;
-	virtual bool SetCustomData(int32 InstanceIndex, const TArray<float>& InCustomData, bool bMarkRenderStateDirty = false) override;
+	virtual bool SetCustomData(int32 InstanceIndex, TArrayView<const float> InCustomData, bool bMarkRenderStateDirty = false) override;
 	virtual bool BatchUpdateInstancesTransforms(int32 StartInstanceIndex, const TArray<FTransform>& NewInstancesTransforms, bool bWorldSpace=false, bool bMarkRenderStateDirty=false, bool bTeleport=false) override;
 	virtual bool BatchUpdateInstancesTransform(int32 StartInstanceIndex, int32 NumInstances, const FTransform& NewInstancesTransform, bool bWorldSpace=false, bool bMarkRenderStateDirty=false, bool bTeleport=false) override;
 	virtual bool BatchUpdateInstancesData(int32 StartInstanceIndex, int32 NumInstances, FInstancedStaticMeshInstanceData* StartInstanceData, bool bMarkRenderStateDirty = false, bool bTeleport = false) override;
diff --git a/Engine/Source/Runtime/Engine/Classes/Components/InstancedStaticMeshComponent.h b/Engine/Source/Runtime/Engine/Classes/Components/InstancedStaticMeshComponent.h
--- a/Engine/Source/Runtime/Engine/Classes/Components/InstancedStaticMeshComponent.h
+++ b/Engine/Source/Runtime/Engine/Classes/Components/InstancedStaticMeshComponent.h
@@ -163,7 +163,7 @@ class ENGINE_API UInstancedStaticMeshComponent : public UStaticMeshComponent, pu
 	virtual bool SetCustomDataValue(int32 InstanceIndex, int32 CustomDataIndex, float CustomDataValue, bool bMarkRenderStateDirty = false);
 
 	/** Per Instance Custom Data */
-	virtual bool SetCustomData(int32 InstanceIndex, const TArray<float>& CustomDataFloats, bool bMarkRenderStateDirty = false); 
+	virtual bool SetCustomData(int32 InstanceIndex, TArrayView<const float> CustomDataFloats, bool bMarkRenderStateDirty = false); 
 
 	/** Preallocated memory to include the new added instances count, to prevent reallloc during the add operation. */
 	virtual void PreAllocateInstancesMemory(int32 AddedInstanceCount);
diff --git a/Engine/Source/Runtime/Engine/Private/HierarchicalInstancedStaticMesh.cpp b/Engine/Source/Runtime/Engine/Private/HierarchicalInstancedStaticMesh.cpp
--- a/Engine/Source/Runtime/Engine/Private/HierarchicalInstancedStaticMesh.cpp
+++ b/Engine/Source/Runtime/Engine/Private/HierarchicalInstancedStaticMesh.cpp
@@ -2410,7 +2410,7 @@ bool UHierarchicalInstancedStaticMeshComponent::SetCustomDataValue(int32 Instanc
 	return Result;
 }
 
-bool UHierarchicalInstancedStaticMeshComponent::SetCustomData(int32 InstanceIndex, const TArray<float>& InCustomData, bool bMarkRenderStateDirty)
+bool UHierarchicalInstancedStaticMeshComponent::SetCustomData(int32 InstanceIndex, TArrayView<const float> InCustomData, bool bMarkRenderStateDirty)
 {
 	if (!PerInstanceSMData.IsValidIndex(InstanceIndex) || InCustomData.Num() == 0)
 	{
diff --git a/Engine/Source/Runtime/Engine/Private/InstancedStaticMesh.cpp b/Engine/Source/Runtime/Engine/Private/InstancedStaticMesh.cpp
--- a/Engine/Source/Runtime/Engine/Private/InstancedStaticMesh.cpp
+++ b/Engine/Source/Runtime/Engine/Private/InstancedStaticMesh.cpp
@@ -3383,7 +3383,7 @@ bool UInstancedStaticMeshComponent::SetCustomDataValue(int32 InstanceIndex, int3
 	return true;
 }
 
-bool UInstancedStaticMeshComponent::SetCustomData(int32 InstanceIndex, const TArray<float>& InCustomData, bool bMarkRenderStateDirty)
+bool UInstancedStaticMeshComponent::SetCustomData(int32 InstanceIndex, TArrayView<const float> InCustomData, bool bMarkRenderStateDirty)
 {
 	if (!PerInstanceSMData.IsValidIndex(InstanceIndex) || InCustomData.Num() == 0)
 	{
diff --git a/Engine/Source/Runtime/Experimental/GeometryCollectionEngine/Private/GeometryCollection/GeometryCollectionComponent.cpp b/Engine/Source/Runtime/Experimental/GeometryCollectionEngine/Private/GeometryCollection/GeometryCollectionComponent.cpp
--- a/Engine/Source/Runtime/Experimental/GeometryCollectionEngine/Private/GeometryCollection/GeometryCollectionComponent.cpp
+++ b/Engine/Source/Runtime/Experimental/GeometryCollectionEngine/Private/GeometryCollection/GeometryCollectionComponent.cpp
@@ -4308,7 +4308,9 @@ void UGeometryCollectionComponent::RegisterToISMPool()
 
 						FGeometryCollectionStaticMeshInstance StaticMeshInstance;
 						StaticMeshInstance.StaticMesh = Mesh;
-						ISMPoolRootProxyMeshIds.Add(ISMPoolComp->AddMeshToGroup(ISMPoolMeshGroupIndex, StaticMeshInstance, 1, bChaos_GC_UseHierarchicalISMForProxyMesh));
+						StaticMeshInstance.NumCustomDataFloats = ISMPoolMaterialCustomData.Num();
+
+						ISMPoolRootProxyMeshIds.Add(ISMPoolComp->AddMeshToGroup(ISMPoolMeshGroupIndex, StaticMeshInstance, 1, ISMPoolMaterialCustomData, bChaos_GC_UseHierarchicalISMForProxyMesh));
 					}
 				}
 
@@ -4361,8 +4363,10 @@ void UGeometryCollectionComponent::AddAutoInstancesToISMPool()
 			if (RestCollection->GetGeometryCollection())
 			{
 				// first count the instance per mesh 
+				const int32 NumMeshes = RestCollection->AutoInstanceMeshes.Num();
 				TArray<int32> InstanceCounts;
-				InstanceCounts.AddZeroed(RestCollection->AutoInstanceMeshes.Num());
+				InstanceCounts.AddZeroed(NumMeshes);
+				int32 TotalInstanceCount = 0;
 
 				const TManagedArray<TSet<int32>>& Children = RestCollection->GetGeometryCollection()->Children;
 
@@ -4375,12 +4379,21 @@ void UGeometryCollectionComponent::AddAutoInstancesToISMPool()
 						if (Children[TransformIndex].Num() == 0)
 						{
 							InstanceCounts[AutoInstanceMeshIndex]++;
+							TotalInstanceCount++;
 						}
 					}
 				}
 
+				// Apply custom instance data from geometry collcection object if custom data array is set up correctly so that it's size matches a multiple of the number of instances.
+				const bool bHasValidAssetCustomInstanceData = TotalInstanceCount && RestCollection->AutoInstanceMaterialCustomData.Num() % TotalInstanceCount == 0;
+				const int32 NumCustomDataFloatsFromAsset = bHasValidAssetCustomInstanceData ? RestCollection->AutoInstanceMaterialCustomData.Num() / TotalInstanceCount : 0;
+				const int32 NumCustomDataFloatsFromComponent = ISMPoolMaterialCustomData.Num();
+				const int32 NumCustomDataFloats = NumCustomDataFloatsFromAsset + NumCustomDataFloatsFromComponent;
+				TArray<float, TInlineAllocator<64>> CustomFloatData;
+				int32 CustomDataFromAssetReadIndex = 0;
+
 				// now register each mesh 
-				for (int32 MeshIndex = 0; MeshIndex < RestCollection->AutoInstanceMeshes.Num(); MeshIndex++)
+				for (int32 MeshIndex = 0; MeshIndex < NumMeshes; MeshIndex++)
 				{
 					const FGeometryCollectionAutoInstanceMesh& AutoInstanceMesh = RestCollection->AutoInstanceMeshes[MeshIndex];
 					if (const UStaticMesh* StaticMesh = AutoInstanceMesh.Mesh)
@@ -4401,7 +4414,24 @@ void UGeometryCollectionComponent::AddAutoInstancesToISMPool()
 						{
 							StaticMeshInstance.MaterialsOverrides = AutoInstanceMesh.Materials;
 						}
-						ISMPoolAutoInstancesMeshIds.Add(ISMPoolComp->AddMeshToGroup(ISMPoolMeshGroupIndex, StaticMeshInstance, InstanceCounts[MeshIndex], bChaos_GC_UseHierarchicalISMForLeafMeshes));
+						StaticMeshInstance.NumCustomDataFloats = NumCustomDataFloats;
+
+						if (NumCustomDataFloats > 0)
+						{
+							const int32 InstanceCount = InstanceCounts[MeshIndex];
+							CustomFloatData.Reset();
+							CustomFloatData.Reserve(InstanceCount * NumCustomDataFloats);
+							
+							for (int32 InstanceIndex = 0; InstanceIndex < InstanceCount; ++InstanceIndex)
+							{
+								CustomFloatData.Append(ISMPoolMaterialCustomData);
+								// Append each set of values from AutoInstanceMaterialCustomData after the common values from ISMPoolMaterialCustomData
+								CustomFloatData.Append(&RestCollection->AutoInstanceMaterialCustomData[CustomDataFromAssetReadIndex], NumCustomDataFloatsFromAsset);
+								CustomDataFromAssetReadIndex += NumCustomDataFloatsFromAsset;
+							}
+						}
+
+						ISMPoolAutoInstancesMeshIds.Add(ISMPoolComp->AddMeshToGroup(ISMPoolMeshGroupIndex, StaticMeshInstance, InstanceCounts[MeshIndex], CustomFloatData, bChaos_GC_UseHierarchicalISMForLeafMeshes));
 					}
 				}
 			}
@@ -4471,7 +4501,7 @@ void UGeometryCollectionComponent::RefreshISMPoolInstances()
 								}
 
 								TArray<FTransform> InstanceTransforms;
-								for (int32 MeshIndex = 0; MeshIndex < RestCollection->AutoInstanceMeshes.Num(); MeshIndex++)
+								for (int32 MeshIndex = 0; MeshIndex < ISMPoolAutoInstancesMeshIds.Num(); MeshIndex++)
 								{
 									InstanceTransforms.Reset(NumTransforms); // Allocate for worst case
 									for (int32 TransformIndex = 0; TransformIndex < NumTransforms; TransformIndex++)
diff --git a/Engine/Source/Runtime/Experimental/GeometryCollectionEngine/Private/GeometryCollection/GeometryCollectionISMPoolComponent.cpp b/Engine/Source/Runtime/Experimental/GeometryCollectionEngine/Private/GeometryCollection/GeometryCollectionISMPoolComponent.cpp
--- a/Engine/Source/Runtime/Experimental/GeometryCollectionEngine/Private/GeometryCollection/GeometryCollectionISMPoolComponent.cpp
+++ b/Engine/Source/Runtime/Experimental/GeometryCollectionEngine/Private/GeometryCollection/GeometryCollectionISMPoolComponent.cpp
@@ -61,6 +61,7 @@ FGeometryCollectionISM::FGeometryCollectionISM(AActor* OwmingActor, const FGeome
 		{
 			ISMC->SetMaterial(MaterialIndex, MeshInstance.MaterialsOverrides[MaterialIndex]);
 		}
+		ISMC->NumCustomDataFloats = MeshInstance.NumCustomDataFloats;
 		ISMC->SetCullDistances(0, 0);
 		ISMC->SetCanEverAffectNavigation(false);
 		ISMC->SetCollisionEnabled(ECollisionEnabled::NoCollision);
@@ -73,9 +74,10 @@ FGeometryCollectionISM::FGeometryCollectionISM(AActor* OwmingActor, const FGeome
 	}
 }
 
-int32 FGeometryCollectionISM::AddInstanceGroup(int32 InstanceCount)
+int32 FGeometryCollectionISM::AddInstanceGroup(int32 InstanceCount, TArrayView<const float> CustomDataFloats)
 {
 	const int32 InstanceGroupIndex = InstanceGroups.AddGroup(InstanceCount);
+	const FInstanceGroups::FInstanceGroupRange& NewInstanceGroup = InstanceGroups.GetGroup(InstanceGroupIndex);
 
 	ISMComponent->PreAllocateInstancesMemory(InstanceCount);
 
@@ -86,10 +88,22 @@ int32 FGeometryCollectionISM::AddInstanceGroup(int32 InstanceCount)
 
 	ISMComponent->AddInstances(ZeroScaleTransforms, false, true);
 
+	if (CustomDataFloats.Num())
+	{
+		const int32 NumCustomDataFloats = ISMComponent->NumCustomDataFloats;
+		if (ensure(NumCustomDataFloats * InstanceCount == CustomDataFloats.Num()))
+		{
+			for (int32 InstanceIndex = 0; InstanceIndex < InstanceCount; ++InstanceIndex)
+			{
+				ISMComponent->SetCustomData(NewInstanceGroup.Start + InstanceIndex, CustomDataFloats.Slice(InstanceIndex * NumCustomDataFloats, NumCustomDataFloats));
+			}
+		}
+	}
+
 	return InstanceGroupIndex;
 }
 
-FGeometryCollectionMeshInfo FGeometryCollectionISMPool::AddISM(UGeometryCollectionISMPoolComponent* OwningComponent, const FGeometryCollectionStaticMeshInstance& MeshInstance, int32 InstanceCount, bool bPreferHISM)
+FGeometryCollectionMeshInfo FGeometryCollectionISMPool::AddISM(UGeometryCollectionISMPoolComponent* OwningComponent, const FGeometryCollectionStaticMeshInstance& MeshInstance, int32 InstanceCount, TArrayView<const float> CustomDataFloats, bool bPreferHISM)
 {
 	FGeometryCollectionMeshInfo Info;
 
@@ -104,7 +118,7 @@ FGeometryCollectionMeshInfo FGeometryCollectionISMPool::AddISM(UGeometryCollecti
 		Info.ISMIndex = *ISMIndex;
 	}
 	// add to the ISM 
-	Info.InstanceGroupIndex = ISMs[Info.ISMIndex].AddInstanceGroup(InstanceCount);
+	Info.InstanceGroupIndex = ISMs[Info.ISMIndex].AddInstanceGroup(InstanceCount, CustomDataFloats);
 	return Info;
 }
 
@@ -182,11 +196,11 @@ void UGeometryCollectionISMPoolComponent::DestroyMeshGroup(FMeshGroupId MeshGrou
 	}
 }
 
-UGeometryCollectionISMPoolComponent::FMeshId UGeometryCollectionISMPoolComponent::AddMeshToGroup(FMeshGroupId MeshGroupId, const FGeometryCollectionStaticMeshInstance& MeshInstance, int32 InstanceCount, bool bPreferHISM)
+UGeometryCollectionISMPoolComponent::FMeshId UGeometryCollectionISMPoolComponent::AddMeshToGroup(FMeshGroupId MeshGroupId, const FGeometryCollectionStaticMeshInstance& MeshInstance, int32 InstanceCount, TArrayView<const float> CustomDataFloats, bool bPreferHISM)
 {
 	if (FGeometryCollectionMeshGroup* MeshGroup = MeshGroups.Find(MeshGroupId))
 	{
-		const FGeometryCollectionMeshInfo ISMInstanceInfo = Pool.AddISM(this, MeshInstance, InstanceCount, bPreferHISM);
+		const FGeometryCollectionMeshInfo ISMInstanceInfo = Pool.AddISM(this, MeshInstance, InstanceCount, CustomDataFloats, bPreferHISM);
 		return MeshGroup->AddMesh(MeshInstance, InstanceCount, ISMInstanceInfo);
 	}
 	UE_LOG(LogChaos, Warning, TEXT("UGeometryCollectionISMPoolComponent : Trying to add a mesh to a mesh group (%d) that does not exists"), MeshGroupId);
diff --git a/Engine/Source/Runtime/Experimental/GeometryCollectionEngine/Private/GeometryCollection/GeometryCollectionObject.cpp b/Engine/Source/Runtime/Experimental/GeometryCollectionEngine/Private/GeometryCollection/GeometryCollectionObject.cpp
--- a/Engine/Source/Runtime/Experimental/GeometryCollectionEngine/Private/GeometryCollection/GeometryCollectionObject.cpp
+++ b/Engine/Source/Runtime/Experimental/GeometryCollectionEngine/Private/GeometryCollection/GeometryCollectionObject.cpp
@@ -416,6 +416,7 @@ void UGeometryCollection::Reset()
 		Materials.Empty();
 		EmbeddedGeometryExemplar.Empty();
 		AutoInstanceMeshes.Empty();
+		AutoInstanceMaterialCustomData.Empty();
 		InvalidateCollection();
 	}
 }
diff --git a/Engine/Source/Runtime/Experimental/GeometryCollectionEngine/Public/GeometryCollection/GeometryCollectionComponent.h b/Engine/Source/Runtime/Experimental/GeometryCollectionEngine/Public/GeometryCollection/GeometryCollectionComponent.h
--- a/Engine/Source/Runtime/Experimental/GeometryCollectionEngine/Public/GeometryCollection/GeometryCollectionComponent.h
+++ b/Engine/Source/Runtime/Experimental/GeometryCollectionEngine/Public/GeometryCollection/GeometryCollectionComponent.h
@@ -1076,6 +1076,13 @@ protected:
 	UPROPERTY(EditAnywhere, BlueprintReadWrite, Category = "ChaosPhysics|Rendering")
 	bool bAutoAssignISMPool;
 
+	/** 
+	 * Values that will be added to the Per Instance Custom Data for every instance in all ISMPools. 
+	 * This will be added in slots after values that come from the Geometry Collection AutoInstanceMaterialCustomData. 
+	 */
+	UPROPERTY(EditAnywhere, BlueprintReadWrite, Category = Rendering)
+	TArray<float> ISMPoolMaterialCustomData;
+
 	TObjectPtr<AGeometryCollectionISMPoolActor> AssignedISMPool = nullptr;
 	int32 ISMPoolMeshGroupIndex = INDEX_NONE;
 	TArray<int32> ISMPoolRootProxyMeshIds;
@@ -1172,6 +1179,7 @@ protected:
 	// We track this on the client to be able to turn off the tick for perf reasons
 	// if we spend a lot of ticks sequentially doing nothing.
 	int64 LastAsyncPhysicsTickMs = 0;
+
 private:
 
 	bool bRenderStateDirty;
diff --git a/Engine/Source/Runtime/Experimental/GeometryCollectionEngine/Public/GeometryCollection/GeometryCollectionISMPoolComponent.h b/Engine/Source/Runtime/Experimental/GeometryCollectionEngine/Public/GeometryCollection/GeometryCollectionISMPoolComponent.h
--- a/Engine/Source/Runtime/Experimental/GeometryCollectionEngine/Public/GeometryCollection/GeometryCollectionISMPoolComponent.h
+++ b/Engine/Source/Runtime/Experimental/GeometryCollectionEngine/Public/GeometryCollection/GeometryCollectionISMPoolComponent.h
@@ -104,10 +104,11 @@ struct FGeometryCollectionStaticMeshInstance
 {
 	UStaticMesh* StaticMesh = nullptr;
 	TArray<UMaterialInterface*> MaterialsOverrides;
+	int32 NumCustomDataFloats = 0;
 
 	bool operator==(const FGeometryCollectionStaticMeshInstance& Other) const 
 	{
-		if (StaticMesh == Other.StaticMesh)
+		if (StaticMesh == Other.StaticMesh && NumCustomDataFloats == Other.NumCustomDataFloats)
 		{
 			if (MaterialsOverrides.Num() == Other.MaterialsOverrides.Num())
 			{
@@ -167,7 +168,7 @@ struct FGeometryCollectionISM
 {
 	FGeometryCollectionISM(AActor* OwmingActor, const FGeometryCollectionStaticMeshInstance& MeshInstance, bool bPreferHISM);
 
-	int32 AddInstanceGroup(int32 InstanceCount);
+	int32 AddInstanceGroup(int32 InstanceCount, TArrayView<const float> CustomDataFloats);
 
 	TObjectPtr<UInstancedStaticMeshComponent> ISMComponent;
 	FInstanceGroups InstanceGroups;
@@ -178,7 +179,7 @@ struct FGeometryCollectionISMPool
 {
 	using FISMIndex = int32;
 
-	FGeometryCollectionMeshInfo AddISM(UGeometryCollectionISMPoolComponent* OwningComponent, const FGeometryCollectionStaticMeshInstance& MeshInstance, int32 InstanceCount, bool bPreferHISM);
+	FGeometryCollectionMeshInfo AddISM(UGeometryCollectionISMPoolComponent* OwningComponent, const FGeometryCollectionStaticMeshInstance& MeshInstance, int32 InstanceCount, TArrayView<const float> CustomDataFloats, bool bPreferHISM);
 	bool BatchUpdateInstancesTransforms(FGeometryCollectionMeshInfo& MeshInfo, int32 StartInstanceIndex, const TArray<FTransform>& NewInstancesTransforms, bool bWorldSpace, bool bMarkRenderStateDirty, bool bTeleport);
 	void RemoveISM(const FGeometryCollectionMeshInfo& MeshInfo);
 
@@ -215,7 +216,7 @@ public:
 	void DestroyMeshGroup(FMeshGroupId MeshGroupId);
 
 	/** Add a static mesh for a nmesh group */
-	FMeshId AddMeshToGroup(FMeshGroupId MeshGroupId, const FGeometryCollectionStaticMeshInstance& MeshInstance, int32 InstanceCount, bool bPreferHISM = true);
+	FMeshId AddMeshToGroup(FMeshGroupId MeshGroupId, const FGeometryCollectionStaticMeshInstance& MeshInstance, int32 InstanceCount, TArrayView<const float> CustomDataFloats, bool bPreferHISM = true);
 
 	/** Add a static mesh for a nmesh group */
 	bool BatchUpdateInstancesTransforms(FMeshGroupId MeshGroupId, FMeshId MeshId, int32 StartInstanceIndex, const TArray<FTransform>& NewInstancesTransforms, bool bWorldSpace = false, bool bMarkRenderStateDirty = false, bool bTeleport = false);
diff --git a/Engine/Source/Runtime/Experimental/GeometryCollectionEngine/Public/GeometryCollection/GeometryCollectionObject.h b/Engine/Source/Runtime/Experimental/GeometryCollectionEngine/Public/GeometryCollection/GeometryCollectionObject.h
--- a/Engine/Source/Runtime/Experimental/GeometryCollectionEngine/Public/GeometryCollection/GeometryCollectionObject.h
+++ b/Engine/Source/Runtime/Experimental/GeometryCollectionEngine/Public/GeometryCollection/GeometryCollectionObject.h
@@ -524,15 +524,17 @@ public:
 	UPROPERTY(EditAnywhere, BlueprintReadWrite, Category = "Rendering")
 	bool bStripRenderDataOnCook;
 
-	/** list of unique static mesh / materials pairs for auto instancing*/
-	UPROPERTY(EditAnywhere, Category = "Rendering")
-	TArray<FGeometryCollectionAutoInstanceMesh> AutoInstanceMeshes;
-
 	/** static mesh to use as a proxy for rendering until the geometry collection is broken */
 	UPROPERTY(BlueprintReadWrite, EditAnywhere, Category = "Rendering")
 	FGeometryCollectionProxyMeshData RootProxyData;
 
+	/** list of unique static mesh / materials pairs for auto instancing*/
+	UPROPERTY(EditAnywhere, Category = "Rendering")
+	TArray<FGeometryCollectionAutoInstanceMesh> AutoInstanceMeshes;
 
+	/** Array of material custom data for all static meshes in AutoInstanceMeshes. */
+	UPROPERTY(EditAnywhere, Category = "Rendering")
+	TArray<float> AutoInstanceMaterialCustomData;
 
 	UFUNCTION(BlueprintCallable, Category = "Nanite")
 	void SetEnableNanite(bool bValue);
