Added ISM pool debug component.
Currently allows ISM selection and some debug bounds/stats.

diff --git a/Engine/Source/Runtime/Experimental/GeometryCollectionEngine/Private/GeometryCollection/GeometryCollectionComponent.cpp b/Engine/Source/Runtime/Experimental/GeometryCollectionEngine/Private/GeometryCollection/GeometryCollectionComponent.cpp
--- a/Engine/Source/Runtime/Experimental/GeometryCollectionEngine/Private/GeometryCollection/GeometryCollectionComponent.cpp
+++ b/Engine/Source/Runtime/Experimental/GeometryCollectionEngine/Private/GeometryCollection/GeometryCollectionComponent.cpp
@@ -4523,6 +4523,22 @@ void UGeometryCollectionComponent::RefreshISMPoolInstances()
 	}
 }
 
+UInstancedStaticMeshComponent* UGeometryCollectionComponent::GetRootProxyISM() const
+{
+	if (AssignedISMPool != nullptr)
+	{
+		if (UGeometryCollectionISMPoolComponent* ISMPoolComp = AssignedISMPool->GetISMPoolComp())
+		{
+			if (RestCollection->AutoInstanceMeshes.Num())
+			{
+				return ISMPoolComp->GetISMForMeshId(ISMPoolMeshGroupIndex, 0);
+			}
+		}
+	}
+
+	return nullptr;
+}
+
 bool UGeometryCollectionComponent::IsRootBroken() const
 {
 	if (DynamicCollection && DynamicCollection->Active.Num() > 0)
diff --git a/Engine/Source/Runtime/Experimental/GeometryCollectionEngine/Private/GeometryCollection/GeometryCollectionISMPoolActor.cpp b/Engine/Source/Runtime/Experimental/GeometryCollectionEngine/Private/GeometryCollection/GeometryCollectionISMPoolActor.cpp
--- a/Engine/Source/Runtime/Experimental/GeometryCollectionEngine/Private/GeometryCollection/GeometryCollectionISMPoolActor.cpp
+++ b/Engine/Source/Runtime/Experimental/GeometryCollectionEngine/Private/GeometryCollection/GeometryCollectionISMPoolActor.cpp
@@ -2,6 +2,7 @@
 
 #include "GeometryCollection/GeometryCollectionISMPoolActor.h"
 #include "GeometryCollection/GeometryCollectionISMPoolComponent.h"
+#include "GeometryCollection/GeometryCollectionISMPoolDebugDrawComponent.h"
 
 #include UE_INLINE_GENERATED_CPP_BY_NAME(GeometryCollectionISMPoolActor)
 
@@ -10,4 +11,13 @@ AGeometryCollectionISMPoolActor::AGeometryCollectionISMPoolActor(const FObjectIn
 {
 	ISMPoolComp = CreateDefaultSubobject<UGeometryCollectionISMPoolComponent>(TEXT("ISMPoolComp"));
 	RootComponent = ISMPoolComp;
+
+#if UE_ENABLE_DEBUG_DRAWING
+	ISMPoolDebugDrawComp = CreateDefaultSubobject<UGeometryCollectionISMPoolDebugDrawComponent>(TEXT("ISMPoolDebug"));
+	ISMPoolDebugDrawComp->SetCollisionEnabled(ECollisionEnabled::NoCollision);
+	ISMPoolDebugDrawComp->SetCanEverAffectNavigation(false);
+	ISMPoolDebugDrawComp->CanCharacterStepUpOn = ECanBeCharacterBase::ECB_No;
+	ISMPoolDebugDrawComp->SetGenerateOverlapEvents(false);
+	ISMPoolDebugDrawComp->SetupAttachment(ISMPoolComp);
+#endif
 }
diff --git a/Engine/Source/Runtime/Experimental/GeometryCollectionEngine/Private/GeometryCollection/GeometryCollectionISMPoolComponent.cpp b/Engine/Source/Runtime/Experimental/GeometryCollectionEngine/Private/GeometryCollection/GeometryCollectionISMPoolComponent.cpp
--- a/Engine/Source/Runtime/Experimental/GeometryCollectionEngine/Private/GeometryCollection/GeometryCollectionISMPoolComponent.cpp
+++ b/Engine/Source/Runtime/Experimental/GeometryCollectionEngine/Private/GeometryCollection/GeometryCollectionISMPoolComponent.cpp
@@ -216,3 +216,19 @@ bool UGeometryCollectionISMPoolComponent::BatchUpdateInstancesTransforms(FMeshGr
 	UE_LOG(LogChaos, Warning, TEXT("UGeometryCollectionISMPoolComponent : Trying to update instance with mesh group (%d) that not exists"), MeshGroupId);
 	return false;
 }
+
+UInstancedStaticMeshComponent* UGeometryCollectionISMPoolComponent::GetISMForMeshId(FMeshGroupId MeshGroupId, FMeshId MeshId) const
+{
+	if (FGeometryCollectionMeshGroup const* MeshGroup = MeshGroups.Find(MeshGroupId))
+	{
+		if (MeshGroup->MeshInfos.IsValidIndex(MeshId))
+		{
+			int32 ISMIndex = MeshGroup->MeshInfos[MeshId].ISMIndex;
+			if (Pool.ISMs.IsValidIndex(ISMIndex))
+			{
+				return Pool.ISMs[ISMIndex].ISMComponent;
+			}
+		}
+	}
+	return nullptr;
+}
diff --git a/Engine/Source/Runtime/Experimental/GeometryCollectionEngine/Public/GeometryCollection/GeometryCollectionComponent.h b/Engine/Source/Runtime/Experimental/GeometryCollectionEngine/Public/GeometryCollection/GeometryCollectionComponent.h
--- a/Engine/Source/Runtime/Experimental/GeometryCollectionEngine/Public/GeometryCollection/GeometryCollectionComponent.h
+++ b/Engine/Source/Runtime/Experimental/GeometryCollectionEngine/Public/GeometryCollection/GeometryCollectionComponent.h
@@ -1012,6 +1012,8 @@ public:
 	const FGeometryDynamicCollection* GetDynamicCollection() const;
 	FGeometryDynamicCollection* GetDynamicCollection();  // TEMP HACK?
 
+	UInstancedStaticMeshComponent* GetRootProxyISM() const;
+
 public:
 	UPROPERTY(BlueprintAssignable, Category = "Collision")
 	FOnChaosPhysicsCollision OnChaosPhysicsCollision;
diff --git a/Engine/Source/Runtime/Experimental/GeometryCollectionEngine/Public/GeometryCollection/GeometryCollectionISMPoolActor.h b/Engine/Source/Runtime/Experimental/GeometryCollectionEngine/Public/GeometryCollection/GeometryCollectionISMPoolActor.h
--- a/Engine/Source/Runtime/Experimental/GeometryCollectionEngine/Public/GeometryCollection/GeometryCollectionISMPoolActor.h
+++ b/Engine/Source/Runtime/Experimental/GeometryCollectionEngine/Public/GeometryCollection/GeometryCollectionISMPoolActor.h
@@ -5,6 +5,7 @@
 #include "GeometryCollectionISMPoolActor.generated.h"
 
 class UGeometryCollectionISMPoolComponent;
+class UGeometryCollectionISMPoolDebugDrawComponent;
 
 UCLASS(ConversionRoot, ComponentWrapperClass)
 class GEOMETRYCOLLECTIONENGINE_API AGeometryCollectionISMPoolActor : public AActor
@@ -15,6 +16,9 @@ private:
 	UPROPERTY(Category = ISMPoolActor, VisibleAnywhere, BlueprintReadOnly, meta = (AllowPrivateAccess = "true"))
 	TObjectPtr<class UGeometryCollectionISMPoolComponent> ISMPoolComp;
 
+	UPROPERTY(Category = ISMPoolActor, VisibleAnywhere, BlueprintReadOnly, meta = (AllowPrivateAccess = "true"))
+	TObjectPtr<class UGeometryCollectionISMPoolDebugDrawComponent> ISMPoolDebugDrawComp;
+
 public:
 	/** Returns ISMPoolComp subobject **/
 	UGeometryCollectionISMPoolComponent* GetISMPoolComp() const { return ISMPoolComp; }
diff --git a/Engine/Source/Runtime/Experimental/GeometryCollectionEngine/Public/GeometryCollection/GeometryCollectionISMPoolComponent.h b/Engine/Source/Runtime/Experimental/GeometryCollectionEngine/Public/GeometryCollection/GeometryCollectionISMPoolComponent.h
--- a/Engine/Source/Runtime/Experimental/GeometryCollectionEngine/Public/GeometryCollection/GeometryCollectionISMPoolComponent.h
+++ b/Engine/Source/Runtime/Experimental/GeometryCollectionEngine/Public/GeometryCollection/GeometryCollectionISMPoolComponent.h
@@ -221,9 +221,15 @@ public:
 	/** Add a static mesh for a nmesh group */
 	bool BatchUpdateInstancesTransforms(FMeshGroupId MeshGroupId, FMeshId MeshId, int32 StartInstanceIndex, const TArray<FTransform>& NewInstancesTransforms, bool bWorldSpace = false, bool bMarkRenderStateDirty = false, bool bTeleport = false);
 
+	/** Get the ISM component related to a mesh id. */
+	UInstancedStaticMeshComponent* GetISMForMeshId(FMeshGroupId MeshGroupId, FMeshId MeshId) const;
+
 private:
 	uint32 NextMeshGroupId = 0;
 	TMap<FMeshGroupId, FGeometryCollectionMeshGroup> MeshGroups;
 	FGeometryCollectionISMPool Pool;
+
+	// Expose internals for debug draw support.
+	friend class UGeometryCollectionISMPoolDebugDrawComponent;
 };
 #endif
\ No newline at end of file
diff --git a/Engine/Source/Runtime/Experimental/GeometryCollectionEngine/Public/GeometryCollection/GeometryCollectionISMPoolDebugDrawComponent.cpp b/Engine/Source/Runtime/Experimental/GeometryCollectionEngine/Public/GeometryCollection/GeometryCollectionISMPoolDebugDrawComponent.cpp
new file mode 100644
--- /dev/null
+++ b/Engine/Source/Runtime/Experimental/GeometryCollectionEngine/Public/GeometryCollection/GeometryCollectionISMPoolDebugDrawComponent.cpp
@@ -0,0 +1,364 @@
+// Copyright Epic Games, Inc. All Rights Reserved.
+
+#include "GeometryCollectionISMPoolDebugDrawComponent.h"
+
+#include "CollisionQueryParams.h"
+#include "ComponentRecreateRenderStateContext.h"
+#include "Components/HierarchicalInstancedStaticMeshComponent.h"
+#include "Engine/StaticMesh.h"
+#include "Engine/World.h"
+#include "EngineUtils.h"
+#include "GameFramework/Pawn.h"
+#include "GameFramework/PlayerController.h"
+#include "GeometryCollectionComponent.h"
+#include "GeometryCollectionISMPoolActor.h"
+#include "GeometryCollectionISMPoolComponent.h"
+#include "Misc/CoreDelegates.h"
+
+#if WITH_EDITOR
+#include "Editor.h"
+#include "EditorViewportClient.h"
+#include "LevelEditorViewport.h"
+#endif
+
+#include UE_INLINE_GENERATED_CPP_BY_NAME(GeometryCollectionISMPoolDebugDrawComponent)
+
+static TAutoConsoleVariable<int32> CVarISMPoolStats(
+	TEXT("p.Chaos.GC.ISMPoolDebugStats"),
+	0,
+	TEXT("Show stats for the ISM pools"),
+	FConsoleVariableDelegate::CreateLambda([](IConsoleVariable* InVariable) { UGeometryCollectionISMPoolDebugDrawComponent::UpdateAllTickEnabled(); }),
+	ECVF_Default);
+
+static TAutoConsoleVariable<int32> CVarISMPoolDebugDraw(
+	TEXT("p.Chaos.GC.ISMPoolDebugDraw"),
+	0,
+	TEXT("Show debug drawing for the ISM pools"),
+	FConsoleVariableDelegate::CreateLambda([](IConsoleVariable* InVariable) { UGeometryCollectionISMPoolDebugDrawComponent::UpdateAllTickEnabled(); FGlobalComponentRecreateRenderStateContext Context; }),
+	ECVF_Default);
+
+
+class FGeometryCollectionISMPoolDebugDrawSceneProxy final : public FDebugRenderSceneProxy
+{
+public:
+	SIZE_T GetTypeHash() const override
+	{
+		static size_t UniquePointer;
+		return reinterpret_cast<size_t>(&UniquePointer);
+	}
+
+	FGeometryCollectionISMPoolDebugDrawSceneProxy(UGeometryCollectionISMPoolDebugDrawComponent const* InComponent)
+		: FDebugRenderSceneProxy(InComponent)
+	{
+		DrawType = FDebugRenderSceneProxy::SolidAndWireMeshes;
+	}
+
+	FPrimitiveViewRelevance GetViewRelevance(const FSceneView* View) const override
+	{
+		FPrimitiveViewRelevance Result;
+		Result.bDrawRelevance = true;
+		Result.bDynamicRelevance = true;
+		Result.bSeparateTranslucency = Result.bNormalTranslucency = true;
+		return Result;
+	}
+};
+
+
+UGeometryCollectionISMPoolDebugDrawComponent::UGeometryCollectionISMPoolDebugDrawComponent(const FObjectInitializer& ObjectInitializer)
+	: Super(ObjectInitializer)
+{
+	bIsEditorOnly = false;
+	bSelectable = false;
+
+	PrimaryComponentTick.bCanEverTick = true;
+	PrimaryComponentTick.bStartWithTickEnabled = false;
+	PrimaryComponentTick.bAllowTickOnDedicatedServer = false;
+}
+
+void UGeometryCollectionISMPoolDebugDrawComponent::UpdateTickEnabled()
+{
+	const bool bForceShowStats = CVarISMPoolStats.GetValueOnAnyThread() != 0;
+	const bool bForceShowBounds = CVarISMPoolDebugDraw.GetValueOnAnyThread() != 0;
+	PrimaryComponentTick.SetTickFunctionEnable(bShowStats || bForceShowStats || bShowBounds || bForceShowBounds);
+}
+
+void UGeometryCollectionISMPoolDebugDrawComponent::UpdateAllTickEnabled()
+{
+	for (TObjectIterator<UGeometryCollectionISMPoolDebugDrawComponent> It; It; ++It)
+	{
+		It->UpdateTickEnabled();
+	}
+}
+
+void UGeometryCollectionISMPoolDebugDrawComponent::BeginPlay()
+{
+	Super::BeginPlay();
+#if UE_ENABLE_DEBUG_DRAWING
+	OnScreenMessagesHandle = FCoreDelegates::OnGetOnScreenMessages.AddUObject(this, &UGeometryCollectionISMPoolDebugDrawComponent::GetOnScreenMessages);
+#endif
+	UpdateTickEnabled();
+}
+
+void UGeometryCollectionISMPoolDebugDrawComponent::EndPlay(const EEndPlayReason::Type EndPlayReason)
+{
+	Super::EndPlay(EndPlayReason);
+#if UE_ENABLE_DEBUG_DRAWING
+	FCoreDelegates::OnGetOnScreenMessages.Remove(OnScreenMessagesHandle);
+#endif
+}
+
+FBoxSphereBounds UGeometryCollectionISMPoolDebugDrawComponent::CalcBounds(const FTransform& LocalToWorld) const
+{
+	return SelectedComponent ? SelectedComponent->CalcBounds(LocalToWorld) : FBox(ForceInitToZero);
+}
+
+#if WITH_EDITOR
+void UGeometryCollectionISMPoolDebugDrawComponent::PostEditChangeProperty(FPropertyChangedEvent& PropertyChangedEvent)
+{
+	Super::PostEditChangeProperty(PropertyChangedEvent);
+
+	if (PropertyChangedEvent.Property)
+	{
+		if (PropertyChangedEvent.Property->GetFName() == GET_MEMBER_NAME_CHECKED(UGeometryCollectionISMPoolDebugDrawComponent, bShowStats) ||
+			PropertyChangedEvent.Property->GetFName() == GET_MEMBER_NAME_CHECKED(UGeometryCollectionISMPoolDebugDrawComponent, bShowBounds))
+		{
+			UpdateTickEnabled();
+		}
+	}
+}
+#endif
+
+void UGeometryCollectionISMPoolDebugDrawComponent::TickComponent(float DeltaTime, enum ELevelTick TickType, FActorComponentTickFunction* ThisTickFunction)
+{
+	Super::TickComponent(DeltaTime, TickType, ThisTickFunction);
+
+#if UE_ENABLE_DEBUG_DRAWING
+	UInstancedStaticMeshComponent const* FoundISM = nullptr;
+
+	const bool bForceShowStats = CVarISMPoolStats.GetValueOnAnyThread() != 0;
+	const bool bForceShowBounds = CVarISMPoolDebugDraw.GetValueOnAnyThread() != 0;
+	if (bShowStats || bForceShowStats || bShowBounds || bForceShowBounds)
+	{
+#if WITH_EDITOR
+		if (GCurrentLevelEditingViewportClient)
+		{
+			FViewport* Viewport = GCurrentLevelEditingViewportClient->Viewport;
+			HHitProxy* HitResult = Viewport->GetHitProxy(Viewport->GetMouseX(), Viewport->GetMouseY());
+			if (HActor* HitActor = HitProxyCast<HActor>(HitResult))
+			{
+				FoundISM = Cast<UInstancedStaticMeshComponent>(HitActor->PrimComponent);
+			}
+		} 
+#endif
+		if (FoundISM == nullptr)
+		{
+			if (UWorld* World = GetWorld())
+			{
+				APlayerController const* Controller = nullptr;
+				for (FConstPlayerControllerIterator Iterator = World->GetPlayerControllerIterator(); Iterator; ++Iterator)
+				{
+					APlayerController* PC = Iterator->Get();
+					if (PC && PC->IsLocalController())
+					{
+						Controller = PC;
+					}
+				}
+
+				if (Controller != nullptr)
+				{
+					FVector CamLoc;
+					FRotator CamRot;
+					Controller->GetPlayerViewPoint(CamLoc, CamRot);
+					const FVector CamForward = CamRot.Vector();
+					const FVector TraceStart = CamLoc;
+					const FVector TraceEnd = TraceStart + CamForward * 20000;
+
+					FHitResult HitResult(ForceInit);
+					FCollisionQueryParams TraceParams(SCENE_QUERY_STAT(ChaosDebugVisibilityTrace), true, Controller->GetPawn());
+					const bool bHit = World->LineTraceSingleByChannel(HitResult, TraceStart, TraceEnd, ECC_Visibility, TraceParams);
+					if (bHit && HitResult.GetComponent() != nullptr)
+					{
+						if (UGeometryCollectionComponent* GeometryCollectionComponent = Cast<UGeometryCollectionComponent>(HitResult.GetComponent()))
+						{
+							// We hit the GC but don't have any information which part.
+							// For now we just select the ISM of the root proxy.
+							FoundISM = GeometryCollectionComponent->GetRootProxyISM();
+						}
+						else if (UInstancedStaticMeshComponent* ISMComponent = Cast<UInstancedStaticMeshComponent>(HitResult.GetComponent()))
+						{
+							FoundISM = ISMComponent;
+						}
+					}
+				}
+			}
+		}
+	}
+
+	if (FoundISM != nullptr && FoundISM->GetOuter() != GetOuter())
+	{
+		FoundISM = nullptr;
+	}
+
+	if (SelectedComponent != FoundISM)
+	{
+		bool bChangeSelection = (SelectedComponent == nullptr) || SelectTimer < 0.f;
+		if (bChangeSelection)
+		{
+#if WITH_EDITOR
+			if (SelectedComponent)
+			{
+				((UInstancedStaticMeshComponent*)SelectedComponent)->PushHoveredToProxy(false);
+			}
+#endif
+			SelectedComponent = FoundISM;
+#if WITH_EDITOR
+			if (SelectedComponent)
+			{
+				((UInstancedStaticMeshComponent*)SelectedComponent)->PushHoveredToProxy(true);
+			}
+#endif
+			SelectTimer = 0.3f;
+			MarkRenderStateDirty();
+		}
+		else
+		{
+			SelectTimer -= DeltaTime;
+		}
+	}
+#endif // UE_ENABLE_DEBUG_DRAWING
+}
+
+#if UE_ENABLE_DEBUG_DRAWING
+
+static void DrawBoxAndChildren(FDebugRenderSceneProxy* DebugProxy, FTransform const& LocalToWorldTransform, TArray<FClusterNode> const& Nodes, int32 NodeIndex, int32 LevelIndex)
+{
+	FClusterNode const& Node = Nodes[NodeIndex];
+
+	const FBox LocalBounds(FBox(Node.BoundMin, Node.BoundMax));
+	
+	static TArray<FColor> Colors = { FColorList::VeryLightGrey, FColorList::CoolCopper, FColorList::GreenYellow, FColorList::CornFlowerBlue, FColorList::DustyRose, FColorList::Red, FColorList::Magenta };
+	const FColor LevelColor = Colors[FMath::Min(LevelIndex, Colors.Num() - 1)];
+	
+	DebugProxy->Boxes.Add(FDebugRenderSceneProxy::FDebugBox(LocalBounds, LevelColor, LocalToWorldTransform));
+
+	if (Node.FirstChild > 0)
+	{
+		for (int32 ChildIndex = Node.FirstChild; ChildIndex <= Node.LastChild; ChildIndex++)
+		{
+			DrawBoxAndChildren(DebugProxy, LocalToWorldTransform, Nodes, ChildIndex, LevelIndex + 1);
+		}
+	}
+}
+
+FDebugRenderSceneProxy* UGeometryCollectionISMPoolDebugDrawComponent::CreateDebugSceneProxy()
+{
+	const bool bForceShowBounds = CVarISMPoolDebugDraw.GetValueOnAnyThread() != 0;
+	if (!(SelectedComponent && (bShowBounds || bForceShowBounds)))
+	{
+		return nullptr;
+	}
+
+	FDebugRenderSceneProxy* DebugProxy = new FGeometryCollectionISMPoolDebugDrawSceneProxy(this);
+
+	if (UHierarchicalInstancedStaticMeshComponent* HISMComponent = Cast<UHierarchicalInstancedStaticMeshComponent>(SelectedComponent))
+	{
+		const FTransform LocalToWorldTransform = HISMComponent->GetComponentTransform();
+		TArray<FClusterNode> TreeNodes;
+		HISMComponent->GetTree(TreeNodes);
+		DrawBoxAndChildren(DebugProxy, LocalToWorldTransform, TreeNodes, 0, 0);
+	}
+	else if(UInstancedStaticMeshComponent* ISMComponent = Cast<UInstancedStaticMeshComponent>(SelectedComponent))
+	{
+		const FBox LocalBounds = ISMComponent->CalcBounds(FTransform::Identity).GetBox();
+		const FTransform LocalToWorldTransform = ISMComponent->GetComponentTransform();
+		DebugProxy->Boxes.Add(FDebugRenderSceneProxy::FDebugBox(LocalBounds, FColorList::CornFlowerBlue, LocalToWorldTransform));
+	}
+
+	return DebugProxy;
+}
+
+void UGeometryCollectionISMPoolDebugDrawComponent::GetOnScreenMessages(TMultiMap<FCoreDelegates::EOnScreenMessageSeverity, FText>& OutMessages)
+{
+	const bool bForceShowStats = CVarISMPoolStats.GetValueOnAnyThread() != 0;
+	if (!bShowStats && !bShowGlobalStats && !bForceShowStats)
+	{
+		return;
+	}
+
+	if (bShowGlobalStats || bForceShowStats)
+	{
+		AGeometryCollectionISMPoolActor const* ISMPoolActor = Cast<AGeometryCollectionISMPoolActor>(GetOuter());
+		UGeometryCollectionISMPoolComponent const* ISMPoolComponent = ISMPoolActor ? ISMPoolActor->GetISMPoolComp() : nullptr;
+
+		if (ISMPoolComponent)
+		{
+			int32 NumISMs = 0;
+			int32 NumHISMs = 0;
+			int32 NumInstancesTotal = 0;
+			TArray<int32> NumInstancesPerISM;
+			NumInstancesPerISM.Reserve(ISMPoolComponent->Pool.ISMs.Num());
+
+			for (FGeometryCollectionISM const& ISM : ISMPoolComponent->Pool.ISMs)
+			{
+				const bool bIsHISM = ISM.ISMComponent->IsA(UHierarchicalInstancedStaticMeshComponent::StaticClass());
+				NumISMs += bIsHISM ? 0 : 1;
+				NumHISMs += bIsHISM ? 1 : 0;
+
+				const int32 NumInstances = ISM.ISMComponent->GetNumRenderInstances();
+				NumInstancesTotal += NumInstances;
+				NumInstancesPerISM.Add(NumInstances);
+			}
+
+			if (NumInstancesTotal > 0)
+			{
+				OutMessages.Add(FCoreDelegates::EOnScreenMessageSeverity::Warning, FText::FromString(*GetOwner()->GetPathName()));
+				OutMessages.Add(FCoreDelegates::EOnScreenMessageSeverity::Info, FText::FromString(FString::Printf(TEXT("Num ISMS %d"), NumISMs)));
+				OutMessages.Add(FCoreDelegates::EOnScreenMessageSeverity::Info, FText::FromString(FString::Printf(TEXT("Num HISMS %d"), NumHISMs)));
+				OutMessages.Add(FCoreDelegates::EOnScreenMessageSeverity::Info, FText::FromString(FString::Printf(TEXT("Num Instances %d"), NumInstancesTotal)));
+				OutMessages.Add(FCoreDelegates::EOnScreenMessageSeverity::Info, FText::FromString(FString::Printf(TEXT("Average Instances Per ISM %f"), (float)NumInstancesTotal / (float)(NumISMs + NumHISMs))));
+
+				NumInstancesPerISM.Sort();
+				for (int32 Index = 0; Index < 10; ++Index)
+				{
+					const int32 Percentile = (Index + 1) * 10;
+					const int32 PercentileIndex = FMath::Min(NumInstancesPerISM.Num() * (Index + 1) / 10, NumInstancesPerISM.Num() - 1);
+					OutMessages.Add(FCoreDelegates::EOnScreenMessageSeverity::Info, FText::FromString(FString::Printf(TEXT("Num Instances at %d percentile: %d"), Percentile, NumInstancesPerISM[PercentileIndex])));
+				}
+			}
+		}
+	}
+
+	if ((bShowStats || bForceShowStats) && SelectedComponent)
+	{
+		TCHAR const* ISMType = SelectedComponent->IsA(UHierarchicalInstancedStaticMeshComponent::StaticClass()) ? TEXT("HISM") : TEXT("ISM");
+
+		FVector BoundsMin, BoundsMax;
+		SelectedComponent->GetLocalBounds(BoundsMin, BoundsMax);
+		const FVector BoundsSize = BoundsMax - BoundsMin;
+
+		const int32 NumInstances = SelectedComponent->GetNumRenderInstances();
+		const int32 NumCustomDataFloats = SelectedComponent->NumCustomDataFloats;
+		const int32 StartCullDistance = SelectedComponent->InstanceStartCullDistance;
+		const int32 EndCullDistance = SelectedComponent->InstanceEndCullDistance;
+		const int32 NumMaterials = SelectedComponent->GetNumMaterials();
+
+		UStaticMesh const* StaticMesh = SelectedComponent->GetStaticMesh();
+		const int32 NumLods = StaticMesh ? StaticMesh->GetNumLODs() : 0;
+
+		FString ISMDescription = FString::Printf(TEXT("Type=%s   Count=%d   Bounds=(%f., %f., %f.)   CullDistance=(%f - %f)   NumLods= %d    NumMaterials=%d   CustomDataSize=%d"), 
+			ISMType, NumInstances, BoundsSize.X, BoundsSize.Z, BoundsSize.Z, StartCullDistance, EndCullDistance, NumLods, NumMaterials, NumCustomDataFloats);
+
+		for (int32 LodIndex = 0; LodIndex < NumLods; ++LodIndex)
+		{
+			const int32 NumVertices = StaticMesh->GetNumVertices(LodIndex);
+			const int32 NumTriangles = StaticMesh->GetNumTriangles(LodIndex);
+			ISMDescription += FString::Printf(TEXT("\nLod %d: Vertices=%d, Triangles=%d."), LodIndex, NumVertices, NumTriangles);
+		}
+
+		OutMessages.Add(FCoreDelegates::EOnScreenMessageSeverity::Warning, FText::FromString(*GetOwner()->GetPathName()));
+		OutMessages.Add(FCoreDelegates::EOnScreenMessageSeverity::Warning, FText::FromString(*SelectedComponent->GetName()));
+		OutMessages.Add(FCoreDelegates::EOnScreenMessageSeverity::Info, FText::FromString(ISMDescription));
+	}
+}
+
+#endif
diff --git a/Engine/Source/Runtime/Experimental/GeometryCollectionEngine/Public/GeometryCollection/GeometryCollectionISMPoolDebugDrawComponent.h b/Engine/Source/Runtime/Experimental/GeometryCollectionEngine/Public/GeometryCollection/GeometryCollectionISMPoolDebugDrawComponent.h
new file mode 100644
--- /dev/null
+++ b/Engine/Source/Runtime/Experimental/GeometryCollectionEngine/Public/GeometryCollection/GeometryCollectionISMPoolDebugDrawComponent.h
@@ -0,0 +1,55 @@
+// Copyright Epic Games, Inc. All Rights Reserved.
+
+#pragma once
+
+#include "Debug/DebugDrawComponent.h"
+#include "DebugRenderSceneProxy.h"
+#include "GeometryCollectionISMPoolDebugDrawComponent.generated.h"
+
+class FGeometryCollectionISMPoolDebugDrawDelegateHelper;
+class UInstancedStaticMeshComponent;
+
+UCLASS(ClassGroup = Debug)
+class UGeometryCollectionISMPoolDebugDrawComponent : public UDebugDrawComponent
+{
+	GENERATED_UCLASS_BODY()
+
+	UPROPERTY(EditAnywhere, Category = Debug)
+	bool bShowGlobalStats = false;
+
+	UPROPERTY(EditAnywhere, Category = Debug)
+	bool bShowStats = false;
+
+	UPROPERTY(EditAnywhere, Category = Debug)
+	bool bShowBounds = false;
+
+	UPROPERTY(Transient)
+	TObjectPtr<const UInstancedStaticMeshComponent> SelectedComponent;
+
+	float SelectTimer = 0.f;
+
+protected:
+	virtual void BeginPlay() override;
+	virtual void EndPlay(const EEndPlayReason::Type EndPlayReason) override;
+	virtual void TickComponent(float DeltaTime, enum ELevelTick TickType, FActorComponentTickFunction* ThisTickFunction) override;
+	virtual FBoxSphereBounds CalcBounds(const FTransform& LocalToWorld) const override;
+
+#if WITH_EDITOR
+	virtual void PostEditChangeProperty(FPropertyChangedEvent& PropertyChangedEvent) override;
+#endif
+
+	void UpdateTickEnabled();
+
+#if UE_ENABLE_DEBUG_DRAWING
+  	virtual FDebugRenderSceneProxy* CreateDebugSceneProxy() override;
+	
+	FDebugDrawDelegateHelper DebugDrawDelegateHelper;
+	virtual FDebugDrawDelegateHelper& GetDebugDrawDelegateHelper() override { return DebugDrawDelegateHelper; }
+
+	FDelegateHandle OnScreenMessagesHandle;
+	void GetOnScreenMessages(TMultiMap<FCoreDelegates::EOnScreenMessageSeverity, FText>& OutMessages);
+#endif
+
+public:
+	static void UpdateAllTickEnabled();
+};
