- Make sure that we always remove elements from Spatial acceleration substructures even when the spatial indices are wrong.
- Detect when a wrong spatial index was used to add or remove elements with ensures (commented out for now)

diff --git a/Engine/Source/Runtime/Experimental/Chaos/Public/Chaos/AABBTree.h b/Engine/Source/Runtime/Experimental/Chaos/Public/Chaos/AABBTree.h
--- a/Engine/Source/Runtime/Experimental/Chaos/Public/Chaos/AABBTree.h
+++ b/Engine/Source/Runtime/Experimental/Chaos/Public/Chaos/AABBTree.h
@@ -1601,7 +1601,7 @@ public:
 		DeAllocateLeafNode(LeafNodeIdx);
 	}
 
-	virtual void RemoveElement(const TPayloadType& Payload)
+	virtual bool RemoveElement(const TPayloadType& Payload)
 	{
 		ensure(bModifyingTreeMultiThreadingFastCheck == false);
 		bModifyingTreeMultiThreadingFastCheck = true;
@@ -1655,9 +1655,12 @@ public:
 
 				PayloadToInfo.Remove(Payload);
 				bShouldRebuild = true;
+				bModifyingTreeMultiThreadingFastCheck = false;
+				return true;
 			}
 		}
 		bModifyingTreeMultiThreadingFastCheck = false;
+		return false;
 	}
 
 	virtual void UpdateElement(const TPayloadType& Payload, const TAABB<T, 3>& NewBounds, bool bInHasBounds) override
diff --git a/Engine/Source/Runtime/Experimental/Chaos/Public/Chaos/BoundingVolume.h b/Engine/Source/Runtime/Experimental/Chaos/Public/Chaos/BoundingVolume.h
--- a/Engine/Source/Runtime/Experimental/Chaos/Public/Chaos/BoundingVolume.h
+++ b/Engine/Source/Runtime/Experimental/Chaos/Public/Chaos/BoundingVolume.h
@@ -241,7 +241,7 @@ public:
 		bIsEmpty = true;
 	}
 
-	virtual void RemoveElement(const TPayloadType& Payload) override
+	virtual bool RemoveElement(const TPayloadType& Payload) override
 	{
 		SCOPE_CYCLE_COUNTER(STAT_BoundingVolumeRemoveElement);
 		if (const FPayloadInfo* PayloadInfo = MPayloadInfo.Find(Payload))
@@ -256,7 +256,9 @@ public:
 			}
 
 			MPayloadInfo.Remove(Payload);
+			return true;
 		}
+		return false;
 	}
 
 	virtual void UpdateElement(const TPayloadType& Payload, const TAABB<T,d>& NewBounds, bool bHasBounds) override
diff --git a/Engine/Source/Runtime/Experimental/Chaos/Public/Chaos/ISpatialAcceleration.h b/Engine/Source/Runtime/Experimental/Chaos/Public/Chaos/ISpatialAcceleration.h
--- a/Engine/Source/Runtime/Experimental/Chaos/Public/Chaos/ISpatialAcceleration.h
+++ b/Engine/Source/Runtime/Experimental/Chaos/Public/Chaos/ISpatialAcceleration.h
@@ -288,9 +288,11 @@ public:
 		check(false);
 	}
 
-	virtual void RemoveElement(const TPayloadType& Payload)
+	// Returns true if element was removed successfully
+	virtual bool RemoveElement(const TPayloadType& Payload)
 	{
 		check(false);	//not implemented
+		return true;
 	}
 
 	virtual void UpdateElement(const TPayloadType& Payload, const TAABB<T, d>& NewBounds, bool bHasBounds)
diff --git a/Engine/Source/Runtime/Experimental/Chaos/Public/Chaos/SpatialAccelerationCollection.h b/Engine/Source/Runtime/Experimental/Chaos/Public/Chaos/SpatialAccelerationCollection.h
--- a/Engine/Source/Runtime/Experimental/Chaos/Public/Chaos/SpatialAccelerationCollection.h
+++ b/Engine/Source/Runtime/Experimental/Chaos/Public/Chaos/SpatialAccelerationCollection.h
@@ -573,7 +573,21 @@ public:
 	virtual void RemoveElementFrom(const TPayloadType& Payload, FSpatialAccelerationIdx SpatialIdx) override
 	{
 		const uint16 UseBucket = ((1 << SpatialIdx.Bucket) & this->ActiveBucketsMask) ? SpatialIdx.Bucket : 0;
-		Buckets[UseBucket].Objects[SpatialIdx.InnerIdx].Acceleration->RemoveElement(Payload);
+		const bool Success = Buckets[UseBucket].Objects[SpatialIdx.InnerIdx].Acceleration->RemoveElement(Payload);
+		//ensure(Success); // We are removing this element with the wrong spatial index
+		if (!Success)
+		{
+			// Make sure that we remove this Payload to prevent crashes
+			for (FSpatialAccelerationIdx Idx : GetAllSpatialIndices())
+			{
+				if (!(Idx == SpatialIdx))
+				{
+					const uint16 Buckt = ((1 << Idx.Bucket) & this->ActiveBucketsMask) ? Idx.Bucket : 0;
+					const bool Removed = Buckets[Buckt].Objects[Idx.InnerIdx].Acceleration->RemoveElement(Payload);
+					//ensure(!Removed); // If this ensure hits, the payload was found in another substructure and successfully removed
+				}
+			}
+		}
 	}
 
 	virtual void UpdateElementIn(const TPayloadType& Payload, const TAABB<T, d>& NewBounds, bool bHasBounds, FSpatialAccelerationIdx SpatialIdx)
