Chaos - undo disable of friction anchors until the system can be fixed

diff --git a/Engine/Source/Runtime/Experimental/Chaos/Private/Chaos/Collision/PBDCollisionConstraint.cpp b/Engine/Source/Runtime/Experimental/Chaos/Private/Chaos/Collision/PBDCollisionConstraint.cpp
--- a/Engine/Source/Runtime/Experimental/Chaos/Private/Chaos/Collision/PBDCollisionConstraint.cpp
+++ b/Engine/Source/Runtime/Experimental/Chaos/Private/Chaos/Collision/PBDCollisionConstraint.cpp
@@ -1037,7 +1037,7 @@ namespace Chaos
 	{
 		int32 MatchIndex = INDEX_NONE;
 
-		const bool bEnableNetworkPhysicsResim = FPhysicsSolverBase::IsNetworkPhysicsPredictionEnabled() && FPhysicsSolverBase::CanResimNetworkPhysicsPrediction();
+		const bool bEnableNetworkPhysicsResim = false;//FPhysicsSolverBase::IsNetworkPhysicsPredictionEnabled() && FPhysicsSolverBase::CanResimNetworkPhysicsPrediction();
 		if (bChaos_Manifold_EnableFrictionRestore && !bEnableNetworkPhysicsResim)
 		{
 			const FManifoldPoint& ManifoldPoint = ManifoldPoints[ManifoldPointIndex];
