Fix crash in acceleration structure.
An overlap query is crashing on the GT because the particle has been deleted but not remove from the acceleration structure.

On the physics thread when a Geometry Collection Proxy is removed it delete all its particles. Those particles could be in the GT or in the PT. There were supposed to be removed from the GT.

The fix ensure all particle from the GT are removed when they should.

diff --git a/Engine/Source/Runtime/Experimental/Chaos/Private/Chaos/PhysicsObjectInterface.cpp b/Engine/Source/Runtime/Experimental/Chaos/Private/Chaos/PhysicsObjectInterface.cpp
--- a/Engine/Source/Runtime/Experimental/Chaos/Private/Chaos/PhysicsObjectInterface.cpp
+++ b/Engine/Source/Runtime/Experimental/Chaos/Private/Chaos/PhysicsObjectInterface.cpp
@@ -495,7 +495,7 @@ namespace Chaos
 		bool bDisabled = !InObjects.IsEmpty();
 		for (FPhysicsObjectHandle Object : InObjects)
 		{
-			bool bParticleDisabled = false;
+			bool bParticleDisabled = true;
 			if (Object)
 			{
 				if (TThreadParticle<Id>* Particle = Object->GetParticle<Id>())
