remove ensure for single shape in drawing unions

diff --git a/Engine/Source/Runtime/Experimental/Chaos/Private/Chaos/ChaosDebugDraw.cpp b/Engine/Source/Runtime/Experimental/Chaos/Private/Chaos/ChaosDebugDraw.cpp
--- a/Engine/Source/Runtime/Experimental/Chaos/Private/Chaos/ChaosDebugDraw.cpp
+++ b/Engine/Source/Runtime/Experimental/Chaos/Private/Chaos/ChaosDebugDraw.cpp
@@ -544,7 +544,7 @@ namespace Chaos
 				{
 					// Retrieve shape from union's shapes array
 					const FPerShapeData* PerShapeData = nullptr;
-					if (ensure(!Shapes.IsSingleShape()))
+					if (!Shapes.IsSingleShape())
 					{
 						const FShapesArray& ShapesArray = *Shapes.GetShapesArray();
 						PerShapeData = ShapesArray[UnionIdx].Get();
