Dataflow : update existing ed graph node pins to reflect changed inputs or outputs on the underlying dataflow nodes

diff --git a/Engine/Source/Runtime/Experimental/Dataflow/Engine/Private/Dataflow/DataflowEdNode.cpp b/Engine/Source/Runtime/Experimental/Dataflow/Engine/Private/Dataflow/DataflowEdNode.cpp
--- a/Engine/Source/Runtime/Experimental/Dataflow/Engine/Private/Dataflow/DataflowEdNode.cpp
+++ b/Engine/Source/Runtime/Experimental/Dataflow/Engine/Private/Dataflow/DataflowEdNode.cpp
@@ -94,6 +94,69 @@ void UDataflowEdNode::AllocateDefaultPins()
 #endif // WITH_EDITOR && !UE_BUILD_SHIPPING
 }
 
+void UDataflowEdNode::UpdatePinsFromDataflowNode()
+{
+	UE_LOG(DATAFLOWNODE_LOG, Verbose, TEXT("UDataflowEdNode::UpdatePinsFromDataflowNode()"));
+	// called on node creation from UI. 
+
+#if WITH_EDITOR && !(UE_BUILD_SHIPPING || UE_BUILD_TEST)
+	if (DataflowGraph)
+	{
+		if (DataflowNodeGuid.IsValid())
+		{
+			if (TSharedPtr<FDataflowNode> DataflowNode = DataflowGraph->FindBaseNode(DataflowNodeGuid))
+			{
+				// remove pins that do not match inputs / outputs anymore
+				TArray<UEdGraphPin*> PinsToRemove;
+				for (UEdGraphPin* Pin : GetAllPins())
+				{
+					if (Pin)
+					{
+						if (Pin->Direction == EEdGraphPinDirection::EGPD_Input)
+						{
+							if (!DataflowNode->FindInput(Pin->GetFName()))
+							{
+								PinsToRemove.Add(Pin);
+							}
+						}
+						else if (Pin->Direction == EEdGraphPinDirection::EGPD_Output)
+						{
+							if (!DataflowNode->FindOutput(Pin->GetFName()))
+							{
+								PinsToRemove.Add(Pin);
+							}
+						}
+					}
+				}
+				for (UEdGraphPin* PinToRemove : PinsToRemove)
+				{
+					RemovePin(PinToRemove);
+				}
+				PinsToRemove.Reset();
+
+				for (const Dataflow::FPin& Pin : DataflowNode->GetPins())
+				{
+					if (Pin.Direction == Dataflow::FPin::EDirection::INPUT)
+					{
+						if (!FindPin(Pin.Name, EEdGraphPinDirection::EGPD_Input))
+						{
+							CreatePin(EEdGraphPinDirection::EGPD_Input, Pin.Type, Pin.Name);
+						}
+					}
+					if (Pin.Direction == Dataflow::FPin::EDirection::OUTPUT)
+					{
+						if (!FindPin(Pin.Name, EEdGraphPinDirection::EGPD_Output))
+						{
+							CreatePin(EEdGraphPinDirection::EGPD_Output, Pin.Type, Pin.Name);
+						}
+					}
+				}
+			}
+		}
+	}
+#endif // WITH_EDITOR && !UE_BUILD_SHIPPING
+}
+
 void UDataflowEdNode::AddOptionPin()
 {
 #if WITH_EDITOR && !(UE_BUILD_SHIPPING || UE_BUILD_TEST)
diff --git a/Engine/Source/Runtime/Experimental/Dataflow/Engine/Private/Dataflow/DataflowObject.cpp b/Engine/Source/Runtime/Experimental/Dataflow/Engine/Private/Dataflow/DataflowObject.cpp
--- a/Engine/Source/Runtime/Experimental/Dataflow/Engine/Private/Dataflow/DataflowObject.cpp
+++ b/Engine/Source/Runtime/Experimental/Dataflow/Engine/Private/Dataflow/DataflowObject.cpp
@@ -67,6 +67,7 @@ void UDataflow::PostLoad()
 		if (DataflowEdNode)
 		{
 			DataflowEdNode->SetDataflowGraph(Dataflow);
+			DataflowEdNode->UpdatePinsFromDataflowNode();
 		}
 
 		if (DisabledNodes.Contains(FName(EdNode->GetName())))
diff --git a/Engine/Source/Runtime/Experimental/Dataflow/Engine/Public/Dataflow/DataflowEdNode.h b/Engine/Source/Runtime/Experimental/Dataflow/Engine/Public/Dataflow/DataflowEdNode.h
--- a/Engine/Source/Runtime/Experimental/Dataflow/Engine/Public/Dataflow/DataflowEdNode.h
+++ b/Engine/Source/Runtime/Experimental/Dataflow/Engine/Public/Dataflow/DataflowEdNode.h
@@ -57,6 +57,8 @@ public:
 	TSharedPtr<const Dataflow::FGraph> GetDataflowGraph() const { return DataflowGraph; }
 	void SetDataflowGraph(TSharedPtr<Dataflow::FGraph> InDataflowGraph) { DataflowGraph = InDataflowGraph; }
 
+	void UpdatePinsFromDataflowNode();
+
 	FGuid GetDataflowNodeGuid() const { return DataflowNodeGuid; }
 	void SetDataflowNodeGuid(FGuid InGuid) { DataflowNodeGuid = InGuid; }
 
