Additional LLM tracking tags for geometrycollection particles.

diff --git a/Engine/Source/Runtime/Experimental/Chaos/Private/PhysicsProxy/GeometryCollectionPhysicsProxy.cpp b/Engine/Source/Runtime/Experimental/Chaos/Private/PhysicsProxy/GeometryCollectionPhysicsProxy.cpp
--- a/Engine/Source/Runtime/Experimental/Chaos/Private/PhysicsProxy/GeometryCollectionPhysicsProxy.cpp
+++ b/Engine/Source/Runtime/Experimental/Chaos/Private/PhysicsProxy/GeometryCollectionPhysicsProxy.cpp
@@ -804,6 +804,7 @@ FAutoConsoleVariableRef CVarReportTooManyChildrenNum(TEXT("p.ReportTooManyChildr
 
 void FGeometryCollectionPhysicsProxy::CreateNonClusteredParticles(Chaos::FPBDRigidsSolver* RigidsSolver, const FGeometryCollection& RestCollection, const FGeometryDynamicCollection& DynamicCollection, const TBitArray<>& EffectiveParticles)
 {
+	LLM_SCOPE_BYNAME(TEXT("Physics/NonClusteredParticles"));
 	const TManagedArray<bool>& SimulatableParticles = DynamicCollection.SimulatableParticles;
 	
 	//const int NumRigids = 0; // ryan - Since we're doing SOA, we start at zero?
diff --git a/Engine/Source/Runtime/Experimental/Chaos/Public/Chaos/PBDRigidsSOAs.h b/Engine/Source/Runtime/Experimental/Chaos/Public/Chaos/PBDRigidsSOAs.h
--- a/Engine/Source/Runtime/Experimental/Chaos/Public/Chaos/PBDRigidsSOAs.h
+++ b/Engine/Source/Runtime/Experimental/Chaos/Public/Chaos/PBDRigidsSOAs.h
@@ -238,18 +238,24 @@ public:
 	
 	TArray<FGeometryParticleHandle*> CreateStaticParticles(int32 NumParticles, const FUniqueIdx* ExistingIndices = nullptr, const FGeometryParticleParameters& Params = FGeometryParticleParameters())
 	{
+		LLM_SCOPE_BYNAME(TEXT("Physics/StaticParticles"));
+
 		auto Results = CreateParticlesHelper<FGeometryParticleHandle>(NumParticles, ExistingIndices, Params.bDisabled ? StaticDisabledParticles : StaticParticles, Params);
 		UpdateViews();
 		return Results;
 	}
 	TArray<FKinematicGeometryParticleHandle*> CreateKinematicParticles(int32 NumParticles, const FUniqueIdx* ExistingIndices = nullptr,  const FKinematicGeometryParticleParameters& Params = FKinematicGeometryParticleParameters())
 	{
+		LLM_SCOPE_BYNAME(TEXT("Physics/KinematicParticles"));
+
 		auto Results = CreateParticlesHelper<FKinematicGeometryParticleHandle>(NumParticles, ExistingIndices, Params.bDisabled ? KinematicDisabledParticles : KinematicParticles, Params);
 		UpdateViews();
 		return Results;
 	}
 	TArray<FPBDRigidParticleHandle*> CreateDynamicParticles(int32 NumParticles, const FUniqueIdx* ExistingIndices = nullptr,  const FPBDRigidParticleParameters& Params = FPBDRigidParticleParameters())
 	{
+		LLM_SCOPE_BYNAME(TEXT("Physics/DynamicParticles"));
+
 		auto Results = CreateParticlesHelper<FPBDRigidParticleHandle>(NumParticles, ExistingIndices, Params.bDisabled ? DynamicDisabledParticles : DynamicParticles, Params);;
 
 		if (!Params.bStartSleeping)
@@ -261,6 +267,8 @@ public:
 	}
 	TArray<TPBDGeometryCollectionParticleHandle<FReal, 3>*> CreateGeometryCollectionParticles(int32 NumParticles, const FUniqueIdx* ExistingIndices = nullptr, const FPBDRigidParticleParameters& Params = FPBDRigidParticleParameters())
 	{
+		LLM_SCOPE_BYNAME(TEXT("Physics/GeometryCollectionParticles"));
+
 		TArray<TPBDGeometryCollectionParticleHandle<FReal, 3>*> Results = CreateParticlesHelper<TPBDGeometryCollectionParticleHandle<FReal, 3>>(
 			NumParticles, ExistingIndices, GeometryCollectionParticles, Params);
 		for (auto* Handle : Results)
@@ -283,6 +291,8 @@ public:
 	/** Used specifically by PBDRigidClustering. These have special properties for maintaining relative order, efficiently switching from kinematic to dynamic, disable to enable, etc... */
 	TArray<FPBDRigidClusteredParticleHandle*> CreateClusteredParticles(int32 NumParticles, const FUniqueIdx* ExistingIndices = nullptr,  const FPBDRigidParticleParameters& Params = FPBDRigidParticleParameters())
 	{
+		LLM_SCOPE_BYNAME(TEXT("Physics/ClusteredParticles"));
+
 		auto NewClustered = CreateParticlesHelper<FPBDRigidClusteredParticleHandle>(NumParticles, ExistingIndices, ClusteredParticles, Params);
 		
 		if (!Params.bDisabled)
