put cluster union particles in FSpatialAccelerationIdx{0, 1} instead similarly to geometry collections

diff --git a/Engine/Source/Runtime/Experimental/Chaos/Private/PBDRigidsSolver.cpp b/Engine/Source/Runtime/Experimental/Chaos/Private/PBDRigidsSolver.cpp
--- a/Engine/Source/Runtime/Experimental/Chaos/Private/PBDRigidsSolver.cpp
+++ b/Engine/Source/Runtime/Experimental/Chaos/Private/PBDRigidsSolver.cpp
@@ -693,7 +693,17 @@ namespace Chaos
 
 		if (FPBDRigidParticle* Particle = Proxy->GetParticle_External())
 		{
-			Particle->SetSpatialIdx(FSpatialAccelerationIdx{ 0, 0 });
+			if (AccelerationStructureSplitStaticAndDynamic == 1)
+			{
+				// It needs to be {0, 1}. Things will crash otherwise. I don't know why.
+				// Probably similarly to the geometry collection physics proxy.
+				Particle->SetSpatialIdx(FSpatialAccelerationIdx{ 0, 1 });
+			}
+			else
+			{
+				Particle->SetSpatialIdx(FSpatialAccelerationIdx{ 0, 0 });
+			}
+
 			Particle->SetUniqueIdx(GetEvolution()->GenerateUniqueIdx());
 		}
 		Proxy->SetSolver(this);
