Chaos - add debug draw for singler contact constraint

diff --git a/Engine/Source/Runtime/Experimental/Chaos/Private/Chaos/ChaosDebugDraw.cpp b/Engine/Source/Runtime/Experimental/Chaos/Private/Chaos/ChaosDebugDraw.cpp
--- a/Engine/Source/Runtime/Experimental/Chaos/Private/Chaos/ChaosDebugDraw.cpp
+++ b/Engine/Source/Runtime/Experimental/Chaos/Private/Chaos/ChaosDebugDraw.cpp
@@ -1702,6 +1702,14 @@ namespace Chaos
 			}
 		}
 
+		void DrawCollision(const FRigidTransform3& SpaceTransform, const FPBDCollisionConstraint& Collision, FRealSingle ColorScale, const FChaosDebugDrawSettings* Settings)
+		{
+			if (FDebugDrawQueue::IsDebugDrawingEnabled())
+			{
+				DrawCollisionImpl(SpaceTransform, Collision, ColorScale, ChaosDebugDrawCollisionDuration, GetChaosDebugDrawSettings(Settings));
+			}
+		}
+
 		void DrawCollisions(const FRigidTransform3& SpaceTransform, const FPBDCollisionConstraints& Collisions, FRealSingle ColorScale, const FChaosDebugDrawSettings* Settings)
 		{
 			if (FDebugDrawQueue::IsDebugDrawingEnabled())
diff --git a/Engine/Source/Runtime/Experimental/Chaos/Public/Chaos/ChaosDebugDraw.h b/Engine/Source/Runtime/Experimental/Chaos/Public/Chaos/ChaosDebugDraw.h
--- a/Engine/Source/Runtime/Experimental/Chaos/Public/Chaos/ChaosDebugDraw.h
+++ b/Engine/Source/Runtime/Experimental/Chaos/Public/Chaos/ChaosDebugDraw.h
@@ -213,6 +213,7 @@ namespace Chaos
 		CHAOS_API void DrawParticleTransforms(const FRigidTransform3& SpaceTransform, const TParticleView<FKinematicGeometryParticles>& ParticlesView, const FChaosDebugDrawSettings* Settings = nullptr);
 		CHAOS_API void DrawParticleTransforms(const FRigidTransform3& SpaceTransform, const TParticleView<FPBDRigidParticles>& ParticlesView, const FChaosDebugDrawSettings* Settings = nullptr);
 		CHAOS_API void DrawParticleCollisions(const FRigidTransform3& SpaceTransform, const FGeometryParticleHandle* Particle, const FPBDCollisionConstraints& Collisions, const FChaosDebugDrawSettings* Settings = nullptr);
+		CHAOS_API void DrawCollision(const FRigidTransform3& SpaceTransform, const FPBDCollisionConstraint& Collision, FRealSingle ColorScale, const FChaosDebugDrawSettings* Settings = nullptr);
 		CHAOS_API void DrawCollisions(const FRigidTransform3& SpaceTransform, const FPBDCollisionConstraints& Collisions, FRealSingle ColorScale, const FChaosDebugDrawSettings* Settings = nullptr);
 		CHAOS_API void DrawCollisions(const FRigidTransform3& SpaceTransform, const Private::FCollisionConstraintAllocator& CollisionAllocator, FRealSingle ColorScale, const FChaosDebugDrawSettings* Settings = nullptr);
 		CHAOS_API void DrawJointConstraints(const FRigidTransform3& SpaceTransform, const TArray<FPBDJointConstraintHandle*>& ConstraintHandles, FRealSingle ColorScale, const FChaosDebugDrawJointFeatures& FeatureMask = FChaosDebugDrawJointFeatures::MakeDefault(), const FChaosDebugDrawSettings* Settings = nullptr);
