move ChaosSQTypes from PhysicsCore to Chaos

diff --git a/Engine/Source/Runtime/Experimental/Chaos/Private/Chaos/Interface/SQTypes.cpp b/Engine/Source/Runtime/Experimental/Chaos/Private/Chaos/Interface/SQTypes.cpp
new file mode 100644
--- /dev/null
+++ b/Engine/Source/Runtime/Experimental/Chaos/Private/Chaos/Interface/SQTypes.cpp
@@ -0,0 +1,39 @@
+// Copyright Epic Games, Inc. All Rights Reserved.
+
+#include "Chaos/Interface/SQTypes.h"
+#include "Chaos/ChaosArchive.h"
+#include "Chaos/Serializable.h"
+#include "Chaos/ParticleHandle.h"
+
+namespace ChaosInterface
+{
+	void FActorShape::Serialize(Chaos::FChaosArchive& Ar)
+	{
+		Ar << Chaos::AsAlwaysSerializable(Actor);
+		auto NonConstShape = const_cast<Chaos::FPerShapeData*>(Shape);
+		Ar << Chaos::AsAlwaysSerializable(NonConstShape);
+		Shape = NonConstShape;
+	}
+
+	void FQueryHit::Serialize(Chaos::FChaosArchive& Ar)
+	{
+		FActorShape::Serialize(Ar);
+		Ar << FaceIndex;
+	}
+
+	void FLocationHit::Serialize(Chaos::FChaosArchive& Ar)
+	{
+		FQueryHit::Serialize(Ar);
+		Ar << Flags;
+		Ar << WorldPosition;
+		Ar << WorldNormal;
+		Ar << Distance;
+	}
+
+	void FRaycastHit::Serialize(Chaos::FChaosArchive& Ar)
+	{
+		FLocationHit::Serialize(Ar);
+		Ar << U;
+		Ar << V;
+	}
+}
diff --git a/Engine/Source/Runtime/Experimental/Chaos/Public/Chaos/Interface/PhysicsInterfaceWrapperShared.h b/Engine/Source/Runtime/Experimental/Chaos/Public/Chaos/Interface/PhysicsInterfaceWrapperShared.h
new file mode 100644
--- /dev/null
+++ b/Engine/Source/Runtime/Experimental/Chaos/Public/Chaos/Interface/PhysicsInterfaceWrapperShared.h
@@ -0,0 +1,119 @@
+// Copyright Epic Games, Inc. All Rights Reserved.
+
+#pragma once
+
+#include "HAL/Platform.h"
+#include "Serialization/Archive.h"
+
+enum class EQueryFlags : uint16
+{
+	None = 0,
+	PreFilter = (1 << 2),
+	PostFilter = (1 << 3),
+	AnyHit = (1 << 4),
+	SkipNarrowPhase = (1 << 5)
+};
+
+inline EQueryFlags operator| (EQueryFlags lhs, EQueryFlags rhs)
+{
+	return static_cast<EQueryFlags>(static_cast<uint16>(lhs) | static_cast<uint16>(rhs));
+}
+
+inline EQueryFlags operator&(EQueryFlags lhs, EQueryFlags rhs)
+{
+	return static_cast<EQueryFlags>(static_cast<uint16>(lhs) & static_cast<uint16>(rhs));
+}
+
+struct FQueryFlags
+{
+	FQueryFlags(EQueryFlags InFlags) : QueryFlags(InFlags) {}
+	explicit operator bool() const { return !!static_cast<uint16>(QueryFlags); }
+	FQueryFlags operator |(EQueryFlags Rhs) const
+	{
+		return FQueryFlags(QueryFlags | Rhs);
+	}
+
+	FQueryFlags& operator |=(EQueryFlags Rhs)
+	{
+		*this = *this | Rhs;
+		return *this;
+	}
+
+	FQueryFlags operator &(EQueryFlags Rhs) const
+	{
+		return FQueryFlags(QueryFlags & Rhs);
+	}
+
+	FQueryFlags& operator &=(EQueryFlags Rhs)
+	{
+		*this = *this & Rhs;
+		return *this;
+	}
+
+	bool operator==(const FQueryFlags& Other) const { return QueryFlags == Other.QueryFlags; }
+
+	EQueryFlags QueryFlags;
+};
+
+inline FArchive& operator<<(FArchive& Ar, FQueryFlags& QueryFlags)
+{
+	return Ar << QueryFlags.QueryFlags;
+}
+
+/** Possible results from a scene query */
+enum class EHitFlags : uint16
+{
+	None = 0,
+	Position = (1 << 0),
+	Normal = (1 << 1),
+	Distance = (1 << 2),
+	UV = (1 << 3),
+	MTD = (1 << 9),
+	FaceIndex = (1 << 10)
+};
+
+inline EHitFlags operator|(EHitFlags lhs, EHitFlags rhs)
+{
+	return static_cast<EHitFlags>(static_cast<uint16>(lhs) | static_cast<uint16>(rhs));
+}
+
+inline EHitFlags operator&(EHitFlags lhs, EHitFlags rhs)
+{
+	return static_cast<EHitFlags>(static_cast<uint16>(lhs) & static_cast<uint16>(rhs));
+}
+
+struct FHitFlags
+{
+	FHitFlags(EHitFlags InFlags = EHitFlags::None) : HitFlags(InFlags) {}
+	explicit operator bool() const { return !!static_cast<uint16>(HitFlags); }
+	FHitFlags operator |(EHitFlags Rhs) const
+	{
+		return FHitFlags(HitFlags | Rhs);
+	}
+
+	FHitFlags& operator |=(EHitFlags Rhs)
+	{
+		*this = *this | Rhs;
+		return *this;
+	}
+
+	FHitFlags operator &(EHitFlags Rhs) const
+	{
+		return FHitFlags(HitFlags & Rhs);
+	}
+
+	FHitFlags& operator &=(EHitFlags Rhs)
+	{
+		*this = *this & Rhs;
+		return *this;
+	}
+
+	bool operator==(const FHitFlags& Other) const { return HitFlags == Other.HitFlags; }
+
+	EHitFlags HitFlags;
+};
+
+inline FArchive& operator<<(FArchive& Ar, FHitFlags& HitFlags)
+{
+	return Ar << HitFlags.HitFlags;
+}
diff --git a/Engine/Source/Runtime/Experimental/Chaos/Public/Chaos/Interface/SQTypes.h b/Engine/Source/Runtime/Experimental/Chaos/Public/Chaos/Interface/SQTypes.h
new file mode 100644
--- /dev/null
+++ b/Engine/Source/Runtime/Experimental/Chaos/Public/Chaos/Interface/SQTypes.h
@@ -0,0 +1,295 @@
+// Copyright Epic Games, Inc. All Rights Reserved.
+
+#pragma once 
+
+#include "Chaos/Declares.h"
+#include "Chaos/Interface/PhysicsInterfaceWrapperShared.h"
+#include "Chaos/ParticleHandleFwd.h"
+#include "Containers/Array.h"
+#include "Containers/ContainerAllocationPolicies.h"
+#include "CoreMinimal.h"
+#include "HAL/PlatformCrt.h"
+#include "Math/UnrealMathSSE.h"
+#include "Math/Vector.h"
+
+namespace Chaos
+{
+	class FChaosArchive;
+	class FPerShapeData;
+}
+
+namespace ChaosInterface
+{
+	struct CHAOS_API FActorShape
+	{
+		Chaos::FGeometryParticle* Actor;
+		const Chaos::FPerShapeData* Shape;
+
+		void Serialize(Chaos::FChaosArchive& Ar);
+	};
+
+	inline Chaos::FChaosArchive& operator<<(Chaos::FChaosArchive& Ar, FActorShape& ActorShape)
+	{
+		ActorShape.Serialize(Ar);
+		return Ar;
+	}
+
+	struct CHAOS_API FQueryHit : public FActorShape
+	{
+		FQueryHit() : FaceIndex(-1), ElementIndex(-1) {}
+
+		/**
+		Face index of touched triangle, for triangle meshes, convex meshes and height fields. Defaults to -1 if face index is not available
+		*/
+
+		int32 FaceIndex; // Signed int to match TArray's size type, and so INDEX_NONE/-1 doesn't underflow.
+
+		int32 ElementIndex; // Currently used to indicate which shape was hit for a particle with multiple shapes.
+
+		void Serialize(Chaos::FChaosArchive& Ar);
+	};
+
+	inline Chaos::FChaosArchive& operator<<(Chaos::FChaosArchive& Ar, FQueryHit& QueryHit)
+	{
+		QueryHit.Serialize(Ar);
+		return Ar;
+	}
+
+	struct CHAOS_API FLocationHit : public FQueryHit
+	{
+		FHitFlags Flags;
+		FVector WorldPosition;
+		FVector WorldNormal;
+		float Distance;	// LWC_TODO: Should be FVector::FReal, but that causes precision issues resulting in collision failures. Investigate!
+
+		void Serialize(Chaos::FChaosArchive& Ar);
+
+		bool operator<(const FLocationHit& Other) const { return Distance < Other.Distance; }
+	};
+
+	inline Chaos::FChaosArchive& operator<<(Chaos::FChaosArchive& Ar, FLocationHit& LocationHit)
+	{
+		LocationHit.Serialize(Ar);
+		return Ar;
+	}
+
+	struct CHAOS_API FRaycastHit : public FLocationHit
+	{
+		float U;
+		float V;
+
+		void Serialize(Chaos::FChaosArchive& Ar);
+
+	};
+
+	inline Chaos::FChaosArchive& operator<<(Chaos::FChaosArchive& Ar, FRaycastHit& RaycastHit)
+	{
+		RaycastHit.Serialize(Ar);
+		return Ar;
+	}
+
+	struct CHAOS_API FOverlapHit : public FQueryHit
+	{
+	};
+
+	inline Chaos::FChaosArchive& operator<<(Chaos::FChaosArchive& Ar, FOverlapHit& OverlapHit)
+	{
+		OverlapHit.Serialize(Ar);
+		return Ar;
+	}
+
+	struct CHAOS_API FSweepHit : public FLocationHit
+	{
+	};
+
+	inline Chaos::FChaosArchive& operator<<(Chaos::FChaosArchive& Ar, FSweepHit& SweepHit)
+	{
+		SweepHit.Serialize(Ar);
+		return Ar;
+	}
+
+	struct CHAOS_API FPTActorShape
+	{
+		Chaos::FGeometryParticleHandle* Actor;
+		const Chaos::FPerShapeData* Shape;
+	};
+
+	struct CHAOS_API FPTQueryHit : public FPTActorShape
+	{
+		/**
+		Face index of touched triangle, for triangle meshes, convex meshes and height fields. Defaults to -1 if face index is not available
+		*/
+
+		int32 FaceIndex = -1; // Signed int to match TArray's size type, and so INDEX_NONE/-1 doesn't underflow.
+
+		int32 ElementIndex = -1; // Currently used to indicate which shape was hit for a particle with multiple shapes.
+	};
+
+	struct CHAOS_API FPTLocationHit : public FPTQueryHit
+	{
+		FHitFlags Flags;
+		FVector WorldPosition;
+		FVector WorldNormal;
+		float Distance;	// LWC_TODO: Should be FVector::FReal, but that causes precision issues resulting in collision failures. Investigate!
+
+		bool operator<(const FPTLocationHit& Other) const { return Distance < Other.Distance; }
+	};
+
+	struct CHAOS_API FPTRaycastHit : public FPTLocationHit
+	{
+		float U;
+		float V;
+	};
+
+	struct CHAOS_API FPTOverlapHit : public FPTQueryHit
+	{
+	};
+
+	struct CHAOS_API FPTSweepHit : public FPTLocationHit
+	{
+	};
+
+#ifndef CHAOS_HIT_BUFFER_SIZE
+#define CHAOS_HIT_BUFFER_SIZE 512 // Preallocated hit buffer size for traces and sweeps.
+#endif // CHAOS_HIT_BUFFER_SIZE
+	static_assert(CHAOS_HIT_BUFFER_SIZE > 0, "Invalid Chaos hit buffer size.");
+
+
+	inline void FinishQueryHelper(TArray<FOverlapHit, TInlineAllocator<CHAOS_HIT_BUFFER_SIZE>>& Hits, const FOverlapHit& BlockingHit, bool bHasBlockingHit)
+	{
+		if (bHasBlockingHit)
+		{
+			Hits.Add(BlockingHit);
+		}
+	}
+
+	inline void FinishQueryHelper(TArray<FPTOverlapHit, TInlineAllocator<CHAOS_HIT_BUFFER_SIZE>>& Hits, const FPTOverlapHit& BlockingHit, bool bHasBlockingHit)
+	{
+		if (bHasBlockingHit)
+		{
+			Hits.Add(BlockingHit);
+		}
+	}
+
+	template <typename HitType>
+	void FinishQueryHelper(TArray<HitType, TInlineAllocator<CHAOS_HIT_BUFFER_SIZE>>& Hits, const HitType& BlockingHit, bool bHasBlockingHit)
+	{
+		Hits.Sort();
+		if (bHasBlockingHit)
+		{
+			int32 FinalNum = Hits.Num() + 1;
+			for (int32 HitIdx = Hits.Num() - 1; HitIdx >= 0; --HitIdx)
+			{
+				if (Hits[HitIdx].Distance >= BlockingHit.Distance)
+				{
+					--FinalNum;
+				}
+				else
+				{
+					break;
+				}
+			}
+			Hits.SetNum(FinalNum);
+			Hits[FinalNum - 1] = BlockingHit;
+		}
+	}
+
+	/** Stores the results of scene queries. This can be passed around to multiple SQAccelerators and is responsible for sorting the results and pruning based on blocking.
+	IncFlushCount / DecFlushCount is used to ensure any final sort / pruning operation is done when all SQAccelerators are finished. If you are passing this into multiple accelerators
+	you should call IncFlushCount / DecFlushCount yourself as otherwise each accelerator will trigger its own sort / prune.
+	*/
+	template<typename HitType>
+	class FSQHitBuffer
+	{
+	public:
+		FSQHitBuffer(bool bSingle = false)
+			: AcceleratorDepth(0)
+			, bHasBlockingHit(false)
+			, bSingleResult(bSingle)
+		{
+			// The Hits array is preallocated through its InlineAllocator, no need to reserve.
+		}
+
+		virtual ~FSQHitBuffer() {}
+
+		//Called 
+		void IncFlushCount()
+		{
+			++AcceleratorDepth;
+		}
+
+		void DecFlushCount()
+		{
+			--AcceleratorDepth;
+			if (AcceleratorDepth == 0)
+			{
+				FinishQuery();
+			}
+		}
+
+		bool HasHit() const { return GetNumHits(); }
+		int32 GetNumHits() const { return Hits.Num(); }
+		HitType* GetHits() { return Hits.GetData(); }
+		const HitType* GetHits() const { return Hits.GetData(); }
+
+		HitType* GetBlock() { return HasBlockingHit() ? &Hits.Last() : nullptr; }
+		const HitType* GetBlock() const { return HasBlockingHit() ? &Hits.Last() : nullptr; }
+
+		bool HasBlockingHit() const { return bHasBlockingHit; }
+
+		void SetBlockingHit(const HitType& InBlockingHit)
+		{
+			CurrentBlockingHit = InBlockingHit;
+			bHasBlockingHit = true;
+		}
+
+		void AddTouchingHit(const HitType& InTouchingHit)
+		{
+			Hits.Add(InTouchingHit);
+		}
+
+		const HitType* GetCurrentHit() const { return bHasBlockingHit ? &CurrentBlockingHit : nullptr; }
+
+		/** Does not do any distance verification. This is up to the SQ code to manage */
+		void InsertHit(const HitType& Hit, bool bBlocking)
+		{
+			// Useful place to break when debugging, but breakpoints can't be set here, as
+			// this gets inlined even in a debug build.
+			//__debugbreak(); 
+			if (bBlocking)
+			{
+				SetBlockingHit(Hit);
+			}
+			else
+			{
+				AddTouchingHit(Hit);
+			}
+		}
+
+		bool WantsSingleResult() const { return bSingleResult; }
+
+	private:
+
+		void FinishQuery()
+		{
+			FinishQueryHelper(Hits, CurrentBlockingHit, bHasBlockingHit);
+		}
+
+		HitType CurrentBlockingHit;
+		int32 AcceleratorDepth;
+		bool bHasBlockingHit;
+		bool bSingleResult;
+
+		TArray<HitType, TInlineAllocator<CHAOS_HIT_BUFFER_SIZE>> Hits;
+	};
+
+	template<typename HitType>
+	class FSQSingleHitBuffer : public FSQHitBuffer<HitType>
+	{
+	public:
+		FSQSingleHitBuffer()
+			: FSQHitBuffer<HitType>(/*bSingle=*/ true)
+		{}
+	};
+
+}
diff --git a/Engine/Source/Runtime/PhysicsCore/Private/ChaosSQTypes.cpp b/Engine/Source/Runtime/PhysicsCore/Private/ChaosSQTypes.cpp
deleted file mode 100644
diff --git a/Engine/Source/Runtime/PhysicsCore/Public/ChaosSQTypes.h b/Engine/Source/Runtime/PhysicsCore/Public/ChaosSQTypes.h
--- a/Engine/Source/Runtime/PhysicsCore/Public/ChaosSQTypes.h
+++ b/Engine/Source/Runtime/PhysicsCore/Public/ChaosSQTypes.h
@@ -2,294 +2,4 @@
 
 #pragma once 
 
-#include "Chaos/Declares.h"
-#include "Chaos/ParticleHandleFwd.h"
-#include "Containers/Array.h"
-#include "Containers/ContainerAllocationPolicies.h"
-#include "CoreMinimal.h"
-#include "HAL/PlatformCrt.h"
-#include "Math/UnrealMathSSE.h"
-#include "Math/Vector.h"
-#include "PhysicsInterfaceWrapperShared.h"
-
-namespace Chaos
-{
-	class FChaosArchive;
-	class FPerShapeData;
-}
-
-namespace ChaosInterface
-{
-	struct FActorShape
-	{
-		Chaos::FGeometryParticle* Actor;
-		const Chaos::FPerShapeData* Shape;
-
-		void Serialize(Chaos::FChaosArchive& Ar);
-	};
-
-	inline Chaos::FChaosArchive& operator<<(Chaos::FChaosArchive& Ar, FActorShape& ActorShape)
-	{
-		ActorShape.Serialize(Ar);
-		return Ar;
-	}
-
-	struct FQueryHit : public FActorShape
-	{
-		FQueryHit() : FaceIndex(-1), ElementIndex(-1) {} 
-
-		/**
-		Face index of touched triangle, for triangle meshes, convex meshes and height fields. Defaults to -1 if face index is not available
-		*/
-
-		int32 FaceIndex; // Signed int to match TArray's size type, and so INDEX_NONE/-1 doesn't underflow.
-
-		int32 ElementIndex; // Currently used to indicate which shape was hit for a particle with multiple shapes.
-
-		void Serialize(Chaos::FChaosArchive& Ar);
-	};
-
-	inline Chaos::FChaosArchive& operator<<(Chaos::FChaosArchive& Ar, FQueryHit& QueryHit)
-	{
-		QueryHit.Serialize(Ar);
-		return Ar;
-	}
-
-	struct FLocationHit : public FQueryHit
-	{
-		FHitFlags Flags;
-		FVector WorldPosition;
-		FVector WorldNormal;
-		float Distance;	// LWC_TODO: Should be FVector::FReal, but that causes precision issues resulting in collision failures. Investigate!
-
-		void Serialize(Chaos::FChaosArchive& Ar);
-
-		bool operator<(const FLocationHit& Other) const { return Distance < Other.Distance; }
-	};
-
-	inline Chaos::FChaosArchive& operator<<(Chaos::FChaosArchive& Ar, FLocationHit& LocationHit)
-	{
-		LocationHit.Serialize(Ar);
-		return Ar;
-	}
-
-	struct FRaycastHit : public FLocationHit
-	{
-		float U;
-		float V;
-
-		void Serialize(Chaos::FChaosArchive& Ar);
-
-	};
-
-	inline Chaos::FChaosArchive& operator<<(Chaos::FChaosArchive& Ar, FRaycastHit& RaycastHit)
-	{
-		RaycastHit.Serialize(Ar);
-		return Ar;
-	}
-
-	struct FOverlapHit : public FQueryHit
-	{
-	};
-
-	inline Chaos::FChaosArchive& operator<<(Chaos::FChaosArchive& Ar, FOverlapHit& OverlapHit)
-	{
-		OverlapHit.Serialize(Ar);
-		return Ar;
-	}
-
-	struct FSweepHit : public FLocationHit
-	{
-	};
-
-	inline Chaos::FChaosArchive& operator<<(Chaos::FChaosArchive& Ar, FSweepHit& SweepHit)
-	{
-		SweepHit.Serialize(Ar);
-		return Ar;
-	}
-
-	struct FPTActorShape
-	{
-		Chaos::FGeometryParticleHandle* Actor;
-		const Chaos::FPerShapeData* Shape;
-	};
-
-	struct FPTQueryHit : public FPTActorShape
-	{
-		/**
-		Face index of touched triangle, for triangle meshes, convex meshes and height fields. Defaults to -1 if face index is not available
-		*/
-
-		int32 FaceIndex = -1; // Signed int to match TArray's size type, and so INDEX_NONE/-1 doesn't underflow.
-
-		int32 ElementIndex = -1; // Currently used to indicate which shape was hit for a particle with multiple shapes.
-	};
-
-	struct FPTLocationHit : public FPTQueryHit
-	{
-		FHitFlags Flags;
-		FVector WorldPosition;
-		FVector WorldNormal;
-		float Distance;	// LWC_TODO: Should be FVector::FReal, but that causes precision issues resulting in collision failures. Investigate!
-
-		bool operator<(const FPTLocationHit& Other) const { return Distance < Other.Distance; }
-	};
-
-	struct FPTRaycastHit : public FPTLocationHit
-	{
-		float U;
-		float V;
-	};
-	
-	struct FPTOverlapHit : public FPTQueryHit
-	{
-	};
-
-	struct FPTSweepHit : public FPTLocationHit
-	{
-	};
-
-	#ifndef CHAOS_HIT_BUFFER_SIZE
-		#define CHAOS_HIT_BUFFER_SIZE 512 // Preallocated hit buffer size for traces and sweeps.
-	#endif // CHAOS_HIT_BUFFER_SIZE
-	static_assert(CHAOS_HIT_BUFFER_SIZE > 0, "Invalid Chaos hit buffer size.");
-
-
-	inline void FinishQueryHelper(TArray<FOverlapHit, TInlineAllocator<CHAOS_HIT_BUFFER_SIZE>>& Hits, const FOverlapHit& BlockingHit, bool bHasBlockingHit)
-	{
-		if (bHasBlockingHit)
-		{
-			Hits.Add(BlockingHit);
-		}
-	}
-
-	inline void FinishQueryHelper(TArray<FPTOverlapHit, TInlineAllocator<CHAOS_HIT_BUFFER_SIZE>>& Hits, const FPTOverlapHit& BlockingHit, bool bHasBlockingHit)
-	{
-		if (bHasBlockingHit)
-		{
-			Hits.Add(BlockingHit);
-		}
-	}
-
-	template <typename HitType>
-	void FinishQueryHelper(TArray<HitType, TInlineAllocator<CHAOS_HIT_BUFFER_SIZE>>& Hits, const HitType& BlockingHit, bool bHasBlockingHit)
-	{
-		Hits.Sort();
-		if (bHasBlockingHit)
-		{
-			int32 FinalNum = Hits.Num() + 1;
-			for (int32 HitIdx = Hits.Num() - 1; HitIdx >= 0; --HitIdx)
-			{
-				if (Hits[HitIdx].Distance >= BlockingHit.Distance)
-				{
-					--FinalNum;
-				}
-				else
-				{
-					break;
-				}
-			}
-			Hits.SetNum(FinalNum);
-			Hits[FinalNum - 1] = BlockingHit;
-		}
-	}
-
-	/** Stores the results of scene queries. This can be passed around to multiple SQAccelerators and is responsible for sorting the results and pruning based on blocking.
-	IncFlushCount / DecFlushCount is used to ensure any final sort / pruning operation is done when all SQAccelerators are finished. If you are passing this into multiple accelerators
-	you should call IncFlushCount / DecFlushCount yourself as otherwise each accelerator will trigger its own sort / prune.
-	*/
-	template<typename HitType>
-	class FSQHitBuffer
-	{
-	public:
-		FSQHitBuffer(bool bSingle = false)
-			: AcceleratorDepth(0)
-			, bHasBlockingHit(false)
-			, bSingleResult(bSingle)
-		{
-			// The Hits array is preallocated through its InlineAllocator, no need to reserve.
-		}
-
-		virtual ~FSQHitBuffer() {}
-
-		//Called 
-		void IncFlushCount()
-		{
-			++AcceleratorDepth;
-		}
-
-		void DecFlushCount()
-		{
-			--AcceleratorDepth;
-			if (AcceleratorDepth == 0)
-			{
-				FinishQuery();
-			}
-		}
-
-		bool HasHit() const { return GetNumHits(); }
-		int32 GetNumHits() const { return Hits.Num(); }
-		HitType* GetHits() { return Hits.GetData(); }
-		const HitType* GetHits() const { return Hits.GetData(); }
-
-		HitType* GetBlock() { return HasBlockingHit() ? &Hits.Last() : nullptr; }
-		const HitType* GetBlock() const { return HasBlockingHit() ? &Hits.Last() : nullptr; }
-
-		bool HasBlockingHit() const { return bHasBlockingHit; }
-
-		void SetBlockingHit(const HitType& InBlockingHit)
-		{
-			CurrentBlockingHit = InBlockingHit;
-			bHasBlockingHit = true;
-		}
-
-		void AddTouchingHit(const HitType& InTouchingHit)
-		{
-			Hits.Add(InTouchingHit);
-		}
-
-		const HitType* GetCurrentHit() const { return bHasBlockingHit ? &CurrentBlockingHit : nullptr; }
-
-		/** Does not do any distance verification. This is up to the SQ code to manage */
-		void InsertHit(const HitType& Hit, bool bBlocking)
-		{
-			// Useful place to break when debugging, but breakpoints can't be set here, as
-			// this gets inlined even in a debug build.
-			//__debugbreak(); 
-			if (bBlocking)
-			{
-				SetBlockingHit(Hit);
-			}
-			else
-			{
-				AddTouchingHit(Hit);
-			}
-		}
-
-		bool WantsSingleResult() const { return bSingleResult; }
-
-	private:
-
-		void FinishQuery()
-		{
-			FinishQueryHelper(Hits, CurrentBlockingHit, bHasBlockingHit);
-		}
-
-		HitType CurrentBlockingHit;
-		int32 AcceleratorDepth;
-		bool bHasBlockingHit;
-		bool bSingleResult;
-
-		TArray<HitType, TInlineAllocator<CHAOS_HIT_BUFFER_SIZE>> Hits;
-	};
-
-	template<typename HitType>
-	class FSQSingleHitBuffer : public FSQHitBuffer<HitType>
-	{
-	public:
-		FSQSingleHitBuffer()
-			: FSQHitBuffer<HitType>(/*bSingle=*/ true)
-		{}
-	};
-
-}
+#include "Chaos/Interface/SQTypes.h"
\ No newline at end of file
diff --git a/Engine/Source/Runtime/PhysicsCore/Public/PhysicsInterfaceWrapperShared.h b/Engine/Source/Runtime/PhysicsCore/Public/PhysicsInterfaceWrapperShared.h
--- a/Engine/Source/Runtime/PhysicsCore/Public/PhysicsInterfaceWrapperShared.h
+++ b/Engine/Source/Runtime/PhysicsCore/Public/PhysicsInterfaceWrapperShared.h
@@ -2,118 +2,4 @@
 
 #pragma once
 
-#include "HAL/Platform.h"
-#include "Serialization/Archive.h"
-
-enum class EQueryFlags : uint16
-{
-	None = 0,
-	PreFilter = (1 << 2),
-	PostFilter = (1 << 3),
-	AnyHit = (1 << 4),
-	SkipNarrowPhase = (1 << 5)
-};
-
-inline EQueryFlags operator| (EQueryFlags lhs, EQueryFlags rhs)
-{
-	return static_cast<EQueryFlags>(static_cast<uint16>(lhs) | static_cast<uint16>(rhs));
-}
-
-inline EQueryFlags operator&(EQueryFlags lhs, EQueryFlags rhs)
-{
-	return static_cast<EQueryFlags>(static_cast<uint16>(lhs) & static_cast<uint16>(rhs));
-}
-
-struct FQueryFlags
-{
-	FQueryFlags(EQueryFlags InFlags) : QueryFlags(InFlags) {}
-	explicit operator bool() const { return !!static_cast<uint16>(QueryFlags); }
-	FQueryFlags operator |(EQueryFlags Rhs) const
-	{
-		return FQueryFlags(QueryFlags | Rhs);
-	}
-
-	FQueryFlags& operator |=(EQueryFlags Rhs)
-	{
-		*this = *this | Rhs;
-		return *this;
-	}
-
-	FQueryFlags operator &(EQueryFlags Rhs) const
-	{
-		return FQueryFlags(QueryFlags & Rhs);
-	}
-
-	FQueryFlags& operator &=(EQueryFlags Rhs)
-	{
-		*this = *this & Rhs;
-		return *this;
-	}
-
-	bool operator==(const FQueryFlags& Other) const { return QueryFlags == Other.QueryFlags; }
-
-	EQueryFlags QueryFlags;
-};
-
-inline FArchive& operator<<(FArchive& Ar, FQueryFlags& QueryFlags)
-{
-	return Ar << QueryFlags.QueryFlags;
-}
-
-/** Possible results from a scene query */
-enum class EHitFlags : uint16
-{
-	None = 0,
-	Position = (1 << 0),
-	Normal = (1 << 1),
-	Distance = (1 << 2),
-	UV = (1 << 3),
-	MTD = (1 << 9),
-	FaceIndex = (1 << 10)
-};
-
-inline EHitFlags operator|(EHitFlags lhs, EHitFlags rhs)
-{
-	return static_cast<EHitFlags>(static_cast<uint16>(lhs) | static_cast<uint16>(rhs));
-}
-
-inline EHitFlags operator&(EHitFlags lhs, EHitFlags rhs)
-{
-	return static_cast<EHitFlags>(static_cast<uint16>(lhs) & static_cast<uint16>(rhs));
-}
-
-struct FHitFlags
-{
-	FHitFlags(EHitFlags InFlags = EHitFlags::None) : HitFlags(InFlags) {}
-	explicit operator bool() const { return !!static_cast<uint16>(HitFlags); }
-	FHitFlags operator |(EHitFlags Rhs) const
-	{
-		return FHitFlags(HitFlags | Rhs);
-	}
-
-	FHitFlags& operator |=(EHitFlags Rhs)
-	{
-		*this = *this | Rhs;
-		return *this;
-	}
-
-	FHitFlags operator &(EHitFlags Rhs) const
-	{
-		return FHitFlags(HitFlags & Rhs);
-	}
-
-	FHitFlags& operator &=(EHitFlags Rhs)
-	{
-		*this = *this & Rhs;
-		return *this;
-	}
-
-	bool operator==(const FHitFlags& Other) const { return HitFlags == Other.HitFlags; }
-
-	EHitFlags HitFlags;
-};
-
-inline FArchive& operator<<(FArchive& Ar, FHitFlags& HitFlags)
-{
-	return Ar << HitFlags.HitFlags;
-}
+#include "Chaos/Interface/PhysicsInterfaceWrapperShared.h"
\ No newline at end of file
