support for implicit object unions in OverlapQuery and SweepQuery

diff --git a/Engine/Source/Runtime/Experimental/Chaos/Private/Chaos/ImplicitObjectUnion.cpp b/Engine/Source/Runtime/Experimental/Chaos/Private/Chaos/ImplicitObjectUnion.cpp
--- a/Engine/Source/Runtime/Experimental/Chaos/Private/Chaos/ImplicitObjectUnion.cpp
+++ b/Engine/Source/Runtime/Experimental/Chaos/Private/Chaos/ImplicitObjectUnion.cpp
@@ -234,6 +234,37 @@ TUniquePtr<FImplicitObject> FImplicitObjectUnion::DeepCopyWithScale(const FVec3&
 	}
 	return MakeUnique<FImplicitObjectUnion>(MoveTemp(CopyOfObjects));
 }
+
+void FImplicitObjectUnion::ForEachObject(TFunctionRef<bool(const FImplicitObject&, const FRigidTransform3&)> Lambda) const
+{
+	if (LargeUnionData)
+	{
+		for (int32 Index = 0; Index < static_cast<int32>(LargeUnionData->GeomParticles.Size()); ++Index)
+		{
+			if (TSerializablePtr<FImplicitObject> SubObject = LargeUnionData->GeomParticles.Geometry(Index))
+			{
+				const FRigidTransform3 SubTransform{ LargeUnionData->GeomParticles.X(Index), LargeUnionData->GeomParticles.R(Index) };
+				if (Lambda(*SubObject, SubTransform))
+				{
+					break;
+				}
+			}
+		}
+	}
+	else
+	{
+		for (const TUniquePtr<FImplicitObject>& Object : MObjects)
+		{
+			if (Object)
+			{
+				if (Lambda(*Object, FRigidTransform3::Identity))
+				{
+					break;
+				}
+			}
+		}
+	}
+}
 	
 void FImplicitObjectUnion::Serialize(FChaosArchive& Ar)
 {
diff --git a/Engine/Source/Runtime/Experimental/Chaos/Public/Chaos/GeometryQueries.h b/Engine/Source/Runtime/Experimental/Chaos/Public/Chaos/GeometryQueries.h
--- a/Engine/Source/Runtime/Experimental/Chaos/Public/Chaos/GeometryQueries.h
+++ b/Engine/Source/Runtime/Experimental/Chaos/Public/Chaos/GeometryQueries.h
@@ -9,6 +9,7 @@
 #include "Chaos/HeightField.h"
 #include "Chaos/ImplicitObject.h"
 #include "Chaos/ImplicitObjectScaled.h"
+#include "Chaos/ImplicitObjectUnion.h"
 #include "Chaos/Levelset.h"
 #include "Chaos/Plane.h"
 #include "Chaos/Sphere.h"
@@ -96,6 +97,27 @@ namespace Chaos
 					const FLevelSet& ALevelSet = static_cast<const FLevelSet&>(A);
 					return ALevelSet.OverlapGeom(B, BToATM, Thickness, OutMTD);
 				}
+				case ImplicitObjectType::Union:
+				case ImplicitObjectType::UnionClustered:
+				{
+					const FImplicitObjectUnion& AUnion = static_cast<const FImplicitObjectUnion&>(A);
+					bool bHit = false;
+					AUnion.ForEachObject(
+						[&bHit, &ATM, &B, &BTM, Thickness, &OutMTD](const FImplicitObject& SubObject, const FRigidTransform3& SubTransform)
+						{
+							const FRigidTransform3 NewATM = SubTransform * ATM;
+							if (OverlapQuery(SubObject, NewATM, B, BTM, Thickness, OutMTD))
+							{
+								bHit = true;
+								return true;
+							}
+
+							return false;
+						}
+					);
+
+					return bHit;
+				}
 				default:
 				{
 					if(IsScaled(AType))
@@ -220,6 +242,42 @@ namespace Chaos
 					bResult = ALevelSet.SweepGeom(B, BToATM, LocalDir, Length, OutTime, LocalPosition, LocalNormal, OutFaceIndex, Thickness, bComputeMTD);
 					break;
 				}
+				case ImplicitObjectType::Union:
+				case ImplicitObjectType::UnionClustered:
+				{
+					const FImplicitObjectUnion& AUnion = static_cast<const FImplicitObjectUnion&>(A);
+
+					bool bHit = false;
+					AUnion.ForEachObject(
+						[&bHit, &ATM, &B, &BTM, &Dir, Length, &OutTime, &OutPosition, &OutNormal, &OutFaceIndex, &OutFaceNormal, Thickness, bComputeMTD](const FImplicitObject& SubObject, const FRigidTransform3& SubTransform)
+						{
+							const FRigidTransform3 NewATM = SubTransform * ATM;
+
+							FReal ObjectTime = 0.0;
+							FVec3 ObjectPosition;
+							FVec3 ObjectNormal;
+							int32 ObjectFaceIndex;
+							FVec3 ObjectFaceNormal;
+
+							if (SweepQuery(SubObject, NewATM, B, BTM, Dir, Length, ObjectTime, ObjectPosition, ObjectNormal, ObjectFaceIndex, ObjectFaceNormal, Thickness, bComputeMTD))
+							{
+								bHit = true;
+								if (ObjectTime < OutTime)
+								{
+									OutTime = ObjectTime;
+									OutPosition = ObjectPosition;
+									OutNormal = ObjectNormal;
+									OutFaceIndex = ObjectFaceIndex;
+									OutFaceNormal = ObjectFaceNormal;
+								}
+							}
+
+							return false;
+						}
+					);
+
+					return bHit;
+				}
 				default:
 				if (IsScaled(AType))
 				{
diff --git a/Engine/Source/Runtime/Experimental/Chaos/Public/Chaos/ImplicitObjectUnion.h b/Engine/Source/Runtime/Experimental/Chaos/Public/Chaos/ImplicitObjectUnion.h
--- a/Engine/Source/Runtime/Experimental/Chaos/Public/Chaos/ImplicitObjectUnion.h
+++ b/Engine/Source/Runtime/Experimental/Chaos/Public/Chaos/ImplicitObjectUnion.h
@@ -149,6 +149,9 @@ class CHAOS_API FImplicitObjectUnion : public FImplicitObject
 
 	TArray<TUniquePtr<FImplicitObject>>& GetObjects() { return MObjects; }
 
+	// The lambda returns TRUE if an object was found and iteration should stop.
+	void ForEachObject(TFunctionRef<bool(const FImplicitObject&, const FRigidTransform3&)> Lambda) const;
+
 	virtual uint32 GetTypeHash() const override
 	{
 		uint32 Result = 0;
