add functions to the physics object interface to make line trace, overlap, and sweep on the UPrimitiveComponent generalizable

- also modified the K2 overlap functions to actually return a hit result as well

diff --git a/Engine/Source/Runtime/Engine/Classes/Components/PrimitiveComponent.h b/Engine/Source/Runtime/Engine/Classes/Components/PrimitiveComponent.h
--- a/Engine/Source/Runtime/Engine/Classes/Components/PrimitiveComponent.h
+++ b/Engine/Source/Runtime/Engine/Classes/Components/PrimitiveComponent.h
@@ -49,6 +49,7 @@ struct FConvexVolume;
 struct FEngineShowFlags;
 struct FNavigableGeometryExport;
 struct FPSOPrecacheParams;
+struct FOverlapResult;
 
 namespace PrimitiveComponentCVars
 {
@@ -2767,7 +2768,7 @@ public:
 	 *  @param  PrimComp        Component to use geometry from to test against this component. Transform of this component is ignored.
 	 *  @param  Pos             Location to place PrimComp geometry at 
 	 *  @param  Rot             Rotation to place PrimComp geometry at 
-	 *  @param	Params			Parameter for trace. TraceTag is only used.
+	 *  @param  Params          Parameter for trace. TraceTag is only used.
 	 *  @return true if PrimComp overlaps this component at the specified location/rotation
 	 */
 	bool ComponentOverlapComponent(class UPrimitiveComponent* PrimComp, const FVector Pos, const FQuat& Rot, const FCollisionQueryParams& Params);
@@ -2782,8 +2783,8 @@ public:
 	 *  Test the collision of the supplied shape at the supplied location, and determine if it overlaps this component.
 	 *
 	 *  @param  Pos             Location to place PrimComp geometry at 
-	 *	@param	Rot				Rotation of PrimComp geometry
-	 *  @param  CollisionShape 	Shape of collision of PrimComp geometry
+	 *	@param  Rot             Rotation of PrimComp geometry
+	 *  @param  CollisionShape  Shape of collision of PrimComp geometry
 	 *  @return true if PrimComp overlaps this component at the specified location/rotation
 	 */
 
@@ -2794,6 +2795,18 @@ public:
 	}
 	virtual bool OverlapComponent(const FVector& Pos, const FQuat& Rot, const FCollisionShape& CollisionShape) const;
 
+	/**
+	 * Test the collision of the supplied shape at the supplied location, and determine if it overlaps this component.
+	 * Also will return information about the overlap.
+	 *  @param  Pos             Location to place PrimComp geometry at
+	 *	@param  Rot             Rotation of PrimComp geometry
+	 *  @param  CollisionShape  Shape of collision of PrimComp geometry
+	 *  @param  OutOverlap      Additional information about what exactly was overlapped.
+	 *  @return true if PrimComp overlaps this component at the specified location/rotation
+	 */
+	virtual bool OverlapComponentWithResult(const FVector& Pos, const FQuat& Rot, const FCollisionShape& CollisionShape, TArray<FOverlapResult>& OutOverlap) const;
+
+
 	/**
 	 * Computes the minimum translation direction (MTD) when an overlap exists between the component and the given shape.
 	 * @param OutMTD			Outputs the MTD to move CollisionShape out of penetration
diff --git a/Engine/Source/Runtime/Engine/Private/Collision/CollisionConversions.cpp b/Engine/Source/Runtime/Engine/Private/Collision/CollisionConversions.cpp
--- a/Engine/Source/Runtime/Engine/Private/Collision/CollisionConversions.cpp
+++ b/Engine/Source/Runtime/Engine/Private/Collision/CollisionConversions.cpp
@@ -771,3 +771,13 @@ bool ConvertOverlapResults(int32 NumOverlaps, ChaosInterface::FPTOverlapHit* Ove
 {
 	return ConvertOverlapResultsImp(NumOverlaps, OverlapResults, QueryFilter, OutOverlaps);
 }
+
+FHitResult ConvertOverlapToHitResult(const FOverlapResult& Overlap)
+{
+	FHitResult Hit;
+	Hit.bBlockingHit = Overlap.bBlockingHit;
+	Hit.Item = Overlap.ItemIndex;
+	Hit.Component = Overlap.Component;
+	Hit.HitObjectHandle = Overlap.OverlapObjectHandle;
+	return Hit;
+}
\ No newline at end of file
diff --git a/Engine/Source/Runtime/Engine/Private/Collision/CollisionConversions.h b/Engine/Source/Runtime/Engine/Private/Collision/CollisionConversions.h
--- a/Engine/Source/Runtime/Engine/Private/Collision/CollisionConversions.h
+++ b/Engine/Source/Runtime/Engine/Private/Collision/CollisionConversions.h
@@ -85,6 +85,10 @@ bool IsBlocking(const FPhysicsShape& PShape, const FCollisionFilterData& QueryFi
 bool ConvertOverlapResults(int32 NumOverlaps, FHitOverlap* POverlapResults, const FCollisionFilterData& QueryFilter, TArray<FOverlapResult>& OutOverlaps);
 bool ConvertOverlapResults(int32 NumOverlaps, ChaosInterface::FPTOverlapHit* POverlapResults, const FCollisionFilterData& QueryFilter, TArray<FOverlapResult>& OutOverlaps);
 
+/**
+ * Converts an overlap result to a hit result.
+ */
+FHitResult ConvertOverlapToHitResult(const FOverlapResult& Overlap);
 
 struct FCompareFHitResultTime
 {
diff --git a/Engine/Source/Runtime/Engine/Private/Components/PrimitiveComponent.cpp b/Engine/Source/Runtime/Engine/Private/Components/PrimitiveComponent.cpp
--- a/Engine/Source/Runtime/Engine/Private/Components/PrimitiveComponent.cpp
+++ b/Engine/Source/Runtime/Engine/Private/Components/PrimitiveComponent.cpp
@@ -5,6 +5,10 @@
 =============================================================================*/
 
 #include "Components/PrimitiveComponent.h"
+
+#include "Chaos/ChaosEngineInterface.h"
+#include "ChaosInterfaceWrapperCore.h"
+#include "Collision/CollisionConversions.h"
 #include "Engine/Level.h"
 #include "Engine/Texture.h"
 #include "GameFramework/DamageType.h"
@@ -30,6 +34,7 @@
 #include "Misc/MapErrors.h"
 #include "CollisionDebugDrawingPublic.h"
 #include "GameFramework/CheatManager.h"
+#include "PhysicsEngine/PhysicsObjectExternalInterface.h"
 #include "RenderingThread.h"
 #include "Streaming/TextureStreamingHelpers.h"
 #include "PrimitiveSceneProxy.h"
@@ -2914,7 +2919,35 @@ extern float DebugLineLifetime;
 
 bool UPrimitiveComponent::LineTraceComponent(struct FHitResult& OutHit, const FVector Start, const FVector End, const struct FCollisionQueryParams& Params)
 {
-	bool bHaveHit = BodyInstance.LineTrace(OutHit, Start, End, Params.bTraceComplex, Params.bReturnPhysicalMaterial); 
+	bool bHaveHit = false;
+	
+	if (FBodyInstance* ThisBodyInstance = GetBodyInstance())
+	{
+		bHaveHit = ThisBodyInstance->LineTrace(OutHit, Start, End, Params.bTraceComplex, Params.bReturnPhysicalMaterial);
+	}
+	else
+	{
+		TArray<Chaos::FPhysicsObjectHandle> Objects = GetAllPhysicsObjects();
+		FLockedReadPhysicsObjectExternalInterface Interface = FPhysicsObjectExternalInterface::LockRead(Objects);
+		Objects = Objects.FilterByPredicate(
+			[&Interface](Chaos::FPhysicsObjectHandle Handle)
+			{
+				return !Interface->AreAllDisabled({ &Handle, 1 });
+			}
+		);
+
+		ChaosInterface::FRaycastHit BestHit;
+		if (Interface->LineTrace(Objects, Start, End, Params.bTraceComplex, BestHit))
+		{
+			bHaveHit = true;
+
+			FCollisionFilterData QueryFilter;
+			QueryFilter.Word1 = 0xFFFFF;
+			ChaosInterface::SetFlags(BestHit, EHitFlags::Distance | EHitFlags::Normal | EHitFlags::Position);
+
+			ConvertQueryImpactHit(GetWorld(), BestHit, OutHit, (End - Start).Size(), QueryFilter, Start, End, nullptr, FTransform{ Start }, true, Params.bReturnPhysicalMaterial);
+		}
+	}
 
 #if !(UE_BUILD_SHIPPING || UE_BUILD_TEST)
 	if (GetWorld()->DebugDrawSceneQueries(Params.TraceTag))
@@ -2933,7 +2966,40 @@ bool UPrimitiveComponent::LineTraceComponent(struct FHitResult& OutHit, const FV
 
 bool UPrimitiveComponent::SweepComponent(struct FHitResult& OutHit, const FVector Start, const FVector End, const FQuat& ShapeWorldRotation, const FCollisionShape &CollisionShape, bool bTraceComplex)
 {
-	return BodyInstance.Sweep(OutHit, Start, End, ShapeWorldRotation, CollisionShape, bTraceComplex);
+	if (FBodyInstance* ThisBodyInstance = GetBodyInstance())
+	{
+		return ThisBodyInstance->Sweep(OutHit, Start, End, ShapeWorldRotation, CollisionShape, bTraceComplex);
+	}
+
+	if (CollisionShape.IsNearlyZero())
+	{
+		FCollisionQueryParams Params;
+		Params.bTraceComplex = bTraceComplex;
+		return LineTraceComponent(OutHit, Start, End, Params);
+	}
+
+	FPhysicsShapeAdapter_Chaos ShapeAdapter(ShapeWorldRotation, CollisionShape);
+	TArray<Chaos::FPhysicsObjectHandle> Objects = GetAllPhysicsObjects();
+	FLockedReadPhysicsObjectExternalInterface Interface = FPhysicsObjectExternalInterface::LockRead(Objects);
+	Objects = Objects.FilterByPredicate(
+		[&Interface](Chaos::FPhysicsObjectHandle Handle)
+		{
+			return !Interface->AreAllDisabled({ &Handle, 1 });
+		}
+	);
+
+	ChaosInterface::FSweepHit BestHit;
+	if (Interface->ShapeSweep(Objects, ShapeAdapter.GetGeometry(), ShapeAdapter.GetGeomPose(Start), End, bTraceComplex, BestHit))
+	{
+		FCollisionFilterData QueryFilter;
+		QueryFilter.Word1 = 0xFFFFF;
+		ChaosInterface::SetFlags(BestHit, EHitFlags::Distance | EHitFlags::Normal | EHitFlags::Position | EHitFlags::FaceIndex);
+
+		ConvertQueryImpactHit(GetWorld(), BestHit, OutHit, (End - Start).Size(), QueryFilter, Start, End, nullptr, FTransform{ Start }, false, false);
+		return true;
+	}
+
+	return false;
 }
 
 bool UPrimitiveComponent::ComponentOverlapComponentImpl(class UPrimitiveComponent* PrimComp, const FVector Pos, const FQuat& Quat, const struct FCollisionQueryParams& Params)
@@ -2960,21 +3026,85 @@ bool UPrimitiveComponent::ComponentOverlapComponentImpl(class UPrimitiveComponen
 		}
 	}
 
-	if(FBodyInstance* BI = PrimComp->GetBodyInstance())
+	FBodyInstance* BI = PrimComp->GetBodyInstance();
+	FBodyInstance* ThisBodyInstance = GetBodyInstance();
+	if(BI && ThisBodyInstance)
 	{
-		if (FBodyInstance* ThisBodyInstance = GetBodyInstance())
+		return BI->OverlapTestForBody(Pos, Quat, ThisBodyInstance);
+	}
+
+	TArray<Chaos::FPhysicsObjectHandle> InObjects = PrimComp->GetAllPhysicsObjects();
+	TArray<Chaos::FPhysicsObjectHandle> ThisObjects = GetAllPhysicsObjects();
+
+	FLockedReadPhysicsObjectExternalInterface Interface = FPhysicsObjectExternalInterface::LockRead(InObjects);
+	InObjects = InObjects.FilterByPredicate(
+		[&Interface](Chaos::FPhysicsObjectHandle Handle)
 		{
-			return BI->OverlapTestForBody(Pos, Quat, ThisBodyInstance);
+			return !Interface->AreAllDisabled({ &Handle, 1 });
 		}
-	}
+	);
+
+	ThisObjects = ThisObjects.FilterByPredicate(
+		[&Interface](Chaos::FPhysicsObjectHandle Handle)
+		{
+			return !Interface->AreAllDisabled({ &Handle, 1 });
+		}
+	);
 
+	for (Chaos::FPhysicsObjectHandle InObject : InObjects)
+	{
+		for (Chaos::FPhysicsObjectHandle ThisObject : ThisObjects)
+		{
+			if (Interface->GetPhysicsObjectOverlap(InObject, FTransform::Identity, ThisObject, FTransform::Identity, Params.bTraceComplex))
+			{
+				return true;
+			}
+		}
+	}
 	return false;
 }
 
 
 bool UPrimitiveComponent::OverlapComponent(const FVector& Pos, const FQuat& Rot, const struct FCollisionShape& CollisionShape) const
 {
-	return BodyInstance.OverlapTest(Pos, Rot, CollisionShape);
+	if (FBodyInstance* ThisBodyInstance = GetBodyInstance())
+	{
+		return ThisBodyInstance->OverlapTest(Pos, Rot, CollisionShape);
+	}
+
+	TArray<FOverlapResult> NopResult;
+	return OverlapComponentWithResult(Pos, Rot, CollisionShape, NopResult);
+}
+
+bool UPrimitiveComponent::OverlapComponentWithResult(const FVector& Pos, const FQuat& Rot, const FCollisionShape& CollisionShape, TArray<FOverlapResult>& OutOverlap) const
+{
+	FPhysicsShapeAdapter_Chaos ShapeAdapter(Rot, CollisionShape);
+	TArray<Chaos::FPhysicsObjectHandle> Objects = GetAllPhysicsObjects();
+	FLockedReadPhysicsObjectExternalInterface Interface = FPhysicsObjectExternalInterface::LockRead(Objects);
+	Objects = Objects.FilterByPredicate(
+		[&Interface](Chaos::FPhysicsObjectHandle Handle)
+		{
+			return !Interface->AreAllDisabled({ &Handle, 1 });
+		}
+	);
+
+	TArray<ChaosInterface::FOverlapHit> OverlapHits;
+	if (Interface->ShapeOverlap(Objects, ShapeAdapter.GetGeometry(), ShapeAdapter.GetGeomPose(Pos), OverlapHits))
+	{
+		TArray<FOverlapResult> Overlaps;
+
+		FCollisionFilterData QueryFilter;
+		QueryFilter.Word2 = 0xFFFFF;
+		ConvertOverlapResults(OverlapHits.Num(), OverlapHits.GetData(), QueryFilter, Overlaps);
+
+		if (!Overlaps.IsEmpty())
+		{
+			OutOverlap = Overlaps;
+			return true;
+		}
+	}
+
+	return false;
 }
 
 bool UPrimitiveComponent::ComputePenetration(FMTDResult& OutMTD, const FCollisionShape & CollisionShape, const FVector& Pos, const FQuat& Rot)
diff --git a/Engine/Source/Runtime/Engine/Private/PrimitiveComponentPhysics.cpp b/Engine/Source/Runtime/Engine/Private/PrimitiveComponentPhysics.cpp
--- a/Engine/Source/Runtime/Engine/Private/PrimitiveComponentPhysics.cpp
+++ b/Engine/Source/Runtime/Engine/Private/PrimitiveComponentPhysics.cpp
@@ -2,6 +2,7 @@
 
 #include "Components/PrimitiveComponent.h"
 #include "AI/NavigationSystemBase.h"
+#include "Collision/CollisionConversions.h"
 #include "EngineLogs.h"
 #include "Logging/MessageLog.h"
 #include "Physics/Experimental/PhysScene_Chaos.h"
@@ -1128,7 +1129,12 @@ bool UPrimitiveComponent::K2_BoxOverlapComponent(FVector InBoxCentre, const FBox
 {
 	FCollisionShape QueryBox = FCollisionShape::MakeBox(InBox.GetExtent());
 
-	bool bHit = OverlapComponent(InBoxCentre, FQuat::Identity, QueryBox);
+	TArray<FOverlapResult> OverlapResult;
+	bool bHit = OverlapComponentWithResult(InBoxCentre, FQuat::Identity, QueryBox, OverlapResult);
+	if (bHit && !OverlapResult.IsEmpty())
+	{
+		OutHit = ConvertOverlapToHitResult(OverlapResult[0]);
+	}
 
 	if(bShowTrace)
 	{
@@ -1144,7 +1150,12 @@ bool UPrimitiveComponent::K2_SphereOverlapComponent(FVector InSphereCentre, floa
 {
 	FCollisionShape QuerySphere = FCollisionShape::MakeSphere(InSphereRadius);
 
-	bool bHit = OverlapComponent(InSphereCentre, FQuat::Identity, QuerySphere);
+	TArray<FOverlapResult> OverlapResult;
+	bool bHit = OverlapComponentWithResult(InSphereCentre, FQuat::Identity, QuerySphere, OverlapResult);
+	if (bHit && !OverlapResult.IsEmpty())
+	{
+		OutHit = ConvertOverlapToHitResult(OverlapResult[0]);
+	}
 
 	if(bShowTrace)
 	{
diff --git a/Engine/Source/Runtime/Experimental/Chaos/Private/Chaos/PhysicsObjectInterface.cpp b/Engine/Source/Runtime/Experimental/Chaos/Private/Chaos/PhysicsObjectInterface.cpp
--- a/Engine/Source/Runtime/Experimental/Chaos/Private/Chaos/PhysicsObjectInterface.cpp
+++ b/Engine/Source/Runtime/Experimental/Chaos/Private/Chaos/PhysicsObjectInterface.cpp
@@ -2,6 +2,7 @@
 #include "Chaos/PhysicsObjectInterface.h"
 
 #include "Chaos/CastingUtilities.h"
+#include "Chaos/CollisionFilterData.h"
 #include "Chaos/GeometryQueries.h"
 #include "Chaos/PhysicsObjectInternal.h"
 #include "PhysicsProxy/SingleParticlePhysicsProxy.h"
@@ -240,6 +241,30 @@ namespace Chaos
 		return AllShapes;
 	}
 
+	template<EThreadContext Id>
+	void FReadPhysicsObjectInterface<Id>::VisitEveryShape(TArrayView<FPhysicsObjectHandle> InObjects, TFunctionRef<bool(FPhysicsObjectHandle, FPerShapeData*)> Lambda)
+	{
+		for (FPhysicsObjectHandle Object : InObjects)
+		{
+			if (!Object)
+			{
+				continue;
+			}
+
+			if (TThreadParticle<Id>* Particle = Object->GetParticle<Id>())
+			{
+				const Chaos::FShapesArray& ShapesArray = Particle->ShapesArray();
+				for (const TUniquePtr<Chaos::FPerShapeData>& Shape : ShapesArray)
+				{
+					if (Lambda(Object, Shape.Get()))
+					{
+						return;
+					}
+				}
+			}
+		}
+	}
+
 	template<EThreadContext Id>
 	bool FReadPhysicsObjectInterface<Id>::GetPhysicsObjectOverlap(FPhysicsObjectHandle ObjectA, FPhysicsObjectHandle ObjectB, bool bTraceComplex, Chaos::FOverlapInfo& OutOverlap)
 	{
@@ -671,6 +696,199 @@ namespace Chaos
 		return FAccelerationStructureHandle{InObject->GetParticle<Id>()};
 	}
 
+	template<EThreadContext Id>
+	bool FReadPhysicsObjectInterface<Id>::LineTrace(TArrayView<FPhysicsObjectHandle> InObjects, const FVector& WorldStart, const FVector& WorldEnd, bool bTraceComplex, ChaosInterface::FRaycastHit& OutBestHit)
+	{
+		bool bHit = false;
+		OutBestHit.Distance = TNumericLimits<float>::Max();
+
+		const FVector Delta = WorldEnd - WorldStart;
+		const FReal DeltaMag = Delta.Size();
+		if (DeltaMag < UE_KINDA_SMALL_NUMBER)
+		{
+			return false;
+		}
+
+		FTransform BestWorldTM = FTransform::Identity;
+
+		for (FPhysicsObjectHandle Object : InObjects)
+		{
+			const FTransform WorldTM = GetTransform(Object);
+			const FVector LocalStart = WorldTM.InverseTransformPositionNoScale(WorldStart);
+			const FVector LocalDelta = WorldTM.InverseTransformVectorNoScale(Delta);
+
+			VisitEveryShape(
+				{ &Object, 1 },
+				[this, &bHit, &WorldTM, &LocalStart, &LocalDelta, &Delta, DeltaMag, &BestWorldTM, bTraceComplex, &OutBestHit](FPhysicsObjectHandle IterObject, FPerShapeData* Shape)
+				{
+					check(Shape);
+
+					FCollisionFilterData ShapeFilter = Shape->GetQueryData();
+					const bool bShapeIsComplex = (ShapeFilter.Word3 & static_cast<uint8>(EFilterFlags::ComplexCollision)) != 0;
+					const bool bShapeIsSimple = (ShapeFilter.Word3 & static_cast<uint8>(EFilterFlags::SimpleCollision)) != 0;
+					if ((bTraceComplex && bShapeIsComplex) || (!bTraceComplex && bShapeIsSimple))
+					{
+						FReal Distance;
+						FVec3 LocalPosition;
+						FVec3 LocalNormal;
+						int32 FaceIndex;
+
+						const bool bRaycastHit = Shape->GetGeometry()->Raycast(
+							LocalStart,
+							LocalDelta / DeltaMag,
+							DeltaMag,
+							0,
+							Distance,
+							LocalPosition,
+							LocalNormal,
+							FaceIndex
+						);
+
+						if (bRaycastHit)
+						{
+							if (Distance < OutBestHit.Distance)
+							{
+								bHit = true;
+								BestWorldTM = WorldTM;
+								OutBestHit.Distance = static_cast<float>(Distance);
+								OutBestHit.WorldNormal = LocalNormal;
+								OutBestHit.WorldPosition = LocalPosition;
+								OutBestHit.Shape = Shape;
+								if constexpr (Id == EThreadContext::External)
+								{
+									OutBestHit.Actor = IterObject->GetParticle<Id>();
+								}
+								OutBestHit.FaceIndex = FaceIndex;
+							}
+						}
+					}
+					return false;
+				}
+			);
+		}
+
+		if (bHit)
+		{
+			OutBestHit.WorldNormal = BestWorldTM.TransformVectorNoScale(OutBestHit.WorldNormal);
+			OutBestHit.WorldPosition = BestWorldTM.TransformPositionNoScale(OutBestHit.WorldPosition);
+		}
+
+		return bHit;
+	}
+
+	template<EThreadContext Id>
+	bool FReadPhysicsObjectInterface<Id>::ShapeOverlap(TArrayView<FPhysicsObjectHandle> InObjects, const Chaos::FImplicitObject& InGeom, const FTransform& GeomTransform, TArray<ChaosInterface::FOverlapHit>& OutOverlaps)
+	{
+		bool bHasOverlap = false;
+		for (FPhysicsObjectHandle Object : InObjects)
+		{
+			const FTransform WorldTM = GetTransform(Object);
+
+			VisitEveryShape(
+				{ &Object, 1 },
+				[this, &bHasOverlap, &WorldTM, &InGeom, &GeomTransform, &OutOverlaps](FPhysicsObjectHandle IterObject, FPerShapeData* Shape)
+				{
+					check(Shape);
+					const bool bOverlap = Chaos::Utilities::CastHelper(
+						InGeom,
+						GeomTransform,
+						[Shape, &WorldTM](const auto& Downcast, const auto& FullTransformB)
+						{
+							return Chaos::OverlapQuery(*Shape->GetGeometry(), WorldTM, Downcast, FullTransformB, 0, nullptr);
+						}
+					);
+
+					if (bOverlap)
+					{
+						bHasOverlap = true;
+
+						ChaosInterface::FOverlapHit Overlap;
+						Overlap.Shape = Shape;
+						if constexpr (Id == EThreadContext::External)
+						{
+							Overlap.Actor = IterObject->GetParticle<Id>();
+						}
+						OutOverlaps.Add(Overlap);
+					}
+					return false;
+				}
+			);
+		}
+		return bHasOverlap;
+	}
+
+	template<EThreadContext Id>
+	bool FReadPhysicsObjectInterface<Id>::ShapeSweep(TArrayView<FPhysicsObjectHandle> InObjects, const Chaos::FImplicitObject& InGeom, const FTransform& StartTM, const FVector& EndPos, bool bSweepComplex, ChaosInterface::FSweepHit& OutBestHit)
+	{
+		bool bHit = false;
+		const FVector StartPos = StartTM.GetTranslation();
+		const FVector Delta = EndPos - StartPos;
+		const FReal DeltaMag = Delta.Size();
+		if (DeltaMag < UE_KINDA_SMALL_NUMBER)
+		{
+			return false;
+		}
+		const FVec3 Dir = Delta / DeltaMag;
+
+		for (FPhysicsObjectHandle Object : InObjects)
+		{
+			const FTransform WorldTM = GetTransform(Object);
+
+			VisitEveryShape(
+				{ &Object, 1 },
+				[&WorldTM, &InGeom, &StartTM, &bHit, &Delta, DeltaMag, &Dir, &OutBestHit, bSweepComplex](FPhysicsObjectHandle IterObject, FPerShapeData* Shape)
+				{
+					check(Shape);
+
+					FCollisionFilterData ShapeFilter = Shape->GetQueryData();
+					const bool bShapeIsComplex = (ShapeFilter.Word3 & static_cast<uint8>(EFilterFlags::ComplexCollision)) != 0;
+					const bool bShapeIsSimple = (ShapeFilter.Word3 & static_cast<uint8>(EFilterFlags::SimpleCollision)) != 0;
+					if ((bSweepComplex && bShapeIsComplex) || (!bSweepComplex && bShapeIsSimple))
+					{
+						FVec3 WorldPosition;
+						FVec3 WorldNormal;
+						FReal Distance;
+						int32 FaceIdx;
+						FVec3 FaceNormal;
+
+						const bool bShapeHit = Chaos::Utilities::CastHelper(
+							InGeom,
+							StartTM,
+							[Shape, &WorldTM, &Dir, DeltaMag, &Distance, &WorldPosition, &WorldNormal, &FaceIdx, &FaceNormal](const auto& Downcast, const auto& FullTransformB)
+							{
+								return Chaos::SweepQuery(*Shape->GetGeometry(), WorldTM, Downcast, FullTransformB, Dir, DeltaMag, Distance, WorldPosition, WorldNormal, FaceIdx, FaceNormal, 0.f, false);
+							}
+						);
+
+						if (bShapeHit)
+						{
+							bHit = true;
+
+							OutBestHit.Shape = Shape;
+							OutBestHit.WorldPosition = WorldPosition;
+							OutBestHit.WorldNormal = WorldNormal;
+							OutBestHit.Distance = static_cast<float>(Distance);
+							OutBestHit.FaceIndex = FaceIdx;
+							if (OutBestHit.Distance > 0.f)
+							{
+								const FVector LocalPosition = WorldTM.InverseTransformPositionNoScale(OutBestHit.WorldPosition);
+								const FVector LocalUnitDir = WorldTM.InverseTransformVectorNoScale(Dir);
+								OutBestHit.FaceIndex = Shape->GetGeometry()->FindMostOpposingFace(LocalPosition, LocalUnitDir, OutBestHit.FaceIndex, 1);
+							}
+							if constexpr (Id == EThreadContext::External)
+							{
+								OutBestHit.Actor = IterObject->GetParticle<Id>();
+							}
+						}
+					}
+					return false;
+				}
+			);
+
+		}
+		return bHit;
+	}
+
 	template<EThreadContext Id>
 	void FWritePhysicsObjectInterface<Id>::PutToSleep(TArrayView<FPhysicsObjectHandle> InObjects)
 	{
diff --git a/Engine/Source/Runtime/Experimental/Chaos/Public/Chaos/PhysicsObjectInterface.h b/Engine/Source/Runtime/Experimental/Chaos/Public/Chaos/PhysicsObjectInterface.h
--- a/Engine/Source/Runtime/Experimental/Chaos/Public/Chaos/PhysicsObjectInterface.h
+++ b/Engine/Source/Runtime/Experimental/Chaos/Public/Chaos/PhysicsObjectInterface.h
@@ -2,6 +2,7 @@
 
 #pragma once
 
+#include "Chaos/Interface/SQTypes.h"
 #include "Chaos/ParticleHandle.h"
 #include "Chaos/PhysicsObject.h"
 #include "Chaos/ISpatialAcceleration.h"
@@ -68,6 +69,9 @@ namespace Chaos
 		TArray<TThreadRigidParticle<Id>*> GetAllRigidParticles(TArrayView<FPhysicsObjectHandle> InObjects);
 		TArray<FPerShapeData*> GetAllShapes(TArrayView<FPhysicsObjectHandle> InObjects);
 
+		// Returns true if a shape is found and we can stop iterating.
+		void VisitEveryShape(TArrayView<FPhysicsObjectHandle> InObjects, TFunctionRef<bool(FPhysicsObjectHandle, FPerShapeData*)> Lambda);
+
 		UE_DEPRECATED(5.3, "GetPhysicsObjectOverlap has been deprecated. Please use the function for the specific overlap metric you wish to compute instead.")
 		bool GetPhysicsObjectOverlap(FPhysicsObjectHandle ObjectA, FPhysicsObjectHandle ObjectB, bool bTraceComplex, Chaos::FOverlapInfo& OutOverlap);
 
@@ -96,6 +100,11 @@ namespace Chaos
 		FBox GetWorldBounds(TArrayView<FPhysicsObjectHandle> InObjects);
 		FClosestPhysicsObjectResult GetClosestPhysicsBodyFromLocation(TArrayView<FPhysicsObjectHandle> InObjects, const FVector& WorldLocation);
 		FAccelerationStructureHandle CreateAccelerationStructureHandle(FPhysicsObjectHandle Handle);
+
+		bool LineTrace(TArrayView<FPhysicsObjectHandle> InObjects, const FVector& WorldStart, const FVector& WorldEnd, bool bTraceComplex, ChaosInterface::FRaycastHit& OutBestHit);
+		bool ShapeOverlap(TArrayView<FPhysicsObjectHandle> InObjects, const Chaos::FImplicitObject& InGeom, const FTransform& GeomTransform, TArray<ChaosInterface::FOverlapHit>& OutOverlaps);
+		bool ShapeSweep(TArrayView<FPhysicsObjectHandle> InObjects, const Chaos::FImplicitObject& InGeom, const FTransform& StartTM, const FVector& EndPos, bool bSweepComplex, ChaosInterface::FSweepHit& OutBestHit);
+
 		friend class FPhysicsObjectInterface;
 	protected:
 		FReadPhysicsObjectInterface() = default;
