rename the GetPhysicsObjectOverlap functions to remove the "Get" to be more consistent in naming with the more recent sweep/overlap functions

diff --git a/Engine/Source/Runtime/Engine/Private/Components/PrimitiveComponent.cpp b/Engine/Source/Runtime/Engine/Private/Components/PrimitiveComponent.cpp
--- a/Engine/Source/Runtime/Engine/Private/Components/PrimitiveComponent.cpp
+++ b/Engine/Source/Runtime/Engine/Private/Components/PrimitiveComponent.cpp
@@ -3055,7 +3055,7 @@ bool UPrimitiveComponent::ComponentOverlapComponentImpl(class UPrimitiveComponen
 	{
 		for (Chaos::FPhysicsObjectHandle ThisObject : ThisObjects)
 		{
-			if (Interface->GetPhysicsObjectOverlap(InObject, FTransform::Identity, ThisObject, FTransform::Identity, Params.bTraceComplex))
+			if (Interface->PhysicsObjectOverlap(InObject, FTransform::Identity, ThisObject, FTransform::Identity, Params.bTraceComplex))
 			{
 				return true;
 			}
diff --git a/Engine/Source/Runtime/Experimental/Chaos/Private/Chaos/PhysicsObjectInterface.cpp b/Engine/Source/Runtime/Experimental/Chaos/Private/Chaos/PhysicsObjectInterface.cpp
--- a/Engine/Source/Runtime/Experimental/Chaos/Private/Chaos/PhysicsObjectInterface.cpp
+++ b/Engine/Source/Runtime/Experimental/Chaos/Private/Chaos/PhysicsObjectInterface.cpp
@@ -273,12 +273,12 @@ namespace Chaos
 		bool bRetOverlap = false;
 		if (OutOverlap.MTD)
 		{
-			bRetOverlap |= GetPhysicsObjectOverlapWithMTD(ObjectA, FTransform::Identity, ObjectB, FTransform::Identity, bTraceComplex, *OutOverlap.MTD);
+			bRetOverlap |= PhysicsObjectOverlapWithMTD(ObjectA, FTransform::Identity, ObjectB, FTransform::Identity, bTraceComplex, *OutOverlap.MTD);
 		}
 
 		if (OutOverlap.AxisOverlap)
 		{
-			bRetOverlap |= GetPhysicsObjectOverlapWithAABB(ObjectA, FTransform::Identity, ObjectB, FTransform::Identity, bTraceComplex, FVector::Zero(), *OutOverlap.AxisOverlap);
+			bRetOverlap |= PhysicsObjectOverlapWithAABB(ObjectA, FTransform::Identity, ObjectB, FTransform::Identity, bTraceComplex, FVector::Zero(), *OutOverlap.AxisOverlap);
 		}
 
 		return bRetOverlap;
@@ -292,28 +292,28 @@ namespace Chaos
 		bool bRetOverlap = false;
 		if (OutOverlap.MTD)
 		{
-			bRetOverlap |= GetPhysicsObjectOverlapWithMTD(ObjectA, InTransformA, ObjectB, InTransformB, bTraceComplex, *OutOverlap.MTD);
+			bRetOverlap |= PhysicsObjectOverlapWithMTD(ObjectA, InTransformA, ObjectB, InTransformB, bTraceComplex, *OutOverlap.MTD);
 		}
 
 		if (OutOverlap.AxisOverlap)
 		{
-			bRetOverlap |= GetPhysicsObjectOverlapWithAABB(ObjectA, InTransformA, ObjectB, InTransformB, bTraceComplex, FVector::Zero(), *OutOverlap.AxisOverlap);
+			bRetOverlap |= PhysicsObjectOverlapWithAABB(ObjectA, InTransformA, ObjectB, InTransformB, bTraceComplex, FVector::Zero(), *OutOverlap.AxisOverlap);
 		}
 
 		return bRetOverlap;
 	}
 
 	template<EThreadContext Id>
-	bool FReadPhysicsObjectInterface<Id>::GetPhysicsObjectOverlap(FPhysicsObjectHandle ObjectA, const FTransform& InTransformA, FPhysicsObjectHandle ObjectB, const FTransform& InTransformB, bool bTraceComplex)
+	bool FReadPhysicsObjectInterface<Id>::PhysicsObjectOverlap(FPhysicsObjectHandle ObjectA, const FTransform& InTransformA, FPhysicsObjectHandle ObjectB, const FTransform& InTransformB, bool bTraceComplex)
 	{
-		TRACE_CPUPROFILER_EVENT_SCOPE(FReadPhysicsObjectInterface<Id>::GetPhysicsObjectOverlap);
+		TRACE_CPUPROFILER_EVENT_SCOPE(FReadPhysicsObjectInterface<Id>::PhysicsObjectOverlap);
 		return PairwiseShapeOverlapHelper(ObjectA, InTransformA, ObjectB, InTransformB, bTraceComplex, false, FVector::Zero(), [](const FShapeOverlapData&, const FShapeOverlapData&, const FMTDInfo&) { return false; });
 	}
 
 	template<EThreadContext Id>
-	bool FReadPhysicsObjectInterface<Id>::GetPhysicsObjectOverlapWithMTD(FPhysicsObjectHandle ObjectA, const FTransform& InTransformA, FPhysicsObjectHandle ObjectB, const FTransform& InTransformB, bool bTraceComplex, FMTDInfo& OutMTD)
+	bool FReadPhysicsObjectInterface<Id>::PhysicsObjectOverlapWithMTD(FPhysicsObjectHandle ObjectA, const FTransform& InTransformA, FPhysicsObjectHandle ObjectB, const FTransform& InTransformB, bool bTraceComplex, FMTDInfo& OutMTD)
 	{
-		TRACE_CPUPROFILER_EVENT_SCOPE(FReadPhysicsObjectInterface<Id>::GetPhysicsObjectOverlapWithMTD);
+		TRACE_CPUPROFILER_EVENT_SCOPE(FReadPhysicsObjectInterface<Id>::PhysicsObjectOverlapWithMTD);
 		OutMTD.Penetration = 0.0;
 		return PairwiseShapeOverlapHelper(
 			ObjectA,
@@ -332,9 +332,9 @@ namespace Chaos
 	}
 
 	template<EThreadContext Id>
-	bool FReadPhysicsObjectInterface<Id>::GetPhysicsObjectOverlapWithAABB(FPhysicsObjectHandle ObjectA, const FTransform& InTransformA, FPhysicsObjectHandle ObjectB, const FTransform& InTransformB, bool bTraceComplex, const FVector& Tolerance, FBox& OutOverlap)
+	bool FReadPhysicsObjectInterface<Id>::PhysicsObjectOverlapWithAABB(FPhysicsObjectHandle ObjectA, const FTransform& InTransformA, FPhysicsObjectHandle ObjectB, const FTransform& InTransformB, bool bTraceComplex, const FVector& Tolerance, FBox& OutOverlap)
 	{
-		TRACE_CPUPROFILER_EVENT_SCOPE(FReadPhysicsObjectInterface<Id>::GetPhysicsObjectOverlapWithAABB);
+		TRACE_CPUPROFILER_EVENT_SCOPE(FReadPhysicsObjectInterface<Id>::PhysicsObjectOverlapWithAABB);
 		OutOverlap = FBox{ EForceInit::ForceInitToZero };
 		return PairwiseShapeOverlapHelper(
 			ObjectA,
@@ -354,9 +354,9 @@ namespace Chaos
 	}
 
 	template<EThreadContext Id>
-	bool FReadPhysicsObjectInterface<Id>::GetPhysicsObjectOverlapWithAABBSize(FPhysicsObjectHandle ObjectA, const FTransform& InTransformA, FPhysicsObjectHandle ObjectB, const FTransform& InTransformB, bool bTraceComplex, const FVector& Tolerance, FVector& OutOverlapSize)
+	bool FReadPhysicsObjectInterface<Id>::PhysicsObjectOverlapWithAABBSize(FPhysicsObjectHandle ObjectA, const FTransform& InTransformA, FPhysicsObjectHandle ObjectB, const FTransform& InTransformB, bool bTraceComplex, const FVector& Tolerance, FVector& OutOverlapSize)
 	{
-		TRACE_CPUPROFILER_EVENT_SCOPE(FReadPhysicsObjectInterface<Id>::GetPhysicsObjectOverlapWithAABBSize);
+		TRACE_CPUPROFILER_EVENT_SCOPE(FReadPhysicsObjectInterface<Id>::PhysicsObjectOverlapWithAABBSize);
 		OutOverlapSize = FVector::Zero();
 		return PairwiseShapeOverlapHelper(
 			ObjectA,
diff --git a/Engine/Source/Runtime/Experimental/Chaos/Public/Chaos/PhysicsObjectInterface.h b/Engine/Source/Runtime/Experimental/Chaos/Public/Chaos/PhysicsObjectInterface.h
--- a/Engine/Source/Runtime/Experimental/Chaos/Public/Chaos/PhysicsObjectInterface.h
+++ b/Engine/Source/Runtime/Experimental/Chaos/Public/Chaos/PhysicsObjectInterface.h
@@ -79,14 +79,14 @@ namespace Chaos
 		bool GetPhysicsObjectOverlapWithTransform(FPhysicsObjectHandle ObjectA, const FTransform& InTransformA, FPhysicsObjectHandle ObjectB, const FTransform& InTransformB, bool bTraceComplex, Chaos::FOverlapInfo& OutOverlap);
 
 		// This function will not compute any overlap heuristic.
-		bool GetPhysicsObjectOverlap(FPhysicsObjectHandle ObjectA, const FTransform& InTransformA, FPhysicsObjectHandle ObjectB, const FTransform& InTransformB, bool bTraceComplex);
+		bool PhysicsObjectOverlap(FPhysicsObjectHandle ObjectA, const FTransform& InTransformA, FPhysicsObjectHandle ObjectB, const FTransform& InTransformB, bool bTraceComplex);
 
 		// This function does the same as GetPhysicsObjectOverlap but also computes the MTD metric.
-		bool GetPhysicsObjectOverlapWithMTD(FPhysicsObjectHandle ObjectA, const FTransform& InTransformA, FPhysicsObjectHandle ObjectB, const FTransform& InTransformB, bool bTraceComplex, FMTDInfo& OutMTD);
+		bool PhysicsObjectOverlapWithMTD(FPhysicsObjectHandle ObjectA, const FTransform& InTransformA, FPhysicsObjectHandle ObjectB, const FTransform& InTransformB, bool bTraceComplex, FMTDInfo& OutMTD);
 
 		// This function does the same as GetPhysicsObjectOverlap but also computes the AABB overlap metric.
-		bool GetPhysicsObjectOverlapWithAABB(FPhysicsObjectHandle ObjectA, const FTransform& InTransformA, FPhysicsObjectHandle ObjectB, const FTransform& InTransformB, bool bTraceComplex, const FVector& Tolerance, FBox& OutOverlap);
-		bool GetPhysicsObjectOverlapWithAABBSize(FPhysicsObjectHandle ObjectA, const FTransform& InTransformA, FPhysicsObjectHandle ObjectB, const FTransform& InTransformB, bool bTraceComplex, const FVector& Tolerance, FVector& OutOverlapSize);
+		bool PhysicsObjectOverlapWithAABB(FPhysicsObjectHandle ObjectA, const FTransform& InTransformA, FPhysicsObjectHandle ObjectB, const FTransform& InTransformB, bool bTraceComplex, const FVector& Tolerance, FBox& OutOverlap);
+		bool PhysicsObjectOverlapWithAABBSize(FPhysicsObjectHandle ObjectA, const FTransform& InTransformA, FPhysicsObjectHandle ObjectB, const FTransform& InTransformB, bool bTraceComplex, const FVector& Tolerance, FVector& OutOverlapSize);
 
 		bool AreAllValid(TArrayView<FPhysicsObjectHandle> InObjects);
 		bool AreAllKinematic(TArrayView<FPhysicsObjectHandle> InObjects);
