Created a new stat gathering method on ManagedArrayCollection to enable easy determination of the marginal cost of adding a new element to a particular group.

diff --git a/Engine/Source/Runtime/Experimental/Chaos/Private/GeometryCollection/ManagedArrayCollection.cpp b/Engine/Source/Runtime/Experimental/Chaos/Private/GeometryCollection/ManagedArrayCollection.cpp
--- a/Engine/Source/Runtime/Experimental/Chaos/Private/GeometryCollection/ManagedArrayCollection.cpp
+++ b/Engine/Source/Runtime/Experimental/Chaos/Private/GeometryCollection/ManagedArrayCollection.cpp
@@ -586,6 +586,26 @@ SIZE_T FManagedArrayCollection::GetAllocatedSize() const
 	return AllocatedSize;
 }
 
+void FManagedArrayCollection::GetElementSizeInfoForGroups(TArray<TPair<FName, SIZE_T>>& OutSizeInfo) const
+{
+	// Group name to total element size map
+	TMap<FName, SIZE_T> GroupToElementSizeMap;
+
+	for (const TPair<FKeyType, FValueType>& Attribute : Map)
+	{
+		SIZE_T& GroupSize = GroupToElementSizeMap.FindOrAdd(Attribute.Key.Get<1>());
+		if (FManagedArrayBase* Array = Attribute.Value.Value)
+		{
+			GroupSize += Array->GetTypeSize();
+		}
+	}
+
+	for (const TPair<FName, SIZE_T>& GroupAndElementSize : GroupToElementSizeMap)
+	{
+		OutSizeInfo.Add(GroupAndElementSize);
+	}
+}
+
 static const FName GuidName("GUID");
 
 // this is a reference wrapper to avoid copying attributes as some may not support copy (unique ptr ones) 
diff --git a/Engine/Source/Runtime/Experimental/Chaos/Public/GeometryCollection/ManagedArrayCollection.h b/Engine/Source/Runtime/Experimental/Chaos/Public/GeometryCollection/ManagedArrayCollection.h
--- a/Engine/Source/Runtime/Experimental/Chaos/Public/GeometryCollection/ManagedArrayCollection.h
+++ b/Engine/Source/Runtime/Experimental/Chaos/Public/GeometryCollection/ManagedArrayCollection.h
@@ -549,6 +549,11 @@ public:
 	*/
 	SIZE_T GetAllocatedSize() const;
 
+	/**
+	 * Get information on the total storage required for each element in each group
+	 */
+	 void GetElementSizeInfoForGroups(TArray<TPair<FName, SIZE_T>>& OutSizeInfo) const;
+
 private:
 
 	/****
