make the physics object interface usable with const handles

diff --git a/Engine/Source/Runtime/Engine/Classes/PhysicsEngine/PhysicsObjectExternalInterface.h b/Engine/Source/Runtime/Engine/Classes/PhysicsEngine/PhysicsObjectExternalInterface.h
--- a/Engine/Source/Runtime/Engine/Classes/PhysicsEngine/PhysicsObjectExternalInterface.h
+++ b/Engine/Source/Runtime/Engine/Classes/PhysicsEngine/PhysicsObjectExternalInterface.h
@@ -47,8 +47,8 @@ using FLockedWritePhysicsObjectExternalInterface = FLockedPhysicsObjectExternalI
 class ENGINE_API FPhysicsObjectExternalInterface : public Chaos::FPhysicsObjectInterface, public PhysicsObjectPhysicsCoreInterface
 {
 public:
-	static FLockedReadPhysicsObjectExternalInterface LockRead(TArrayView<Chaos::FPhysicsObjectHandle> InObjects);
-	static FLockedWritePhysicsObjectExternalInterface LockWrite(TArrayView<Chaos::FPhysicsObjectHandle> InObjects);
+	static FLockedReadPhysicsObjectExternalInterface LockRead(TArrayView<const Chaos::FConstPhysicsObjectHandle> InObjects);
+	static FLockedWritePhysicsObjectExternalInterface LockWrite(TArrayView<const Chaos::FPhysicsObjectHandle> InObjects);
 
 	static Chaos::FReadPhysicsObjectInterface_External GetRead_AssumesLocked();
 	static Chaos::FWritePhysicsObjectInterface_External GetWrite_AssumesLocked();
diff --git a/Engine/Source/Runtime/Engine/Private/PhysicsEngine/PhysicsObjectExternalInterface.cpp b/Engine/Source/Runtime/Engine/Private/PhysicsEngine/PhysicsObjectExternalInterface.cpp
--- a/Engine/Source/Runtime/Engine/Private/PhysicsEngine/PhysicsObjectExternalInterface.cpp
+++ b/Engine/Source/Runtime/Engine/Private/PhysicsEngine/PhysicsObjectExternalInterface.cpp
@@ -1,12 +1,12 @@
 // Copyright Epic Games, Inc. All Rights Reserved.
 #include "PhysicsEngine/PhysicsObjectExternalInterface.h"
 
-FLockedReadPhysicsObjectExternalInterface FPhysicsObjectExternalInterface::LockRead(TArrayView<Chaos::FPhysicsObjectHandle> InObjects)
+FLockedReadPhysicsObjectExternalInterface FPhysicsObjectExternalInterface::LockRead(TArrayView<const Chaos::FConstPhysicsObjectHandle> InObjects)
 {
 	return FLockedReadPhysicsObjectExternalInterface{ GetScene(InObjects), CreateReadInterface<Chaos::EThreadContext::External>() };
 }
 
-FLockedWritePhysicsObjectExternalInterface FPhysicsObjectExternalInterface::LockWrite(TArrayView<Chaos::FPhysicsObjectHandle> InObjects)
+FLockedWritePhysicsObjectExternalInterface FPhysicsObjectExternalInterface::LockWrite(TArrayView<const Chaos::FPhysicsObjectHandle> InObjects)
 {
 	return FLockedWritePhysicsObjectExternalInterface{ GetScene(InObjects), CreateWriteInterface<Chaos::EThreadContext::External>() };
 }
diff --git a/Engine/Source/Runtime/Experimental/Chaos/Private/Chaos/PhysicsObjectInterface.cpp b/Engine/Source/Runtime/Experimental/Chaos/Private/Chaos/PhysicsObjectInterface.cpp
--- a/Engine/Source/Runtime/Experimental/Chaos/Private/Chaos/PhysicsObjectInterface.cpp
+++ b/Engine/Source/Runtime/Experimental/Chaos/Private/Chaos/PhysicsObjectInterface.cpp
@@ -11,7 +11,7 @@
 namespace
 {
 	template<Chaos::EThreadContext Id>
-	void SetParticleStateHelper(Chaos::FPhysicsObjectHandle PhysicsObject, Chaos::EObjectStateType State)
+	void SetParticleStateHelper(const Chaos::FPhysicsObjectHandle PhysicsObject, Chaos::EObjectStateType State)
 	{
 		if (!PhysicsObject)
 		{
@@ -79,7 +79,7 @@ FName FClosestPhysicsObjectResult::HitName() const
 namespace Chaos
 {
 	template<EThreadContext Id>
-	FPhysicsObjectHandle FReadPhysicsObjectInterface<Id>::GetRootObject(FPhysicsObjectHandle Object)
+	FPhysicsObjectHandle FReadPhysicsObjectInterface<Id>::GetRootObject(const FConstPhysicsObjectHandle Object)
 	{
 		if (!Object)
 		{
@@ -90,19 +90,19 @@ namespace Chaos
 	}
 
 	template<EThreadContext Id>
-	bool FReadPhysicsObjectInterface<Id>::HasChildren(FPhysicsObjectHandle Object)
+	bool FReadPhysicsObjectInterface<Id>::HasChildren(const FConstPhysicsObjectHandle Object)
 	{
 		return Object ? Object->HasChildren<Id>() : false;
 	}
 
 	template<EThreadContext Id>
-	FTransform FReadPhysicsObjectInterface<Id>::GetTransform(FPhysicsObjectHandle Object)
+	FTransform FReadPhysicsObjectInterface<Id>::GetTransform(const FConstPhysicsObjectHandle Object)
 	{
 		return FTransform{ GetR(Object), GetX(Object) };
 	}
 
 	template<EThreadContext Id>
-	FVector FReadPhysicsObjectInterface<Id>::GetX(FPhysicsObjectHandle Object)
+	FVector FReadPhysicsObjectInterface<Id>::GetX(const FConstPhysicsObjectHandle Object)
 	{
 		if (!Object)
 		{
@@ -118,7 +118,7 @@ namespace Chaos
 	}
 
 	template<EThreadContext Id>
-	FVector FReadPhysicsObjectInterface<Id>::GetCoM(FPhysicsObjectHandle Object)
+	FVector FReadPhysicsObjectInterface<Id>::GetCoM(const FConstPhysicsObjectHandle Object)
 	{
 		if (!Object)
 		{
@@ -137,13 +137,13 @@ namespace Chaos
 	}
 
 	template<EThreadContext Id>
-	FVector FReadPhysicsObjectInterface<Id>::GetWorldCoM(FPhysicsObjectHandle Object)
+	FVector FReadPhysicsObjectInterface<Id>::GetWorldCoM(const FConstPhysicsObjectHandle Object)
 	{
 		return GetX(Object) + GetR(Object).RotateVector(GetCoM(Object));
 	}
 
 	template<EThreadContext Id>
-	FQuat FReadPhysicsObjectInterface<Id>::GetR(FPhysicsObjectHandle Object)
+	FQuat FReadPhysicsObjectInterface<Id>::GetR(const FConstPhysicsObjectHandle Object)
 	{
 		if (!Object)
 		{
@@ -159,7 +159,7 @@ namespace Chaos
 	}
 
 	template<EThreadContext Id>
-	FSpatialAccelerationIdx FReadPhysicsObjectInterface<Id>::GetSpatialIndex(FPhysicsObjectHandle Object)
+	FSpatialAccelerationIdx FReadPhysicsObjectInterface<Id>::GetSpatialIndex(const FConstPhysicsObjectHandle Object)
 	{
 		if (!Object)
 		{
@@ -175,12 +175,12 @@ namespace Chaos
 	}
 
 	template<EThreadContext Id>
-	TArray<TThreadParticle<Id>*> FReadPhysicsObjectInterface<Id>::GetAllParticles(TArrayView<FPhysicsObjectHandle> InObjects)
+	TArray<TThreadParticle<Id>*> FReadPhysicsObjectInterface<Id>::GetAllParticles(TArrayView<const FConstPhysicsObjectHandle> InObjects)
 	{
 		TArray<TThreadParticle<Id>*> Particles;
 		Particles.Reserve(InObjects.Num());
 
-		for (const FPhysicsObjectHandle& Handle : InObjects)
+		for (const FConstPhysicsObjectHandle& Handle : InObjects)
 		{
 			if (TThreadParticle<Id>* Particle = Handle->GetParticle<Id>())
 			{
@@ -192,12 +192,12 @@ namespace Chaos
 	}
 
 	template<EThreadContext Id>
-	TArray<TThreadRigidParticle<Id>*> FReadPhysicsObjectInterface<Id>::GetAllRigidParticles(TArrayView<FPhysicsObjectHandle> InObjects)
+	TArray<TThreadRigidParticle<Id>*> FReadPhysicsObjectInterface<Id>::GetAllRigidParticles(TArrayView<const FConstPhysicsObjectHandle> InObjects)
 	{
 		TArray<TThreadRigidParticle<Id>*> Particles;
 		Particles.Reserve(InObjects.Num());
 
-		for (const FPhysicsObjectHandle& Handle : InObjects)
+		for (const FConstPhysicsObjectHandle& Handle : InObjects)
 		{
 			if (!Handle)
 			{
@@ -217,11 +217,11 @@ namespace Chaos
 	}
 
 	template<EThreadContext Id>
-	TArray<FPerShapeData*> FReadPhysicsObjectInterface<Id>::GetAllShapes(TArrayView<FPhysicsObjectHandle> InObjects)
+	TArray<FPerShapeData*> FReadPhysicsObjectInterface<Id>::GetAllShapes(TArrayView<const FConstPhysicsObjectHandle> InObjects)
 	{
 		TArray<FPerShapeData*> AllShapes;
 
-		for (FPhysicsObjectHandle Object : InObjects)
+		for (const FConstPhysicsObjectHandle Object : InObjects)
 		{
 			if (!Object)
 			{
@@ -242,9 +242,9 @@ namespace Chaos
 	}
 
 	template<EThreadContext Id>
-	void FReadPhysicsObjectInterface<Id>::VisitEveryShape(TArrayView<FPhysicsObjectHandle> InObjects, TFunctionRef<bool(FPhysicsObjectHandle, FPerShapeData*)> Lambda)
+	void FReadPhysicsObjectInterface<Id>::VisitEveryShape(TArrayView<const FConstPhysicsObjectHandle> InObjects, TFunctionRef<bool(const FConstPhysicsObjectHandle, FPerShapeData*)> Lambda)
 	{
-		for (FPhysicsObjectHandle Object : InObjects)
+		for (const FConstPhysicsObjectHandle Object : InObjects)
 		{
 			if (!Object)
 			{
@@ -266,7 +266,7 @@ namespace Chaos
 	}
 
 	template<EThreadContext Id>
-	bool FReadPhysicsObjectInterface<Id>::GetPhysicsObjectOverlap(FPhysicsObjectHandle ObjectA, FPhysicsObjectHandle ObjectB, bool bTraceComplex, Chaos::FOverlapInfo& OutOverlap)
+	bool FReadPhysicsObjectInterface<Id>::GetPhysicsObjectOverlap(const FConstPhysicsObjectHandle ObjectA, const FConstPhysicsObjectHandle ObjectB, bool bTraceComplex, Chaos::FOverlapInfo& OutOverlap)
 	{
 		TRACE_CPUPROFILER_EVENT_SCOPE(FReadPhysicsObjectInterface<Id>::GetPhysicsObjectOverlap);
 		// This is slow and inefficient and hence deprecated.
@@ -285,7 +285,7 @@ namespace Chaos
 	}
 
 	template<EThreadContext Id>
-	bool FReadPhysicsObjectInterface<Id>::GetPhysicsObjectOverlapWithTransform(FPhysicsObjectHandle ObjectA, const FTransform& InTransformA, FPhysicsObjectHandle ObjectB, const FTransform& InTransformB, bool bTraceComplex, Chaos::FOverlapInfo& OutOverlap)
+	bool FReadPhysicsObjectInterface<Id>::GetPhysicsObjectOverlapWithTransform(const FConstPhysicsObjectHandle ObjectA, const FTransform& InTransformA, const FConstPhysicsObjectHandle ObjectB, const FTransform& InTransformB, bool bTraceComplex, Chaos::FOverlapInfo& OutOverlap)
 	{
 		TRACE_CPUPROFILER_EVENT_SCOPE(FReadPhysicsObjectInterface<Id>::GetPhysicsObjectOverlapWithTransform);
 		// This is slow and inefficient and hence deprecated.
@@ -304,14 +304,14 @@ namespace Chaos
 	}
 
 	template<EThreadContext Id>
-	bool FReadPhysicsObjectInterface<Id>::PhysicsObjectOverlap(FPhysicsObjectHandle ObjectA, const FTransform& InTransformA, FPhysicsObjectHandle ObjectB, const FTransform& InTransformB, bool bTraceComplex)
+	bool FReadPhysicsObjectInterface<Id>::PhysicsObjectOverlap(const FConstPhysicsObjectHandle ObjectA, const FTransform& InTransformA, const FConstPhysicsObjectHandle ObjectB, const FTransform& InTransformB, bool bTraceComplex)
 	{
 		TRACE_CPUPROFILER_EVENT_SCOPE(FReadPhysicsObjectInterface<Id>::PhysicsObjectOverlap);
 		return PairwiseShapeOverlapHelper(ObjectA, InTransformA, ObjectB, InTransformB, bTraceComplex, false, FVector::Zero(), [](const FShapeOverlapData&, const FShapeOverlapData&, const FMTDInfo&) { return false; });
 	}
 
 	template<EThreadContext Id>
-	bool FReadPhysicsObjectInterface<Id>::PhysicsObjectOverlapWithMTD(FPhysicsObjectHandle ObjectA, const FTransform& InTransformA, FPhysicsObjectHandle ObjectB, const FTransform& InTransformB, bool bTraceComplex, FMTDInfo& OutMTD)
+	bool FReadPhysicsObjectInterface<Id>::PhysicsObjectOverlapWithMTD(const FConstPhysicsObjectHandle ObjectA, const FTransform& InTransformA, const FConstPhysicsObjectHandle ObjectB, const FTransform& InTransformB, bool bTraceComplex, FMTDInfo& OutMTD)
 	{
 		TRACE_CPUPROFILER_EVENT_SCOPE(FReadPhysicsObjectInterface<Id>::PhysicsObjectOverlapWithMTD);
 		OutMTD.Penetration = 0.0;
@@ -332,7 +332,7 @@ namespace Chaos
 	}
 
 	template<EThreadContext Id>
-	bool FReadPhysicsObjectInterface<Id>::PhysicsObjectOverlapWithAABB(FPhysicsObjectHandle ObjectA, const FTransform& InTransformA, FPhysicsObjectHandle ObjectB, const FTransform& InTransformB, bool bTraceComplex, const FVector& Tolerance, FBox& OutOverlap)
+	bool FReadPhysicsObjectInterface<Id>::PhysicsObjectOverlapWithAABB(const FConstPhysicsObjectHandle ObjectA, const FTransform& InTransformA, const FConstPhysicsObjectHandle ObjectB, const FTransform& InTransformB, bool bTraceComplex, const FVector& Tolerance, FBox& OutOverlap)
 	{
 		TRACE_CPUPROFILER_EVENT_SCOPE(FReadPhysicsObjectInterface<Id>::PhysicsObjectOverlapWithAABB);
 		OutOverlap = FBox{ EForceInit::ForceInitToZero };
@@ -354,7 +354,7 @@ namespace Chaos
 	}
 
 	template<EThreadContext Id>
-	bool FReadPhysicsObjectInterface<Id>::PhysicsObjectOverlapWithAABBSize(FPhysicsObjectHandle ObjectA, const FTransform& InTransformA, FPhysicsObjectHandle ObjectB, const FTransform& InTransformB, bool bTraceComplex, const FVector& Tolerance, FVector& OutOverlapSize)
+	bool FReadPhysicsObjectInterface<Id>::PhysicsObjectOverlapWithAABBSize(const FConstPhysicsObjectHandle ObjectA, const FTransform& InTransformA, const FConstPhysicsObjectHandle ObjectB, const FTransform& InTransformB, bool bTraceComplex, const FVector& Tolerance, FVector& OutOverlapSize)
 	{
 		TRACE_CPUPROFILER_EVENT_SCOPE(FReadPhysicsObjectInterface<Id>::PhysicsObjectOverlapWithAABBSize);
 		OutOverlapSize = FVector::Zero();
@@ -376,7 +376,7 @@ namespace Chaos
 	}
 
 	template<EThreadContext Id>
-	bool FReadPhysicsObjectInterface<Id>::PairwiseShapeOverlapHelper(FPhysicsObjectHandle ObjectA, const FTransform& InTransformA, FPhysicsObjectHandle ObjectB, const FTransform& InTransformB, bool bTraceComplex, bool bComputeMTD, const FVector& Tolerance, const TFunction<bool(const FShapeOverlapData&, const FShapeOverlapData&, const FMTDInfo&)>& Lambda)
+	bool FReadPhysicsObjectInterface<Id>::PairwiseShapeOverlapHelper(const FConstPhysicsObjectHandle ObjectA, const FTransform& InTransformA, const FConstPhysicsObjectHandle ObjectB, const FTransform& InTransformB, bool bTraceComplex, bool bComputeMTD, const FVector& Tolerance, const TFunction<bool(const FShapeOverlapData&, const FShapeOverlapData&, const FMTDInfo&)>& Lambda)
 	{
 		TArray<FPerShapeData*> ShapesA = GetAllShapes({ &ObjectA, 1 });
 		const FTransform TransformA = FTransform{ GetR(ObjectA), GetX(ObjectA) } *InTransformA;
@@ -460,10 +460,10 @@ namespace Chaos
 	}
 
 	template<EThreadContext Id>
-	bool FReadPhysicsObjectInterface<Id>::AreAllValid(TArrayView<FPhysicsObjectHandle> InObjects)
+	bool FReadPhysicsObjectInterface<Id>::AreAllValid(TArrayView<const FConstPhysicsObjectHandle> InObjects)
 	{
 		bool bCheck = !InObjects.IsEmpty();
-		for (FPhysicsObjectHandle Object : InObjects)
+		for (const FConstPhysicsObjectHandle Object : InObjects)
 		{
 			bCheck &= (Object != nullptr && Object->IsValid());
 		}
@@ -471,10 +471,10 @@ namespace Chaos
 	}
 
 	template<EThreadContext Id>
-	bool FReadPhysicsObjectInterface<Id>::AreAllKinematic(TArrayView<FPhysicsObjectHandle> InObjects)
+	bool FReadPhysicsObjectInterface<Id>::AreAllKinematic(TArrayView<const FConstPhysicsObjectHandle> InObjects)
 	{
 		bool bCheck = !InObjects.IsEmpty();
-		for (FPhysicsObjectHandle Object : InObjects)
+		for (const FConstPhysicsObjectHandle Object : InObjects)
 		{
 			bCheck &= (Object && Object->IsValid() && Object->ObjectState<Id>() == EObjectStateType::Kinematic);
 		}
@@ -482,10 +482,10 @@ namespace Chaos
 	}
 
 	template<EThreadContext Id>
-	bool FReadPhysicsObjectInterface<Id>::AreAllSleeping(TArrayView<FPhysicsObjectHandle> InObjects)
+	bool FReadPhysicsObjectInterface<Id>::AreAllSleeping(TArrayView<const FConstPhysicsObjectHandle> InObjects)
 	{
 		bool bCheck = !InObjects.IsEmpty();
-		for (FPhysicsObjectHandle Object : InObjects)
+		for (const FConstPhysicsObjectHandle Object : InObjects)
 		{
 			bCheck &= (Object && Object->IsValid() && Object->ObjectState<Id>() == EObjectStateType::Sleeping);
 		}
@@ -493,10 +493,10 @@ namespace Chaos
 	}
 
 	template<EThreadContext Id>
-	bool FReadPhysicsObjectInterface<Id>::AreAllRigidBody(TArrayView<FPhysicsObjectHandle> InObjects)
+	bool FReadPhysicsObjectInterface<Id>::AreAllRigidBody(TArrayView<const FConstPhysicsObjectHandle> InObjects)
 	{
 		bool bCheck = !InObjects.IsEmpty();
-		for (FPhysicsObjectHandle Object : InObjects)
+		for (const FConstPhysicsObjectHandle Object : InObjects)
 		{
 			bCheck &= (Object && Object->IsValid() && Object->ObjectState<Id>() != EObjectStateType::Static);
 		}
@@ -504,10 +504,10 @@ namespace Chaos
 	}
 
 	template<EThreadContext Id>
-	bool FReadPhysicsObjectInterface<Id>::AreAllDynamic(TArrayView<FPhysicsObjectHandle> InObjects)
+	bool FReadPhysicsObjectInterface<Id>::AreAllDynamic(TArrayView<const FConstPhysicsObjectHandle> InObjects)
 	{
 		bool bCheck = !InObjects.IsEmpty();
-		for (FPhysicsObjectHandle Object : InObjects)
+		for (const FConstPhysicsObjectHandle Object : InObjects)
 		{
 			bCheck &= (Object && Object->IsValid() && Object->ObjectState<Id>() == EObjectStateType::Dynamic);
 		}
@@ -515,10 +515,10 @@ namespace Chaos
 	}
 
 	template<EThreadContext Id>
-	bool FReadPhysicsObjectInterface<Id>::AreAllDisabled(TArrayView<FPhysicsObjectHandle> InObjects)
+	bool FReadPhysicsObjectInterface<Id>::AreAllDisabled(TArrayView<const FConstPhysicsObjectHandle> InObjects)
 	{
 		bool bDisabled = !InObjects.IsEmpty();
-		for (FPhysicsObjectHandle Object : InObjects)
+		for (const FConstPhysicsObjectHandle Object : InObjects)
 		{
 			bool bParticleDisabled = true;
 			if (Object)
@@ -534,14 +534,14 @@ namespace Chaos
 	}
 
 	template<EThreadContext Id>
-	bool FReadPhysicsObjectInterface<Id>::AreAllShapesQueryEnabled(TArrayView<FPhysicsObjectHandle> InObjects)
+	bool FReadPhysicsObjectInterface<Id>::AreAllShapesQueryEnabled(TArrayView<const FConstPhysicsObjectHandle> InObjects)
 	{
 		if (InObjects.IsEmpty())
 		{
 			return false;
 		}
 
-		for (FPhysicsObjectHandle Object : InObjects)
+		for (const FConstPhysicsObjectHandle Object : InObjects)
 		{
 			if (Object)
 			{
@@ -561,10 +561,10 @@ namespace Chaos
 	}
 
 	template<EThreadContext Id>
-	float FReadPhysicsObjectInterface<Id>::GetMass(TArrayView<FPhysicsObjectHandle> InObjects)
+	float FReadPhysicsObjectInterface<Id>::GetMass(TArrayView<const FConstPhysicsObjectHandle> InObjects)
 	{
 		float Mass = 0.f;
-		for (FPhysicsObjectHandle Object : InObjects)
+		for (const FConstPhysicsObjectHandle Object : InObjects)
 		{
 			if (Object)
 			{
@@ -581,10 +581,10 @@ namespace Chaos
 	}
 
 	template<EThreadContext Id>
-	FBox FReadPhysicsObjectInterface<Id>::GetBounds(TArrayView<FPhysicsObjectHandle> InObjects)
+	FBox FReadPhysicsObjectInterface<Id>::GetBounds(TArrayView<const FConstPhysicsObjectHandle> InObjects)
 	{
 		FBox RetBox(ForceInit);
-		for (FPhysicsObjectHandle Object : InObjects)
+		for (const FConstPhysicsObjectHandle Object : InObjects)
 		{
 			if (!Object)
 			{
@@ -613,10 +613,10 @@ namespace Chaos
 	}
 
 	template<EThreadContext Id>
-	FBox FReadPhysicsObjectInterface<Id>::GetWorldBounds(TArrayView<FPhysicsObjectHandle> InObjects)
+	FBox FReadPhysicsObjectInterface<Id>::GetWorldBounds(TArrayView<const FConstPhysicsObjectHandle> InObjects)
 	{
 		FBox RetBox(ForceInit);
-		for (FPhysicsObjectHandle Object : InObjects)
+		for (const FConstPhysicsObjectHandle Object : InObjects)
 		{
 			if (!Object)
 			{
@@ -647,10 +647,10 @@ namespace Chaos
 	}
 
 	template<EThreadContext Id>
-	FClosestPhysicsObjectResult FReadPhysicsObjectInterface<Id>::GetClosestPhysicsBodyFromLocation(TArrayView<FPhysicsObjectHandle> InObjects, const FVector& WorldLocation)
+	FClosestPhysicsObjectResult FReadPhysicsObjectInterface<Id>::GetClosestPhysicsBodyFromLocation(TArrayView<const FConstPhysicsObjectHandle> InObjects, const FVector& WorldLocation)
 	{
 		FClosestPhysicsObjectResult AggregateResult;
-		for (FPhysicsObjectHandle Object : InObjects)
+		for (const FConstPhysicsObjectHandle Object : InObjects)
 		{
 			if (!Object)
 			{
@@ -670,7 +670,7 @@ namespace Chaos
 
 			if (const FImplicitObject* Geometry = Particle->Geometry().Get())
 			{
-				Result.PhysicsObject = Object;
+				Result.PhysicsObject = const_cast<FPhysicsObjectHandle>(Object);
 
 				Chaos::FVec3 Normal;
 				Result.ClosestDistance = static_cast<double>(Geometry->PhiWithNormal(LocalLocation, Normal));
@@ -691,13 +691,13 @@ namespace Chaos
 	}
 
 	template<EThreadContext Id>
-	FAccelerationStructureHandle FReadPhysicsObjectInterface<Id>::CreateAccelerationStructureHandle(FPhysicsObjectHandle InObject)
+	FAccelerationStructureHandle FReadPhysicsObjectInterface<Id>::CreateAccelerationStructureHandle(const FConstPhysicsObjectHandle InObject)
 	{
 		return FAccelerationStructureHandle{InObject->GetParticle<Id>()};
 	}
 
 	template<EThreadContext Id>
-	bool FReadPhysicsObjectInterface<Id>::LineTrace(TArrayView<FPhysicsObjectHandle> InObjects, const FVector& WorldStart, const FVector& WorldEnd, bool bTraceComplex, ChaosInterface::FRaycastHit& OutBestHit)
+	bool FReadPhysicsObjectInterface<Id>::LineTrace(TArrayView<const FConstPhysicsObjectHandle> InObjects, const FVector& WorldStart, const FVector& WorldEnd, bool bTraceComplex, ChaosInterface::FRaycastHit& OutBestHit)
 	{
 		bool bHit = false;
 		OutBestHit.Distance = TNumericLimits<float>::Max();
@@ -711,7 +711,7 @@ namespace Chaos
 
 		FTransform BestWorldTM = FTransform::Identity;
 
-		for (FPhysicsObjectHandle Object : InObjects)
+		for (const FConstPhysicsObjectHandle Object : InObjects)
 		{
 			const FTransform WorldTM = GetTransform(Object);
 			const FVector LocalStart = WorldTM.InverseTransformPositionNoScale(WorldStart);
@@ -719,7 +719,7 @@ namespace Chaos
 
 			VisitEveryShape(
 				{ &Object, 1 },
-				[this, &bHit, &WorldTM, &LocalStart, &LocalDelta, &Delta, DeltaMag, &BestWorldTM, bTraceComplex, &OutBestHit](FPhysicsObjectHandle IterObject, FPerShapeData* Shape)
+				[this, &bHit, &WorldTM, &LocalStart, &LocalDelta, &Delta, DeltaMag, &BestWorldTM, bTraceComplex, &OutBestHit](const FConstPhysicsObjectHandle IterObject, FPerShapeData* Shape)
 				{
 					check(Shape);
 
@@ -777,16 +777,16 @@ namespace Chaos
 	}
 
 	template<EThreadContext Id>
-	bool FReadPhysicsObjectInterface<Id>::ShapeOverlap(TArrayView<FPhysicsObjectHandle> InObjects, const Chaos::FImplicitObject& InGeom, const FTransform& GeomTransform, TArray<ChaosInterface::FOverlapHit>& OutOverlaps)
+	bool FReadPhysicsObjectInterface<Id>::ShapeOverlap(TArrayView<const FConstPhysicsObjectHandle> InObjects, const Chaos::FImplicitObject& InGeom, const FTransform& GeomTransform, TArray<ChaosInterface::FOverlapHit>& OutOverlaps)
 	{
 		bool bHasOverlap = false;
-		for (FPhysicsObjectHandle Object : InObjects)
+		for (const FConstPhysicsObjectHandle Object : InObjects)
 		{
 			const FTransform WorldTM = GetTransform(Object);
 
 			VisitEveryShape(
 				{ &Object, 1 },
-				[this, &bHasOverlap, &WorldTM, &InGeom, &GeomTransform, &OutOverlaps](FPhysicsObjectHandle IterObject, FPerShapeData* Shape)
+				[this, &bHasOverlap, &WorldTM, &InGeom, &GeomTransform, &OutOverlaps](const FConstPhysicsObjectHandle IterObject, FPerShapeData* Shape)
 				{
 					check(Shape);
 					const bool bOverlap = Chaos::Utilities::CastHelper(
@@ -818,7 +818,7 @@ namespace Chaos
 	}
 
 	template<EThreadContext Id>
-	bool FReadPhysicsObjectInterface<Id>::ShapeSweep(TArrayView<FPhysicsObjectHandle> InObjects, const Chaos::FImplicitObject& InGeom, const FTransform& StartTM, const FVector& EndPos, bool bSweepComplex, ChaosInterface::FSweepHit& OutBestHit)
+	bool FReadPhysicsObjectInterface<Id>::ShapeSweep(TArrayView<const FConstPhysicsObjectHandle> InObjects, const Chaos::FImplicitObject& InGeom, const FTransform& StartTM, const FVector& EndPos, bool bSweepComplex, ChaosInterface::FSweepHit& OutBestHit)
 	{
 		bool bHit = false;
 		const FVector StartPos = StartTM.GetTranslation();
@@ -830,13 +830,13 @@ namespace Chaos
 		}
 		const FVec3 Dir = Delta / DeltaMag;
 
-		for (FPhysicsObjectHandle Object : InObjects)
+		for (const FConstPhysicsObjectHandle Object : InObjects)
 		{
 			const FTransform WorldTM = GetTransform(Object);
 
 			VisitEveryShape(
 				{ &Object, 1 },
-				[&WorldTM, &InGeom, &StartTM, &bHit, &Delta, DeltaMag, &Dir, &OutBestHit, bSweepComplex](FPhysicsObjectHandle IterObject, FPerShapeData* Shape)
+				[&WorldTM, &InGeom, &StartTM, &bHit, &Delta, DeltaMag, &Dir, &OutBestHit, bSweepComplex](const FConstPhysicsObjectHandle IterObject, FPerShapeData* Shape)
 				{
 					check(Shape);
 
@@ -890,9 +890,9 @@ namespace Chaos
 	}
 
 	template<EThreadContext Id>
-	void FWritePhysicsObjectInterface<Id>::PutToSleep(TArrayView<FPhysicsObjectHandle> InObjects)
+	void FWritePhysicsObjectInterface<Id>::PutToSleep(TArrayView<const FPhysicsObjectHandle> InObjects)
 	{
-		for (FPhysicsObjectHandle Object : InObjects)
+		for (const FPhysicsObjectHandle Object : InObjects)
 		{
 			if (!Object)
 			{
@@ -908,9 +908,9 @@ namespace Chaos
 	}
 
 	template<EThreadContext Id>
-	void FWritePhysicsObjectInterface<Id>::WakeUp(TArrayView<FPhysicsObjectHandle> InObjects)
+	void FWritePhysicsObjectInterface<Id>::WakeUp(TArrayView<const FPhysicsObjectHandle> InObjects)
 	{
-		for (FPhysicsObjectHandle Object : InObjects)
+		for (const FPhysicsObjectHandle Object : InObjects)
 		{
 			if (!Object)
 			{
@@ -936,9 +936,9 @@ namespace Chaos
 	}
 
 	template<EThreadContext Id>
-	void FWritePhysicsObjectInterface<Id>::AddForce(TArrayView<FPhysicsObjectHandle> InObjects, const FVector& Force, bool bInvalidate)
+	void FWritePhysicsObjectInterface<Id>::AddForce(TArrayView<const FPhysicsObjectHandle> InObjects, const FVector& Force, bool bInvalidate)
 	{
-		for (FPhysicsObjectHandle Object : InObjects)
+		for (const FPhysicsObjectHandle Object : InObjects)
 		{
 			if (!Object)
 			{
@@ -964,9 +964,9 @@ namespace Chaos
 	}
 
 	template<EThreadContext Id>
-	void FWritePhysicsObjectInterface<Id>::AddTorque(TArrayView<FPhysicsObjectHandle> InObjects, const FVector& Torque, bool bInvalidate)
+	void FWritePhysicsObjectInterface<Id>::AddTorque(TArrayView<const FPhysicsObjectHandle> InObjects, const FVector& Torque, bool bInvalidate)
 	{
-		for (FPhysicsObjectHandle Object : InObjects)
+		for (const FPhysicsObjectHandle Object : InObjects)
 		{
 			if (!Object)
 			{
@@ -992,9 +992,9 @@ namespace Chaos
 	}
 
 	template<EThreadContext Id>
-	void FWritePhysicsObjectInterface<Id>::UpdateShapeCollisionFlags(TArrayView<FPhysicsObjectHandle> InObjects, bool bSimCollision, bool bQueryCollision)
+	void FWritePhysicsObjectInterface<Id>::UpdateShapeCollisionFlags(TArrayView<const FPhysicsObjectHandle> InObjects, bool bSimCollision, bool bQueryCollision)
 	{
-		for (FPhysicsObjectHandle Object : InObjects)
+		for (const FPhysicsObjectHandle Object : InObjects)
 		{
 			if (!Object)
 			{
@@ -1015,9 +1015,9 @@ namespace Chaos
 	}
 
 	template<EThreadContext Id>
-	void FWritePhysicsObjectInterface<Id>::UpdateShapeFilterData(TArrayView<FPhysicsObjectHandle> InObjects, const FCollisionFilterData& QueryData, const FCollisionFilterData& SimData)
+	void FWritePhysicsObjectInterface<Id>::UpdateShapeFilterData(TArrayView<const FPhysicsObjectHandle> InObjects, const FCollisionFilterData& QueryData, const FCollisionFilterData& SimData)
 	{
-		for (FPhysicsObjectHandle Object : InObjects)
+		for (const FPhysicsObjectHandle Object : InObjects)
 		{
 			if (!Object)
 			{
@@ -1035,7 +1035,7 @@ namespace Chaos
 		}
 	}
 
-	void FPhysicsObjectInterface::SetName(FPhysicsObjectHandle Object, const FName& InName)
+	void FPhysicsObjectInterface::SetName(const FPhysicsObjectHandle Object, const FName& InName)
 	{
 		if (!Object)
 		{
@@ -1045,7 +1045,7 @@ namespace Chaos
 		Object->SetName(InName);
 	}
 
-	FName FPhysicsObjectInterface::GetName(FPhysicsObjectHandle Object)
+	FName FPhysicsObjectInterface::GetName(const FConstPhysicsObjectHandle Object)
 	{
 		if (!Object)
 		{
@@ -1055,7 +1055,7 @@ namespace Chaos
 		return Object->GetBodyName();
 	}
 
-	void FPhysicsObjectInterface::SetId(FPhysicsObjectHandle Object, int32 InId)
+	void FPhysicsObjectInterface::SetId(const FPhysicsObjectHandle Object, int32 InId)
 	{
 		if (!Object)
 		{
@@ -1065,7 +1065,7 @@ namespace Chaos
 		Object->SetBodyIndex(InId);
 	}
 
-	int32 FPhysicsObjectInterface::GetId(FPhysicsObjectHandle Object)
+	int32 FPhysicsObjectInterface::GetId(const FConstPhysicsObjectHandle Object)
 	{
 		if (!Object)
 		{
@@ -1075,10 +1075,10 @@ namespace Chaos
 		return Object->GetBodyIndex();
 	}
 
-	FPBDRigidsSolver* FPhysicsObjectInterface::GetSolver(TArrayView<FPhysicsObjectHandle> InObjects)
+	FPBDRigidsSolver* FPhysicsObjectInterface::GetSolver(TArrayView<const FConstPhysicsObjectHandle> InObjects)
 	{
 		FPBDRigidsSolver* RetSolver = nullptr;
-		for (FPhysicsObjectHandle Object : InObjects)
+		for (const FConstPhysicsObjectHandle Object : InObjects)
 		{
 			if (!Object)
 			{
@@ -1086,7 +1086,7 @@ namespace Chaos
 			}
 
 			FPBDRigidsSolver* Solver = nullptr;
-			if (IPhysicsProxyBase* Proxy = Object->PhysicsProxy())
+			if (const IPhysicsProxyBase* Proxy = Object->PhysicsProxy())
 			{
 				Solver = Proxy->GetSolver<FPBDRigidsSolver>();
 			}
@@ -1107,17 +1107,17 @@ namespace Chaos
 		return RetSolver;
 	}
 
-	IPhysicsProxyBase* FPhysicsObjectInterface::GetProxy(TArrayView<FPhysicsObjectHandle> InObjects)
+	IPhysicsProxyBase* FPhysicsObjectInterface::GetProxy(TArrayView<const FConstPhysicsObjectHandle> InObjects)
 	{
 		IPhysicsProxyBase* RetProxy = nullptr;
-		for (FPhysicsObjectHandle Object : InObjects)
+		for (const FConstPhysicsObjectHandle Object : InObjects)
 		{
 			if (!Object)
 			{
 				continue;
 			}
 
-			IPhysicsProxyBase* Proxy = Object->PhysicsProxy();
+			IPhysicsProxyBase* Proxy = const_cast<IPhysicsProxyBase*>(Object->PhysicsProxy());
 			if (!Proxy)
 			{
 				return nullptr;
diff --git a/Engine/Source/Runtime/Experimental/Chaos/Private/Chaos/PhysicsObjectInternal.h b/Engine/Source/Runtime/Experimental/Chaos/Private/Chaos/PhysicsObjectInternal.h
--- a/Engine/Source/Runtime/Experimental/Chaos/Private/Chaos/PhysicsObjectInternal.h
+++ b/Engine/Source/Runtime/Experimental/Chaos/Private/Chaos/PhysicsObjectInternal.h
@@ -23,7 +23,7 @@ namespace Chaos
 		const FName& GetBodyName() const { return BodyName; }
 
 		template<EThreadContext Id>
-		EObjectStateType ObjectState()
+		EObjectStateType ObjectState() const
 		{
 			TThreadParticle<Id>* Particle = GetParticle<Id>();
 			if (!Particle)
@@ -42,12 +42,12 @@ namespace Chaos
 		}
 
 		template<EThreadContext Id>
-		FPhysicsObjectHandle GetRootObject()
+		FPhysicsObjectHandle GetRootObject() const
 		{
 			EPhysicsProxyType ProxyType = Proxy->GetType();
 
 			TThreadParticle<Id>* Particle = GetParticle<Id>();
-			FPhysicsObjectHandle CurrentObject = this;
+			FPhysicsObjectHandle CurrentObject = const_cast<FPhysicsObjectHandle>(this);
 			FPhysicsObjectHandle CurrentParent = GetParentObject<Id>();
 			while (CurrentParent && IsParticleDisabled<Id>(Particle))
 			{
@@ -101,7 +101,7 @@ namespace Chaos
 		}
 
 		template <EThreadContext Id>
-		TThreadParticle<Id>* GetRootParticle()
+		TThreadParticle<Id>* GetRootParticle() const
 		{
 			FPhysicsObjectHandle Root = GetRootObject<Id>();
 			if (!Root)
@@ -112,7 +112,7 @@ namespace Chaos
 		}
 
 		template <EThreadContext Id>
-		TThreadParticle<Id>* GetParticle()
+		TThreadParticle<Id>* GetParticle() const
 		{
 			if (Proxy)
 			{
diff --git a/Engine/Source/Runtime/Experimental/Chaos/Public/Chaos/PhysicsObject.h b/Engine/Source/Runtime/Experimental/Chaos/Public/Chaos/PhysicsObject.h
--- a/Engine/Source/Runtime/Experimental/Chaos/Public/Chaos/PhysicsObject.h
+++ b/Engine/Source/Runtime/Experimental/Chaos/Public/Chaos/PhysicsObject.h
@@ -16,6 +16,7 @@ namespace Chaos
 	 */
 	struct FPhysicsObject;
 	using FPhysicsObjectHandle = FPhysicsObject*;
+	using FConstPhysicsObjectHandle = const FPhysicsObject*;
 
 	struct CHAOS_API FPhysicsObjectDeleter
 	{
diff --git a/Engine/Source/Runtime/Experimental/Chaos/Public/Chaos/PhysicsObjectInterface.h b/Engine/Source/Runtime/Experimental/Chaos/Public/Chaos/PhysicsObjectInterface.h
--- a/Engine/Source/Runtime/Experimental/Chaos/Public/Chaos/PhysicsObjectInterface.h
+++ b/Engine/Source/Runtime/Experimental/Chaos/Public/Chaos/PhysicsObjectInterface.h
@@ -55,55 +55,55 @@ namespace Chaos
 	class CHAOS_API FReadPhysicsObjectInterface
 	{
 	public:
-		FPhysicsObjectHandle GetRootObject(FPhysicsObjectHandle Object);
-		bool HasChildren(FPhysicsObjectHandle Object);
+		FPhysicsObjectHandle GetRootObject(const FConstPhysicsObjectHandle Object);
+		bool HasChildren(const FConstPhysicsObjectHandle Object);
 
-		FTransform GetTransform(FPhysicsObjectHandle Object);
-		FVector GetX(FPhysicsObjectHandle Object);
-		FVector GetCoM(FPhysicsObjectHandle Object);
-		FVector GetWorldCoM(FPhysicsObjectHandle Object);
-		FQuat GetR(FPhysicsObjectHandle Object);
-		FSpatialAccelerationIdx GetSpatialIndex(FPhysicsObjectHandle Object);
+		FTransform GetTransform(const FConstPhysicsObjectHandle Object);
+		FVector GetX(const FConstPhysicsObjectHandle Object);
+		FVector GetCoM(const FConstPhysicsObjectHandle Object);
+		FVector GetWorldCoM(const FConstPhysicsObjectHandle Object);
+		FQuat GetR(const FConstPhysicsObjectHandle Object);
+		FSpatialAccelerationIdx GetSpatialIndex(const FConstPhysicsObjectHandle Object);
 
-		TArray<TThreadParticle<Id>*> GetAllParticles(TArrayView<FPhysicsObjectHandle> InObjects);
-		TArray<TThreadRigidParticle<Id>*> GetAllRigidParticles(TArrayView<FPhysicsObjectHandle> InObjects);
-		TArray<FPerShapeData*> GetAllShapes(TArrayView<FPhysicsObjectHandle> InObjects);
+		TArray<TThreadParticle<Id>*> GetAllParticles(TArrayView<const FConstPhysicsObjectHandle> InObjects);
+		TArray<TThreadRigidParticle<Id>*> GetAllRigidParticles(TArrayView<const FConstPhysicsObjectHandle> InObjects);
+		TArray<FPerShapeData*> GetAllShapes(TArrayView<const FConstPhysicsObjectHandle> InObjects);
 
 		// Returns true if a shape is found and we can stop iterating.
-		void VisitEveryShape(TArrayView<FPhysicsObjectHandle> InObjects, TFunctionRef<bool(FPhysicsObjectHandle, FPerShapeData*)> Lambda);
+		void VisitEveryShape(TArrayView<const FConstPhysicsObjectHandle> InObjects, TFunctionRef<bool(const FConstPhysicsObjectHandle, FPerShapeData*)> Lambda);
 
 		UE_DEPRECATED(5.3, "GetPhysicsObjectOverlap has been deprecated. Please use the function for the specific overlap metric you wish to compute instead.")
-		bool GetPhysicsObjectOverlap(FPhysicsObjectHandle ObjectA, FPhysicsObjectHandle ObjectB, bool bTraceComplex, Chaos::FOverlapInfo& OutOverlap);
+		bool GetPhysicsObjectOverlap(const FConstPhysicsObjectHandle ObjectA, const FConstPhysicsObjectHandle ObjectB, bool bTraceComplex, Chaos::FOverlapInfo& OutOverlap);
 
 		UE_DEPRECATED(5.3, "GetPhysicsObjectOverlapWithTransform has been deprecated. Please use the function for the specific overlap metric you wish to compute instead.")
-		bool GetPhysicsObjectOverlapWithTransform(FPhysicsObjectHandle ObjectA, const FTransform& InTransformA, FPhysicsObjectHandle ObjectB, const FTransform& InTransformB, bool bTraceComplex, Chaos::FOverlapInfo& OutOverlap);
+		bool GetPhysicsObjectOverlapWithTransform(const FConstPhysicsObjectHandle ObjectA, const FTransform& InTransformA, const FConstPhysicsObjectHandle ObjectB, const FTransform& InTransformB, bool bTraceComplex, Chaos::FOverlapInfo& OutOverlap);
 
 		// This function will not compute any overlap heuristic.
-		bool PhysicsObjectOverlap(FPhysicsObjectHandle ObjectA, const FTransform& InTransformA, FPhysicsObjectHandle ObjectB, const FTransform& InTransformB, bool bTraceComplex);
+		bool PhysicsObjectOverlap(const FConstPhysicsObjectHandle ObjectA, const FTransform& InTransformA, const FConstPhysicsObjectHandle ObjectB, const FTransform& InTransformB, bool bTraceComplex);
 
 		// This function does the same as GetPhysicsObjectOverlap but also computes the MTD metric.
-		bool PhysicsObjectOverlapWithMTD(FPhysicsObjectHandle ObjectA, const FTransform& InTransformA, FPhysicsObjectHandle ObjectB, const FTransform& InTransformB, bool bTraceComplex, FMTDInfo& OutMTD);
+		bool PhysicsObjectOverlapWithMTD(const FConstPhysicsObjectHandle ObjectA, const FTransform& InTransformA, const FConstPhysicsObjectHandle ObjectB, const FTransform& InTransformB, bool bTraceComplex, FMTDInfo& OutMTD);
 
 		// This function does the same as GetPhysicsObjectOverlap but also computes the AABB overlap metric.
-		bool PhysicsObjectOverlapWithAABB(FPhysicsObjectHandle ObjectA, const FTransform& InTransformA, FPhysicsObjectHandle ObjectB, const FTransform& InTransformB, bool bTraceComplex, const FVector& Tolerance, FBox& OutOverlap);
-		bool PhysicsObjectOverlapWithAABBSize(FPhysicsObjectHandle ObjectA, const FTransform& InTransformA, FPhysicsObjectHandle ObjectB, const FTransform& InTransformB, bool bTraceComplex, const FVector& Tolerance, FVector& OutOverlapSize);
-
-		bool AreAllValid(TArrayView<FPhysicsObjectHandle> InObjects);
-		bool AreAllKinematic(TArrayView<FPhysicsObjectHandle> InObjects);
-		bool AreAllSleeping(TArrayView<FPhysicsObjectHandle> InObjects);
-		bool AreAllRigidBody(TArrayView<FPhysicsObjectHandle> InObjects);
-		bool AreAllDynamic(TArrayView<FPhysicsObjectHandle> InObjects);
-		bool AreAllDisabled(TArrayView<FPhysicsObjectHandle> InObjects);
-		bool AreAllShapesQueryEnabled(TArrayView<FPhysicsObjectHandle> InObjects);
-		float GetMass(TArrayView<FPhysicsObjectHandle> InObjects);
-		FBox GetBounds(TArrayView<FPhysicsObjectHandle> InObjects);
-		FBox GetWorldBounds(TArrayView<FPhysicsObjectHandle> InObjects);
-		FClosestPhysicsObjectResult GetClosestPhysicsBodyFromLocation(TArrayView<FPhysicsObjectHandle> InObjects, const FVector& WorldLocation);
-		FAccelerationStructureHandle CreateAccelerationStructureHandle(FPhysicsObjectHandle Handle);
-
-		bool LineTrace(TArrayView<FPhysicsObjectHandle> InObjects, const FVector& WorldStart, const FVector& WorldEnd, bool bTraceComplex, ChaosInterface::FRaycastHit& OutBestHit);
-		bool ShapeOverlap(TArrayView<FPhysicsObjectHandle> InObjects, const Chaos::FImplicitObject& InGeom, const FTransform& GeomTransform, TArray<ChaosInterface::FOverlapHit>& OutOverlaps);
-		bool ShapeSweep(TArrayView<FPhysicsObjectHandle> InObjects, const Chaos::FImplicitObject& InGeom, const FTransform& StartTM, const FVector& EndPos, bool bSweepComplex, ChaosInterface::FSweepHit& OutBestHit);
+		bool PhysicsObjectOverlapWithAABB(const FConstPhysicsObjectHandle ObjectA, const FTransform& InTransformA, const FConstPhysicsObjectHandle ObjectB, const FTransform& InTransformB, bool bTraceComplex, const FVector& Tolerance, FBox& OutOverlap);
+		bool PhysicsObjectOverlapWithAABBSize(const FConstPhysicsObjectHandle ObjectA, const FTransform& InTransformA, const FConstPhysicsObjectHandle ObjectB, const FTransform& InTransformB, bool bTraceComplex, const FVector& Tolerance, FVector& OutOverlapSize);
+
+		bool AreAllValid(TArrayView<const FConstPhysicsObjectHandle> InObjects);
+		bool AreAllKinematic(TArrayView<const FConstPhysicsObjectHandle> InObjects);
+		bool AreAllSleeping(TArrayView<const FConstPhysicsObjectHandle> InObjects);
+		bool AreAllRigidBody(TArrayView<const FConstPhysicsObjectHandle> InObjects);
+		bool AreAllDynamic(TArrayView<const FConstPhysicsObjectHandle> InObjects);
+		bool AreAllDisabled(TArrayView<const FConstPhysicsObjectHandle> InObjects);
+		bool AreAllShapesQueryEnabled(TArrayView<const FConstPhysicsObjectHandle> InObjects);
+		float GetMass(TArrayView<const FConstPhysicsObjectHandle> InObjects);
+		FBox GetBounds(TArrayView<const FConstPhysicsObjectHandle> InObjects);
+		FBox GetWorldBounds(TArrayView<const FConstPhysicsObjectHandle> InObjects);
+		FClosestPhysicsObjectResult GetClosestPhysicsBodyFromLocation(TArrayView<const FConstPhysicsObjectHandle> InObjects, const FVector& WorldLocation);
+		FAccelerationStructureHandle CreateAccelerationStructureHandle(const FConstPhysicsObjectHandle Handle);
+
+		bool LineTrace(TArrayView<const FConstPhysicsObjectHandle> InObjects, const FVector& WorldStart, const FVector& WorldEnd, bool bTraceComplex, ChaosInterface::FRaycastHit& OutBestHit);
+		bool ShapeOverlap(TArrayView<const FConstPhysicsObjectHandle> InObjects, const Chaos::FImplicitObject& InGeom, const FTransform& GeomTransform, TArray<ChaosInterface::FOverlapHit>& OutOverlaps);
+		bool ShapeSweep(TArrayView<const FConstPhysicsObjectHandle> InObjects, const Chaos::FImplicitObject& InGeom, const FTransform& StartTM, const FVector& EndPos, bool bSweepComplex, ChaosInterface::FSweepHit& OutBestHit);
 
 		friend class FPhysicsObjectInterface;
 	protected:
@@ -119,7 +119,7 @@ namespace Chaos
 		/**
 		 * For every pair of shapes that overlap, allows the caller to perform some computation. If additional pairs of shapes need to be examined, the input TFunction should return true.
 		 */
-		bool PairwiseShapeOverlapHelper(FPhysicsObjectHandle ObjectA, const FTransform& InTransformA, FPhysicsObjectHandle ObjectB, const FTransform& InTransformB, bool bTraceComplex, bool bComputeMTD, const FVector& Tolerance, const TFunction<bool(const FShapeOverlapData&, const FShapeOverlapData&, const FMTDInfo&)>& Lambda);
+		bool PairwiseShapeOverlapHelper(const FConstPhysicsObjectHandle ObjectA, const FTransform& InTransformA, const FConstPhysicsObjectHandle ObjectB, const FTransform& InTransformB, bool bTraceComplex, bool bComputeMTD, const FVector& Tolerance, const TFunction<bool(const FShapeOverlapData&, const FShapeOverlapData&, const FMTDInfo&)>& Lambda);
 	};
 
 	using FReadPhysicsObjectInterface_External = FReadPhysicsObjectInterface<EThreadContext::External>;
@@ -132,23 +132,23 @@ namespace Chaos
 	class CHAOS_API FWritePhysicsObjectInterface: public FReadPhysicsObjectInterface<Id>
 	{
 	public:
-		void PutToSleep(TArrayView<FPhysicsObjectHandle> InObjects);
-		void WakeUp(TArrayView<FPhysicsObjectHandle> InObjects);
-		void AddForce(TArrayView<FPhysicsObjectHandle> InObjects, const FVector& Force, bool bInvalidate);
-		void AddTorque(TArrayView<FPhysicsObjectHandle> InObjects, const FVector& Torque, bool bInvalidate);
+		void PutToSleep(TArrayView<const FPhysicsObjectHandle> InObjects);
+		void WakeUp(TArrayView<const FPhysicsObjectHandle> InObjects);
+		void AddForce(TArrayView<const FPhysicsObjectHandle> InObjects, const FVector& Force, bool bInvalidate);
+		void AddTorque(TArrayView<const FPhysicsObjectHandle> InObjects, const FVector& Torque, bool bInvalidate);
 		
-		void UpdateShapeCollisionFlags(TArrayView<FPhysicsObjectHandle> InObjects, bool bSimCollision, bool bQueryCollision);
-		void UpdateShapeFilterData(TArrayView<FPhysicsObjectHandle> InObjects, const FCollisionFilterData& QueryData, const FCollisionFilterData& SimData);
+		void UpdateShapeCollisionFlags(TArrayView<const FPhysicsObjectHandle> InObjects, bool bSimCollision, bool bQueryCollision);
+		void UpdateShapeFilterData(TArrayView<const FPhysicsObjectHandle> InObjects, const FCollisionFilterData& QueryData, const FCollisionFilterData& SimData);
 
 		template<typename TPayloadType, typename T, int d>
-		void AddToSpatialAcceleration(TArrayView<FPhysicsObjectHandle> InObjects, ISpatialAcceleration<TPayloadType, T, d>* SpatialAcceleration)
+		void AddToSpatialAcceleration(TArrayView<const FPhysicsObjectHandle> InObjects, ISpatialAcceleration<TPayloadType, T, d>* SpatialAcceleration)
 		{
 			if (!SpatialAcceleration)
 			{
 				return;
 			}
 
-			for (FPhysicsObjectHandle Handle : InObjects)
+			for (const FConstPhysicsObjectHandle Handle : InObjects)
 			{
 				const FBox WorldBounds = this->GetWorldBounds({ &Handle, 1 });
 				const FAABB3 ChaosWorldBounds{ WorldBounds.Min, WorldBounds.Max };
@@ -158,14 +158,14 @@ namespace Chaos
 		}
 
 		template<typename TPayloadType, typename T, int d>
-		void RemoveFromSpatialAcceleration(TArrayView<FPhysicsObjectHandle> InObjects, ISpatialAcceleration<TPayloadType, T, d>* SpatialAcceleration)
+		void RemoveFromSpatialAcceleration(TArrayView<const FPhysicsObjectHandle> InObjects, ISpatialAcceleration<TPayloadType, T, d>* SpatialAcceleration)
 		{
 			if (!SpatialAcceleration)
 			{
 				return;
 			}
 
-			for (FPhysicsObjectHandle Handle : InObjects)
+			for (const FConstPhysicsObjectHandle Handle : InObjects)
 			{
 				FAccelerationStructureHandle AccelerationHandle = this->CreateAccelerationStructureHandle(Handle);
 				SpatialAcceleration->RemoveElementFrom(AccelerationHandle, this->GetSpatialIndex(Handle));
@@ -188,14 +188,14 @@ namespace Chaos
 	class CHAOS_API FPhysicsObjectInterface
 	{
 	public:
-		static void SetName(FPhysicsObjectHandle Object, const FName& InName);
-		static FName GetName(FPhysicsObjectHandle Object);
+		static void SetName(const FPhysicsObjectHandle Object, const FName& InName);
+		static FName GetName(const FConstPhysicsObjectHandle Object);
 
-		static void SetId(FPhysicsObjectHandle Object, int32 InId);
-		static int32 GetId(FPhysicsObjectHandle Object);
+		static void SetId(const FPhysicsObjectHandle Object, int32 InId);
+		static int32 GetId(const FConstPhysicsObjectHandle Object);
 
-		static FPBDRigidsSolver* GetSolver(TArrayView<FPhysicsObjectHandle> InObjects);
-		static IPhysicsProxyBase* GetProxy(TArrayView<FPhysicsObjectHandle> InObjects);
+		static FPBDRigidsSolver* GetSolver(TArrayView<const FConstPhysicsObjectHandle> InObjects);
+		static IPhysicsProxyBase* GetProxy(TArrayView<const FConstPhysicsObjectHandle> InObjects);
 
 	protected:
 		// This function should not be called without an appropriate read-lock on the relevant scene.
diff --git a/Engine/Source/Runtime/PhysicsCore/Private/PhysicsObjectPhysicsCoreInterface.cpp b/Engine/Source/Runtime/PhysicsCore/Private/PhysicsObjectPhysicsCoreInterface.cpp
--- a/Engine/Source/Runtime/PhysicsCore/Private/PhysicsObjectPhysicsCoreInterface.cpp
+++ b/Engine/Source/Runtime/PhysicsCore/Private/PhysicsObjectPhysicsCoreInterface.cpp
@@ -5,7 +5,7 @@
 #include "Chaos/PhysicsObjectInterface.h"
 #include "PBDRigidsSolver.h"
 
-FChaosScene* PhysicsObjectPhysicsCoreInterface::GetScene(TArrayView<Chaos::FPhysicsObjectHandle> InObjects)
+FChaosScene* PhysicsObjectPhysicsCoreInterface::GetScene(TArrayView<const Chaos::FConstPhysicsObjectHandle> InObjects)
 {
 	Chaos::FPBDRigidsSolver* Solver = Chaos::FPhysicsObjectInterface::GetSolver(InObjects);
 	if (!Solver)
diff --git a/Engine/Source/Runtime/PhysicsCore/Public/PhysicsObjectPhysicsCoreInterface.h b/Engine/Source/Runtime/PhysicsCore/Public/PhysicsObjectPhysicsCoreInterface.h
--- a/Engine/Source/Runtime/PhysicsCore/Public/PhysicsObjectPhysicsCoreInterface.h
+++ b/Engine/Source/Runtime/PhysicsCore/Public/PhysicsObjectPhysicsCoreInterface.h
@@ -10,5 +10,5 @@ class FChaosScene;
 class PHYSICSCORE_API PhysicsObjectPhysicsCoreInterface
 {
 public:
-	static FChaosScene* GetScene(TArrayView<Chaos::FPhysicsObjectHandle> InObjects);
+	static FChaosScene* GetScene(TArrayView<const Chaos::FConstPhysicsObjectHandle> InObjects);
 };
\ No newline at end of file
