dataflow : fix serialization to be able to laod nodes that have added outputs causing the index to be shifted during deserialization
This is hack to avoid using versioning for now, we should fix it properly in a near future by seperating the input and output serialization

diff --git a/Engine/Source/Runtime/Experimental/Dataflow/Core/Private/Dataflow/DataflowGraph.cpp b/Engine/Source/Runtime/Experimental/Dataflow/Core/Private/Dataflow/DataflowGraph.cpp
--- a/Engine/Source/Runtime/Experimental/Dataflow/Core/Private/Dataflow/DataflowGraph.cpp
+++ b/Engine/Source/Runtime/Experimental/Dataflow/Core/Private/Dataflow/DataflowGraph.cpp
@@ -178,19 +178,48 @@ namespace Dataflow
 					ensure(!NodeGuidMap.Contains(ArGuid));
 					NodeGuidMap.Add(ArGuid, Node);
 
-					int ArNumInner;
-					Ar << ArNumInner;
-					TArray< FDataflowConnection* > IO;
-					IO.Append(Node->GetOutputs());
-					IO.Append(Node->GetInputs());
-					for (int Cdx = 0; Cdx < ArNumInner; Cdx++)
+					int ArNumInputsOutputs;
+					Ar << ArNumInputsOutputs;
+					TArray< FDataflowConnection* > InputsOutputs;
+					InputsOutputs.Append(Node->GetOutputs());
+					InputsOutputs.Append(Node->GetInputs());
+
+					// skip offset is use to correct the mismatch of outputs have been added
+					int SkipOffset = 0;
+					for (int ConnectionIndex = 0; ConnectionIndex < ArNumInputsOutputs; ConnectionIndex++)
 					{
 						Ar << ArGuid << ArType << ArName;
-						if (Cdx < IO.Num())
+
+						int AdjustedConnectionIndex = ConnectionIndex + SkipOffset;
+						if (InputsOutputs.IsValidIndex(AdjustedConnectionIndex))
 						{
-							IO[Cdx]->SetGuid(ArGuid);
-							ensure(!ConnectionGuidMap.Contains(ArGuid));
-							ConnectionGuidMap.Add(ArGuid, IO[Cdx]);
+							FDataflowConnection* Connection = InputsOutputs[AdjustedConnectionIndex];
+
+							// if the name does not match this means the node has changed since the last serialization 
+							// ( added outputs for example that shift the index )
+							// in that case we try to recover by finding the next good node
+							// note we cannot just find by name as some nodes have inputs and outputs named the same 
+							// todo: implement a better way to serialize inputs and outputs seperately to avoid this case
+							while (Connection && Connection->GetName() != ArName)
+							{
+								SkipOffset++;
+								AdjustedConnectionIndex = ConnectionIndex + SkipOffset;
+								if (InputsOutputs.IsValidIndex(AdjustedConnectionIndex))
+								{
+									Connection = InputsOutputs[AdjustedConnectionIndex];
+								}
+								else
+								{
+									Connection = nullptr;
+								}
+							}
+							if (Connection)
+							{
+								check(Connection->GetType() == ArType);
+								Connection->SetGuid(ArGuid);
+								ensure(!ConnectionGuidMap.Contains(ArGuid));
+								ConnectionGuidMap.Add(ArGuid, Connection);
+							}
 						}
 					}
 
