Iris
- Hooked up new delegate in NetSerializerDelegates that can be conditionally bound for
serializers that needs to update cached data based on loaded modules.

diff --git a/Engine/Plugins/Runtime/GameplayAbilities/Source/GameplayAbilities/Private/AbilitySystemGlobals.cpp b/Engine/Plugins/Runtime/GameplayAbilities/Source/GameplayAbilities/Private/AbilitySystemGlobals.cpp
--- a/Engine/Plugins/Runtime/GameplayAbilities/Source/GameplayAbilities/Private/AbilitySystemGlobals.cpp
+++ b/Engine/Plugins/Runtime/GameplayAbilities/Source/GameplayAbilities/Private/AbilitySystemGlobals.cpp
@@ -265,10 +265,6 @@ void UAbilitySystemGlobals::InitTargetDataScriptStructCache()
 {
 	TargetDataStructCache.InitForType(FGameplayAbilityTargetData::StaticStruct());
 	EffectContextStructCache.InitForType(FGameplayEffectContext::StaticStruct());
-#if UE_WITH_IRIS
-	InitGameplayAbilityTargetDataHandleNetSerializerTypeCache();
-	InitGameplayEffectContextHandleNetSerializerTypeCache();
-#endif // UE_WITH_IRIS
 }
 
 // --------------------------------------------------------------------
diff --git a/Engine/Plugins/Runtime/GameplayAbilities/Source/GameplayAbilities/Private/Serialization/GameplayAbilityTargetDataHandleNetSerializer.cpp b/Engine/Plugins/Runtime/GameplayAbilities/Source/GameplayAbilities/Private/Serialization/GameplayAbilityTargetDataHandleNetSerializer.cpp
--- a/Engine/Plugins/Runtime/GameplayAbilities/Source/GameplayAbilities/Private/Serialization/GameplayAbilityTargetDataHandleNetSerializer.cpp
+++ b/Engine/Plugins/Runtime/GameplayAbilities/Source/GameplayAbilities/Private/Serialization/GameplayAbilityTargetDataHandleNetSerializer.cpp
@@ -65,11 +65,13 @@ private:
 	class FNetSerializerRegistryDelegates final : private UE::Net::FNetSerializerRegistryDelegates
 	{
 	public:
+		FNetSerializerRegistryDelegates();
 		virtual ~FNetSerializerRegistryDelegates();
 
 	private:
 		virtual void OnPreFreezeNetSerializerRegistry() override;
 		virtual void OnPostFreezeNetSerializerRegistry() override;
+		virtual void OnLoadedModulesUpdated() override;
 	};
 
 	static FGameplayAbilityTargetDataHandleNetSerializer::FNetSerializerRegistryDelegates NetSerializerRegistryDelegates;
@@ -235,6 +237,11 @@ void FGameplayAbilityTargetDataHandleNetSerializer::SetArrayNum(FGameplayAbility
 static const FName PropertyNetSerializerRegistry_NAME_GameplayAbilityTargetDataHandle(TEXT("GameplayAbilityTargetDataHandle"));
 UE_NET_IMPLEMENT_NAMED_STRUCT_NETSERIALIZER_INFO(PropertyNetSerializerRegistry_NAME_GameplayAbilityTargetDataHandle, FGameplayAbilityTargetDataHandleNetSerializer);
 
+FGameplayAbilityTargetDataHandleNetSerializer::FNetSerializerRegistryDelegates::FNetSerializerRegistryDelegates()
+: UE::Net::FNetSerializerRegistryDelegates(EFlags::ShouldBindLoadedModulesUpdatedDelegate)
+{
+}
+
 FGameplayAbilityTargetDataHandleNetSerializer::FNetSerializerRegistryDelegates::~FNetSerializerRegistryDelegates()
 {
 	UE_NET_UNREGISTER_NETSERIALIZER_INFO(PropertyNetSerializerRegistry_NAME_GameplayAbilityTargetDataHandle);
@@ -248,6 +255,10 @@ void FGameplayAbilityTargetDataHandleNetSerializer::FNetSerializerRegistryDelega
 void FGameplayAbilityTargetDataHandleNetSerializer::FNetSerializerRegistryDelegates::OnPostFreezeNetSerializerRegistry()
 {
 	bIsPostFreezeCalled = true;
+}
+
+void FGameplayAbilityTargetDataHandleNetSerializer::FNetSerializerRegistryDelegates::OnLoadedModulesUpdated()
+{
 	InitGameplayAbilityTargetDataHandleNetSerializerTypeCache();
 }
 
diff --git a/Engine/Plugins/Runtime/GameplayAbilities/Source/GameplayAbilities/Private/Serialization/GameplayEffectContextHandleNetSerializer.cpp b/Engine/Plugins/Runtime/GameplayAbilities/Source/GameplayAbilities/Private/Serialization/GameplayEffectContextHandleNetSerializer.cpp
--- a/Engine/Plugins/Runtime/GameplayAbilities/Source/GameplayAbilities/Private/Serialization/GameplayEffectContextHandleNetSerializer.cpp
+++ b/Engine/Plugins/Runtime/GameplayAbilities/Source/GameplayAbilities/Private/Serialization/GameplayEffectContextHandleNetSerializer.cpp
@@ -36,11 +36,13 @@ private:
 	class FNetSerializerRegistryDelegates final : private UE::Net::FNetSerializerRegistryDelegates
 	{
 	public:
+		FNetSerializerRegistryDelegates();
 		virtual ~FNetSerializerRegistryDelegates();
 
 	private:
 		virtual void OnPreFreezeNetSerializerRegistry() override;
 		virtual void OnPostFreezeNetSerializerRegistry() override;
+		virtual void OnLoadedModulesUpdated() override;
 	};
 
 	static FGameplayEffectContextHandleNetSerializer::FNetSerializerRegistryDelegates NetSerializerRegistryDelegates;
@@ -66,6 +68,11 @@ void FGameplayEffectContextHandleNetSerializer::InitTypeCache()
 static const FName PropertyNetSerializerRegistry_NAME_GameplayEffectContextHandle("GameplayEffectContextHandle");
 UE_NET_IMPLEMENT_NAMED_STRUCT_NETSERIALIZER_INFO(PropertyNetSerializerRegistry_NAME_GameplayEffectContextHandle, FGameplayEffectContextHandleNetSerializer);
 
+FGameplayEffectContextHandleNetSerializer::FNetSerializerRegistryDelegates::FNetSerializerRegistryDelegates()
+: UE::Net::FNetSerializerRegistryDelegates(EFlags::ShouldBindLoadedModulesUpdatedDelegate)
+{
+}
+
 FGameplayEffectContextHandleNetSerializer::FNetSerializerRegistryDelegates::~FNetSerializerRegistryDelegates()
 {
 	UE_NET_UNREGISTER_NETSERIALIZER_INFO(PropertyNetSerializerRegistry_NAME_GameplayEffectContextHandle);
@@ -79,6 +86,10 @@ void FGameplayEffectContextHandleNetSerializer::FNetSerializerRegistryDelegates:
 void FGameplayEffectContextHandleNetSerializer::FNetSerializerRegistryDelegates::OnPostFreezeNetSerializerRegistry()
 {
 	bIsPostFreezeCalled = true;
+}
+
+void FGameplayEffectContextHandleNetSerializer::FNetSerializerRegistryDelegates::OnLoadedModulesUpdated()
+{
 	InitGameplayEffectContextHandleNetSerializerTypeCache();
 }
 
diff --git a/Engine/Source/Runtime/Engine/Private/GameFramework/RootMotionSourceGroupNetSerializer.cpp b/Engine/Source/Runtime/Engine/Private/GameFramework/RootMotionSourceGroupNetSerializer.cpp
--- a/Engine/Source/Runtime/Engine/Private/GameFramework/RootMotionSourceGroupNetSerializer.cpp
+++ b/Engine/Source/Runtime/Engine/Private/GameFramework/RootMotionSourceGroupNetSerializer.cpp
@@ -103,11 +103,13 @@ private:
 	class FNetSerializerRegistryDelegates final : private UE::Net::FNetSerializerRegistryDelegates
 	{
 	public:
+		FNetSerializerRegistryDelegates();
 		virtual ~FNetSerializerRegistryDelegates();
 
 	private:
 		virtual void OnPreFreezeNetSerializerRegistry() override;
 		virtual void OnPostFreezeNetSerializerRegistry() override;
+		virtual void OnLoadedModulesUpdated() override;
 
 		inline static const FName PropertyNetSerializerRegistry_NAME_RootMotionSourceGroup = FName("RootMotionSourceGroup");
 		UE_NET_IMPLEMENT_NAMED_STRUCT_NETSERIALIZER_INFO(PropertyNetSerializerRegistry_NAME_RootMotionSourceGroup, FRootMotionSourceGroupNetSerializer);
@@ -529,6 +531,11 @@ void FRootMotionSourceGroupNetSerializer::SetPendingRootMotionSourcesNum(FRootMo
 	ArrayContainer.PendingAddRootMotionSources.SetNum(static_cast<SSIZE_T>(Num));
 }
 
+FRootMotionSourceGroupNetSerializer::FNetSerializerRegistryDelegates::FNetSerializerRegistryDelegates()
+: UE::Net::FNetSerializerRegistryDelegates(EFlags::ShouldBindLoadedModulesUpdatedDelegate)
+{
+}
+
 FRootMotionSourceGroupNetSerializer::FNetSerializerRegistryDelegates::~FNetSerializerRegistryDelegates()
 {
 	UE_NET_UNREGISTER_NETSERIALIZER_INFO(PropertyNetSerializerRegistry_NAME_RootMotionSourceGroup);
@@ -542,7 +549,11 @@ void FRootMotionSourceGroupNetSerializer::FNetSerializerRegistryDelegates::OnPre
 void FRootMotionSourceGroupNetSerializer::FNetSerializerRegistryDelegates::OnPostFreezeNetSerializerRegistry()
 {
 	bIsPostFreezeCalled = true;
-	FRootMotionSourceGroupNetSerializer::InitTypeCache();
+}
+
+void FRootMotionSourceGroupNetSerializer::FNetSerializerRegistryDelegates::OnLoadedModulesUpdated()
+{
+	UE::Net::FRootMotionSourceGroupNetSerializer::InitTypeCache();
 }
 
 }
diff --git a/Engine/Source/Runtime/Experimental/Iris/Core/Private/Iris/IrisCoreModule.cpp b/Engine/Source/Runtime/Experimental/Iris/Core/Private/Iris/IrisCoreModule.cpp
--- a/Engine/Source/Runtime/Experimental/Iris/Core/Private/Iris/IrisCoreModule.cpp
+++ b/Engine/Source/Runtime/Experimental/Iris/Core/Private/Iris/IrisCoreModule.cpp
@@ -3,6 +3,7 @@
 #include "CoreTypes.h"
 #include "Modules/ModuleInterface.h"
 #include "Modules/ModuleManager.h"
+#include "Misc/CoreDelegates.h"
 #include "Iris/ReplicationState/DefaultPropertyNetSerializerInfos.h"
 #include "Iris/ReplicationState/PropertyNetSerializerInfoRegistry.h"
 #include "Iris/ReplicationSystem/LegacyPushModel.h"
@@ -34,6 +35,8 @@ private:
 		// Iris requires NetCore
 		FModuleManager::LoadModuleChecked<IModuleInterface>("NetCore");
 
+		FCoreDelegates::OnAllModuleLoadingPhasesComplete.AddRaw(this, &FIrisCoreModule::OnAllModuleLoadingPhasesComplete);
+
 		// Check command line for whether we should override the net.Iris.UseIrisReplication cvar, as we need to do that early
 		int32 UseIrisReplication;
 		if(FParse::Value(FCommandLine::Get(), TEXT("UseIrisReplication="), UseIrisReplication))
@@ -67,6 +70,14 @@ private:
 		return false;
 	}
 
+	void OnAllModuleLoadingPhasesComplete()
+	{
+		bAllowLoadedModulesUpdatedCallback = true;
+		UE::Net::Private::FInternalNetSerializerDelegates::BroadcastLoadedModulesUpdated();
+
+		FCoreDelegates::OnAllModuleLoadingPhasesComplete.RemoveAll(this);
+	}
+
 	void OnModulesChanged(FName ModuleThatChanged, EModuleChangeReason ReasonForChange)
 	{
 		switch (ReasonForChange)
@@ -75,6 +86,10 @@ private:
 			{
 				UE::Net::Private::FInternalNetSerializerDelegates::BroadcastPreFreezeNetSerializerRegistry();
 				UE::Net::Private::FInternalNetSerializerDelegates::BroadcastPostFreezeNetSerializerRegistry();
+				if (bAllowLoadedModulesUpdatedCallback)
+				{
+					UE::Net::Private::FInternalNetSerializerDelegates::BroadcastLoadedModulesUpdated();
+				}
 			}
 			break;
 
@@ -87,5 +102,6 @@ private:
 
 private:
 	FDelegateHandle ModulesChangedHandle;
+	bool bAllowLoadedModulesUpdatedCallback = false;
 };
 IMPLEMENT_MODULE(FIrisCoreModule, IrisCore);
diff --git a/Engine/Source/Runtime/Experimental/Iris/Core/Private/Iris/Serialization/InternalNetSerializerDelegates.cpp b/Engine/Source/Runtime/Experimental/Iris/Core/Private/Iris/Serialization/InternalNetSerializerDelegates.cpp
--- a/Engine/Source/Runtime/Experimental/Iris/Core/Private/Iris/Serialization/InternalNetSerializerDelegates.cpp
+++ b/Engine/Source/Runtime/Experimental/Iris/Core/Private/Iris/Serialization/InternalNetSerializerDelegates.cpp
@@ -19,6 +19,12 @@ void FInternalNetSerializerDelegates::BroadcastPostFreezeNetSerializerRegistry()
 	PostFreezeNetSerializerRegistryDelegate.Clear();
 }
 
+void FInternalNetSerializerDelegates::BroadcastLoadedModulesUpdated()
+{
+	FSimpleMulticastDelegate& LoadedModulesUpdatedDelegate = GetLoadedModulesUpdatedDelegate();
+	LoadedModulesUpdatedDelegate.Broadcast();
+}
+
 FSimpleMulticastDelegate& FInternalNetSerializerDelegates::GetPreFreezeNetSerializerRegistryDelegate()
 {
 	static FSimpleMulticastDelegate PreFreezeNetSerializerRegistryDelegate;
@@ -31,4 +37,10 @@ FSimpleMulticastDelegate& FInternalNetSerializerDelegates::GetPostFreezeNetSeria
 	return PostFreezeNetSerializerRegistryDelegate;
 }
 
+FSimpleMulticastDelegate& FInternalNetSerializerDelegates::GetLoadedModulesUpdatedDelegate()
+{
+	static FSimpleMulticastDelegate LoadedModulesUpdatedDelegate;
+	return LoadedModulesUpdatedDelegate;
+}
+
 }
diff --git a/Engine/Source/Runtime/Experimental/Iris/Core/Private/Iris/Serialization/InternalNetSerializerDelegates.h b/Engine/Source/Runtime/Experimental/Iris/Core/Private/Iris/Serialization/InternalNetSerializerDelegates.h
--- a/Engine/Source/Runtime/Experimental/Iris/Core/Private/Iris/Serialization/InternalNetSerializerDelegates.h
+++ b/Engine/Source/Runtime/Experimental/Iris/Core/Private/Iris/Serialization/InternalNetSerializerDelegates.h
@@ -12,9 +12,11 @@ class FInternalNetSerializerDelegates
 public:
 	static void BroadcastPreFreezeNetSerializerRegistry();
 	static void BroadcastPostFreezeNetSerializerRegistry();
+	static void BroadcastLoadedModulesUpdated();
 
 	static FSimpleMulticastDelegate& GetPreFreezeNetSerializerRegistryDelegate();
 	static FSimpleMulticastDelegate& GetPostFreezeNetSerializerRegistryDelegate();
+	static FSimpleMulticastDelegate& GetLoadedModulesUpdatedDelegate();
 };
 
 }
diff --git a/Engine/Source/Runtime/Experimental/Iris/Core/Private/Iris/Serialization/NetSerializerDelegates.cpp b/Engine/Source/Runtime/Experimental/Iris/Core/Private/Iris/Serialization/NetSerializerDelegates.cpp
--- a/Engine/Source/Runtime/Experimental/Iris/Core/Private/Iris/Serialization/NetSerializerDelegates.cpp
+++ b/Engine/Source/Runtime/Experimental/Iris/Core/Private/Iris/Serialization/NetSerializerDelegates.cpp
@@ -8,9 +8,18 @@ namespace UE::Net
 
 using namespace UE::Net::Private;
 
-FNetSerializerRegistryDelegates::FNetSerializerRegistryDelegates()
+FNetSerializerRegistryDelegates::FNetSerializerRegistryDelegates(EFlags Flags)
 : PreFreezeDelegate(FInternalNetSerializerDelegates::GetPreFreezeNetSerializerRegistryDelegate().AddRaw(this, &FNetSerializerRegistryDelegates::PreFreezeNetSerializerRegistry))
 , PostFreezeDelegate(FInternalNetSerializerDelegates::GetPostFreezeNetSerializerRegistryDelegate().AddRaw(this, &FNetSerializerRegistryDelegates::PostFreezeNetSerializerRegistry))
+{
+	if ((Flags & EFlags::ShouldBindLoadedModulesUpdatedDelegate) != 0U)
+	{
+		LoadedModulesUpdatedDelegate = FInternalNetSerializerDelegates::GetLoadedModulesUpdatedDelegate().AddRaw(this, &FNetSerializerRegistryDelegates::LoadedModulesUpdated);
+	}
+}
+
+FNetSerializerRegistryDelegates::FNetSerializerRegistryDelegates()
+: FNetSerializerRegistryDelegates(EFlags::None)
 {
 }
 
@@ -24,6 +33,10 @@ FNetSerializerRegistryDelegates::~FNetSerializerRegistryDelegates()
 	{
 		FInternalNetSerializerDelegates::GetPostFreezeNetSerializerRegistryDelegate().Remove(PostFreezeDelegate);
 	}
+	if (LoadedModulesUpdatedDelegate.IsValid())
+	{
+		FInternalNetSerializerDelegates::GetLoadedModulesUpdatedDelegate().Remove(LoadedModulesUpdatedDelegate);
+	}
 }
 
 void FNetSerializerRegistryDelegates::OnPreFreezeNetSerializerRegistry()
@@ -34,6 +47,10 @@ void FNetSerializerRegistryDelegates::OnPostFreezeNetSerializerRegistry()
 {
 }
 
+void FNetSerializerRegistryDelegates::OnLoadedModulesUpdated()
+{
+}
+
 void FNetSerializerRegistryDelegates::PreFreezeNetSerializerRegistry()
 {
 	OnPreFreezeNetSerializerRegistry();
@@ -48,4 +65,9 @@ void FNetSerializerRegistryDelegates::PostFreezeNetSerializerRegistry()
 	PostFreezeDelegate.Reset();
 }
 
+void FNetSerializerRegistryDelegates::LoadedModulesUpdated()
+{
+	OnLoadedModulesUpdated();
+}
+
 }
diff --git a/Engine/Source/Runtime/Experimental/Iris/Core/Private/Iris/Serialization/PolymorphicNetSerializer.cpp b/Engine/Source/Runtime/Experimental/Iris/Core/Private/Iris/Serialization/PolymorphicNetSerializer.cpp
--- a/Engine/Source/Runtime/Experimental/Iris/Core/Private/Iris/Serialization/PolymorphicNetSerializer.cpp
+++ b/Engine/Source/Runtime/Experimental/Iris/Core/Private/Iris/Serialization/PolymorphicNetSerializer.cpp
@@ -19,10 +19,15 @@ void FPolymorphicNetSerializerScriptStructCache::InitForType(const UScriptStruct
 {
 	IRIS_PROFILER_SCOPE(FPolymorphicNetSerializerScriptStructCache_InitForType);
 
-	RegisteredTypes.Reset();
+	TArray<FTypeInfo> UpdatedRegisteredTypes;
+	UpdatedRegisteredTypes.Reserve(RegisteredTypes.Num());
+
+	CommonTraits = EReplicationStateTraits::None;
+
+	bool bFoundNewScriptStructs = false;
 
 	// Find all script structs of this type and add them to the list and build descriptor
-	// (not sure of a better way to do this but it should only happen once at startup)
+	// (not sure of a better way to do this but it should only happen at startup)
 	for (TObjectIterator<UScriptStruct> It; It; ++It)
 	{
 		if (It->IsChildOf(InScriptStruct))
@@ -30,13 +35,23 @@ void FPolymorphicNetSerializerScriptStructCache::InitForType(const UScriptStruct
 			FTypeInfo Entry;
 			Entry.ScriptStruct = *It;
 
-			// Get or create descriptor
-			FReplicationStateDescriptorBuilder::FParameters Params;
-			Entry.Descriptor = FReplicationStateDescriptorBuilder::CreateDescriptorForStruct(*It, Params);
+			// See if we already had a descriptor for the struct
+			const uint32 ExistingTypeIndex = GetTypeIndex(Entry.ScriptStruct);
+			if (ExistingTypeIndex != InvalidTypeIndex)
+			{
+				Entry.Descriptor = GetTypeInfo(ExistingTypeIndex)->Descriptor;
+			}
+			else
+			{
+				// Get or create descriptor
+				FReplicationStateDescriptorBuilder::FParameters Params;
+				Entry.Descriptor = FReplicationStateDescriptorBuilder::CreateDescriptorForStruct(*It, Params);
+				bFoundNewScriptStructs = true;
+			}
 			
 			if (Entry.Descriptor.IsValid())
 			{
-				RegisteredTypes.Add(Entry);
+				UpdatedRegisteredTypes.Add(Entry);
 
 				CommonTraits |= (Entry.Descriptor->Traits & EReplicationStateTraits::HasObjectReference);
 			}
@@ -46,14 +61,20 @@ void FPolymorphicNetSerializerScriptStructCache::InitForType(const UScriptStruct
 			}
 		}
 	}
-	
-	RegisteredTypes.Sort([](const FTypeInfo& A, const FTypeInfo& B) { return A.ScriptStruct->GetName().ToLower() > B.ScriptStruct->GetName().ToLower(); });
 
-	if ((uint32)RegisteredTypes.Num() > MaxRegisteredTypeCount)
+	UpdatedRegisteredTypes.Sort([](const FTypeInfo& A, const FTypeInfo& B) { return A.ScriptStruct->GetName().ToLower() > B.ScriptStruct->GetName().ToLower(); });
+
+	if ((uint32)UpdatedRegisteredTypes.Num() > MaxRegisteredTypeCount)
 	{
-		UE_LOG(LogIris, Error, TEXT("FPolymorphicNetSerializerScriptStructCache::InitForType Too many types (%u of %u) of base %s"), RegisteredTypes.Num(), MaxRegisteredTypeCount,  ToCStr(InScriptStruct->GetName()));
+		UE_LOG(LogIris, Error, TEXT("FPolymorphicNetSerializerScriptStructCache::InitForType Too many types (%u of %u) of base %s"), UpdatedRegisteredTypes.Num(), MaxRegisteredTypeCount,  ToCStr(InScriptStruct->GetName()));
 		check((uint32)RegisteredTypes.Num() <= MaxRegisteredTypeCount);
 	}
+
+	// Log if we updated the types
+	const bool bWasUpdated = (UpdatedRegisteredTypes.Num() != RegisteredTypes.Num()) || bFoundNewScriptStructs;
+	UE_CLOG(bWasUpdated, LogIris, Log,TEXT("FPolymorphicNetSerializerScriptStructCache::InitForType Updated CachedTypedIds for base %s, NumCachedTypeIDs: %d"), ToCStr(InScriptStruct->GetName()), UpdatedRegisteredTypes.Num());
+
+	RegisteredTypes = MoveTemp(UpdatedRegisteredTypes);
 }
 
 }
diff --git a/Engine/Source/Runtime/Experimental/Iris/Core/Public/Iris/Serialization/NetSerializerDelegates.h b/Engine/Source/Runtime/Experimental/Iris/Core/Public/Iris/Serialization/NetSerializerDelegates.h
--- a/Engine/Source/Runtime/Experimental/Iris/Core/Public/Iris/Serialization/NetSerializerDelegates.h
+++ b/Engine/Source/Runtime/Experimental/Iris/Core/Public/Iris/Serialization/NetSerializerDelegates.h
@@ -13,13 +13,22 @@ namespace UE::Net
  * and OnPostFreezeNetSerializerRegistry() to perform any additional fixup
  * needed, such as creating a descriptor for a struct that you forward your
  * NetSerializer's calls to.
- * The virtual functions will be called zero or one times each.
+ * The virtual functions OnPreFreezeNetSerializerRegistry and OnPostFreezeNetSerializerRegistry will be called zero or one times each.
  */
 class FNetSerializerRegistryDelegates
 {
 public:
+
+	enum EFlags : uint32
+	{
+		None = 0U,
+		ShouldBindLoadedModulesUpdatedDelegate = 1U,
+	};
+
 	IRISCORE_API FNetSerializerRegistryDelegates();
 
+	explicit IRISCORE_API FNetSerializerRegistryDelegates(EFlags Flags);
+
 	/** Implement a destructor to unregister your NetSerializers. */
 	IRISCORE_API virtual ~FNetSerializerRegistryDelegates();
 
@@ -39,12 +48,23 @@ protected:
 	 */
 	IRISCORE_API virtual void OnPostFreezeNetSerializerRegistry();
 
+	/**
+	 * OnLoadedModulesUpdated delegate will be called every time a module has been loaded or unloaded
+	 * It gives the serializer an opportunity to update cached data that could be affected by loading or unloading modules.
+	 * NOTE: It will only be called if the EFlags::ShouldBindLoadedModulesUpdatedDelegate is set
+	 */
+	IRISCORE_API virtual void OnLoadedModulesUpdated();
+
+
 private:
+
 	void PreFreezeNetSerializerRegistry();
 	void PostFreezeNetSerializerRegistry();
+	void LoadedModulesUpdated();
 
 	FDelegateHandle PreFreezeDelegate;
 	FDelegateHandle PostFreezeDelegate;
+	FDelegateHandle LoadedModulesUpdatedDelegate;
 };
 
 }
