Add warning logs in AABBTree.h

diff --git a/Engine/Source/Runtime/Experimental/Chaos/Public/Chaos/AABBTree.h b/Engine/Source/Runtime/Experimental/Chaos/Public/Chaos/AABBTree.h
--- a/Engine/Source/Runtime/Experimental/Chaos/Public/Chaos/AABBTree.h
+++ b/Engine/Source/Runtime/Experimental/Chaos/Public/Chaos/AABBTree.h
@@ -431,7 +431,10 @@ struct TAABBTreeLeafArray : public TBoundsWrapperHelper<TPayloadType, T, bComput
 				Elems.RemoveAtSwap(Idx);
 				break;
 			}
-			ensure(Idx != Elems.Num() - 1); // Make sure the payload was actually in here
+			if (UNLIKELY(!ensure(Idx != Elems.Num() - 1))) // Make sure the payload was actually in here
+			{
+				UE_LOG(LogChaos, Warning, TEXT("AABBTree: Element not removed"));
+			}
 		}
 		bDirtyLeaf = true;
 	}
@@ -1196,7 +1199,10 @@ public:
 			AllocatedNodeIdx = Nodes.AddUninitialized(1);;
 			Nodes[AllocatedNodeIdx].bLeaf = false;
 		}
-		ensure(Nodes[AllocatedNodeIdx].bLeaf == false);
+		if (UNLIKELY(!ensure(Nodes[AllocatedNodeIdx].bLeaf == false)))
+		{
+			UE_LOG(LogChaos, Warning, TEXT("AABBTree: Allocated internal node is a leaf"));
+		}
 		return AllocatedNodeIdx;
 	}
 
@@ -1239,7 +1245,10 @@ public:
 		Nodes[AllocatedNodeIdx].ChildrenBounds[0] = ExpandedBounds;
 
 		Nodes[AllocatedNodeIdx].ParentNode = INDEX_NONE;
-		ensure(Nodes[AllocatedNodeIdx].bLeaf == true);
+		if (UNLIKELY(!ensure(Nodes[AllocatedNodeIdx].bLeaf == true)))
+		{
+			UE_LOG(LogChaos, Warning, TEXT("AABBTree: Allocated leaf node is not a leaf"));
+		}
 		
 		return NodeAndLeafIndices{ AllocatedNodeIdx , LeafIndex };
 	}
@@ -1248,7 +1257,10 @@ public:
 	{
 		Nodes[NodeIdx].ChildrenNodes[1] = FirstFreeInternalNode;
 		FirstFreeInternalNode = NodeIdx;
-		ensure(Nodes[NodeIdx].bLeaf == false);
+		if (UNLIKELY(!ensure(Nodes[NodeIdx].bLeaf == false)))
+		{
+			UE_LOG(LogChaos, Warning, TEXT("AABBTree: Deallocated Internal node is a leaf"));
+		}
 	}
 
 	void  DeAllocateLeafNode(int32 NodeIdx)
@@ -1258,7 +1270,10 @@ public:
 
 		Nodes[NodeIdx].ChildrenNodes[1] = FirstFreeLeafNode;
 		FirstFreeLeafNode = NodeIdx;
-		ensure(Nodes[NodeIdx].bLeaf == true);
+		if (UNLIKELY(!ensure(Nodes[NodeIdx].bLeaf == true)))
+		{
+			UE_LOG(LogChaos, Warning, TEXT("AABBTree: Deallocated Leaf node is not a leaf"));
+		}
 	}
 
 	// Is the input node Child 0 or Child 1?
@@ -1273,7 +1288,10 @@ public:
 		}
 		else
 		{
-			ensure(Nodes[ParentIdx].ChildrenNodes[1] == NodeIdx);
+			if (UNLIKELY(!ensure(Nodes[ParentIdx].ChildrenNodes[1] == NodeIdx)))
+			{
+				UE_LOG(LogChaos, Warning, TEXT("AABBTree: Child node not found"));
+			}
 			return 1;
 		}
 	}
@@ -1432,10 +1450,16 @@ public:
 			// Modify NodeIdx 
 			Nodes[NodeIdx].ChildrenNodes[AuntLocalChildIdx] = BestGrandChildToSwap;
 			// Modify BestGrandChildToSwap
-			ensure(BestGrandChildToSwap != NodeIdx);
+			if (UNLIKELY(!ensure(BestGrandChildToSwap != NodeIdx)))
+			{
+				UE_LOG(LogChaos, Warning, TEXT("AABBTree: 1: Rotate Node loop detected"));
+			}
 			Nodes[BestGrandChildToSwap].ParentNode = NodeIdx;
 			// Modify BestAuntToSwap
-			ensure(BestAuntToSwap != MotherOfBestGrandChild);
+			if (UNLIKELY(!ensure(BestAuntToSwap != MotherOfBestGrandChild)))
+			{
+				UE_LOG(LogChaos, Warning, TEXT("AABBTree: 2: Rotate Node loop detected"));
+			}
 			Nodes[BestAuntToSwap].ParentNode = MotherOfBestGrandChild;
 			// Modify MotherOfBestGrandChild
 			Nodes[MotherOfBestGrandChild].ChildrenNodes[GrandChildLocalChildIdx] = BestAuntToSwap;
@@ -1489,7 +1513,10 @@ public:
 		const int32 OldParent = Nodes[BestSibling].ParentNode;
 		const int32 NewParent = AllocateInternalNode();
 		FNode& NewParentNode = Nodes[NewParent];
-		ensure(NewParent != OldParent);
+		if (UNLIKELY(!ensure(NewParent != OldParent)))
+		{
+			UE_LOG(LogChaos, Warning, TEXT("AABBTree: 1: Insert leaf loop detected"));
+		}
 		NewParentNode.ParentNode = OldParent;
 		NewParentNode.ChildrenNodes[0] = BestSibling;
 		NewParentNode.ChildrenNodes[1] = NewLeafIndices.NodeIdx;
@@ -1510,9 +1537,15 @@ public:
 			RootNode = NewParent;
 		}
 
-		ensure(BestSibling != NewParent);
+		if (UNLIKELY(!ensure(BestSibling != NewParent)))
+		{
+			UE_LOG(LogChaos, Warning, TEXT("AABBTree: 2: Insert leaf loop detected"));
+		}
 		Nodes[BestSibling].ParentNode = NewParent;
-		ensure(NewLeafIndices.NodeIdx != NewParent);
+		if (UNLIKELY(!ensure(NewLeafIndices.NodeIdx != NewParent)))
+		{
+			UE_LOG(LogChaos, Warning, TEXT("AABBTree: 3: Insert leaf loop detected"));
+		}
 		Nodes[NewLeafIndices.NodeIdx].ParentNode = NewParent;
 
 		UpdateAncestorBounds(NewParent, true);
@@ -1533,7 +1566,10 @@ public:
 		}*/
 		while (ParentNodeIdx != INDEX_NONE)
 		{
-			ensure(CurrentNodeIdx != ParentNodeIdx);
+			if (UNLIKELY(!ensure(CurrentNodeIdx != ParentNodeIdx)))
+			{
+				UE_LOG(LogChaos, Warning, TEXT("AABBTree: 1: UpdateAncestorBounds loop detected"));
+			}
 			int32 ChildIndex = WhichChildAmI(CurrentNodeIdx);
 			Nodes[ParentNodeIdx].ChildrenBounds[ChildIndex] = Nodes[CurrentNodeIdx].ChildrenBounds[0];
 			if (!Nodes[CurrentNodeIdx].bLeaf)
@@ -1553,7 +1589,11 @@ public:
 
 	void RemoveLeafNode(int32 LeafNodeIdx, const TPayloadType& Payload)
 	{
-		ensure(Nodes[LeafNodeIdx].bLeaf == true);
+		if (UNLIKELY(!ensure(Nodes[LeafNodeIdx].bLeaf == true)))
+		{
+			UE_LOG(LogChaos, Warning, TEXT("AABBTree: RemoveLeafNode on node that is not a leaf"));
+		}
+
 
 		int32 LeafIdx = Nodes[LeafNodeIdx].ChildrenNodes[0];
 
@@ -1589,7 +1629,12 @@ public:
 			{
 				RootNode = SiblingNodeIdx;
 			}
-			ensure(SiblingNodeIdx != GrandParentNodeIdx);
+
+			if (UNLIKELY(!ensure(SiblingNodeIdx != GrandParentNodeIdx)))
+			{
+				UE_LOG(LogChaos, Warning, TEXT("AABBTree: RemoveLeafNode loop detected"));
+			}
+
 			Nodes[SiblingNodeIdx].ParentNode = GrandParentNodeIdx;
 			UpdateAncestorBounds(SiblingNodeIdx);
 			DeAllocateInternalNode(ParentNodeIdx);
@@ -1603,7 +1648,10 @@ public:
 
 	virtual bool RemoveElement(const TPayloadType& Payload)
 	{
-		ensure(bModifyingTreeMultiThreadingFastCheck == false);
+		if (UNLIKELY(!ensure(bModifyingTreeMultiThreadingFastCheck == false)))
+		{
+			UE_LOG(LogChaos, Warning, TEXT("AABBTree: RemoveElement unsafe updated from multiple threads detected"));
+		}
 		bModifyingTreeMultiThreadingFastCheck = true;
 		if (ensure(bMutable))
 		{
@@ -1665,7 +1713,10 @@ public:
 
 	virtual void UpdateElement(const TPayloadType& Payload, const TAABB<T, 3>& NewBounds, bool bInHasBounds) override
 	{
-		ensure(bModifyingTreeMultiThreadingFastCheck == false);
+		if (UNLIKELY(!ensure(bModifyingTreeMultiThreadingFastCheck == false)))
+		{
+			UE_LOG(LogChaos, Warning, TEXT("AABBTree: UpdateElement unsafe updated from multiple threads detected"));
+		}
 		bModifyingTreeMultiThreadingFastCheck = true;
 #if !WITH_EDITOR
 		//CSV_SCOPED_TIMING_STAT(ChaosPhysicsTimers, AABBTreeUpdateElement)
