potential fix for PVS studio static analysis in the cluster union physics proxy

diff --git a/Engine/Source/Runtime/Experimental/Chaos/Private/PhysicsProxy/ClusterUnionPhysicsProxy.cpp b/Engine/Source/Runtime/Experimental/Chaos/Private/PhysicsProxy/ClusterUnionPhysicsProxy.cpp
--- a/Engine/Source/Runtime/Experimental/Chaos/Private/PhysicsProxy/ClusterUnionPhysicsProxy.cpp
+++ b/Engine/Source/Runtime/Experimental/Chaos/Private/PhysicsProxy/ClusterUnionPhysicsProxy.cpp
@@ -290,9 +290,10 @@ namespace Chaos
 		Particle_External->SetGeometry(PullData.SharedGeometry);
 		Particle_External->SetObjectState(PullData.ObjectState, true, /*bInvalidate=*/false);
 
-		for (int32 ShapeIndex = 0; ShapeIndex < Particle_External->ShapesArray().Num(); ++ShapeIndex)
+		const FShapesArray& ShapeArray = Particle_External->ShapesArray();
+		for (int32 ShapeIndex = 0; ShapeIndex < ShapeArray.Num(); ++ShapeIndex)
 		{
-			if (const TUniquePtr<Chaos::FPerShapeData>& ShapeData = Particle_External->ShapesArray()[ShapeIndex])
+			if (const TUniquePtr<Chaos::FPerShapeData>& ShapeData = ShapeArray[ShapeIndex])
 			{
 				if (PullData.CollisionData.IsValidIndex(ShapeIndex))
 				{
