potential fix for PVS studio static analysis in PBDRigidClustering

diff --git a/Engine/Source/Runtime/Experimental/Chaos/Private/Chaos/PBDRigidClustering.cpp b/Engine/Source/Runtime/Experimental/Chaos/Private/Chaos/PBDRigidClustering.cpp
--- a/Engine/Source/Runtime/Experimental/Chaos/Private/Chaos/PBDRigidClustering.cpp
+++ b/Engine/Source/Runtime/Experimental/Chaos/Private/Chaos/PBDRigidClustering.cpp
@@ -1495,11 +1495,16 @@ namespace Chaos
 		// we should probably have a way to retrieve all the active clusters per proxy instead of having to do this iteration
 		for (FPBDRigidClusteredParticleHandle* ClusteredHandle : GetTopLevelClusterParents())
 		{
+			if (!ClusteredHandle)
+			{
+				continue;
+			}
+
 			const bool bIsInputProxy = ClusteredHandle->PhysicsProxy() == Proxy;
 
 			// This handles the case where we want to break a GC but it's still in a cluster union.
 			const bool bIsInPhysicsProxiesSet = ClusteredHandle->PhysicsProxies().Contains(Proxy);
-			if (ClusteredHandle && (bIsInputProxy || bIsInPhysicsProxiesSet))
+			if (bIsInputProxy || bIsInPhysicsProxiesSet)
 			{
 				// Now we need to go from the parent cluster union particle to the GC particle that corresponds to the proxy.
 				if (bIsInPhysicsProxiesSet)
