Geometry collection : optimize  scale on remove feature
- Add oiptions on the GC asset to disable scale on remove, saving transform computation
- Remove the "UniformScale" attribute on the dynamic collection ( saving a FTransform per piece per compoinent )
- Move update of transform from decay from CalculateGlobalMatrices to the post physics update callback
- Make better use of the collection decay facade throughout the remove on sleep and remove on break code

diff --git a/Engine/Source/Runtime/Experimental/Chaos/Public/GeometryCollection/ManagedArrayAccessor.h b/Engine/Source/Runtime/Experimental/Chaos/Public/GeometryCollection/ManagedArrayAccessor.h
--- a/Engine/Source/Runtime/Experimental/Chaos/Public/GeometryCollection/ManagedArrayAccessor.h
+++ b/Engine/Source/Runtime/Experimental/Chaos/Public/GeometryCollection/ManagedArrayAccessor.h
@@ -60,6 +60,13 @@ public:
 		return Collection->AddElements(NumElements, Group);
 	}
 
+	/** array style accessor */
+	const T& operator[](int32 Index) const
+	{
+		check(IsValid());
+		return (*ConstAttributeArray)[Index];
+	}
+
 	/** get the attribute for read only */
 	const TManagedArray<T>& Get() const
 	{
@@ -81,6 +88,14 @@ public:
 		return *AttributeArray;
 	}
 
+	/** per index modification */
+	void ModifyAt(int32 Index, const T& NewValue)
+	{
+		check(AttributeArray != nullptr && !IsConst());
+		AttributeArray->MarkDirty();
+		(*AttributeArray)[Index] = NewValue;
+	}
+
 	/** add the attribute if it does not exists yet */
 	TManagedArray<T>& Add(ManageArrayAccessor::EPersistencePolicy PersistencePolicy = ManageArrayAccessor::EPersistencePolicy::MakePersistent,
 		FName InGroupDependency = FName(NAME_None))
diff --git a/Engine/Source/Runtime/Experimental/GeometryCollectionEngine/Private/GeometryCollection/GeometryCollectionComponent.cpp b/Engine/Source/Runtime/Experimental/GeometryCollectionEngine/Private/GeometryCollection/GeometryCollectionComponent.cpp
--- a/Engine/Source/Runtime/Experimental/GeometryCollectionEngine/Private/GeometryCollection/GeometryCollectionComponent.cpp
+++ b/Engine/Source/Runtime/Experimental/GeometryCollectionEngine/Private/GeometryCollection/GeometryCollectionComponent.cpp
@@ -2822,6 +2822,8 @@ void UGeometryCollectionComponent::UpdateRenderSystemsIfNeeded(bool bDynamicColl
 
 	if (bDynamicCollectionDirty)
 	{
+		UpdateRemovalIfNeeded();
+
 		// #todo review: When we've made changes to ISMC, we need to move this function call to SetRenderDynamicData_Concurrent
 		RefreshEmbeddedGeometry();
 
@@ -2858,6 +2860,68 @@ void UGeometryCollectionComponent::UpdateNavigationDataIfNeeded(bool bDynamicCol
 	}
 }
 
+void UGeometryCollectionComponent::UpdateRemovalIfNeeded()
+{
+	// if removal is enabled, update the dynamic collection transform based on the decay 
+	// todo: we could optimize this using a list of transform to update from when we update the decay values
+	if (DynamicCollection && bAllowRemovalOnBreak && bAllowRemovalOnSleep)
+	{
+		FGeometryCollectionDecayDynamicFacade DecayFacade(*DynamicCollection);
+		if (DecayFacade.IsValid())
+		{
+			const FTransform ZeroScaleTransform(FQuat::Identity, FVector::Zero(), FVector(0, 0, 0));
+
+			if (RestCollection->bScaleOnRemoval)
+			{
+				const FTransform InverseComponentTransform = GetComponentTransform().Inverse();
+				for (int32 Idx = 0; Idx < GetTransformArray().Num(); ++Idx)
+				{
+					// only update values if the decay has changed 
+					const float Decay = DecayFacade.GetDecay(Idx);
+					if (Decay > 0.f && Decay <= 1.f)
+					{
+						const float Scale = 1.0 - Decay;
+						if (Scale < UE_SMALL_NUMBER)
+						{
+							DynamicCollection->Transform[Idx].SetScale3D(FVector::ZeroVector);
+						}
+						else
+						{
+							float ShrinkRadius = 0.0f;
+							UE::Math::TSphere<double> AccumulatedSphere;
+							// todo(chaos) : find a faster way to do that ( precompute the data ? )
+							if (CalculateInnerSphere(Idx, AccumulatedSphere))
+							{
+								ShrinkRadius = -AccumulatedSphere.W;
+							}
+
+							const FQuat LocalRotation = (InverseComponentTransform * FTransform(GlobalMatrices[Idx]).Inverse()).GetRotation();
+							const FVector LocalDown = LocalRotation.RotateVector(FVector(0.f, 0.f, ShrinkRadius));
+							const FVector CenterOfMass = DynamicCollection->MassToLocal[Idx].GetTranslation();
+							const FVector ScaleCenter = LocalDown + CenterOfMass;
+							const FTransform ScaleTransform(FQuat::Identity, ScaleCenter * FVector::FReal(1.f - Scale), FVector(Scale));
+							DynamicCollection->Transform[Idx].SetScale3D(FVector::ZeroVector);
+						}
+					}
+				}
+			}
+			else
+			{
+				// if we do not use scale we still need to make the piece to disappear when decay has reached 1
+				for (int32 Idx = 0; Idx < GetTransformArray().Num(); ++Idx)
+				{
+					const float Scale = (1.0 - DecayFacade.GetDecay(Idx));
+					if (Scale < UE_SMALL_NUMBER)
+					{
+						DynamicCollection->Transform[Idx].SetScale3D(FVector::ZeroVector);
+					}
+				}
+			}
+
+		}
+	}
+}
+
 void UGeometryCollectionComponent::RequestUpdateRepData()
 {
 	if (FPhysScene* PhysScene = GetInnerChaosScene())
@@ -4024,51 +4088,7 @@ void UGeometryCollectionComponent::CalculateGlobalMatrices()
 		}
 		else
 		{
-			// Have to fully rebuild
-			if (DynamicCollection 
-				&& DynamicCollection->HasAttribute("UniformScale", FGeometryCollection::TransformGroup)
-				&& DynamicCollection->HasAttribute("Decay", FGeometryCollection::TransformGroup))
-			{
-				const TManagedArray<float>& Decay = DynamicCollection->GetAttribute<float>("Decay", FGeometryCollection::TransformGroup);
-				TManagedArray<FTransform>& UniformScale = DynamicCollection->ModifyAttribute<FTransform>("UniformScale", FGeometryCollection::TransformGroup);
-
-				const FTransform InverseComponentTransform = GetComponentTransform().Inverse();
-				const FTransform ZeroScaleTransform(FQuat::Identity, FVector::Zero(), FVector(0, 0, 0));
-				for (int32 Idx = 0; Idx < GetTransformArray().Num(); ++Idx)
-				{
-					// only update values if the decay has changed 
-					if (Decay[Idx] > 0.f && Decay[Idx] <= 1.f)
-					{
-						const float Scale = 1.0 - Decay[Idx];
-						if (Scale < UE_SMALL_NUMBER)
-						{
-							UniformScale[Idx] = ZeroScaleTransform;
-						}
-						else
-						{
-							float ShrinkRadius = 0.0f;
-							UE::Math::TSphere<double> AccumulatedSphere;
-							// todo(chaos) : find a faster way to do that ( precompute the data ? )
-							if (CalculateInnerSphere(Idx, AccumulatedSphere))
-							{
-								ShrinkRadius = -AccumulatedSphere.W;
-							}
-						
-							const FQuat LocalRotation = (InverseComponentTransform * FTransform(GlobalMatrices[Idx]).Inverse()).GetRotation();
-							const FVector LocalDown = LocalRotation.RotateVector(FVector(0.f, 0.f, ShrinkRadius));
-							const FVector CenterOfMass = DynamicCollection->MassToLocal[Idx].GetTranslation();
-							const FVector ScaleCenter = LocalDown + CenterOfMass;
-							UniformScale[Idx] = FTransform(FQuat::Identity, ScaleCenter * (FVector::FReal)(1.f - Scale), FVector(Scale));
-						}
-					}
-				}
-				
-				GeometryCollectionAlgo::GlobalMatrices(GetTransformArray(), GetParentArray(), UniformScale, GlobalMatrices);
-			}
-			else
-			{ 
-				GeometryCollectionAlgo::GlobalMatrices(GetTransformArray(), GetParentArray(), GlobalMatrices);		
-			}
+			GeometryCollectionAlgo::GlobalMatrices(GetTransformArray(), GetParentArray(), GlobalMatrices);
 		}
 	}
 	
@@ -4586,11 +4606,11 @@ struct FGeometryCollectionDecayContext
 
 void UGeometryCollectionComponent::UpdateDecay(int32 TransformIdx, float UpdatedDecay, bool bUseClusterCrumbling, bool bHasDynamicInternalClusterParent, FGeometryCollectionDecayContext& ContextInOut)
 {
-	TManagedArray<float>& Decay = ContextInOut.DecayFacade.DecayAttribute.Modify();
-	if (UpdatedDecay > Decay[TransformIdx])
+	float Decay = ContextInOut.DecayFacade.GetDecay(TransformIdx);
+	if (UpdatedDecay > Decay)
 	{
 		ContextInOut.DirtyDynamicCollection = true;
-		Decay[TransformIdx] = UpdatedDecay;
+		Decay = UpdatedDecay;
 
 		if (bUseClusterCrumbling)
 		{
@@ -4600,21 +4620,24 @@ void UGeometryCollectionComponent::UpdateDecay(int32 TransformIdx, float Updated
 				if (InternalClusterItemindex.IsValid())
 				{
 					ContextInOut.ToCrumble.AddUnique(InternalClusterItemindex);
-					Decay[TransformIdx] = 0.0f;
+					Decay = 0.0f;
 				}
 			}
 			else
 			{
 				ContextInOut.ToCrumble.AddUnique(FGeometryCollectionItemIndex::CreateTransformItemIndex(TransformIdx));
-				Decay[TransformIdx] = 0.0f;
+				Decay = 0.0f;
 			}
 		}
-		else if (Decay[TransformIdx] >= 1.0f)
+		else if (Decay >= 1.0f)
 		{
 			// Disable the particle if it has decayed the requisite time
-			Decay[TransformIdx] = 1.0f;
+			Decay = 1.0f;
 			ContextInOut.ToDisable.Add(TransformIdx);
 		}
+
+		// push back Decay in the attribute
+		ContextInOut.DecayFacade.SetDecay(TransformIdx, Decay);
 	}
 }
 
@@ -4640,8 +4663,8 @@ void UGeometryCollectionComponent::IncrementSleepTimer(float DeltaTime)
 
 			const TManagedArray<int32>& OriginalParents = RestCollection->GetGeometryCollection()->Parent;
 
-			TManagedArray<float>& Decay = DecayFacade.DecayAttribute.Modify();
-			for (int32 TransformIdx = 0; TransformIdx < Decay.Num(); ++TransformIdx)
+			const int32 NumTransforms = OriginalParents.Num();
+			for (int32 TransformIdx = 0; TransformIdx < NumTransforms; ++TransformIdx)
 			{
 				const bool HasInternalClusterParent = DynamicStateFacade.HasInternalClusterParent(TransformIdx);
 				if (HasInternalClusterParent)
@@ -4662,7 +4685,7 @@ void UGeometryCollectionComponent::IncrementSleepTimer(float DeltaTime)
 					if (OriginalParents[TransformIdx] > INDEX_NONE)
 					{
 						// if decay has started we do not need to check slow moving or sleeping state anymore  
-						bool ShouldUpdateTimer = (Decay[TransformIdx] > 0);
+						bool ShouldUpdateTimer = (DecayFacade.GetDecay(TransformIdx) > 0);
 						if (!ShouldUpdateTimer && RestCollection->bSlowMovingAsSleeping)
 						{
 							const FVector CurrentPosition = DynamicCollection->Transform[TransformIdx].GetTranslation();
@@ -4710,8 +4733,8 @@ void UGeometryCollectionComponent::IncrementBreakTimer(float DeltaTime)
 
 			const TManagedArray<int32>* InitialLevels = PhysicsProxy->GetPhysicsCollection().FindAttribute<int32>("InitialLevel", FGeometryCollection::TransformGroup);
 
-			TManagedArray<float>& Decay = DecayFacade.DecayAttribute.Modify();
-			for (int32 TransformIdx = 0; TransformIdx < Decay.Num(); ++TransformIdx)
+			const int32 NumTransforms = OriginalParents.Num();
+			for (int32 TransformIdx = 0; TransformIdx < NumTransforms; ++TransformIdx)
 			{
 				const bool HasInternalClusterParent = DynamicStateFacade.HasInternalClusterParent(TransformIdx);
 				if (HasInternalClusterParent)
@@ -5101,4 +5124,4 @@ void UGeometryCollectionComponent::SetEnableDamageFromCollision(bool bValue)
 	{
 		PhysicsProxy->SetEnableDamageFromCollision_External(bValue);
 	}
-}
\ No newline at end of file
+}
diff --git a/Engine/Source/Runtime/Experimental/GeometryCollectionEngine/Private/GeometryCollection/GeometryCollectionEngineRemoval.cpp b/Engine/Source/Runtime/Experimental/GeometryCollectionEngine/Private/GeometryCollection/GeometryCollectionEngineRemoval.cpp
--- a/Engine/Source/Runtime/Experimental/GeometryCollectionEngine/Private/GeometryCollection/GeometryCollectionEngineRemoval.cpp
+++ b/Engine/Source/Runtime/Experimental/GeometryCollectionEngine/Private/GeometryCollection/GeometryCollectionEngineRemoval.cpp
@@ -203,19 +203,26 @@ float FGeometryCollectionRemoveOnSleepDynamicFacade::ComputeDecay(int32 Transfor
 
 FGeometryCollectionDecayDynamicFacade::FGeometryCollectionDecayDynamicFacade(FManagedArrayCollection& InCollection)
 	: DecayAttribute(InCollection, "Decay", FGeometryCollection::TransformGroup)
-	, UniformScaleAttribute(InCollection, "UniformScale", FGeometryCollection::TransformGroup)
 {
 }
 
 bool FGeometryCollectionDecayDynamicFacade::IsValid() const
 {
-	return DecayAttribute.IsValid()
-		&& UniformScaleAttribute.IsValid()
-	;
+	// on check the decay since the uniform attribute is optional ( see bScaleOnRemoval )
+	return DecayAttribute.IsValid();
 }
 
 void FGeometryCollectionDecayDynamicFacade::AddAttributes()
 {
 	DecayAttribute.AddAndFill(0.0f);
-	UniformScaleAttribute.AddAndFill(FTransform::Identity);
-}
\ No newline at end of file
+}
+
+float FGeometryCollectionDecayDynamicFacade::GetDecay(int32 TransformIndex) const
+{
+	return DecayAttribute[TransformIndex];
+}
+
+void FGeometryCollectionDecayDynamicFacade::SetDecay(int32 TransformIndex, float DecayValue)
+{
+	DecayAttribute.ModifyAt(TransformIndex, DecayValue);
+}
diff --git a/Engine/Source/Runtime/Experimental/GeometryCollectionEngine/Private/GeometryCollection/GeometryCollectionObject.cpp b/Engine/Source/Runtime/Experimental/GeometryCollectionEngine/Private/GeometryCollection/GeometryCollectionObject.cpp
--- a/Engine/Source/Runtime/Experimental/GeometryCollectionEngine/Private/GeometryCollection/GeometryCollectionObject.cpp
+++ b/Engine/Source/Runtime/Experimental/GeometryCollectionEngine/Private/GeometryCollection/GeometryCollectionObject.cpp
@@ -108,6 +108,7 @@ UGeometryCollection::UGeometryCollection(const FObjectInitializer& ObjectInitial
 	, MinimumMassClamp(0.1f)
 	, bImportCollisionFromSource(false)
 	, bRemoveOnMaxSleep(false)
+	, bScaleOnRemoval(true)
 	, MaximumSleepTime(5.0, 10.0)
 	, RemovalDuration(2.5, 5.0)
 	, bSlowMovingAsSleeping(true)
diff --git a/Engine/Source/Runtime/Experimental/GeometryCollectionEngine/Public/GeometryCollection/GeometryCollectionComponent.h b/Engine/Source/Runtime/Experimental/GeometryCollectionEngine/Public/GeometryCollection/GeometryCollectionComponent.h
--- a/Engine/Source/Runtime/Experimental/GeometryCollectionEngine/Public/GeometryCollection/GeometryCollectionComponent.h
+++ b/Engine/Source/Runtime/Experimental/GeometryCollectionEngine/Public/GeometryCollection/GeometryCollectionComponent.h
@@ -1277,6 +1277,7 @@ private:
 	
 	void UpdateRenderSystemsIfNeeded(bool bDynamicCollectionDirty);
 	void UpdateNavigationDataIfNeeded(bool bDynamicCollectionDirty);
+	void UpdateRemovalIfNeeded();
 
 	void BuildInitialFilterData();
 
diff --git a/Engine/Source/Runtime/Experimental/GeometryCollectionEngine/Public/GeometryCollection/GeometryCollectionEngineRemoval.h b/Engine/Source/Runtime/Experimental/GeometryCollectionEngine/Public/GeometryCollection/GeometryCollectionEngineRemoval.h
--- a/Engine/Source/Runtime/Experimental/GeometryCollectionEngine/Public/GeometryCollection/GeometryCollectionEngineRemoval.h
+++ b/Engine/Source/Runtime/Experimental/GeometryCollectionEngine/Public/GeometryCollection/GeometryCollectionEngineRemoval.h
@@ -153,10 +153,13 @@ public:
 	/** Add the necessary attributes if they are missing and initialize them if necessary */
 	void AddAttributes();
 
+	/** Get decay value for a specific transform index */
+	float GetDecay(int32 TransformIndex) const;
+
+	/** Set decay value for a specific transform index */
+	void SetDecay(int32 TransformIndex, float DecayValue);
+
+private:
 	/** state of decay ([0-1] range) */
-	// @todo(chaos) this should eventually move to a common removal facade when break and sleep removal are consolidated
 	TManagedArrayAccessor<float> DecayAttribute;
-
-	/** scale transform used to shrink the geometry collection piece */
-	TManagedArrayAccessor<FTransform> UniformScaleAttribute;
 };
diff --git a/Engine/Source/Runtime/Experimental/GeometryCollectionEngine/Public/GeometryCollection/GeometryCollectionObject.h b/Engine/Source/Runtime/Experimental/GeometryCollectionEngine/Public/GeometryCollection/GeometryCollectionObject.h
--- a/Engine/Source/Runtime/Experimental/GeometryCollectionEngine/Public/GeometryCollection/GeometryCollectionObject.h
+++ b/Engine/Source/Runtime/Experimental/GeometryCollectionEngine/Public/GeometryCollection/GeometryCollectionObject.h
@@ -631,6 +631,10 @@ public:
 	int32 MaximumCollisionParticles_DEPRECATED;
 #endif
 
+	/** When enabled, particle will scale down (shrink) when using being removed ( using both remove on sleep or remove on break ) - Enabled by default */
+	UPROPERTY(EditAnywhere, BlueprintReadOnly, Category = Removal)
+	bool bScaleOnRemoval;
+
 	/** Remove particle from simulation and dissolve rendered geometry once sleep threshold has been exceeded. */
 	UPROPERTY(EditAnywhere, BlueprintReadOnly, Category = Removal, meta = (DisplayName = "Remove on Sleep"))
 	bool bRemoveOnMaxSleep;
