Add MaxHealth to PT userdata for collisions, and read health data in PhysicalStrainHealth sim callback so that we can know when a physical collision will do enough damage to destroy a GC, and no longer block strain-based destruction.

diff --git a/Engine/Source/Runtime/Experimental/Chaos/Private/Chaos/StrainModification.cpp b/Engine/Source/Runtime/Experimental/Chaos/Private/Chaos/StrainModification.cpp
--- a/Engine/Source/Runtime/Experimental/Chaos/Private/Chaos/StrainModification.cpp
+++ b/Engine/Source/Runtime/Experimental/Chaos/Private/Chaos/StrainModification.cpp
@@ -11,6 +11,11 @@ const FGeometryCollectionPhysicsProxy* Chaos::FStrainedProxyModifier::GetProxy()
 	return Proxy;
 }
 
+const Chaos::FPBDRigidParticleHandle* Chaos::FStrainedProxyModifier::GetRootHandle() const
+{
+	return RootHandle;
+}
+
 int32 Chaos::FStrainedProxyModifier::GetNumRestBreakables() const
 {
 	if (RestChildren)
@@ -78,6 +83,21 @@ void Chaos::FStrainedProxyModifier::ClearStrains()
 	}
 }
 
+Chaos::FPBDRigidClusteredParticleHandle* Chaos::FStrainedProxyModifier::InitRootHandle(FGeometryCollectionPhysicsProxy* Proxy)
+{
+	if (Proxy == nullptr) { return nullptr; }
+
+	// Get the root index of the proxy
+	FSimulationParameters& Parameters = Proxy->GetSimParameters();
+	const int32 RootIndex = Parameters.InitialRootIndex;
+	if (RootIndex == INDEX_NONE) { return nullptr; }
+
+	// Return the particle handle corresponding to the root index
+	TArray<Chaos::FPBDRigidClusteredParticleHandle*>& ParticleHandles = Proxy->GetSolverParticleHandles();
+	if (!ParticleHandles.IsValidIndex(RootIndex)) { return nullptr; }
+	return ParticleHandles[RootIndex];
+}
+
 const TSet<int32>* Chaos::FStrainedProxyModifier::InitRestChildren(FGeometryCollectionPhysicsProxy* Proxy)
 {
 	if (Proxy == nullptr) { return nullptr; }
diff --git a/Engine/Source/Runtime/Experimental/Chaos/Public/Chaos/StrainModification.h b/Engine/Source/Runtime/Experimental/Chaos/Public/Chaos/StrainModification.h
--- a/Engine/Source/Runtime/Experimental/Chaos/Public/Chaos/StrainModification.h
+++ b/Engine/Source/Runtime/Experimental/Chaos/Public/Chaos/StrainModification.h
@@ -26,6 +26,7 @@ namespace Chaos
 		FStrainedProxyModifier(FRigidClustering& InRigidClustering, FGeometryCollectionPhysicsProxy* InProxy)
 			: RigidClustering(InRigidClustering)
 			, Proxy(InProxy)
+			, RootHandle(InitRootHandle(InProxy))
 			, RestChildren(InitRestChildren(InProxy))
 		{ }
 
@@ -38,6 +39,9 @@ namespace Chaos
 		// Get the proxy that owns the strained cluster or clusters
 		const FGeometryCollectionPhysicsProxy* GetProxy() const;
 
+		// Get the physics handle for the strained parent cluster
+		const Chaos::FPBDRigidParticleHandle* GetRootHandle() const;
+
 		// Get the number of level-1 strainable entities (number of rest-children in the per-particle
 		// strain model, or number of rest-connections in the edge/area model).
 		int32 GetNumRestBreakables() const;
@@ -50,10 +54,13 @@ namespace Chaos
 
 	private:
 
+		static Chaos::FPBDRigidClusteredParticleHandle* InitRootHandle(FGeometryCollectionPhysicsProxy* Proxy);
+
 		static const TSet<int32>* InitRestChildren(FGeometryCollectionPhysicsProxy* Proxy);
 
 		FRigidClustering& RigidClustering;
 		FGeometryCollectionPhysicsProxy* Proxy;
+		Chaos::FPBDRigidClusteredParticleHandle* RootHandle;
 		const TSet<int32>* RestChildren;
 	};
 
