New GeometryCollection helpers for getting allocation info.

diff --git a/Engine/Source/Runtime/Engine/Public/Rendering/PositionVertexBuffer.h b/Engine/Source/Runtime/Engine/Public/Rendering/PositionVertexBuffer.h
--- a/Engine/Source/Runtime/Engine/Public/Rendering/PositionVertexBuffer.h
+++ b/Engine/Source/Runtime/Engine/Public/Rendering/PositionVertexBuffer.h
@@ -96,6 +96,8 @@ public:
 		return NumVertices;
 	}
 	ENGINE_API bool GetAllowCPUAccess() const;
+	
+	FORCEINLINE SIZE_T GetAllocatedSize() const { return (Data != nullptr) ? Stride * NumVertices : 0; }
 
 	/** Create an RHI vertex buffer with CPU data. CPU data may be discarded after creation (see TResourceArray::Discard) */
 	FBufferRHIRef CreateRHIBuffer_RenderThread();
diff --git a/Engine/Source/Runtime/Experimental/GeometryCollectionEngine/Private/GeometryCollection/GeometryCollectionObject.cpp b/Engine/Source/Runtime/Experimental/GeometryCollectionEngine/Private/GeometryCollection/GeometryCollectionObject.cpp
--- a/Engine/Source/Runtime/Experimental/GeometryCollectionEngine/Private/GeometryCollection/GeometryCollectionObject.cpp
+++ b/Engine/Source/Runtime/Experimental/GeometryCollectionEngine/Private/GeometryCollection/GeometryCollectionObject.cpp
@@ -1263,6 +1263,27 @@ TSharedPtr<FGeometryCollection, ESPMode::ThreadSafe> UGeometryCollection::CopyCo
 
 #endif
 
+FGeometryCollectionRenderResourceSizeInfo UGeometryCollection::GetRenderResourceSizeInfo() const
+{
+	FGeometryCollectionRenderResourceSizeInfo InfoOut;
+	const FGeometryCollectionMeshResources& MeshResources = RenderData->MeshResource;
+	InfoOut.MeshResourcesSize += MeshResources.IndexBuffer.GetIndexDataSize();
+	InfoOut.MeshResourcesSize += MeshResources.PositionVertexBuffer.GetAllocatedSize();
+	InfoOut.MeshResourcesSize += MeshResources.StaticMeshVertexBuffer.GetResourceSize();
+	InfoOut.MeshResourcesSize += MeshResources.ColorVertexBuffer.GetAllocatedSize();
+	InfoOut.MeshResourcesSize += MeshResources.BoneMapVertexBuffer.GetAllocatedSize();
+
+	const Nanite::FResources& NaniteResource = RenderData->NaniteResource;
+	InfoOut.NaniteResourcesSize += NaniteResource.RootData.GetAllocatedSize();
+	InfoOut.NaniteResourcesSize += NaniteResource.ImposterAtlas.GetAllocatedSize();
+	InfoOut.NaniteResourcesSize += NaniteResource.HierarchyNodes.GetAllocatedSize();
+	InfoOut.NaniteResourcesSize += NaniteResource.HierarchyRootOffsets.GetAllocatedSize();
+	InfoOut.NaniteResourcesSize += NaniteResource.PageStreamingStates.GetAllocatedSize();
+	InfoOut.NaniteResourcesSize += NaniteResource.PageDependencies.GetAllocatedSize();
+
+	return InfoOut;
+}
+
 void UGeometryCollection::InitResources()
 {
 	if (RenderData)
diff --git a/Engine/Source/Runtime/Experimental/GeometryCollectionEngine/Private/GeometryCollection/GeometryCollectionRenderData.h b/Engine/Source/Runtime/Experimental/GeometryCollectionEngine/Private/GeometryCollection/GeometryCollectionRenderData.h
--- a/Engine/Source/Runtime/Experimental/GeometryCollectionEngine/Private/GeometryCollection/GeometryCollectionRenderData.h
+++ b/Engine/Source/Runtime/Experimental/GeometryCollectionEngine/Private/GeometryCollection/GeometryCollectionRenderData.h
@@ -29,6 +29,8 @@ public:
 		return *(uint16*)(Data + VertexIndex * PixelFormatStride);
 	}
 
+	SIZE_T GetAllocatedSize() const { return (BoneMapData != nullptr) ? BoneMapData->GetResourceSize() : 0; }
+
 protected:
 	void CleanUp();
 	void AllocateData(bool bInNeedsCPUAccess = true);
diff --git a/Engine/Source/Runtime/Experimental/GeometryCollectionEngine/Public/GeometryCollection/GeometryCollectionObject.h b/Engine/Source/Runtime/Experimental/GeometryCollectionEngine/Public/GeometryCollection/GeometryCollectionObject.h
--- a/Engine/Source/Runtime/Experimental/GeometryCollectionEngine/Public/GeometryCollection/GeometryCollectionObject.h
+++ b/Engine/Source/Runtime/Experimental/GeometryCollectionEngine/Public/GeometryCollection/GeometryCollectionObject.h
@@ -322,6 +322,20 @@ struct GEOMETRYCOLLECTIONENGINE_API FGeometryCollectionProxyMeshData
 	TArray<TObjectPtr<UStaticMesh>> ProxyMeshes;
 };
 
+USTRUCT()
+struct GEOMETRYCOLLECTIONENGINE_API FGeometryCollectionRenderResourceSizeInfo
+{
+	GENERATED_BODY();
+
+	// Total size of the arrays for the MeshResources
+	UPROPERTY()
+	uint64 MeshResourcesSize = 0;
+
+	// Total size of the arrays for the NaniteResources
+	UPROPERTY()
+	uint64 NaniteResourcesSize = 0;
+};
+
 /**
 * UGeometryCollectionObject (UObject)
 *
@@ -422,6 +436,9 @@ public:
 	/** copy a collection and remove geometry from it */
 	static TSharedPtr<FGeometryCollection, ESPMode::ThreadSafe> CopyCollectionAndRemoveGeometry(const TSharedPtr<const FGeometryCollection, ESPMode::ThreadSafe>& CollectionToCopy);
 
+	/** get the size of the render data resources associated with this collection */
+	FGeometryCollectionRenderResourceSizeInfo GetRenderResourceSizeInfo() const;
+
 #if WITH_EDITOR
 	/** If this flag is set, we only regenerate simulation data when requested via CreateSimulationData() */
 	bool bManualDataCreate;
