Chaos - CCD debug draw cleanup

diff --git a/Engine/Source/Runtime/Experimental/Chaos/Private/Chaos/CCDUtilities.cpp b/Engine/Source/Runtime/Experimental/Chaos/Private/Chaos/CCDUtilities.cpp
--- a/Engine/Source/Runtime/Experimental/Chaos/Private/Chaos/CCDUtilities.cpp
+++ b/Engine/Source/Runtime/Experimental/Chaos/Private/Chaos/CCDUtilities.cpp
@@ -7,7 +7,7 @@
 #include "Chaos/ChaosDebugDraw.h"
 #include "Chaos/DebugDrawQueue.h"
 
-//PRAGMA_DISABLE_OPTIMIZATION
+//UE_DISABLE_OPTIMIZATION
 
 namespace Chaos
 {
@@ -411,18 +411,6 @@ namespace Chaos
 			if (CVars::ChaosSolverDrawCCDInteractions)
 			{
 				const DebugDraw::FChaosDebugDrawSettings& DebugDrawSettings = CVars::ChaosSolverDebugDebugDrawSettings;
-				if (CCDParticle0)
-				{
-					const FRigidTransform3 SweepWorldTransform0 = FRigidTransform3(FConstGenericParticleHandle(CCDParticle0->Particle)->X(), FConstGenericParticleHandle(CCDParticle0->Particle)->Q());
-					const FRigidTransform3 ShapeWorldTransformXQ0 = CCDConstraint->SweptConstraint->GetShapeRelativeTransform0() * SweepWorldTransform0;
-					DebugDraw::DrawShape(ShapeWorldTransformXQ0, CCDConstraint->SweptConstraint->GetShape0()->GetLeafGeometry(), CCDConstraint->SweptConstraint->GetShape0(), FColor::Magenta, &DebugDrawSettings);
-				}
-				if (CCDParticle1)
-				{
-					const FRigidTransform3 SweepWorldTransform1 = FRigidTransform3(FConstGenericParticleHandle(CCDParticle1->Particle)->X(), FConstGenericParticleHandle(CCDParticle1->Particle)->Q());
-					const FRigidTransform3 ShapeWorldTransformXR1 = CCDConstraint->SweptConstraint->GetShapeRelativeTransform1() * SweepWorldTransform1;
-					DebugDraw::DrawShape(ShapeWorldTransformXR1, CCDConstraint->SweptConstraint->GetShape1()->GetLeafGeometry(), CCDConstraint->SweptConstraint->GetShape1(), FColor::Magenta, &DebugDrawSettings);
-				}
 				for (int32 ManifoldPointIndex = 0; ManifoldPointIndex < CCDConstraint->SweptConstraint->NumManifoldPoints(); ++ManifoldPointIndex)
 				{
 					const FManifoldPoint& ManifoldPoint = CCDConstraint->SweptConstraint->GetManifoldPoint(ManifoldPointIndex);
@@ -495,20 +483,12 @@ namespace Chaos
 		// Debugdraw the shapes at the final position
 		if (CVars::ChaosSolverDrawCCDInteractions)
 		{
-			for (FCCDConstraint* CCDConstraint : SortedCCDConstraints)
+			for (int32 ParticleIndex = ParticleStart; ParticleIndex < ParticleStart + ParticleNum; ParticleIndex++)
 			{
-				FCCDParticle* CCDParticle0 = CCDConstraint->Particle[0];
-				FCCDParticle* CCDParticle1 = CCDConstraint->Particle[1];
-				if (CCDParticle0)
-				{
-					const FRigidTransform3 ShapeWorldTransformPQ0 = CCDConstraint->SweptConstraint->GetShapeRelativeTransform0() * FConstGenericParticleHandle(CCDParticle0->Particle)->GetTransformPQ();
-					DebugDraw::DrawShape(ShapeWorldTransformPQ0, CCDConstraint->SweptConstraint->GetShape0()->GetLeafGeometry(), CCDConstraint->SweptConstraint->GetShape0(), FColor::Green, &CVars::ChaosSolverDebugDebugDrawSettings);
-				}
-				if (CCDParticle1)
-				{
-					const FRigidTransform3 ShapeWorldTransformPQ1 = CCDConstraint->SweptConstraint->GetShapeRelativeTransform1() * FConstGenericParticleHandle(CCDParticle1->Particle)->GetTransformPQ();
-					DebugDraw::DrawShape(ShapeWorldTransformPQ1, CCDConstraint->SweptConstraint->GetShape1()->GetLeafGeometry(), CCDConstraint->SweptConstraint->GetShape1(), FColor::Green, &CVars::ChaosSolverDebugDebugDrawSettings);
-				}
+				const FGeometryParticleHandle* Particle = GroupedCCDParticles[ParticleIndex]->Particle;
+				UE_LOG(LogChaos, Warning, TEXT("Particle <%d> TOI %f %s"), Particle->ParticleID().LocalID, GroupedCCDParticles[ParticleIndex]->TOI, *(*Particle->DebugName().Get()))
+
+				DebugDraw::DrawParticleShapes(FRigidTransform3::Identity, Particle, FColor::Green, &CVars::ChaosSolverDebugDebugDrawSettings);
 			}
 		}
 #endif
