Fix CIS - remove accidentally submitted log

diff --git a/Engine/Source/Runtime/Experimental/Chaos/Private/Chaos/CCDUtilities.cpp b/Engine/Source/Runtime/Experimental/Chaos/Private/Chaos/CCDUtilities.cpp
--- a/Engine/Source/Runtime/Experimental/Chaos/Private/Chaos/CCDUtilities.cpp
+++ b/Engine/Source/Runtime/Experimental/Chaos/Private/Chaos/CCDUtilities.cpp
@@ -486,8 +486,6 @@ namespace Chaos
 			for (int32 ParticleIndex = ParticleStart; ParticleIndex < ParticleStart + ParticleNum; ParticleIndex++)
 			{
 				const FGeometryParticleHandle* Particle = GroupedCCDParticles[ParticleIndex]->Particle;
-				UE_LOG(LogChaos, Warning, TEXT("Particle <%d> TOI %f %s"), Particle->ParticleID().LocalID, GroupedCCDParticles[ParticleIndex]->TOI, *(*Particle->DebugName().Get()))
-
 				DebugDraw::DrawParticleShapes(FRigidTransform3::Identity, Particle, FColor::Green, &CVars::ChaosSolverDebugDebugDrawSettings);
 			}
 		}
