Iris
* Always startup IrisCore module properly even if -UseIrisReplication is off.
* Always to go from iris to non-iris servers using the same client.

diff --git a/Engine/Source/Runtime/Experimental/Iris/Core/Private/Iris/IrisCoreModule.cpp b/Engine/Source/Runtime/Experimental/Iris/Core/Private/Iris/IrisCoreModule.cpp
--- a/Engine/Source/Runtime/Experimental/Iris/Core/Private/Iris/IrisCoreModule.cpp
+++ b/Engine/Source/Runtime/Experimental/Iris/Core/Private/Iris/IrisCoreModule.cpp
@@ -44,14 +44,11 @@ private:
 			UE::Net::SetUseIrisReplication(UseIrisReplication > 0);
 		}
 
-		if (UE::Net::ShouldUseIrisReplication())
-		{
-			RegisterPropertyNetSerializerSelectorTypes();
+		RegisterPropertyNetSerializerSelectorTypes();
 
-			UE_NET_IRIS_INIT_LEGACY_PUSH_MODEL();
+		UE_NET_IRIS_INIT_LEGACY_PUSH_MODEL();
 		
-			ModulesChangedHandle = FModuleManager::Get().OnModulesChanged().AddRaw(this, &FIrisCoreModule::OnModulesChanged);
-		}
+		ModulesChangedHandle = FModuleManager::Get().OnModulesChanged().AddRaw(this, &FIrisCoreModule::OnModulesChanged);
 	}
 
 	virtual void ShutdownModule() override
