Geometry collection : specialize bool typed managed array use a BitArray instead of a TArray

diff --git a/Engine/Plugins/Experimental/Dataflow/Source/DataflowEditor/Private/Dataflow/CollectionSpreadSheetWidget.cpp b/Engine/Plugins/Experimental/Dataflow/Source/DataflowEditor/Private/Dataflow/CollectionSpreadSheetWidget.cpp
--- a/Engine/Plugins/Experimental/Dataflow/Source/DataflowEditor/Private/Dataflow/CollectionSpreadSheetWidget.cpp
+++ b/Engine/Plugins/Experimental/Dataflow/Source/DataflowEditor/Private/Dataflow/CollectionSpreadSheetWidget.cpp
@@ -46,6 +46,11 @@ namespace CollectionSpreadSheetHelpers
 		return Value ? FString("true") : FString("false");
 	}
 
+	FString AttributeValueToString(const FConstBitReference& Value)
+	{
+		return Value ? FString("true") : FString("false");
+	}
+
 	FString AttributeValueToString(TSet<int32> Value)
 	{
 		TArray<int32> Array = Value.Array();
diff --git a/Engine/Plugins/Experimental/GeometryCollectionPlugin/Source/GeometryCollectionNodes/Private/Dataflow/GeometryCollectionNodes.cpp b/Engine/Plugins/Experimental/GeometryCollectionPlugin/Source/GeometryCollectionNodes/Private/Dataflow/GeometryCollectionNodes.cpp
--- a/Engine/Plugins/Experimental/GeometryCollectionPlugin/Source/GeometryCollectionNodes/Private/Dataflow/GeometryCollectionNodes.cpp
+++ b/Engine/Plugins/Experimental/GeometryCollectionPlugin/Source/GeometryCollectionNodes/Private/Dataflow/GeometryCollectionNodes.cpp
@@ -954,7 +954,8 @@ void FGetCollectionAttributeDataTypedDataflowNode::Evaluate(Dataflow::FContext&
 					if (TypeStr == FString("Bool"))
 					{
 						const TManagedArray<bool>& AttributeArr = InCollection.GetAttribute<bool>(FName(*AttrName), GroupNameToUse);
-						SetValue<TArray<bool>>(Context, AttributeArr.GetConstArray(), &BoolAttributeData);
+						TArray<bool> BoolArray = AttributeArr.GetAsBoolArray();
+						SetValue<TArray<bool>>(Context, MoveTemp(BoolArray), &BoolAttributeData);
 					}
 					else if (TypeStr == FString("Float"))
 					{
diff --git a/Engine/Source/Runtime/Experimental/Chaos/Public/GeometryCollection/ManagedArray.h b/Engine/Source/Runtime/Experimental/Chaos/Public/GeometryCollection/ManagedArray.h
--- a/Engine/Source/Runtime/Experimental/Chaos/Public/GeometryCollection/ManagedArray.h
+++ b/Engine/Source/Runtime/Experimental/Chaos/Public/GeometryCollection/ManagedArray.h
@@ -95,6 +95,12 @@ namespace ManagedArrayTypeSize
 		return Array.GetAllocatedSize();
 	}
 
+	template<typename T>
+	inline SIZE_T GetAllocatedSize(const TBitArray<>& Array)
+	{
+		return Array.GetAllocatedSize();
+	}
+
 	template<typename T>
 	inline SIZE_T GetAllocatedSize(const TSet<T>& Set)
 	{
@@ -648,6 +654,343 @@ void CopyRangeHelper(TArray<TUniquePtr<T>>& Target, const TManagedArrayBase<TUni
 	}
 }
 
+/***
+*  BitArray Managed Array base
+*/
+class FManagedBitArrayBase : public FManagedArrayBase
+{
+
+public:
+
+	FORCEINLINE FManagedBitArrayBase()
+	{}
+
+	FORCEINLINE FManagedBitArrayBase(const FManagedBitArrayBase& Other) = delete;
+
+	FORCEINLINE FManagedBitArrayBase(FManagedBitArrayBase&& Other)
+		: Array(MoveTemp(Other.Array))
+	{}
+
+	FORCEINLINE FManagedBitArrayBase& operator=(FManagedBitArrayBase&& Other)
+	{
+		return this->operator=(MoveTemp(Other.Array));
+	}
+
+	FORCEINLINE FManagedBitArrayBase& operator=(TBitArray<>&& Other)
+	{
+		// ryan - is it okay to check that the size matches?
+		ensureMsgf(Array.Num() == 0 || Array.Num() == Other.Num(), TEXT("FManagedBitArrayBase::operator=(TArray<T>&&) : Invalid array size."));
+		Array = MoveTemp(Other);
+		return *this;
+	}
+
+	virtual ~FManagedBitArrayBase()
+	{}
+
+	virtual void RemoveElements(const TArray<int32>& SortedDeletionList) override
+	{
+		if (SortedDeletionList.Num() == 0)
+		{
+			return;
+		}
+
+		// try to batch as many element as possible
+		int32 RangeStart = SortedDeletionList.Last();
+		for (int32 ii = SortedDeletionList.Num() - 1; ii > -1; --ii)
+		{
+			const int32 NumTopRemove = (RangeStart - SortedDeletionList[ii] + 1);
+			if (ii == 0)
+			{
+				Array.RemoveAt(SortedDeletionList[0], NumTopRemove);
+			}
+			else if (SortedDeletionList[ii] != (SortedDeletionList[ii - 1] + 1)) // compare this and previous values to make sure the difference is only 1.
+			{
+				Array.RemoveAt(SortedDeletionList[ii], NumTopRemove);
+				RangeStart = SortedDeletionList[ii - 1];
+			}
+		}
+	}
+
+	/**
+	* Init from a predefined Array of matching type
+	*/
+	virtual void Init(const FManagedArrayBase& NewArray) override
+	{
+		ensureMsgf(NewArray.GetTypeSize() == GetTypeSize(), TEXT("FManagedBitArrayBase::Init : Invalid array types."));
+		const FManagedBitArrayBase& TypedConstArray = static_cast<const FManagedBitArrayBase&>(NewArray);
+		
+		const int32 Size = TypedConstArray.Num();
+		Resize(Size);
+		for (int32 Index = 0; Index < Size; Index++)
+		{
+			Array[Index] = TypedConstArray[Index];
+		}
+	}
+
+	virtual SIZE_T GetAllocatedSize() const override
+	{
+		return ManagedArrayTypeSize::GetAllocatedSize(Array);
+	}
+
+	/**
+	* Copy from a predefined Array of matching type
+	*/
+	virtual void CopyRange(const FManagedArrayBase& ConstArray, int32 Start, int32 Stop, int32 Offset = 0) override
+	{
+		ensureMsgf(ConstArray.GetTypeSize() == GetTypeSize(), TEXT("TManagedArrayBase<T>::Init : Invalid array types."));
+		if (ensureMsgf(Stop + Offset < Array.Num(), TEXT("Error : Index out of bounds")))
+		{
+			const FManagedBitArrayBase& TypedConstArray = static_cast<const FManagedBitArrayBase&>(ConstArray);
+			for (int32 Sdx = Start, Tdx = Start + Offset; Sdx < ConstArray.Num() && Tdx < Array.Num() && Sdx < Stop; Sdx++, Tdx++)
+			{
+				Array[Tdx] = TypedConstArray[Sdx];
+			}
+		}
+	}
+
+	/**
+	 * Fill the array with \p Value.
+	 */
+	void Fill(const bool Value)
+	{
+		for (int32 Idx = 0; Idx < Array.Num(); ++Idx)
+		{
+			Array[Idx] = Value;
+		}
+	}
+
+	void Fill(const TArray<bool>& BoolArray)
+	{
+		check(BoolArray.Num() == Array.Num());
+		for (int32 Idx = 0; Idx < Array.Num(); Idx++)
+		{
+			Array[Idx] = BoolArray[Idx];
+		}
+	}
+
+	virtual void ExchangeArrays(FManagedArrayBase& NewArray) override
+	{
+		//It's up to the caller to make sure that the two arrays are of the same type
+		ensureMsgf(NewArray.GetTypeSize() == GetTypeSize(), TEXT("FManagedBitArrayBase::Exchange : Invalid array types."));
+		FManagedBitArrayBase& NewTypedArray = static_cast<FManagedBitArrayBase&>(NewArray);
+
+		Exchange(*this, NewTypedArray);
+	}
+
+	/**
+	* Returning a reference to the element at index.
+	*
+	* @returns Array element reference
+	*/
+	FORCEINLINE FBitReference operator[](int Index)
+	{
+		// @todo : optimization
+		// TArray->operator(Index) will perform checks against the 
+		// the array. It might be worth implementing the memory
+		// management directly on the ManagedArray, to avoid the
+		// overhead of the TArray.
+		return Array[Index];
+	}
+	FORCEINLINE const FConstBitReference operator[](int Index) const
+	{
+		return Array[Index];
+	}
+
+	/**
+	* Helper function for returning the internal const array
+	*
+	* @returns const array of all the elements
+	*/
+	FORCEINLINE const TBitArray<>& GetConstArray()
+	{
+		return Array;
+	}
+
+	FORCEINLINE const TBitArray<>& GetConstArray() const
+	{
+		return Array;
+	}
+
+	/**
+	* Helper function to convert the bitArray to a bool array 
+	* this is creating a new array and copy convert each bit to a bool
+	*/
+	TArray<bool> GetAsBoolArray() const
+	{
+		TArray<bool> BoolArray;
+		BoolArray.SetNumUninitialized(Array.Num());
+		for (int32 Idx = 0; Idx < Array.Num(); Idx++)
+		{
+			BoolArray[Idx] = Array[Idx];
+		}
+		return BoolArray;
+	}
+
+
+	/**
+	* Helper function for returning a typed pointer to the first array entry.
+	*
+	* @returns Pointer to first array entry or nullptr if ArrayMax == 0.
+	*/
+	//FORCEINLINE ElementType* GetData()
+	//{
+	//	return Array.GetData();
+	//}
+
+	/**
+	* Helper function for returning a typed pointer to the first array entry.
+	*
+	* @returns Pointer to first array entry or nullptr if ArrayMax == 0.
+	*/
+	//FORCEINLINE const ElementType* GetData() const
+	//{
+	//	return Array.GetData();
+	//}
+
+	/**
+	* Helper function returning the size of the inner type.
+	*
+	* @returns Size in bytes of array type.
+	*/
+	FORCEINLINE size_t GetTypeSize() const override
+	{
+		// this is not true but this ios the smallest we can represent
+		return 1;
+	}
+
+	/**
+	* Returning the size of the array
+	*
+	* @returns Array size
+	*/
+	FORCEINLINE int32 Num() const override
+	{
+		return Array.Num();
+	}
+
+	FORCEINLINE int32 Max() const override
+	{
+		return Array.Max();
+	}
+
+	FORCEINLINE bool Contains(const bool Item) const
+	{
+		return Array.Contains(Item);
+	}
+
+	/**
+	* Find first index of the element
+	*/
+	int32 Find(const bool Item) const
+	{
+		return Array.Find(Item);
+	}
+
+	/**
+	* Count the number of entries match \p Item.
+	*/
+	int32 Count(const bool Item) const
+	{
+		const int32 NumSetBits = Array.CountSetBits();
+		return (Item) ? NumSetBits : (Array.Num() - NumSetBits);
+	}
+
+	/**
+	* Checks if index is in array range.
+	*
+	* @param Index Index to check.
+	*/
+	FORCEINLINE void RangeCheck(int32 Index) const
+	{
+		checkf((Index >= 0) & (Index < Array.Num()), TEXT("Array index out of bounds: %i from an array of size %i"), Index, Array.Num());
+	}
+
+	/**
+	* Serialization Support
+	*
+	* @param Chaos::FChaosArchive& Ar
+	*/
+	virtual void Serialize(Chaos::FChaosArchive& Ar)
+	{
+		// we need to keep the backward compatibility with TManagedArray<bool> when it inherited from TManagedArrayBase<T>
+		Ar.UsingCustomVersion(FDestructionObjectVersion::GUID);
+		int Version = 1;
+		Ar << Version;
+
+		// for now always go through a bool array, in the future we can have a more optimized path 
+		TArray<bool> BoolArray;
+		if (Ar.IsSaving())
+		{
+			BoolArray = GetAsBoolArray();
+		}
+		if (Ar.CustomVer(FDestructionObjectVersion::GUID) < FDestructionObjectVersion::BulkSerializeArrays)
+		{
+			Ar << BoolArray;
+		}
+		else
+		{
+			TryBulkSerializeManagedArray(Ar, BoolArray);
+		}
+		if (Ar.IsLoading())
+		{
+			Resize(BoolArray.Num());
+			Fill(BoolArray);
+		}
+
+	}
+private:
+	/**
+	* Protected Resize to prevent external resizing of the array
+	*
+	* @param New array size.
+	*/
+	void Resize(const int32 Size)
+	{
+		if (Size > Array.Num())
+		{
+			Array.Add(false, Size - Array.Num());
+		}
+		else if (Size < Array.Num())
+		{
+			const int32 NumToRemove = (Array.Num() - Size);
+			Array.RemoveAt((Array.Num() - NumToRemove - 1), NumToRemove);
+		}
+	}
+
+	/**
+	* Protected Reserve to prevent external reservation of the array
+	*
+	* @param New array reservation size.
+	*/
+	void Reserve(const int32 Size)
+	{
+		Array.Reserve(Size);
+	}
+
+	/**
+	* Protected clear to prevent external clearing of the array
+	*
+	*/
+	void Empty()
+	{
+		Array.Empty();
+	}
+
+	void Reorder(const TArray<int32>& NewOrder) override
+	{
+		const int32 NumElements = Num();
+		check(NewOrder.Num() == NumElements);
+		TBitArray<> NewArray(false, NumElements);
+		for (int32 OriginalIdx = 0; OriginalIdx < NumElements; ++OriginalIdx)
+		{
+			NewArray[OriginalIdx] = Array[NewOrder[OriginalIdx]];
+		}
+		Exchange(Array, NewArray);
+	}
+
+	TBitArray<> Array;
+};
+
 //
 //
 //
@@ -763,6 +1106,10 @@ public:
 };
 
 
+template<>
+class TManagedArray<bool> : public FManagedBitArrayBase
+{};
+
 template<>
 class TManagedArray<TSet<int32>> : public TManagedArrayBase<TSet<int32>>
 {
