Chaos - CCD fixes
- don't duplicate ShouldUsCCD call in UpdateConstraintFromGeometrySweptImpl
- fix bounds for sweeps against heightfields
- run collision detection even if the sweep misses because we may be within CullDistance of something even if the sweep misses
- fix sweep distance calculation in ShouldEnableCCD

diff --git a/Engine/Source/Runtime/Experimental/Chaos/Private/Chaos/Collision/PBDCollisionConstraint.cpp b/Engine/Source/Runtime/Experimental/Chaos/Private/Chaos/Collision/PBDCollisionConstraint.cpp
--- a/Engine/Source/Runtime/Experimental/Chaos/Private/Chaos/Collision/PBDCollisionConstraint.cpp
+++ b/Engine/Source/Runtime/Experimental/Chaos/Private/Chaos/Collision/PBDCollisionConstraint.cpp
@@ -13,7 +13,7 @@
 #include "Chaos/PBDCollisionConstraints.h"
 
 
-//PRAGMA_DISABLE_OPTIMIZATION
+//UE_DISABLE_OPTIMIZATION
 
 namespace Chaos
 {
@@ -399,8 +399,18 @@ namespace Chaos
 		const FRealSingle MinBounds0 = FConstGenericParticleHandle(Particle[0])->CCDEnabled() ? FRealSingle(Implicit[0]->BoundingBox().Extents().GetAbsMin()) : FRealSingle(0);
 		const FRealSingle MinBounds1 = FConstGenericParticleHandle(Particle[1])->CCDEnabled() ? FRealSingle(Implicit[1]->BoundingBox().Extents().GetAbsMin()) : FRealSingle(0);
 		const FRealSingle MinBounds = FMath::Max(MinBounds0, MinBounds1);
+
+		// If we perform a CCD sweep and it finds a hit, but the the penetration depth at the end of
+		// the sweep (T=1) is less than this, we ignore the sweep result (no rewind) and just run
+		// regular collision detection at T=1 (i.e., as if CCD is not enabled for this tick).
+		// If this is zero, then we will always perform CCD rewind when we get a sweep hit.
 		CCDEnablePenetration = MinBounds * CVars::CCDEnableThresholdBoundsScale;
-		CCDTargetPenetration = FMath::Min(MinBounds * CVars::CCDAllowedDepthBoundsScale, CCDEnablePenetration);
+		
+		// If we have a CCD sweep hit above the enable penetration, we will rewind the body so
+		// that the penetration depth CCDTargetPenetration (or less). A new contact manifold is then calculated.
+		// NOTE: this can be zero but if it is we will get very poor behaviour because there will be
+		// nothing for the contact resolution to do since we fully depenetrated in the CCD rewind.
+		CCDTargetPenetration = MinBounds * CVars::CCDAllowedDepthBoundsScale;
 	}
 
 	void FPBDCollisionConstraint::Activate()
diff --git a/Engine/Source/Runtime/Experimental/Chaos/Private/Chaos/Collision/ParticlePairMidPhase.cpp b/Engine/Source/Runtime/Experimental/Chaos/Private/Chaos/Collision/ParticlePairMidPhase.cpp
--- a/Engine/Source/Runtime/Experimental/Chaos/Private/Chaos/Collision/ParticlePairMidPhase.cpp
+++ b/Engine/Source/Runtime/Experimental/Chaos/Private/Chaos/Collision/ParticlePairMidPhase.cpp
@@ -13,6 +13,8 @@
 #include "Chaos/ChaosDebugDraw.h"
 #include "ChaosStats.h"
 
+//UE_DISABLE_OPTIMIZATION
+
 extern bool Chaos_Collision_NarrowPhase_AABBBoundsCheck;
 
 namespace Chaos
@@ -415,8 +417,8 @@ namespace Chaos
 
 			// We need the previous transform for the swept collision detector. It assumes that the current
 			// transform has been set on the constraint. 
-			// We assume that the particle's center of mass moved in a stright line and that it's rotation has 
-			// not changed so we calculate the previous transform from the current one and the velocity.
+			// We assume that the particle's center of mass moved in a straight line and that it's rotation has 
+			// not changed so we calculate sthe previous transform from the current one and the velocity.
 			// NOTE: These are actor transforms, not CoM transforms
 			// @todo(chaos): Pass both start and end transforms to the collision detector
 			const FRigidTransform3 CCDParticleWorldTransform0 = FRigidTransform3(P0->P() - P0->V() * Dt, P0->Q());
@@ -441,15 +443,13 @@ namespace Chaos
 			}
 #endif
 
-			// If we did get a hit but it's at TOI = 1, treat this constraint as a regular non-swept constraint (skip the rewind)
-			if ((!bDidSweep) || Constraint->GetCCDTimeOfImpact() == FReal(1))
+			// If we did not get a sweep hit (TOI > 1) or did not sweep (bDidSweep = false), we need to run standard collision detection at T=1.
+			// Likewise, if we did get a sweep hit but it's at TOI = 1, treat this constraint as a regular non-swept constraint and skip the rewind.
+			// NOTE: The sweep will report TOI==1 for "shallow" sweep hits below the CCD thresholds in the constraint.
+			if ((!bDidSweep) || (Constraint->GetCCDTimeOfImpact() >= FReal(1)))
 			{
+				// @todo(chaos): should we use a reduced cull distance if we get here? The cull distance will have been set based on movement speed...
 				Collisions::UpdateConstraint(*Constraint.Get(), Constraint->GetShapeWorldTransform0(), Constraint->GetShapeWorldTransform1(), Dt);
-			}
-
-			// If the sweep did not find a hit, or we are treating it like a regular contact, we will skip the CCD rewind step for this contact
-			if (Constraint->GetCCDTimeOfImpact() >= FReal(1))
-			{
 				Constraint->SetCCDSweepEnabled(false);
 			}
 
@@ -1092,10 +1092,11 @@ namespace Chaos
 			FConstGenericParticleHandle ConstParticle0 = FConstGenericParticleHandle(Particle0);
 			FConstGenericParticleHandle ConstParticle1 = FConstGenericParticleHandle(Particle1);
 
+			// See comments in FSingleShapePairCollisionDetector::GenerateCollisionCCDImpl for why we use V() here
 			FReal LengthCCD = 0;
 			FVec3 DirCCD = FVec3(0);
-			const FVec3 DeltaX0 = (ConstParticle0->ObjectState() == EObjectStateType::Kinematic) ? (ConstParticle0->V() * Dt) : (ConstParticle0->P() - ConstParticle0->X());
-			const FVec3 DeltaX1 = (ConstParticle1->ObjectState() == EObjectStateType::Kinematic) ? (ConstParticle1->V() * Dt) : (ConstParticle1->P() - ConstParticle1->X());
+			const FVec3 DeltaX0 = ConstParticle0->V() * Dt;
+			const FVec3 DeltaX1 = ConstParticle1->V() * Dt;
 			const bool bUseCCD = Collisions::ShouldUseCCD(Particle0, DeltaX0, Particle1, DeltaX1, DirCCD, LengthCCD);
 
 			return bUseCCD;
diff --git a/Engine/Source/Runtime/Experimental/Chaos/Private/Chaos/CollisionResolution.cpp b/Engine/Source/Runtime/Experimental/Chaos/Private/Chaos/CollisionResolution.cpp
--- a/Engine/Source/Runtime/Experimental/Chaos/Private/Chaos/CollisionResolution.cpp
+++ b/Engine/Source/Runtime/Experimental/Chaos/Private/Chaos/CollisionResolution.cpp
@@ -35,7 +35,7 @@
 #include "Chaos/Transform.h"
 #include "Chaos/TriangleMeshImplicitObject.h"
 
-//PRAGMA_DISABLE_OPTIMIZATION
+//UE_DISABLE_OPTIMIZATION
 
 
 // @todo(chaos): clean up the contact creation time rejection to avoid extra transforms
@@ -553,7 +553,7 @@ namespace Chaos
 
 		void UpdateSphereHeightFieldConstraintSwept(TGeometryParticleHandle<FReal, 3>* Particle0, const TSphere<FReal, 3>& A, const FRigidTransform3& ATransform, const FHeightField& B, const FRigidTransform3& BTransform, const FVec3& Dir, const FReal Length, const FReal Dt, FPBDCollisionConstraint& Constraint)
 		{
-			FReal TOI = 1.0f;
+			FReal TOI = FPBDCollisionConstraint::MaxTOI;
 			UpdateContactPointNoCull(Constraint, GJKImplicitSweptContactPoint(A, ATransform, B, BTransform, Dir, Length, Constraint.GetCCDEnablePenetration(), Constraint.GetCCDTargetPenetration(), TOI), Dt, CCDNoCullAllShapePairs);
 			Constraint.SetCCDTimeOfImpact(TOI);
 		}
@@ -645,7 +645,7 @@ namespace Chaos
 		template<typename TriMeshType>
 		void UpdateSphereTriangleMeshConstraintSwept(TGeometryParticleHandle<FReal, 3>* Particle0, const TSphere<FReal, 3>& Sphere0, const FRigidTransform3& WorldTransform0, const TriMeshType& TriangleMesh1, const FRigidTransform3& WorldTransform1, const FVec3& Dir, const FReal Length, const FReal Dt, FPBDCollisionConstraint& Constraint)
 		{
-			FReal TOI = 1.0f;
+			FReal TOI = FPBDCollisionConstraint::MaxTOI;
 			UpdateContactPointNoCull(Constraint, SphereTriangleMeshSweptContactPoint(Sphere0, WorldTransform0, TriangleMesh1, WorldTransform1, Dir, Length, Constraint.GetCCDEnablePenetration(), Constraint.GetCCDTargetPenetration(), TOI), Dt, CCDNoCullAllShapePairs);
 			Constraint.SetCCDTimeOfImpact(TOI);
 		}
@@ -775,7 +775,7 @@ namespace Chaos
 		void UpdateCapsuleHeightFieldConstraintSwept(TGeometryParticleHandle<FReal, 3>* Particle0, const FCapsule& A, const FRigidTransform3& ATransform, const FHeightField& B, const FRigidTransform3& BTransform, const FVec3& Dir, const FReal Length, const FReal Dt, FPBDCollisionConstraint& Constraint)
 		{
 			CONDITIONAL_SCOPE_CYCLE_COUNTER(STAT_Collisions_UpdateCapsuleHeightFieldConstraintSwept, ConstraintsDetailedStats);
-			FReal TOI = 1.0f;
+			FReal TOI = FPBDCollisionConstraint::MaxTOI;
 			UpdateContactPointNoCull(Constraint, GJKImplicitSweptContactPoint(A, ATransform, B, BTransform, Dir, Length, Constraint.GetCCDEnablePenetration(), Constraint.GetCCDTargetPenetration(), TOI), Dt, CCDNoCullAllShapePairs);
 			Constraint.SetCCDTimeOfImpact(TOI);
 		}
@@ -803,7 +803,7 @@ namespace Chaos
 		void UpdateCapsuleTriangleMeshConstraintSwept(TGeometryParticleHandle<FReal, 3>* Particle0, const FCapsule& Capsule0, const FRigidTransform3& WorldTransform0, const TriMeshType& TriangleMesh1, const FRigidTransform3& WorldTransform1, const FVec3& Dir, const FReal Length, const FReal Dt, FPBDCollisionConstraint& Constraint)
 		{
 			CONDITIONAL_SCOPE_CYCLE_COUNTER(STAT_Collisions_UpdateCapsuleTriangleMeshConstraint, ConstraintsDetailedStats);
-			FReal TOI = 1.0f;
+			FReal TOI = FPBDCollisionConstraint::MaxTOI;
 			UpdateContactPointNoCull(Constraint, CapsuleTriangleMeshSweptContactPoint(Capsule0, WorldTransform0, TriangleMesh1, WorldTransform1, Dir, Length, Constraint.GetCCDEnablePenetration(), Constraint.GetCCDTargetPenetration(), TOI), Dt, CCDNoCullAllShapePairs);
 			Constraint.SetCCDTimeOfImpact(TOI);
 		}
@@ -850,7 +850,7 @@ namespace Chaos
 		void UpdateConvexHeightFieldConstraintSwept(TGeometryParticleHandle<FReal, 3>* Particle0, const FImplicitObject& A, const FRigidTransform3& ATransform, const FHeightField& B, const FRigidTransform3& BTransform, const FVec3& Dir, const FReal Length, const FReal Dt, FPBDCollisionConstraint& Constraint)
 		{
 			CONDITIONAL_SCOPE_CYCLE_COUNTER(STAT_Collisions_UpdateConvexHeightFieldConstraintSwept, ConstraintsDetailedStats);
-			FReal TOI = 1.0f;
+			FReal TOI = FPBDCollisionConstraint::MaxTOI;
 			UpdateContactPointNoCull(Constraint, GJKImplicitSweptContactPoint(A, ATransform, B, BTransform, Dir, Length, Constraint.GetCCDEnablePenetration(), Constraint.GetCCDTargetPenetration(), TOI), Dt, CCDNoCullAllShapePairs);
 			Constraint.SetCCDTimeOfImpact(TOI);
 		}
@@ -877,7 +877,7 @@ namespace Chaos
 		void UpdateConvexTriangleMeshConstraintSwept(TGeometryParticleHandle<FReal, 3>* Particle0, const FImplicitObject& Convex0, const FRigidTransform3& WorldTransform0, const TriMeshType& TriangleMesh1, const FRigidTransform3& WorldTransform1, const FVec3& Dir, const FReal Length, const FReal Dt, FPBDCollisionConstraint& Constraint)
 		{
 			CONDITIONAL_SCOPE_CYCLE_COUNTER(STAT_Collisions_UpdateConvexTriangleMeshConstraintSwept, ConstraintsDetailedStats);
-			FReal TOI = 1.0f;
+			FReal TOI = FPBDCollisionConstraint::MaxTOI;
 			UpdateContactPointNoCull(Constraint, ConvexTriangleMeshSweptContactPoint(Convex0, WorldTransform0, TriangleMesh1, WorldTransform1, Dir, Length, Constraint.GetCCDEnablePenetration(), Constraint.GetCCDTargetPenetration(), TOI), Dt, CCDNoCullAllShapePairs);
 			Constraint.SetCCDTimeOfImpact(TOI);
 		}
@@ -1064,12 +1064,11 @@ namespace Chaos
 			const FImplicitObject& Implicit0 = *Constraint.GetImplicit0();
 			const FImplicitObject& Implicit1 = *Constraint.GetImplicit1();
 
-			// @todo(chaos): this work has alreday been done in the MidPhase - check and remove/test
-			FReal LengthCCD = 0.0f;
-			FVec3 DirCCD(0.0f);
 			const FVec3 DeltaX0 = Constraint.GetShapeWorldTransform0().GetTranslation() - WorldTransform0.GetLocation();
 			const FVec3 DeltaX1 = Constraint.GetShapeWorldTransform1().GetTranslation() - WorldTransform1.GetLocation();
-			bool bUseCCD = ShouldUseCCD(Particle0, DeltaX0, Particle1, DeltaX1, DirCCD, LengthCCD);
+			FVec3 DirCCD = DeltaX0 - DeltaX1;
+			const FReal LengthCCD = DirCCD.SafeNormalize();
+			const bool bUseCCD = (LengthCCD > 0);
 
 			Constraint.ResetCCDTimeOfImpact();
 
diff --git a/Engine/Source/Runtime/Experimental/Chaos/Private/Chaos/HeightField.cpp b/Engine/Source/Runtime/Experimental/Chaos/Private/Chaos/HeightField.cpp
--- a/Engine/Source/Runtime/Experimental/Chaos/Private/Chaos/HeightField.cpp
+++ b/Engine/Source/Runtime/Experimental/Chaos/Private/Chaos/HeightField.cpp
@@ -19,7 +19,7 @@
 #include "Chaos/Triangle.h"
 #include "Chaos/TriangleRegister.h"
 
-//PRAGMA_DISABLE_OPTIMIZATION
+//UE_DISABLE_OPTIMIZATION
 
 namespace Chaos
 {
@@ -450,7 +450,7 @@ namespace Chaos
 	{
 	public:
 
-		THeightfieldSweepVisitorCCD(const typename FHeightField::FDataType* InData, const GeomQueryType& InQueryGeom, const FRigidTransform3& InStartTM, const FVec3& InDir, const FReal Length, const FReal InIgnorePenetration, const FReal InTargetPenetration)
+		THeightfieldSweepVisitorCCD(const typename FHeightField::FDataType* InData, const GeomQueryType& InQueryGeom, const FAABB3& RotatedQueryBounds, const FRigidTransform3& InStartTM, const FVec3& InDir, const FReal Length, const FReal InIgnorePenetration, const FReal InTargetPenetration)
 			: OutDistance(TNumericLimits<FReal>::Max())
 			, OutPhi(TNumericLimits<FReal>::Max())
 			, OutPosition(0)
@@ -462,10 +462,9 @@ namespace Chaos
 			, IgnorePenetration(InIgnorePenetration)
 			, TargetPenetration(InTargetPenetration)
 		{
-			const FAABB3 QueryBounds = InQueryGeom.BoundingBox();
-			const FVec3 BoundsStartPoint = StartTM.TransformPositionNoScale(QueryBounds.Center());
 			const FVec3 StartPoint = StartTM.GetTranslation();
-			const FVec3 Inflation3D = QueryBounds.Extents() * FReal(0.5);
+			const FVec3 BoundsStartPoint = StartPoint + RotatedQueryBounds.Center();
+			const FVec3 Inflation3D = RotatedQueryBounds.Extents() * FReal(0.5);
 
 			const VectorRegister4Float LengthSimd = MakeVectorRegisterFloatFromDouble(VectorSetFloat1(Length));
 			BoundsStartPointSimd = MakeVectorRegisterFloatFromDouble(MakeVectorRegister(BoundsStartPoint.X, BoundsStartPoint.Y, BoundsStartPoint.Z, 0.0));
@@ -2205,9 +2204,10 @@ namespace Chaos
 		const FReal IgnorePenetration = CVars::bCCDNewTargetDepthMode ? InIgnorePenetration : 0;
 		const FReal TargetPenetration = CVars::bCCDNewTargetDepthMode ? InTargetPenetration : 0;
 
+		const FAABB3 QueryBounds = QueryGeom.BoundingBox().TransformedAABB(FRigidTransform3(FVec3(0), StartTM.GetRotation()));
+
 		bool bHit = false;
-		THeightfieldSweepVisitorCCD<QueryGeomType> SQVisitor(&GeomData, QueryGeom, StartTM, Dir, Length, IgnorePenetration, TargetPenetration);
-		const FAABB3 QueryBounds = QueryGeom.BoundingBox();
+		THeightfieldSweepVisitorCCD<QueryGeomType> SQVisitor(&GeomData, QueryGeom, QueryBounds, StartTM, Dir, Length, IgnorePenetration, TargetPenetration);
 		const FVec3 StartPoint = StartTM.TransformPositionNoScale(QueryBounds.Center());
 
 		const FVec3 Inflation3D = QueryBounds.Extents() * FReal(0.5);
