Make Initial velocity attributes in dynamic collection optional to avoid allocating them when not in use

diff --git a/Engine/Source/Runtime/Experimental/Chaos/Private/GeometryCollectionProxyData.cpp b/Engine/Source/Runtime/Experimental/Chaos/Private/GeometryCollectionProxyData.cpp
--- a/Engine/Source/Runtime/Experimental/Chaos/Private/GeometryCollectionProxyData.cpp
+++ b/Engine/Source/Runtime/Experimental/Chaos/Private/GeometryCollectionProxyData.cpp
@@ -56,16 +56,57 @@ FGeometryDynamicCollection::FGeometryDynamicCollection()
 	AddExternalAttribute("CollisionStructureID", FTransformCollection::TransformGroup, CollisionStructureID);
 	AddExternalAttribute<int32>(FGeometryDynamicCollection::DynamicStateAttribute, FTransformCollection::TransformGroup, DynamicState);
 	AddExternalAttribute(ImplicitsAttribute, FTransformCollection::TransformGroup, Implicits);
-	AddExternalAttribute("InitialAngularVelocity", FTransformCollection::TransformGroup, InitialAngularVelocity);
-	AddExternalAttribute("InitialLinearVelocity", FTransformCollection::TransformGroup, InitialLinearVelocity);
 	AddExternalAttribute("MassToLocal", FTransformCollection::TransformGroup, MassToLocal);
-	//AddExternalAttribute(ShapesQueryDataAttribute, FTransformCollection::TransformGroup, ShapeQueryData);
-	//AddExternalAttribute(ShapesSimDataAttribute, FTransformCollection::TransformGroup, ShapeSimData);
 	AddExternalAttribute(SimplicialsAttribute, FTransformCollection::TransformGroup, Simplicials);
 	AddExternalAttribute(SimulatableParticlesAttribute, FGeometryCollection::TransformGroup, SimulatableParticles);
 
 }
 
+void FGeometryDynamicCollection::CopyInitialVelocityAttributesFrom(const FGeometryDynamicCollection& SourceCollection)
+{
+	FInitialVelocityFacade InitialVelocityFacade(*this);
+	InitialVelocityFacade.CopyFrom(SourceCollection);
+}
+
+FGeometryDynamicCollection::FInitialVelocityFacade::FInitialVelocityFacade(FGeometryDynamicCollection& DynamicCollection)
+	: InitialLinearVelocityAttribute(DynamicCollection, "InitialLinearVelocity", FTransformCollection::TransformGroup)
+	, InitialAngularVelocityAttribute(DynamicCollection, "InitialAngularVelocity", FTransformCollection::TransformGroup)
+{}
+
+FGeometryDynamicCollection::FInitialVelocityFacade::FInitialVelocityFacade(const FGeometryDynamicCollection& DynamicCollection)
+	: InitialLinearVelocityAttribute(DynamicCollection, "InitialLinearVelocity", FTransformCollection::TransformGroup)
+	, InitialAngularVelocityAttribute(DynamicCollection, "InitialAngularVelocity", FTransformCollection::TransformGroup)
+{}
+
+bool FGeometryDynamicCollection::FInitialVelocityFacade::IsValid() const
+{
+	return InitialLinearVelocityAttribute.IsValid() && InitialAngularVelocityAttribute.IsValid();
+}
+
+void FGeometryDynamicCollection::FInitialVelocityFacade::DefineSchema()
+{
+	InitialLinearVelocityAttribute.Add();
+	InitialAngularVelocityAttribute.Add();
+}
+
+void FGeometryDynamicCollection::FInitialVelocityFacade::Fill(const FVector3f& InitialLinearVelocity, const FVector3f& InitialAngularVelocity)
+{
+	check(IsValid());
+	InitialLinearVelocityAttribute.Fill(InitialLinearVelocity);
+	InitialAngularVelocityAttribute.Fill(InitialAngularVelocity);
+}
+
+void FGeometryDynamicCollection::FInitialVelocityFacade::CopyFrom(const FGeometryDynamicCollection& SourceCollection)
+{
+	FInitialVelocityFacade SourceInitialVelocityFacade(SourceCollection);
+	if (SourceInitialVelocityFacade.IsValid())
+	{
+		DefineSchema();
+		InitialLinearVelocityAttribute.Copy(SourceInitialVelocityFacade.InitialLinearVelocityAttribute);
+		InitialAngularVelocityAttribute.Copy(SourceInitialVelocityFacade.InitialAngularVelocityAttribute);
+	}
+}
+
 ///////////////////////////////////////////////////////////////////////////////////////////////////////////////////////////
 
 FGeometryCollectionDynamicStateFacade::FGeometryCollectionDynamicStateFacade(FManagedArrayCollection& InCollection)
diff --git a/Engine/Source/Runtime/Experimental/Chaos/Private/PhysicsProxy/GeometryCollectionPhysicsProxy.cpp b/Engine/Source/Runtime/Experimental/Chaos/Private/PhysicsProxy/GeometryCollectionPhysicsProxy.cpp
--- a/Engine/Source/Runtime/Experimental/Chaos/Private/PhysicsProxy/GeometryCollectionPhysicsProxy.cpp
+++ b/Engine/Source/Runtime/Experimental/Chaos/Private/PhysicsProxy/GeometryCollectionPhysicsProxy.cpp
@@ -664,6 +664,7 @@ void FGeometryCollectionPhysicsProxy::Initialize(Chaos::FPBDRigidsEvolutionBase
 	TransformGroupSkipList.Add(DynamicCollection.SimplicialsAttribute);
 
 	PhysicsThreadCollection.CopyMatchingAttributesFrom(DynamicCollection, &SkipList);
+	PhysicsThreadCollection.CopyInitialVelocityAttributesFrom(DynamicCollection);
 
 	// make sure we copy the anchored information over to the physics thread collection
 	const Chaos::Facades::FCollectionAnchoringFacade DynamicCollectionAnchoringFacade(DynamicCollection);
@@ -738,15 +739,12 @@ void FGeometryCollectionPhysicsProxy::InitializeDynamicCollection(FGeometryDynam
 	DynamicCollection.CopyMatchingAttributesFrom(RestCollection, &SkipList);
 
 
-	//
 	// User defined initial velocities need to be populated. 
-	//
+	if (Params.InitialVelocityType == EInitialVelocityTypeEnum::Chaos_Initial_Velocity_User_Defined)
 	{
-		if (Params.InitialVelocityType == EInitialVelocityTypeEnum::Chaos_Initial_Velocity_User_Defined)
-		{
-			DynamicCollection.InitialLinearVelocity.Fill(FVector3f(Params.InitialLinearVelocity));
-			DynamicCollection.InitialAngularVelocity.Fill(FVector3f(Params.InitialAngularVelocity));
-		}
+		FGeometryDynamicCollection::FInitialVelocityFacade InitialVelocityFacade = DynamicCollection.GetInitialVelocityFacade();
+		InitialVelocityFacade.DefineSchema();
+		InitialVelocityFacade.Fill(FVector3f(Params.InitialLinearVelocity), FVector3f(Params.InitialAngularVelocity));
 	}
 
 	// process simplicials
@@ -929,8 +927,6 @@ void FGeometryCollectionPhysicsProxy::InitializeBodiesPT(Chaos::FPBDRigidsSolver
 		const TManagedArray<int32>& CollisionGroup = DynamicCollection.CollisionGroup;
 		const TManagedArray<bool>& SimulatableParticles = DynamicCollection.SimulatableParticles;
 		const TManagedArray<FTransform>& MassToLocal = DynamicCollection.MassToLocal;
-		const TManagedArray<FVector3f>& InitialAngularVelocity = DynamicCollection.InitialAngularVelocity;
-		const TManagedArray<FVector3f>& InitialLinearVelocity = DynamicCollection.InitialLinearVelocity;
 		const TManagedArray<FGeometryDynamicCollection::FSharedImplicit>& Implicits = DynamicCollection.Implicits;
 		const TManagedArray<TUniquePtr<FCollisionStructureManager::FSimplicial>>& Simplicials = DynamicCollection.Simplicials;
 		const TManagedArray<TSet<int32>>& Children = DynamicCollection.Children;
@@ -1067,6 +1063,8 @@ void FGeometryCollectionPhysicsProxy::InitializeBodiesPT(Chaos::FPBDRigidsSolver
 
 		if (Parameters.InitialVelocityType == EInitialVelocityTypeEnum::Chaos_Initial_Velocity_User_Defined)
 		{
+			FGeometryDynamicCollection::FInitialVelocityFacade InitialVelocityFacade = DynamicCollection.GetInitialVelocityFacade();
+			check(InitialVelocityFacade.IsValid());
 			// A previous implementation of this went wide on this loop.  The general 
 			// rule of thumb for parallelization is that each thread needs at least
 			// 1000 operations in order to overcome the expense of threading.  I don't
@@ -1077,8 +1075,8 @@ void FGeometryCollectionPhysicsProxy::InitializeBodiesPT(Chaos::FPBDRigidsSolver
 				{
 					if (DynamicState[TransformGroupIndex] == (int32)EObjectStateTypeEnum::Chaos_Object_Dynamic)
 					{
-						Handle->SetV(InitialLinearVelocity[TransformGroupIndex]);
-						Handle->SetW(InitialAngularVelocity[TransformGroupIndex]);
+						Handle->SetV(InitialVelocityFacade.InitialAngularVelocityAttribute[TransformGroupIndex]);
+						Handle->SetW(InitialVelocityFacade.InitialAngularVelocityAttribute[TransformGroupIndex]);
 					}
 				}
 			}
@@ -4346,6 +4344,8 @@ void FGeometryCollectionPhysicsProxy::FieldParameterUpdateCallback(Chaos::FPBDRi
 
 								bool bHasStateChanged = false;
 								
+								const FGeometryDynamicCollection::FInitialVelocityFacade InitialVelocityFacade = Collection.GetInitialVelocityFacade();
+
 								const TFieldArrayView<FFieldContextIndex>& EvaluatedSamples = FieldContext.GetEvaluatedSamples();
 								for (const FFieldContextIndex& Index : EvaluatedSamples)
 								{
@@ -4359,18 +4359,21 @@ void FGeometryCollectionPhysicsProxy::FieldParameterUpdateCallback(Chaos::FPBDRi
 										const int8 ResultState = static_cast<int8>(CurrResult);
 										const int32 TransformIndex = HandleToTransformGroupIndex[RigidHandle];
 
-										// Update of the handles object state. No need to update 
-										// the initial velocities since it is done after this function call in InitializeBodiesPT
-										if (bUpdateViews && (Parameters.InitialVelocityType == EInitialVelocityTypeEnum::Chaos_Initial_Velocity_User_Defined))
-										{
-											bHasStateChanged |= ReportDynamicStateResult(RigidSolver, static_cast<Chaos::EObjectStateType>(ResultState), RigidHandle,
-												true, Collection.InitialLinearVelocity[TransformIndex], true, Collection.InitialAngularVelocity[TransformIndex]);
-										}
-										else
+										Chaos::FVec3 InitialLinearVelocity{ 0 };
+										Chaos::FVec3 InitialAngularVelocity{ 0 };
+
+										// bUpdateViews is used here to avoid applying the velocity twice when running the Initial fields ( see InitializeBodiesPT )
+										const bool bHasInitialVelocities = (bUpdateViews && (Parameters.InitialVelocityType == EInitialVelocityTypeEnum::Chaos_Initial_Velocity_User_Defined));
+										if (bHasInitialVelocities)
 										{
-											bHasStateChanged |= ReportDynamicStateResult(RigidSolver, static_cast<Chaos::EObjectStateType>(ResultState), RigidHandle,
-												false, Chaos::FVec3(0), false, Chaos::FVec3(0));
+											InitialLinearVelocity = InitialVelocityFacade.InitialLinearVelocityAttribute[TransformIndex];
+											InitialAngularVelocity = InitialVelocityFacade.InitialAngularVelocityAttribute[TransformIndex];
 										}
+
+										// Update of the handles object state. 
+										bHasStateChanged |= ReportDynamicStateResult(RigidSolver, static_cast<Chaos::EObjectStateType>(ResultState), RigidHandle,
+												bHasInitialVelocities, InitialLinearVelocity, bHasInitialVelocities, InitialAngularVelocity);
+
 										// Update of the Collection dynamic state. It will be used just after to set the initial velocity
 										Collection.DynamicState[TransformIndex] = ResultState;
 									}
@@ -4400,15 +4403,24 @@ void FGeometryCollectionPhysicsProxy::FieldParameterUpdateCallback(Chaos::FPBDRi
 							{
 								SCOPE_CYCLE_COUNTER(STAT_ParamUpdateField_LinearVelocity);
 								{
-									static_cast<const FFieldNode<FVector>*>(FieldCommand.RootNode.Get())->Evaluate(FieldContext, ResultsView);
-									for (const FFieldContextIndex& Index : FieldContext.GetEvaluatedSamples())
+									FGeometryDynamicCollection::FInitialVelocityFacade InitialVelocityFacade = Collection.GetInitialVelocityFacade();
+									if (InitialVelocityFacade.IsValid())
 									{
-										Chaos::TPBDRigidParticleHandle<Chaos::FReal, 3>* RigidHandle = ParticleHandles[Index.Sample]->CastToRigidParticle();
-										if (RigidHandle)
+										static_cast<const FFieldNode<FVector>*>(FieldCommand.RootNode.Get())->Evaluate(FieldContext, ResultsView);
+										for (const FFieldContextIndex& Index : FieldContext.GetEvaluatedSamples())
 										{
-											Collection.InitialLinearVelocity[HandleToTransformGroupIndex[RigidHandle]] = FVector3f(ResultsView[Index.Result]);
+											Chaos::TPBDRigidParticleHandle<Chaos::FReal, 3>* RigidHandle = ParticleHandles[Index.Sample]->CastToRigidParticle();
+											if (RigidHandle)
+											{
+												const int32 TransformIndex = HandleToTransformGroupIndex[RigidHandle];
+												InitialVelocityFacade.InitialLinearVelocityAttribute.ModifyAt(TransformIndex, FVector3f(ResultsView[Index.Result]));
+											}
 										}
 									}
+									else
+									{
+										UE_LOG(LogChaos, Warning, TEXT("Geometry collection : trying to apply a InitialLinearVelocity field but InitialVelocityType is not set to None"))
+									}
 								}
 							}
 							else
@@ -4422,15 +4434,24 @@ void FGeometryCollectionPhysicsProxy::FieldParameterUpdateCallback(Chaos::FPBDRi
 							{
 								SCOPE_CYCLE_COUNTER(STAT_ParamUpdateField_AngularVelocity);
 								{
-									static_cast<const FFieldNode<FVector>*>(FieldCommand.RootNode.Get())->Evaluate(FieldContext, ResultsView);
-									for (const FFieldContextIndex& Index : FieldContext.GetEvaluatedSamples())
+									FGeometryDynamicCollection::FInitialVelocityFacade InitialVelocityFacade = Collection.GetInitialVelocityFacade();
+									if (InitialVelocityFacade.IsValid())
 									{
-										Chaos::TPBDRigidParticleHandle<Chaos::FReal, 3>* RigidHandle = ParticleHandles[Index.Sample]->CastToRigidParticle();
-										if (RigidHandle)
+										static_cast<const FFieldNode<FVector>*>(FieldCommand.RootNode.Get())->Evaluate(FieldContext, ResultsView);
+										for (const FFieldContextIndex& Index : FieldContext.GetEvaluatedSamples())
 										{
-											Collection.InitialAngularVelocity[HandleToTransformGroupIndex[RigidHandle]] = FVector3f(ResultsView[Index.Result]);
+											Chaos::TPBDRigidParticleHandle<Chaos::FReal, 3>* RigidHandle = ParticleHandles[Index.Sample]->CastToRigidParticle();
+											if (RigidHandle)
+											{
+												const int32 TransformIndex = HandleToTransformGroupIndex[RigidHandle];
+												InitialVelocityFacade.InitialAngularVelocityAttribute.ModifyAt(TransformIndex, FVector3f(ResultsView[Index.Result]));
+											}
 										}
 									}
+									else
+									{
+										UE_LOG(LogChaos, Warning, TEXT("Geometry collection : trying to apply a InitialLinearVelocity field but InitialVelocityType is not set to None"))
+									}
 								}
 							}
 							else
diff --git a/Engine/Source/Runtime/Experimental/Chaos/Public/GeometryCollectionProxyData.h b/Engine/Source/Runtime/Experimental/Chaos/Public/GeometryCollectionProxyData.h
--- a/Engine/Source/Runtime/Experimental/Chaos/Public/GeometryCollectionProxyData.h
+++ b/Engine/Source/Runtime/Experimental/Chaos/Public/GeometryCollectionProxyData.h
@@ -50,7 +50,6 @@ protected:
 
 class CHAOS_API FGeometryDynamicCollection : public FTransformDynamicCollection
 {
-
 public:
 	FGeometryDynamicCollection();
 	FGeometryDynamicCollection(FGeometryDynamicCollection&) = delete;
@@ -79,13 +78,29 @@ public:
 	TManagedArray<int32> CollisionStructureID;
 	TManagedArray<int32> DynamicState;
 	TManagedArray<FSharedImplicit> Implicits;
-	TManagedArray<FVector3f> InitialAngularVelocity;
-	TManagedArray<FVector3f> InitialLinearVelocity;
 	TManagedArray<FTransform> MassToLocal;
-	//TManagedArray<TArray<FCollisionFilterData>> ShapeQueryData;
-	//TManagedArray<TArray<FCollisionFilterData>> ShapeSimData;
 	TManagedArray<TUniquePtr<FCollisionStructureManager::FSimplicial>> Simplicials;
 	TManagedArray<bool> SimulatableParticles;
+
+public:
+	struct FInitialVelocityFacade
+	{
+		FInitialVelocityFacade(FGeometryDynamicCollection& DynamicCollection);
+		FInitialVelocityFacade(const FGeometryDynamicCollection& DynamicCollection);
+
+		bool IsValid() const;
+		void DefineSchema();
+		void Fill(const FVector3f& InitialLinearVelocity, const FVector3f& InitialAngularVelocity);
+		void CopyFrom(const FGeometryDynamicCollection& SourceCollection);
+
+		TManagedArrayAccessor<FVector3f> InitialLinearVelocityAttribute;
+		TManagedArrayAccessor<FVector3f> InitialAngularVelocityAttribute;
+	};
+
+	FInitialVelocityFacade GetInitialVelocityFacade() { return FInitialVelocityFacade(*this); }
+	FInitialVelocityFacade GetInitialVelocityFacade() const { return FInitialVelocityFacade(*this); }
+
+	void CopyInitialVelocityAttributesFrom(const FGeometryDynamicCollection& SourceCollection);
 };
 
 /**
