introduce a new FPhysicsObjectId alias to add some type clarity to the various functions on the IPhysicsComponent that deal with this id

diff --git a/Engine/Source/Runtime/Engine/Classes/Components/PrimitiveComponent.h b/Engine/Source/Runtime/Engine/Classes/Components/PrimitiveComponent.h
--- a/Engine/Source/Runtime/Engine/Classes/Components/PrimitiveComponent.h
+++ b/Engine/Source/Runtime/Engine/Classes/Components/PrimitiveComponent.h
@@ -2854,10 +2854,10 @@ public:
 
 	//~ Begin IPhysicsComponent Interface.
 public:
-	virtual Chaos::FPhysicsObject* GetPhysicsObjectById(int32 Id) const override;
+	virtual Chaos::FPhysicsObject* GetPhysicsObjectById(Chaos::FPhysicsObjectId Id) const override;
 	virtual Chaos::FPhysicsObject* GetPhysicsObjectByName(const FName& Name) const override;
 	virtual TArray<Chaos::FPhysicsObject*> GetAllPhysicsObjects() const override;
-	virtual int32 GetIdFromGTParticle(Chaos::FGeometryParticle* Particle) const override;
+	virtual Chaos::FPhysicsObjectId GetIdFromGTParticle(Chaos::FGeometryParticle* Particle) const override;
 	//~ End IPhysicsComponent Interface.
 };
 
diff --git a/Engine/Source/Runtime/Engine/Classes/Components/SkeletalMeshComponent.h b/Engine/Source/Runtime/Engine/Classes/Components/SkeletalMeshComponent.h
--- a/Engine/Source/Runtime/Engine/Classes/Components/SkeletalMeshComponent.h
+++ b/Engine/Source/Runtime/Engine/Classes/Components/SkeletalMeshComponent.h
@@ -2644,7 +2644,7 @@ public:
 	bool IsPostEvaluatingAnimation() const { return bPostEvaluatingAnimation; }
 
 	//~ Begin IPhysicsComponent Interface.
-	virtual Chaos::FPhysicsObject* GetPhysicsObjectById(int32 Id) const override;
+	virtual Chaos::FPhysicsObject* GetPhysicsObjectById(Chaos::FPhysicsObjectId Id) const override;
 	virtual Chaos::FPhysicsObject* GetPhysicsObjectByName(const FName& Name) const override;
 	virtual TArray<Chaos::FPhysicsObject*> GetAllPhysicsObjects() const override;
 	//~ End IPhysicsComponent Interface.
diff --git a/Engine/Source/Runtime/Engine/Classes/Interfaces/IPhysicsComponent.h b/Engine/Source/Runtime/Engine/Classes/Interfaces/IPhysicsComponent.h
--- a/Engine/Source/Runtime/Engine/Classes/Interfaces/IPhysicsComponent.h
+++ b/Engine/Source/Runtime/Engine/Classes/Interfaces/IPhysicsComponent.h
@@ -20,9 +20,9 @@ class IPhysicsComponent
 	GENERATED_BODY()
 
 public:
-	virtual Chaos::FPhysicsObject* GetPhysicsObjectById(int32 Id) const = 0;
+	virtual Chaos::FPhysicsObject* GetPhysicsObjectById(Chaos::FPhysicsObjectId Id) const = 0;
 	virtual Chaos::FPhysicsObject* GetPhysicsObjectByName(const FName& Name) const = 0;
 	virtual TArray<Chaos::FPhysicsObject*> GetAllPhysicsObjects() const = 0;
 	
-	virtual int32 GetIdFromGTParticle(Chaos::FGeometryParticle* Particle) const = 0;
+	virtual Chaos::FPhysicsObjectId GetIdFromGTParticle(Chaos::FGeometryParticle* Particle) const = 0;
 };
\ No newline at end of file
diff --git a/Engine/Source/Runtime/Engine/Classes/PhysicsEngine/ClusterUnionComponent.h b/Engine/Source/Runtime/Engine/Classes/PhysicsEngine/ClusterUnionComponent.h
--- a/Engine/Source/Runtime/Engine/Classes/PhysicsEngine/ClusterUnionComponent.h
+++ b/Engine/Source/Runtime/Engine/Classes/PhysicsEngine/ClusterUnionComponent.h
@@ -218,10 +218,10 @@ public:
 
 	//~ Begin IPhysicsComponent Interface.
 public:
-	virtual Chaos::FPhysicsObject* GetPhysicsObjectById(int32 Id) const override;
+	virtual Chaos::FPhysicsObject* GetPhysicsObjectById(Chaos::FPhysicsObjectId Id) const override;
 	virtual Chaos::FPhysicsObject* GetPhysicsObjectByName(const FName& Name) const override;
 	virtual TArray<Chaos::FPhysicsObject*> GetAllPhysicsObjects() const override;
-	virtual int32 GetIdFromGTParticle(Chaos::FGeometryParticle* Particle) const override;
+	virtual Chaos::FPhysicsObjectId GetIdFromGTParticle(Chaos::FGeometryParticle* Particle) const override;
 	//~ End IPhysicsComponent Interface.
 
 	//~ Begin UObject Interface.
diff --git a/Engine/Source/Runtime/Engine/Private/PhysicsEngine/ClusterUnionComponent.cpp b/Engine/Source/Runtime/Engine/Private/PhysicsEngine/ClusterUnionComponent.cpp
--- a/Engine/Source/Runtime/Engine/Private/PhysicsEngine/ClusterUnionComponent.cpp
+++ b/Engine/Source/Runtime/Engine/Private/PhysicsEngine/ClusterUnionComponent.cpp
@@ -261,7 +261,7 @@ bool UClusterUnionComponent::HasValidPhysicsState() const
 	return PhysicsProxy != nullptr;
 }
 
-Chaos::FPhysicsObject* UClusterUnionComponent::GetPhysicsObjectById(int32 Id) const
+Chaos::FPhysicsObject* UClusterUnionComponent::GetPhysicsObjectById(Chaos::FPhysicsObjectId Id) const
 {
 	if (!PhysicsProxy)
 	{
@@ -280,7 +280,7 @@ TArray<Chaos::FPhysicsObject*> UClusterUnionComponent::GetAllPhysicsObjects() co
 	return { GetPhysicsObjectById(0) };
 }
 
-int32 UClusterUnionComponent::GetIdFromGTParticle(Chaos::FGeometryParticle* Particle) const
+Chaos::FPhysicsObjectId UClusterUnionComponent::GetIdFromGTParticle(Chaos::FGeometryParticle* Particle) const
 {
 	return 0;
 }
diff --git a/Engine/Source/Runtime/Engine/Private/PrimitiveComponentPhysics.cpp b/Engine/Source/Runtime/Engine/Private/PrimitiveComponentPhysics.cpp
--- a/Engine/Source/Runtime/Engine/Private/PrimitiveComponentPhysics.cpp
+++ b/Engine/Source/Runtime/Engine/Private/PrimitiveComponentPhysics.cpp
@@ -1196,7 +1196,7 @@ void UPrimitiveComponent::UpdatePhysicsToRBChannels()
 	}
 }
 
-Chaos::FPhysicsObject* UPrimitiveComponent::GetPhysicsObjectById(int32 Id) const
+Chaos::FPhysicsObject* UPrimitiveComponent::GetPhysicsObjectById(Chaos::FPhysicsObjectId Id) const
 {
 	if (!BodyInstance.IsValidBodyInstance())
 	{
@@ -1217,7 +1217,7 @@ TArray<Chaos::FPhysicsObject*> UPrimitiveComponent::GetAllPhysicsObjects() const
 	return Bodies;
 }
 
-int32 UPrimitiveComponent::GetIdFromGTParticle(Chaos::FGeometryParticle* Particle) const
+Chaos::FPhysicsObjectId UPrimitiveComponent::GetIdFromGTParticle(Chaos::FGeometryParticle* Particle) const
 {
 	return 0;
 }
diff --git a/Engine/Source/Runtime/Engine/Private/SkeletalMeshComponentPhysics.cpp b/Engine/Source/Runtime/Engine/Private/SkeletalMeshComponentPhysics.cpp
--- a/Engine/Source/Runtime/Engine/Private/SkeletalMeshComponentPhysics.cpp
+++ b/Engine/Source/Runtime/Engine/Private/SkeletalMeshComponentPhysics.cpp
@@ -4046,7 +4046,7 @@ FTransform USkeletalMeshComponent::GetComponentTransformFromBodyInstance(FBodyIn
 }
 
 
-Chaos::FPhysicsObject* USkeletalMeshComponent::GetPhysicsObjectById(int32 Id) const
+Chaos::FPhysicsObject* USkeletalMeshComponent::GetPhysicsObjectById(Chaos::FPhysicsObjectId Id) const
 {
 	if (!Bodies.IsValidIndex(Id) || !Bodies[Id] || !Bodies[Id]->ActorHandle)
 	{
diff --git a/Engine/Source/Runtime/Experimental/Chaos/Public/Chaos/PhysicsObject.h b/Engine/Source/Runtime/Experimental/Chaos/Public/Chaos/PhysicsObject.h
--- a/Engine/Source/Runtime/Experimental/Chaos/Public/Chaos/PhysicsObject.h
+++ b/Engine/Source/Runtime/Experimental/Chaos/Public/Chaos/PhysicsObject.h
@@ -17,6 +17,7 @@ namespace Chaos
 	struct FPhysicsObject;
 	using FPhysicsObjectHandle = FPhysicsObject*;
 	using FConstPhysicsObjectHandle = const FPhysicsObject*;
+	using FPhysicsObjectId = int32;
 
 	struct CHAOS_API FPhysicsObjectDeleter
 	{
diff --git a/Engine/Source/Runtime/Experimental/GeometryCollectionEngine/Private/GeometryCollection/GeometryCollectionComponent.cpp b/Engine/Source/Runtime/Experimental/GeometryCollectionEngine/Private/GeometryCollection/GeometryCollectionComponent.cpp
--- a/Engine/Source/Runtime/Experimental/GeometryCollectionEngine/Private/GeometryCollection/GeometryCollectionComponent.cpp
+++ b/Engine/Source/Runtime/Experimental/GeometryCollectionEngine/Private/GeometryCollection/GeometryCollectionComponent.cpp
@@ -5049,7 +5049,7 @@ void UGeometryCollectionComponent::PostLoad()
 	}
 }
 
-Chaos::FPhysicsObject* UGeometryCollectionComponent::GetPhysicsObjectById(int32 Id) const
+Chaos::FPhysicsObject* UGeometryCollectionComponent::GetPhysicsObjectById(Chaos::FPhysicsObjectId Id) const
 {
 	if (!PhysicsProxy)
 	{
@@ -5100,7 +5100,7 @@ TArray<Chaos::FPhysicsObject*> UGeometryCollectionComponent::GetAllPhysicsObject
 	return Objects;
 }
 
-int32 UGeometryCollectionComponent::GetIdFromGTParticle(Chaos::FGeometryParticle* Particle) const
+Chaos::FPhysicsObjectId UGeometryCollectionComponent::GetIdFromGTParticle(Chaos::FGeometryParticle* Particle) const
 {
 	if (!PhysicsProxy || !Particle)
 	{
diff --git a/Engine/Source/Runtime/Experimental/GeometryCollectionEngine/Public/GeometryCollection/GeometryCollectionComponent.h b/Engine/Source/Runtime/Experimental/GeometryCollectionEngine/Public/GeometryCollection/GeometryCollectionComponent.h
--- a/Engine/Source/Runtime/Experimental/GeometryCollectionEngine/Public/GeometryCollection/GeometryCollectionComponent.h
+++ b/Engine/Source/Runtime/Experimental/GeometryCollectionEngine/Public/GeometryCollection/GeometryCollectionComponent.h
@@ -1301,9 +1301,9 @@ private:
 
 	//~ Begin IPhysicsComponent Interface.
 public:
-	virtual Chaos::FPhysicsObject* GetPhysicsObjectById(int32 Id) const override;
+	virtual Chaos::FPhysicsObject* GetPhysicsObjectById(Chaos::FPhysicsObjectId Id) const override;
 	virtual Chaos::FPhysicsObject* GetPhysicsObjectByName(const FName& Name) const override;
 	virtual TArray<Chaos::FPhysicsObject*> GetAllPhysicsObjects() const override;
-	virtual int32 GetIdFromGTParticle(Chaos::FGeometryParticle* Particle) const override;
+	virtual Chaos::FPhysicsObjectId GetIdFromGTParticle(Chaos::FGeometryParticle* Particle) const override;
 	//~ End IPhysicsComponent Interface.
 };
