the physics object interface and the SQ types should be able to statically select the shape instance to operate on

diff --git a/Engine/Source/Runtime/Engine/Private/Collision/WorldCollision.cpp b/Engine/Source/Runtime/Engine/Private/Collision/WorldCollision.cpp
--- a/Engine/Source/Runtime/Engine/Private/Collision/WorldCollision.cpp
+++ b/Engine/Source/Runtime/Engine/Private/Collision/WorldCollision.cpp
@@ -369,13 +369,13 @@ bool UWorld::ComponentOverlapMultiByChannel(TArray<struct FOverlapResult>& OutOv
 	TArray<struct FOverlapResult> TempOverlaps;
 	for (Chaos::FPhysicsObjectHandle Object : PhysicsObjects)
 	{
-		TArray<Chaos::FPerShapeData*> Shapes = Interface->GetAllShapes({&Object, 1});
+		TArray<Chaos::FShapeInstanceProxy*> Shapes = Interface->GetAllThreadShapes({&Object, 1});
 
 		// Determine how to convert the local space of this body instance to the test space
 		const FTransform ObjectToWorld = Interface->GetTransform(Object);
 		const FTransform ObjectToTest = ComponentToTest * WorldToComponent * ObjectToWorld;
 
-		for (Chaos::FPerShapeData* Shape : Shapes)
+		for (Chaos::FShapeInstanceProxy* Shape : Shapes)
 		{
 			FPhysicsShapeHandle ShapeRef{ Shape, nullptr };
 
@@ -471,8 +471,8 @@ bool UWorld::ComponentSweepMultiByChannel(TArray<struct FHitResult>& OutHits, cl
 		}		
 		ParamsCopy.ClearIgnoredActors(); // This will be populated a bit later
 		// All actors pointed to by shapes should be ignored (This deals with welded Actors)
-		TArray<Chaos::FPerShapeData*> Shapes = Interface->GetAllShapes({ &Object, 1 });
-		for (Chaos::FPerShapeData* Shape : Shapes)
+		TArray<Chaos::FShapeInstanceProxy*> Shapes = Interface->GetAllThreadShapes({ &Object, 1 });
+		for (Chaos::FShapeInstanceProxy* Shape : Shapes)
 		{
 			FCollisionFilterData ShapeFilter = ChaosInterface::GetQueryFilterData(*Shape);
 			uint32 ShapeActorID = ShapeFilter.Word0; // Unique Actor ID is saved in word 0
@@ -484,7 +484,7 @@ bool UWorld::ComponentSweepMultiByChannel(TArray<struct FHitResult>& OutHits, cl
 			ParamsCopy.AddIgnoredActor(ActorID);
 		}
 
-		for (Chaos::FPerShapeData* Shape : Shapes)
+		for (Chaos::FShapeInstanceProxy* Shape : Shapes)
 		{
 			FPhysicsShapeHandle ShapeHandle{ Shape, nullptr };
 
diff --git a/Engine/Source/Runtime/Experimental/Chaos/Private/Chaos/PhysicsObjectInterface.cpp b/Engine/Source/Runtime/Experimental/Chaos/Private/Chaos/PhysicsObjectInterface.cpp
--- a/Engine/Source/Runtime/Experimental/Chaos/Private/Chaos/PhysicsObjectInterface.cpp
+++ b/Engine/Source/Runtime/Experimental/Chaos/Private/Chaos/PhysicsObjectInterface.cpp
@@ -3,7 +3,6 @@
 
 #include "Chaos/CastingUtilities.h"
 #include "Chaos/CollisionFilterData.h"
-#include "Chaos/GeometryQueries.h"
 #include "Chaos/PhysicsObjectInternal.h"
 #include "PhysicsProxy/SingleParticlePhysicsProxy.h"
 #include "Math/Transform.h"
@@ -175,9 +174,13 @@ namespace Chaos
 	}
 
 	template<EThreadContext Id>
-	TThreadParticle<Id>* FReadPhysicsObjectInterface<Id>::GetParticle(const FConstPhysicsObjectHandle Object)
+	TThreadParticle<Id>* FReadPhysicsObjectInterface<Id>::GetParticle(const FConstPhysicsObjectHandle Handle)
 	{
-		return Object->GetParticle<Id>();
+		if (!Handle)
+		{
+			return nullptr;
+		}
+		return Handle->GetParticle<Id>();
 	}
 
 	template<EThreadContext Id>
@@ -248,8 +251,10 @@ namespace Chaos
 	}
 
 	template<EThreadContext Id>
-	void FReadPhysicsObjectInterface<Id>::VisitEveryShape(TArrayView<const FConstPhysicsObjectHandle> InObjects, TFunctionRef<bool(const FConstPhysicsObjectHandle, FPerShapeData*)> Lambda)
+	TArray<TThreadShapeInstance<Id>*> FReadPhysicsObjectInterface<Id>::GetAllThreadShapes(TArrayView<const FConstPhysicsObjectHandle> InObjects)
 	{
+		TArray<TThreadShapeInstance<Id>*> AllShapes;
+
 		for (const FConstPhysicsObjectHandle Object : InObjects)
 		{
 			if (!Object)
@@ -259,8 +264,29 @@ namespace Chaos
 
 			if (TThreadParticle<Id>* Particle = Object->GetParticle<Id>())
 			{
-				const Chaos::FShapesArray& ShapesArray = Particle->ShapesArray();
-				for (const TUniquePtr<Chaos::FPerShapeData>& Shape : ShapesArray)
+				for (const TUniquePtr<TThreadShapeInstance<Id>>& Shape : Particle->ShapeInstances())
+				{
+					AllShapes.Add(Shape.Get());
+				}
+			}
+		}
+
+		return AllShapes;
+	}
+
+	template<EThreadContext Id>
+	void FReadPhysicsObjectInterface<Id>::VisitEveryShape(TArrayView<const FConstPhysicsObjectHandle> InObjects, TFunctionRef<bool(const FConstPhysicsObjectHandle, TThreadShapeInstance<Id>*)> Lambda)
+	{
+		for (const FConstPhysicsObjectHandle Object : InObjects)
+		{
+			if (!Object)
+			{
+				continue;
+			}
+
+			if (TThreadParticle<Id>* Particle = Object->GetParticle<Id>())
+			{
+				for (const TUniquePtr<TThreadShapeInstance<Id>>& Shape : Particle->ShapeInstances())
 				{
 					if (Lambda(Object, Shape.Get()))
 					{
@@ -384,11 +410,11 @@ namespace Chaos
 	template<EThreadContext Id>
 	bool FReadPhysicsObjectInterface<Id>::PairwiseShapeOverlapHelper(const FConstPhysicsObjectHandle ObjectA, const FTransform& InTransformA, const FConstPhysicsObjectHandle ObjectB, const FTransform& InTransformB, bool bTraceComplex, bool bComputeMTD, const FVector& Tolerance, const TFunction<bool(const FShapeOverlapData&, const FShapeOverlapData&, const FMTDInfo&)>& Lambda)
 	{
-		TArray<FPerShapeData*> ShapesA = GetAllShapes({ &ObjectA, 1 });
+		TArray<TThreadShapeInstance<Id>*> ShapesA = GetAllThreadShapes({ &ObjectA, 1 });
 		const FTransform TransformA = FTransform{ GetR(ObjectA), GetX(ObjectA) } *InTransformA;
 		const FBox BoxA = GetWorldBounds({ &ObjectA, 1 }).TransformBy(InTransformA);
 
-		TArray<FPerShapeData*> ShapesB = GetAllShapes({ &ObjectB, 1 });
+		TArray<TThreadShapeInstance<Id>*> ShapesB = GetAllThreadShapes({ &ObjectB, 1 });
 		const FTransform TransformB = FTransform{ GetR(ObjectB), GetX(ObjectB) } *InTransformB;
 		const FBox BoxB = GetWorldBounds({ &ObjectB, 1 }).TransformBy(InTransformB);
 
@@ -398,7 +424,7 @@ namespace Chaos
 		}
 
 		bool bFoundOverlap = false;
-		for (FPerShapeData* B : ShapesB)
+		for (TThreadShapeInstance<Id>* B : ShapesB)
 		{
 			if (!B)
 			{
@@ -415,7 +441,7 @@ namespace Chaos
 			// At this point on, this function should be mirror the Overlap_GeomInternal function in PhysInterface_Chaos.cpp.
 			// ShapeA is equivalent to InInstance and GeomB is equivalent to InGeom.
 
-			for (FPerShapeData* A : ShapesA)
+			for (TThreadShapeInstance<Id>* A : ShapesA)
 			{
 				if (!A)
 				{
@@ -702,199 +728,6 @@ namespace Chaos
 		return FAccelerationStructureHandle{InObject->GetParticle<Id>()};
 	}
 
-	template<EThreadContext Id>
-	bool FReadPhysicsObjectInterface<Id>::LineTrace(TArrayView<const FConstPhysicsObjectHandle> InObjects, const FVector& WorldStart, const FVector& WorldEnd, bool bTraceComplex, ChaosInterface::FRaycastHit& OutBestHit)
-	{
-		bool bHit = false;
-		OutBestHit.Distance = TNumericLimits<float>::Max();
-
-		const FVector Delta = WorldEnd - WorldStart;
-		const FReal DeltaMag = Delta.Size();
-		if (DeltaMag < UE_KINDA_SMALL_NUMBER)
-		{
-			return false;
-		}
-
-		FTransform BestWorldTM = FTransform::Identity;
-
-		for (const FConstPhysicsObjectHandle Object : InObjects)
-		{
-			const FTransform WorldTM = GetTransform(Object);
-			const FVector LocalStart = WorldTM.InverseTransformPositionNoScale(WorldStart);
-			const FVector LocalDelta = WorldTM.InverseTransformVectorNoScale(Delta);
-
-			VisitEveryShape(
-				{ &Object, 1 },
-				[this, &bHit, &WorldTM, &LocalStart, &LocalDelta, &Delta, DeltaMag, &BestWorldTM, bTraceComplex, &OutBestHit](const FConstPhysicsObjectHandle IterObject, FPerShapeData* Shape)
-				{
-					check(Shape);
-
-					FCollisionFilterData ShapeFilter = Shape->GetQueryData();
-					const bool bShapeIsComplex = (ShapeFilter.Word3 & static_cast<uint8>(EFilterFlags::ComplexCollision)) != 0;
-					const bool bShapeIsSimple = (ShapeFilter.Word3 & static_cast<uint8>(EFilterFlags::SimpleCollision)) != 0;
-					if ((bTraceComplex && bShapeIsComplex) || (!bTraceComplex && bShapeIsSimple))
-					{
-						FReal Distance;
-						FVec3 LocalPosition;
-						FVec3 LocalNormal;
-						int32 FaceIndex;
-
-						const bool bRaycastHit = Shape->GetGeometry()->Raycast(
-							LocalStart,
-							LocalDelta / DeltaMag,
-							DeltaMag,
-							0,
-							Distance,
-							LocalPosition,
-							LocalNormal,
-							FaceIndex
-						);
-
-						if (bRaycastHit)
-						{
-							if (Distance < OutBestHit.Distance)
-							{
-								bHit = true;
-								BestWorldTM = WorldTM;
-								OutBestHit.Distance = static_cast<float>(Distance);
-								OutBestHit.WorldNormal = LocalNormal;
-								OutBestHit.WorldPosition = LocalPosition;
-								OutBestHit.Shape = Shape;
-								if constexpr (Id == EThreadContext::External)
-								{
-									OutBestHit.Actor = IterObject->GetParticle<Id>();
-								}
-								OutBestHit.FaceIndex = FaceIndex;
-							}
-						}
-					}
-					return false;
-				}
-			);
-		}
-
-		if (bHit)
-		{
-			OutBestHit.WorldNormal = BestWorldTM.TransformVectorNoScale(OutBestHit.WorldNormal);
-			OutBestHit.WorldPosition = BestWorldTM.TransformPositionNoScale(OutBestHit.WorldPosition);
-		}
-
-		return bHit;
-	}
-
-	template<EThreadContext Id>
-	bool FReadPhysicsObjectInterface<Id>::ShapeOverlap(TArrayView<const FConstPhysicsObjectHandle> InObjects, const Chaos::FImplicitObject& InGeom, const FTransform& GeomTransform, TArray<ChaosInterface::FOverlapHit>& OutOverlaps)
-	{
-		bool bHasOverlap = false;
-		for (const FConstPhysicsObjectHandle Object : InObjects)
-		{
-			const FTransform WorldTM = GetTransform(Object);
-
-			VisitEveryShape(
-				{ &Object, 1 },
-				[this, &bHasOverlap, &WorldTM, &InGeom, &GeomTransform, &OutOverlaps](const FConstPhysicsObjectHandle IterObject, FPerShapeData* Shape)
-				{
-					check(Shape);
-					const bool bOverlap = Chaos::Utilities::CastHelper(
-						InGeom,
-						GeomTransform,
-						[Shape, &WorldTM](const auto& Downcast, const auto& FullTransformB)
-						{
-							return Chaos::OverlapQuery(*Shape->GetGeometry(), WorldTM, Downcast, FullTransformB, 0, nullptr);
-						}
-					);
-
-					if (bOverlap)
-					{
-						bHasOverlap = true;
-
-						ChaosInterface::FOverlapHit Overlap;
-						Overlap.Shape = Shape;
-						if constexpr (Id == EThreadContext::External)
-						{
-							Overlap.Actor = IterObject->GetParticle<Id>();
-						}
-						OutOverlaps.Add(Overlap);
-					}
-					return false;
-				}
-			);
-		}
-		return bHasOverlap;
-	}
-
-	template<EThreadContext Id>
-	bool FReadPhysicsObjectInterface<Id>::ShapeSweep(TArrayView<const FConstPhysicsObjectHandle> InObjects, const Chaos::FImplicitObject& InGeom, const FTransform& StartTM, const FVector& EndPos, bool bSweepComplex, ChaosInterface::FSweepHit& OutBestHit)
-	{
-		bool bHit = false;
-		const FVector StartPos = StartTM.GetTranslation();
-		const FVector Delta = EndPos - StartPos;
-		const FReal DeltaMag = Delta.Size();
-		if (DeltaMag < UE_KINDA_SMALL_NUMBER)
-		{
-			return false;
-		}
-		const FVec3 Dir = Delta / DeltaMag;
-
-		for (const FConstPhysicsObjectHandle Object : InObjects)
-		{
-			const FTransform WorldTM = GetTransform(Object);
-
-			VisitEveryShape(
-				{ &Object, 1 },
-				[&WorldTM, &InGeom, &StartTM, &bHit, &Delta, DeltaMag, &Dir, &OutBestHit, bSweepComplex](const FConstPhysicsObjectHandle IterObject, FPerShapeData* Shape)
-				{
-					check(Shape);
-
-					FCollisionFilterData ShapeFilter = Shape->GetQueryData();
-					const bool bShapeIsComplex = (ShapeFilter.Word3 & static_cast<uint8>(EFilterFlags::ComplexCollision)) != 0;
-					const bool bShapeIsSimple = (ShapeFilter.Word3 & static_cast<uint8>(EFilterFlags::SimpleCollision)) != 0;
-					if ((bSweepComplex && bShapeIsComplex) || (!bSweepComplex && bShapeIsSimple))
-					{
-						FVec3 WorldPosition;
-						FVec3 WorldNormal;
-						FReal Distance;
-						int32 FaceIdx;
-						FVec3 FaceNormal;
-
-						const bool bShapeHit = Chaos::Utilities::CastHelper(
-							InGeom,
-							StartTM,
-							[Shape, &WorldTM, &Dir, DeltaMag, &Distance, &WorldPosition, &WorldNormal, &FaceIdx, &FaceNormal](const auto& Downcast, const auto& FullTransformB)
-							{
-								return Chaos::SweepQuery(*Shape->GetGeometry(), WorldTM, Downcast, FullTransformB, Dir, DeltaMag, Distance, WorldPosition, WorldNormal, FaceIdx, FaceNormal, 0.f, false);
-							}
-						);
-
-						if (bShapeHit)
-						{
-							bHit = true;
-
-							OutBestHit.Shape = Shape;
-							OutBestHit.WorldPosition = WorldPosition;
-							OutBestHit.WorldNormal = WorldNormal;
-							OutBestHit.Distance = static_cast<float>(Distance);
-							OutBestHit.FaceIndex = FaceIdx;
-							if (OutBestHit.Distance > 0.f)
-							{
-								const FVector LocalPosition = WorldTM.InverseTransformPositionNoScale(OutBestHit.WorldPosition);
-								const FVector LocalUnitDir = WorldTM.InverseTransformVectorNoScale(Dir);
-								OutBestHit.FaceIndex = Shape->GetGeometry()->FindMostOpposingFace(LocalPosition, LocalUnitDir, OutBestHit.FaceIndex, 1);
-							}
-							if constexpr (Id == EThreadContext::External)
-							{
-								OutBestHit.Actor = IterObject->GetParticle<Id>();
-							}
-						}
-					}
-					return false;
-				}
-			);
-
-		}
-		return bHit;
-	}
-
 	template<EThreadContext Id>
 	void FWritePhysicsObjectInterface<Id>::PutToSleep(TArrayView<const FPhysicsObjectHandle> InObjects)
 	{
diff --git a/Engine/Source/Runtime/Experimental/Chaos/Private/Chaos/PhysicsObjectInternal.h b/Engine/Source/Runtime/Experimental/Chaos/Private/Chaos/PhysicsObjectInternal.h
--- a/Engine/Source/Runtime/Experimental/Chaos/Private/Chaos/PhysicsObjectInternal.h
+++ b/Engine/Source/Runtime/Experimental/Chaos/Private/Chaos/PhysicsObjectInternal.h
@@ -2,6 +2,7 @@
 #pragma once
 #include "Chaos/PhysicsObject.h"
 #include "Chaos/ParticleHandleFwd.h"
+#include "Chaos/ShapeInstanceFwd.h"
 #include "Containers/Set.h"
 #include "Framework/Threading.h"
 
diff --git a/Engine/Source/Runtime/Experimental/Chaos/Public/Chaos/Interface/SQTypes.h b/Engine/Source/Runtime/Experimental/Chaos/Public/Chaos/Interface/SQTypes.h
--- a/Engine/Source/Runtime/Experimental/Chaos/Public/Chaos/Interface/SQTypes.h
+++ b/Engine/Source/Runtime/Experimental/Chaos/Public/Chaos/Interface/SQTypes.h
@@ -8,10 +8,13 @@
 #include "Containers/Array.h"
 #include "Containers/ContainerAllocationPolicies.h"
 #include "CoreMinimal.h"
+#include "Framework/ThreadContextEnum.h"
 #include "HAL/PlatformCrt.h"
 #include "Math/UnrealMathSSE.h"
 #include "Math/Vector.h"
 
+#include <type_traits>
+
 namespace Chaos
 {
 	class FChaosArchive;
@@ -149,6 +152,15 @@ namespace ChaosInterface
 	{
 	};
 
+	template<Chaos::EThreadContext Id>
+	using TThreadRaycastHit = std::conditional_t<Id == Chaos::EThreadContext::External, FRaycastHit, FPTRaycastHit>;
+
+	template<Chaos::EThreadContext Id>
+	using TThreadOverlapHit = std::conditional_t<Id == Chaos::EThreadContext::External, FOverlapHit, FPTOverlapHit>;
+
+	template<Chaos::EThreadContext Id>
+	using TThreadSweepHit = std::conditional_t<Id == Chaos::EThreadContext::External, FSweepHit, FPTSweepHit>;
+
 #ifndef CHAOS_HIT_BUFFER_SIZE
 #define CHAOS_HIT_BUFFER_SIZE 512 // Preallocated hit buffer size for traces and sweeps.
 #endif // CHAOS_HIT_BUFFER_SIZE
diff --git a/Engine/Source/Runtime/Experimental/Chaos/Public/Chaos/PhysicsObjectInterface.h b/Engine/Source/Runtime/Experimental/Chaos/Public/Chaos/PhysicsObjectInterface.h
--- a/Engine/Source/Runtime/Experimental/Chaos/Public/Chaos/PhysicsObjectInterface.h
+++ b/Engine/Source/Runtime/Experimental/Chaos/Public/Chaos/PhysicsObjectInterface.h
@@ -3,8 +3,10 @@
 #pragma once
 
 #include "Chaos/Interface/SQTypes.h"
+#include "Chaos/GeometryQueries.h"
 #include "Chaos/ParticleHandle.h"
 #include "Chaos/PhysicsObject.h"
+#include "Chaos/ShapeInstanceFwd.h"
 #include "Chaos/ISpatialAcceleration.h"
 #include "Framework/Threading.h"
 #include "Containers/ArrayView.h"
@@ -68,10 +70,14 @@ namespace Chaos
 		TThreadParticle<Id>* GetParticle(const FConstPhysicsObjectHandle Object);
 		TArray<TThreadParticle<Id>*> GetAllParticles(TArrayView<const FConstPhysicsObjectHandle> InObjects);
 		TArray<TThreadRigidParticle<Id>*> GetAllRigidParticles(TArrayView<const FConstPhysicsObjectHandle> InObjects);
+
+		UE_DEPRECATED(5.3, "GetAllShapes has been deprecated. Please use GetAllThreadShapes instead.")
 		TArray<FPerShapeData*> GetAllShapes(TArrayView<const FConstPhysicsObjectHandle> InObjects);
 
+		TArray<TThreadShapeInstance<Id>*> GetAllThreadShapes(TArrayView<const FConstPhysicsObjectHandle> InObjects);
+
 		// Returns true if a shape is found and we can stop iterating.
-		void VisitEveryShape(TArrayView<const FConstPhysicsObjectHandle> InObjects, TFunctionRef<bool(const FConstPhysicsObjectHandle, FPerShapeData*)> Lambda);
+		void VisitEveryShape(TArrayView<const FConstPhysicsObjectHandle> InObjects, TFunctionRef<bool(const FConstPhysicsObjectHandle, TThreadShapeInstance<Id>*)> Lambda);
 
 		UE_DEPRECATED(5.3, "GetPhysicsObjectOverlap has been deprecated. Please use the function for the specific overlap metric you wish to compute instead.")
 		bool GetPhysicsObjectOverlap(const FConstPhysicsObjectHandle ObjectA, const FConstPhysicsObjectHandle ObjectB, bool bTraceComplex, Chaos::FOverlapInfo& OutOverlap);
@@ -102,9 +108,192 @@ namespace Chaos
 		FClosestPhysicsObjectResult GetClosestPhysicsBodyFromLocation(TArrayView<const FConstPhysicsObjectHandle> InObjects, const FVector& WorldLocation);
 		FAccelerationStructureHandle CreateAccelerationStructureHandle(const FConstPhysicsObjectHandle Handle);
 
-		bool LineTrace(TArrayView<const FConstPhysicsObjectHandle> InObjects, const FVector& WorldStart, const FVector& WorldEnd, bool bTraceComplex, ChaosInterface::FRaycastHit& OutBestHit);
-		bool ShapeOverlap(TArrayView<const FConstPhysicsObjectHandle> InObjects, const Chaos::FImplicitObject& InGeom, const FTransform& GeomTransform, TArray<ChaosInterface::FOverlapHit>& OutOverlaps);
-		bool ShapeSweep(TArrayView<const FConstPhysicsObjectHandle> InObjects, const Chaos::FImplicitObject& InGeom, const FTransform& StartTM, const FVector& EndPos, bool bSweepComplex, ChaosInterface::FSweepHit& OutBestHit);
+		template<typename TRaycastHit>
+		bool LineTrace(TArrayView<const FConstPhysicsObjectHandle> InObjects, const FVector& WorldStart, const FVector& WorldEnd, bool bTraceComplex, TRaycastHit& OutBestHit)
+		{
+			static_assert(std::is_same_v<TRaycastHit, ChaosInterface::TThreadRaycastHit<Id>>);
+			bool bHit = false;
+			OutBestHit.Distance = TNumericLimits<float>::Max();
+
+			const FVector Delta = WorldEnd - WorldStart;
+			const FReal DeltaMag = Delta.Size();
+			if (DeltaMag < UE_KINDA_SMALL_NUMBER)
+			{
+				return false;
+			}
+
+			FTransform BestWorldTM = FTransform::Identity;
+
+			for (const FConstPhysicsObjectHandle Object : InObjects)
+			{
+				const FTransform WorldTM = GetTransform(Object);
+				const FVector LocalStart = WorldTM.InverseTransformPositionNoScale(WorldStart);
+				const FVector LocalDelta = WorldTM.InverseTransformVectorNoScale(Delta);
+
+				VisitEveryShape(
+					{ &Object, 1 },
+					[this, &bHit, &WorldTM, &LocalStart, &LocalDelta, &Delta, DeltaMag, &BestWorldTM, bTraceComplex, &OutBestHit](const FConstPhysicsObjectHandle IterObject, TThreadShapeInstance<Id>* Shape)
+					{
+						check(Shape);
+
+						FCollisionFilterData ShapeFilter = Shape->GetQueryData();
+						const bool bShapeIsComplex = (ShapeFilter.Word3 & static_cast<uint8>(EFilterFlags::ComplexCollision)) != 0;
+						const bool bShapeIsSimple = (ShapeFilter.Word3 & static_cast<uint8>(EFilterFlags::SimpleCollision)) != 0;
+						if ((bTraceComplex && bShapeIsComplex) || (!bTraceComplex && bShapeIsSimple))
+						{
+							FReal Distance;
+							FVec3 LocalPosition;
+							FVec3 LocalNormal;
+							int32 FaceIndex;
+
+							const bool bRaycastHit = Shape->GetGeometry()->Raycast(
+								LocalStart,
+								LocalDelta / DeltaMag,
+								DeltaMag,
+								0,
+								Distance,
+								LocalPosition,
+								LocalNormal,
+								FaceIndex
+							);
+
+							if (bRaycastHit)
+							{
+								if (Distance < OutBestHit.Distance)
+								{
+									bHit = true;
+									BestWorldTM = WorldTM;
+									OutBestHit.Distance = static_cast<float>(Distance);
+									OutBestHit.WorldNormal = LocalNormal;
+									OutBestHit.WorldPosition = LocalPosition;
+									OutBestHit.Shape = Shape;
+									OutBestHit.Actor = GetParticle(IterObject);
+									OutBestHit.FaceIndex = FaceIndex;
+								}
+							}
+						}
+						return false;
+					}
+				);
+			}
+
+			if (bHit)
+			{
+				OutBestHit.WorldNormal = BestWorldTM.TransformVectorNoScale(OutBestHit.WorldNormal);
+				OutBestHit.WorldPosition = BestWorldTM.TransformPositionNoScale(OutBestHit.WorldPosition);
+			}
+
+			return bHit;
+		}
+
+		template<typename TOverlapHit>
+		bool ShapeOverlap(TArrayView<const FConstPhysicsObjectHandle> InObjects, const Chaos::FImplicitObject& InGeom, const FTransform& GeomTransform, TArray<TOverlapHit>& OutOverlaps)
+		{
+			static_assert(std::is_same_v<TOverlapHit, ChaosInterface::TThreadOverlapHit<Id>>);
+			bool bHasOverlap = false;
+			for (const FConstPhysicsObjectHandle Object : InObjects)
+			{
+				const FTransform WorldTM = GetTransform(Object);
+
+				VisitEveryShape(
+					{ &Object, 1 },
+					[this, &bHasOverlap, &WorldTM, &InGeom, &GeomTransform, &OutOverlaps](const FConstPhysicsObjectHandle IterObject, TThreadShapeInstance<Id>* Shape)
+					{
+						check(Shape);
+						const bool bOverlap = Chaos::Utilities::CastHelper(
+							InGeom,
+							GeomTransform,
+							[Shape, &WorldTM](const auto& Downcast, const auto& FullTransformB)
+							{
+								return Chaos::OverlapQuery(*Shape->GetGeometry(), WorldTM, Downcast, FullTransformB, 0, nullptr);
+							}
+						);
+
+						if (bOverlap)
+						{
+							bHasOverlap = true;
+
+							ChaosInterface::FOverlapHit Overlap;
+							Overlap.Shape = Shape;
+							Overlap.Actor = GetParticle(IterObject);
+							OutOverlaps.Add(Overlap);
+						}
+						return false;
+					}
+				);
+			}
+			return bHasOverlap;
+		}
+
+		template<typename TSweepHit>
+		bool ShapeSweep(TArrayView<const FConstPhysicsObjectHandle> InObjects, const Chaos::FImplicitObject& InGeom, const FTransform& StartTM, const FVector& EndPos, bool bSweepComplex, TSweepHit& OutBestHit)
+		{
+			static_assert(std::is_same_v<TSweepHit, ChaosInterface::TThreadSweepHit<Id>>);
+			bool bHit = false;
+			const FVector StartPos = StartTM.GetTranslation();
+			const FVector Delta = EndPos - StartPos;
+			const FReal DeltaMag = Delta.Size();
+			if (DeltaMag < UE_KINDA_SMALL_NUMBER)
+			{
+				return false;
+			}
+			const FVec3 Dir = Delta / DeltaMag;
+
+			for (const FConstPhysicsObjectHandle Object : InObjects)
+			{
+				const FTransform WorldTM = GetTransform(Object);
+
+				VisitEveryShape(
+					{ &Object, 1 },
+					[this, &WorldTM, &InGeom, &StartTM, &bHit, &Delta, DeltaMag, &Dir, &OutBestHit, bSweepComplex](const FConstPhysicsObjectHandle IterObject, TThreadShapeInstance<Id>* Shape)
+					{
+						check(Shape);
+
+						FCollisionFilterData ShapeFilter = Shape->GetQueryData();
+						const bool bShapeIsComplex = (ShapeFilter.Word3 & static_cast<uint8>(EFilterFlags::ComplexCollision)) != 0;
+						const bool bShapeIsSimple = (ShapeFilter.Word3 & static_cast<uint8>(EFilterFlags::SimpleCollision)) != 0;
+						if ((bSweepComplex && bShapeIsComplex) || (!bSweepComplex && bShapeIsSimple))
+						{
+							FVec3 WorldPosition;
+							FVec3 WorldNormal;
+							FReal Distance;
+							int32 FaceIdx;
+							FVec3 FaceNormal;
+
+							const bool bShapeHit = Chaos::Utilities::CastHelper(
+								InGeom,
+								StartTM,
+								[Shape, &WorldTM, &Dir, DeltaMag, &Distance, &WorldPosition, &WorldNormal, &FaceIdx, &FaceNormal](const auto& Downcast, const auto& FullTransformB)
+								{
+									return Chaos::SweepQuery(*Shape->GetGeometry(), WorldTM, Downcast, FullTransformB, Dir, DeltaMag, Distance, WorldPosition, WorldNormal, FaceIdx, FaceNormal, 0.f, false);
+								}
+							);
+
+							if (bShapeHit)
+							{
+								bHit = true;
+
+								OutBestHit.Shape = Shape;
+								OutBestHit.WorldPosition = WorldPosition;
+								OutBestHit.WorldNormal = WorldNormal;
+								OutBestHit.Distance = static_cast<float>(Distance);
+								OutBestHit.FaceIndex = FaceIdx;
+								if (OutBestHit.Distance > 0.f)
+								{
+									const FVector LocalPosition = WorldTM.InverseTransformPositionNoScale(OutBestHit.WorldPosition);
+									const FVector LocalUnitDir = WorldTM.InverseTransformVectorNoScale(Dir);
+									OutBestHit.FaceIndex = Shape->GetGeometry()->FindMostOpposingFace(LocalPosition, LocalUnitDir, OutBestHit.FaceIndex, 1);
+								}
+								OutBestHit.Actor = GetParticle(IterObject);
+							}
+						}
+						return false;
+					}
+				);
+
+			}
+			return bHit;
+		}
 
 		friend class FPhysicsObjectInterface;
 	protected:
@@ -113,7 +302,7 @@ namespace Chaos
 	private:
 		struct FShapeOverlapData
 		{
-			FPerShapeData* Shape;
+			TThreadShapeInstance<Id>* Shape;
 			FAABB3 BoundingBox;
 		};
 
diff --git a/Engine/Source/Runtime/Experimental/Chaos/Public/Chaos/ShapeInstanceFwd.h b/Engine/Source/Runtime/Experimental/Chaos/Public/Chaos/ShapeInstanceFwd.h
--- a/Engine/Source/Runtime/Experimental/Chaos/Public/Chaos/ShapeInstanceFwd.h
+++ b/Engine/Source/Runtime/Experimental/Chaos/Public/Chaos/ShapeInstanceFwd.h
@@ -2,8 +2,10 @@
 #pragma once
 
 #include "Containers/Array.h"
+#include "Framework/ThreadContextEnum.h"
 #include "Templates/SharedPointer.h"
 #include "Templates/UniquePtr.h"
+#include <type_traits>
 
 // @todo(chaos): use TRefCountPtr rather than TSharedPtr below - it is smaller (one atomic int instead of 2 atomic ints and 2 pointers)
 
@@ -29,4 +31,10 @@ namespace Chaos
 	using FShapeInstancePtr = TUniquePtr<FShapeInstance>;
 	using FConstShapeInstancePtr = TUniquePtr<const FShapeInstance>;
 	using FShapeInstanceArray = TArray<FShapeInstancePtr, TInlineAllocator<1>>;
+
+	template<EThreadContext Id>
+	using TThreadShapeInstance = std::conditional_t<Id == EThreadContext::External, FShapeInstanceProxy, FShapeInstance>;
+
+	template<EThreadContext Id>
+	using TThreadShapeInstanceArray = std::conditional_t<Id == EThreadContext::External, FShapeInstanceProxyArray, FShapeInstanceArray>;
 }
