fix for client crash when removing a part of a cluster union that would split the union into two pieces one of which would be instantly destroyed

- fixed an issue where a disabled cluster particle that doesn't have a cluster parent could be dirtied and added to the SQ
- disable particles immediately on the PT when unregistering cluster unions and geometry collections

diff --git a/Engine/Source/Runtime/Experimental/Chaos/Private/PBDRigidsSolver.cpp b/Engine/Source/Runtime/Experimental/Chaos/Private/PBDRigidsSolver.cpp
--- a/Engine/Source/Runtime/Experimental/Chaos/Private/PBDRigidsSolver.cpp
+++ b/Engine/Source/Runtime/Experimental/Chaos/Private/PBDRigidsSolver.cpp
@@ -677,6 +677,14 @@ namespace Chaos
 
 		EnqueueCommandImmediate([InProxy, this]()
 			{
+				for (FPBDRigidClusteredParticleHandle* Handle : InProxy->GetSolverParticleHandles())
+				{
+					if (Handle)
+					{
+						MEvolution->DisableParticle(Handle);
+					}
+				}
+
 				GeometryCollectionPhysicsProxies_Internal.RemoveSingle(InProxy);
 				InProxy->SyncBeforeDestroy();
 				InProxy->ResetDirtyIdx();
@@ -720,9 +728,18 @@ namespace Chaos
 		Proxy->MarkDeleted();
 		RemoveDirtyProxy(Proxy);
 
+		if (FPBDRigidParticle* GTParticle = Proxy->GetParticle_External())
+		{
+			GTParticle->SetProxy(nullptr);
+		}
+
 		EnqueueCommandImmediate(
 			[Proxy, this]()
 			{
+				if (Proxy->GetParticle_Internal())
+				{
+					MEvolution->DisableParticle(Proxy->GetParticle_Internal());
+				}
 				ClusterUnionPhysicsProxies_Internal.RemoveSingle(Proxy);
 				Proxy->SyncBeforeDestroy();
 				Proxy->ResetDirtyIdx();
diff --git a/Engine/Source/Runtime/Experimental/Chaos/Public/Chaos/PBDRigidsEvolution.h b/Engine/Source/Runtime/Experimental/Chaos/Public/Chaos/PBDRigidsEvolution.h
--- a/Engine/Source/Runtime/Experimental/Chaos/Public/Chaos/PBDRigidsEvolution.h
+++ b/Engine/Source/Runtime/Experimental/Chaos/Public/Chaos/PBDRigidsEvolution.h
@@ -537,6 +537,17 @@ public:
 							return;
 						}
 					}
+					else
+					{
+						// There's probably no way we get here since if we're clustering the parent should always be a clustered.
+						ensure(false);
+						return;
+					}
+				}
+				else
+				{
+					// Disabled cluster particle that doesn't have a parent cluster. MUST BE IGNORED.
+					return;
 				}
 			}
 			else
