fix an issue where sweeps against implicit object unions would not properly set the hit location/normal/etc.

diff --git a/Engine/Source/Runtime/Experimental/Chaos/Public/Chaos/GeometryQueries.h b/Engine/Source/Runtime/Experimental/Chaos/Public/Chaos/GeometryQueries.h
--- a/Engine/Source/Runtime/Experimental/Chaos/Public/Chaos/GeometryQueries.h
+++ b/Engine/Source/Runtime/Experimental/Chaos/Public/Chaos/GeometryQueries.h
@@ -248,6 +248,7 @@ namespace Chaos
 					const FImplicitObjectUnion& AUnion = static_cast<const FImplicitObjectUnion&>(A);
 
 					bool bHit = false;
+					OutTime = TNumericLimits<FReal>::Max();
 					AUnion.ForEachObject(
 						[&bHit, &ATM, &B, &BTM, &Dir, Length, &OutTime, &OutPosition, &OutNormal, &OutFaceIndex, &OutFaceNormal, Thickness, bComputeMTD](const FImplicitObject& SubObject, const FRigidTransform3& SubTransform)
 						{
