Chaos - fix trii mesh collisions. Don't throw away contact points within cull distance of the deepest contact.

This was causing problems with fast moving bodies (e.g., a box) that hit the mesh at an angle. We would throw away all but one corner contact and the body then rotates into the surafce and possibly ends up more than halfway through.

diff --git a/Engine/Source/Runtime/Experimental/Chaos/Private/Chaos/Collision/ContactTriangles.cpp b/Engine/Source/Runtime/Experimental/Chaos/Private/Chaos/Collision/ContactTriangles.cpp
--- a/Engine/Source/Runtime/Experimental/Chaos/Private/Chaos/Collision/ContactTriangles.cpp
+++ b/Engine/Source/Runtime/Experimental/Chaos/Private/Chaos/Collision/ContactTriangles.cpp
@@ -158,7 +158,7 @@ namespace Chaos
 		// NOTE: relies on the sort above, and retains order
 		if (TriangleContactPoints.Num() > 4)
 		{
-			PruneUnnecessaryContactPoints(FReal(0.1), FReal(0.1));
+			PruneUnnecessaryContactPoints();
 		}
 
 		DebugDrawContactPoints(FColor::Yellow, 0.5);
@@ -370,7 +370,7 @@ namespace Chaos
 	// NOTE: ContactPoints must be sorted on Phi on entry, and will be sorted on exit
 	// @todo(chaos): we should only remove shallow contacts that are hidden by nearby deeper contacts. It should be ok to have shallow contacts
 	// far from the deeper ones. E.g., a box landing on plane at an angle.
-	void FContactTriangleCollector::PruneUnnecessaryContactPoints(const FReal PhiTolerance, const FReal DistanceTolerance)
+	void FContactTriangleCollector::PruneUnnecessaryContactPoints()
 	{
 		// Remove all points except for the deepest one, and ones with phis similar to it
 		int32 NewContactPointCount = TriangleContactPoints.Num() > 0 ? 1 : 0;
diff --git a/Engine/Source/Runtime/Experimental/Chaos/Private/Chaos/CollisionOneShotManifolds.cpp b/Engine/Source/Runtime/Experimental/Chaos/Private/Chaos/CollisionOneShotManifolds.cpp
--- a/Engine/Source/Runtime/Experimental/Chaos/Private/Chaos/CollisionOneShotManifolds.cpp
+++ b/Engine/Source/Runtime/Experimental/Chaos/Private/Chaos/CollisionOneShotManifolds.cpp
@@ -89,6 +89,12 @@ namespace Chaos
 	FAutoConsoleVariableRef CVarChaos_Collision_OneSidedTriangleMesh(TEXT("p.Chaos.Collision.OneSidedTriangleMesh"), bChaos_Collision_OneSidedTriangleMesh, TEXT(""));
 	FAutoConsoleVariableRef CVarChaos_Collision_OneSidedHeightfield(TEXT("p.Chaos.Collision.OneSidedHeightField"), bChaos_Collision_OneSidedHeightField, TEXT(""));
 
+	// Ueed to reject contacts against tri meshes
+	FRealSingle Chaos_Collision_TriMeshPhiToleranceScale = 1.0f;	// A multipler on cull distance. Points farther than this from the deepest point are ignored
+	FRealSingle Chaos_Collision_TriMeshDistanceTolerance = 0.1f;	// Points closer than this to a deeper point are ignored
+	FAutoConsoleVariableRef CVarChaos_Collision_TriMeshDistanceolerance(TEXT("p.Chaos.Collision.TriangeMeshDistanceTolerance"), Chaos_Collision_TriMeshDistanceTolerance, TEXT(""));
+	FAutoConsoleVariableRef CVarChaos_Collision_TriMeshPhiToleranceScale(TEXT("p.Chaos.Collision.TriangeMeshPhiToleranceScale"), Chaos_Collision_TriMeshPhiToleranceScale, TEXT(""));
+
 	bool bChaos_Collision_UseCapsuleTriMesh2 = true;
 	FAutoConsoleVariableRef CVarChaos_Collision_UseCapsuleTriMesh2(TEXT("p.Chaos.Collision.UseCapsuleTriMesh2"), bChaos_Collision_UseCapsuleTriMesh2, TEXT(""));
 
diff --git a/Engine/Source/Runtime/Experimental/Chaos/Private/Chaos/CollisionOneShotManifoldsMeshShapes.cpp b/Engine/Source/Runtime/Experimental/Chaos/Private/Chaos/CollisionOneShotManifoldsMeshShapes.cpp
--- a/Engine/Source/Runtime/Experimental/Chaos/Private/Chaos/CollisionOneShotManifoldsMeshShapes.cpp
+++ b/Engine/Source/Runtime/Experimental/Chaos/Private/Chaos/CollisionOneShotManifoldsMeshShapes.cpp
@@ -20,12 +20,14 @@
 #include "Chaos/Utilities.h"
 #include "ChaosStats.h"
 
-//PRAGMA_DISABLE_OPTIMIZATION
+//UE_DISABLE_OPTIMIZATION
 
 namespace Chaos
 {
 	extern bool bChaos_Collision_OneSidedTriangleMesh;
 	extern bool bChaos_Collision_OneSidedHeightField;
+	extern FRealSingle Chaos_Collision_TriMeshDistanceTolerance;
+	extern FRealSingle Chaos_Collision_TriMeshPhiToleranceScale;
 
 	extern bool bChaos_Collision_UseCapsuleTriMesh2;
 	extern bool bChaos_Collision_UseConvexTriMesh2;
@@ -33,6 +35,11 @@ namespace Chaos
 	namespace Collisions
 	{
 
+		inline FReal CalculateTriMeshPhiTolerance(const FReal CullDistance)
+		{
+			return Chaos_Collision_TriMeshPhiToleranceScale * CullDistance;
+		}
+
 		template <typename TriMeshType>
 		void ConstructSphereTriangleMeshOneShotManifold(const TSphere<FReal, 3>& Sphere, const FRigidTransform3& SphereWorldTransform, const TriMeshType& TriangleMesh, const FRigidTransform3& TriMeshWorldTransform, const FReal Dt, FPBDCollisionConstraint& Constraint)
 		{
@@ -164,8 +171,10 @@ namespace Chaos
 			const FTriangleMeshImplicitObject* Mesh = UnwrapImplicit<FTriangleMeshImplicitObject>(InMesh, MeshScale, MeshMargin);
 			check(Mesh != nullptr);
 
-			FContactTriangleCollector MeshContacts(bChaos_Collision_OneSidedTriangleMesh, QuadraticTransform);
 			const FReal CullDistance = Constraint.GetCullDistance();
+			const FReal PhiTolerance = CalculateTriMeshPhiTolerance(CullDistance);
+			const FReal DistanceTolerance = Chaos_Collision_TriMeshDistanceTolerance;
+			FContactTriangleCollector MeshContacts(bChaos_Collision_OneSidedTriangleMesh, PhiTolerance, DistanceTolerance, QuadraticTransform);
 
 			if (const FImplicitSphere3* Sphere = Quadratic.template GetObject<FImplicitSphere3>())
 			{
@@ -190,8 +199,10 @@ namespace Chaos
 			ensure(QuadraticTransform.GetScale3D() == FVec3(1));
 			ensure(MeshTransform.GetScale3D() == FVec3(1));
 
-			FContactTriangleCollector MeshContacts(bChaos_Collision_OneSidedHeightField, QuadraticTransform);
 			const FReal CullDistance = Constraint.GetCullDistance();
+			const FReal PhiTolerance = CalculateTriMeshPhiTolerance(CullDistance);
+			const FReal DistanceTolerance = Chaos_Collision_TriMeshDistanceTolerance;
+			FContactTriangleCollector MeshContacts(bChaos_Collision_OneSidedHeightField, PhiTolerance, DistanceTolerance, QuadraticTransform);
 
 			if (const FImplicitSphere3* Sphere = Quadratic.template GetObject<FImplicitSphere3>())
 			{
@@ -226,8 +237,10 @@ namespace Chaos
 			const FTriangleMeshImplicitObject* Mesh = UnwrapImplicit<FTriangleMeshImplicitObject>(InMesh, MeshScale, MeshMargin);
 			check(Mesh != nullptr);
 
-			FContactTriangleCollector MeshContacts(bChaos_Collision_OneSidedTriangleMesh, ConvexTransform);
 			const FReal CullDistance = Constraint.GetCullDistance();
+			const FReal PhiTolerance = CalculateTriMeshPhiTolerance(CullDistance);
+			const FReal DistanceTolerance = Chaos_Collision_TriMeshDistanceTolerance;
+			FContactTriangleCollector MeshContacts(bChaos_Collision_OneSidedTriangleMesh, PhiTolerance, DistanceTolerance, ConvexTransform);
 
 			if (const FImplicitBox3* RawBox = Convex.template GetObject<FImplicitBox3>())
 			{
@@ -263,10 +276,13 @@ namespace Chaos
 			ensure(ConvexTransform.GetScale3D() == FVec3(1));
 			ensure(MeshTransform.GetScale3D() == FVec3(1));
 
-			FContactTriangleCollector MeshContacts(bChaos_Collision_OneSidedHeightField, ConvexTransform);
-			const FReal CullDistance = Constraint.GetCullDistance();
 			const FVec3 MeshScale = FVec3(1);	// Scale is built into heightfield
 
+			const FReal CullDistance = Constraint.GetCullDistance();
+			const FReal PhiTolerance = CalculateTriMeshPhiTolerance(CullDistance);
+			const FReal DistanceTolerance = Chaos_Collision_TriMeshDistanceTolerance;
+			FContactTriangleCollector MeshContacts(bChaos_Collision_OneSidedHeightField, PhiTolerance, DistanceTolerance, ConvexTransform);
+
 			if (const FImplicitBox3* RawBox = Convex.template GetObject<FImplicitBox3>())
 			{
 				ConstructConvexMeshOneShotManifold(*RawBox, ConvexTransform, Mesh, MeshTransform, MeshScale, CullDistance, MeshContacts);
diff --git a/Engine/Source/Runtime/Experimental/Chaos/Public/Chaos/Collision/ContactTriangles.h b/Engine/Source/Runtime/Experimental/Chaos/Public/Chaos/Collision/ContactTriangles.h
--- a/Engine/Source/Runtime/Experimental/Chaos/Public/Chaos/Collision/ContactTriangles.h
+++ b/Engine/Source/Runtime/Experimental/Chaos/Public/Chaos/Collision/ContactTriangles.h
@@ -223,8 +223,23 @@ namespace Chaos
 	class FContactTriangleCollector
 	{
 	public:
+		UE_DEPRECATED(5.3, "Use the constructor below which take PhiTolerance etc")
 		FContactTriangleCollector(const bool bInOneSided, const FRigidTransform3& InConvexTransform)
-			: bOneSidedCollision(bInOneSided)
+			: PhiTolerance(0.1)
+			, DistanceTolerance(0.1)
+			, bOneSidedCollision(bInOneSided)
+			, ConvexTransform(InConvexTransform)
+		{
+		}
+
+		FContactTriangleCollector(
+			const bool bInOneSided,
+			const FReal InPhiTolerance,
+			const FReal InDistanceTolerance,
+			const FRigidTransform3& InConvexTransform)
+			: PhiTolerance(InPhiTolerance)
+			, DistanceTolerance(InDistanceTolerance)
+			, bOneSidedCollision(bInOneSided)
 			, ConvexTransform(InConvexTransform)
 		{
 		}
@@ -268,7 +283,7 @@ namespace Chaos
 		void PruneEdgeContactPoints();
 		void FixInvalidNormalContactPoints();
 		void PruneInfacingContactPoints();
-		void PruneUnnecessaryContactPoints(const FReal PhiTolerance, const FReal DistanceTolerance);
+		void PruneUnnecessaryContactPoints();
 		void ReduceManifoldContactPointsTriangeMesh();
 		void FinalizeContacts(const FRigidTransform3& MeshToConvexTransform);
 
@@ -294,6 +309,12 @@ namespace Chaos
 		// A list of all the vertices of contacts with faces and edges. Used to prune vertex contacts
 		TSet<FContactVertexID> ContactVertices;
 
+		// We remove contacts that are shallower by this much compared to the deepest contact (negative to disable this functionality)
+		FReal PhiTolerance;
+
+		// We remove contacts that are closer than this to any other contact (negative to disable this functionality)
+		FReal DistanceTolerance;
+
 		bool bOneSidedCollision;
 
 		FRigidTransform3 ConvexTransform;
