Chaos: Fixed repeated calls to update geometry collection particle views during geometry collection proxy initialization. Defer to one view update at the end of the creation loop.

diff --git a/Engine/Source/Runtime/Experimental/Chaos/Private/PBDRigidsSolver.cpp b/Engine/Source/Runtime/Experimental/Chaos/Private/PBDRigidsSolver.cpp
--- a/Engine/Source/Runtime/Experimental/Chaos/Private/PBDRigidsSolver.cpp
+++ b/Engine/Source/Runtime/Experimental/Chaos/Private/PBDRigidsSolver.cpp
@@ -1325,7 +1325,8 @@ namespace Chaos
 		};
 
 		//need to create new particle handles
-		DirtyProxiesData->ForEachProxy([this, &ProcessProxyPT, Manager](int32 DataIdx,FDirtyProxy& Dirty)
+		int32 NumInitializedGCProxies = 0;
+		DirtyProxiesData->ForEachProxy([this, &ProcessProxyPT, Manager, &NumInitializedGCProxies](int32 DataIdx,FDirtyProxy& Dirty)
 		{
 			if(Dirty.Proxy->GetIgnoreDataOnStep_Internal() != CurrentFrame)
 			{
@@ -1354,6 +1355,8 @@ namespace Chaos
 						{
 							// Finish registration on the physics thread...
 							Proxy->InitializeBodiesPT(this, GetParticles());
+							++NumInitializedGCProxies;
+
 							GeometryCollectionPhysicsProxies_Internal.Add(Proxy);
 						}
 						Proxy->PushToPhysicsState();
@@ -1391,6 +1394,11 @@ namespace Chaos
 			}
 		});
 
+		if(NumInitializedGCProxies > 0)
+		{
+			GetParticles().UpdateGeometryCollectionViews(true);
+		}
+
 		//need to create new constraint handles
 		DirtyProxiesData->ForEachProxy([this, Manager, RewindData](int32 DataIdx, FDirtyProxy& Dirty)
 		{
diff --git a/Engine/Source/Runtime/Experimental/Chaos/Private/PhysicsProxy/GeometryCollectionPhysicsProxy.cpp b/Engine/Source/Runtime/Experimental/Chaos/Private/PhysicsProxy/GeometryCollectionPhysicsProxy.cpp
--- a/Engine/Source/Runtime/Experimental/Chaos/Private/PhysicsProxy/GeometryCollectionPhysicsProxy.cpp
+++ b/Engine/Source/Runtime/Experimental/Chaos/Private/PhysicsProxy/GeometryCollectionPhysicsProxy.cpp
@@ -1042,11 +1042,7 @@ void FGeometryCollectionPhysicsProxy::InitializeBodiesPT(Chaos::FPBDRigidsSolver
 			}
 		}
 
-		// After population, the states of each particle could have changed
-		// @todo(Chaos) Do we really need to maintain the Views while running this function? We do it twice, but there's also a Dirty flag built-in
-		Particles.UpdateGeometryCollectionViews();
-
-		for (FFieldSystemCommand& Cmd : Parameters.InitializationCommands)
+		for(FFieldSystemCommand& Cmd : Parameters.InitializationCommands)
 		{
 			if(Cmd.MetaData.Contains(FFieldSystemMetaData::EMetaType::ECommandData_ProcessingResolution))
 			{
@@ -1055,7 +1051,7 @@ void FGeometryCollectionPhysicsProxy::InitializeBodiesPT(Chaos::FPBDRigidsSolver
 
 			FFieldSystemMetaDataProcessingResolution* ResolutionData = new FFieldSystemMetaDataProcessingResolution(EFieldResolutionType::Field_Resolution_Maximum);
 
-			Cmd.MetaData.Add( FFieldSystemMetaData::EMetaType::ECommandData_ProcessingResolution, TUniquePtr<FFieldSystemMetaDataProcessingResolution>(ResolutionData));
+			Cmd.MetaData.Add(FFieldSystemMetaData::EMetaType::ECommandData_ProcessingResolution, TUniquePtr<FFieldSystemMetaDataProcessingResolution>(ResolutionData));
 			Commands.Add(Cmd);
 		}
 		Parameters.InitializationCommands.Empty();
@@ -1218,11 +1214,6 @@ void FGeometryCollectionPhysicsProxy::InitializeBodiesPT(Chaos::FPBDRigidsSolver
 				}
 			}
 
-			// We've likely changed the state of leaf nodes, which are geometry
-			// collection particles.  Update which particle views they belong in,
-			// as well as views of clustered particles.
-			Particles.UpdateGeometryCollectionViews(true); 
-
 			const TManagedArray<TSet<int32>>* Connections = RestCollection->FindAttribute<TSet<int32>>("Connections", FGeometryCollection::TransformGroup);
 			const bool bGenerateConnectionGraph = (!Connections) || (Connections->Num() != RestCollection->Transform.Num()) || bGeometryCollectionAlwaysGenerateConnectionGraph;
 			if (bGenerateConnectionGraph)
