Chaos - remove array slack from Unions
- fix natvis for ManagedArray

diff --git a/Engine/Source/Runtime/Experimental/Chaos/Private/Chaos/ImplicitObjectUnion.cpp b/Engine/Source/Runtime/Experimental/Chaos/Private/Chaos/ImplicitObjectUnion.cpp
--- a/Engine/Source/Runtime/Experimental/Chaos/Private/Chaos/ImplicitObjectUnion.cpp
+++ b/Engine/Source/Runtime/Experimental/Chaos/Private/Chaos/ImplicitObjectUnion.cpp
@@ -74,6 +74,7 @@ void FImplicitObjectUnion::Combine(TArray<TUniquePtr<FImplicitObject>>& OtherObj
 		MLocalBoundingBox.GrowToInclude(OtherObjects[i]->BoundingBox());
 	}
 
+	MObjects.Reserve(MObjects.Num() + OtherObjects.Num());
 	for (TUniquePtr<FImplicitObject>& ChildObject: OtherObjects)
 	{
 		MObjects.Add(MoveTemp(ChildObject));
diff --git a/Engine/Source/Runtime/Experimental/Chaos/Private/Chaos/Island/SolverIsland.cpp b/Engine/Source/Runtime/Experimental/Chaos/Private/Chaos/Island/SolverIsland.cpp
--- a/Engine/Source/Runtime/Experimental/Chaos/Private/Chaos/Island/SolverIsland.cpp
+++ b/Engine/Source/Runtime/Experimental/Chaos/Private/Chaos/Island/SolverIsland.cpp
@@ -8,18 +8,23 @@ namespace Chaos
 	namespace Private
 	{
 	
-	FPBDIsland::FPBDIsland(const int32 MaxConstraintContainers)
-	: IslandIndex(INDEX_NONE)
-	, bIsSleeping(false)
-	, bNeedsResim(false)
-	, bIsPersistent(true)
-	, SleepCounter(0)
-	, IslandParticles()
-	, IslandConstraintsByType()
-	, NumConstraints(0)
-{
-	IslandConstraintsByType.SetNum(MaxConstraintContainers);
-}
+		FPBDIsland::FPBDIsland(const int32 MaxConstraintContainers)
+			: IslandIndex(INDEX_NONE)
+			, bIsSleeping(false)
+			, bNeedsResim(false)
+			, bIsPersistent(true)
+			, SleepCounter(0)
+			, IslandParticles()
+			, IslandConstraintsByType()
+			, NumConstraints(0)
+		{
+			// NOTE: TIinlineAllocator used for IslandConstraintsByType.
+			// We do not expect this to allocate but one day we will have to figure 
+			// out how to allow users to add new constraint types and then we will 
+			// need to do something to avoid many small allocations (e.g., use a pool)
+			IslandConstraintsByType.Reserve(MaxConstraintContainers);
+			IslandConstraintsByType.SetNum(MaxConstraintContainers);
+		}
 
 		void FPBDIsland::Reuse()
 		{
diff --git a/Engine/Source/Runtime/Experimental/Chaos/Private/Chaos/ShapeInstance.cpp b/Engine/Source/Runtime/Experimental/Chaos/Private/Chaos/ShapeInstance.cpp
--- a/Engine/Source/Runtime/Experimental/Chaos/Private/Chaos/ShapeInstance.cpp
+++ b/Engine/Source/Runtime/Experimental/Chaos/Private/Chaos/ShapeInstance.cpp
@@ -27,6 +27,7 @@ namespace Chaos
 			const int32 OldShapeNum = ShapesArray.Num();
 			if (const auto* Union = Geometry->template GetObject<FImplicitObjectUnion>())
 			{
+				ShapesArray.Reserve(Union->GetObjects().Num());
 				ShapesArray.SetNum(Union->GetObjects().Num());
 
 				for (int32 ShapeIndex = 0; ShapeIndex < ShapesArray.Num(); ++ShapeIndex)
@@ -47,6 +48,7 @@ namespace Chaos
 			}
 			else
 			{
+				ShapesArray.Reserve(1);
 				ShapesArray.SetNum(1);
 				if (OldShapeNum == 0)
 				{
