fix headless chaos failing when trying to set initialLinearVelocity

#jir anone

diff --git a/Engine/Source/Runtime/Experimental/Chaos/Private/PhysicsProxy/GeometryCollectionPhysicsProxy.cpp b/Engine/Source/Runtime/Experimental/Chaos/Private/PhysicsProxy/GeometryCollectionPhysicsProxy.cpp
--- a/Engine/Source/Runtime/Experimental/Chaos/Private/PhysicsProxy/GeometryCollectionPhysicsProxy.cpp
+++ b/Engine/Source/Runtime/Experimental/Chaos/Private/PhysicsProxy/GeometryCollectionPhysicsProxy.cpp
@@ -1071,7 +1071,7 @@ void FGeometryCollectionPhysicsProxy::InitializeBodiesPT(Chaos::FPBDRigidsSolver
 				{
 					if (DynamicState[TransformGroupIndex] == (int32)EObjectStateTypeEnum::Chaos_Object_Dynamic)
 					{
-						Handle->SetV(InitialVelocityFacade.InitialAngularVelocityAttribute[TransformGroupIndex]);
+						Handle->SetV(InitialVelocityFacade.InitialLinearVelocityAttribute[TransformGroupIndex]);
 						Handle->SetW(InitialVelocityFacade.InitialAngularVelocityAttribute[TransformGroupIndex]);
 					}
 				}
