Iris - Bump max RPC size.

#rnx

diff --git a/Engine/Source/Runtime/Experimental/Iris/Core/Private/Iris/ReplicationSystem/NetBlob/NetRPC.h b/Engine/Source/Runtime/Experimental/Iris/Core/Private/Iris/ReplicationSystem/NetBlob/NetRPC.h
--- a/Engine/Source/Runtime/Experimental/Iris/Core/Private/Iris/ReplicationSystem/NetBlob/NetRPC.h
+++ b/Engine/Source/Runtime/Experimental/Iris/Core/Private/Iris/ReplicationSystem/NetBlob/NetRPC.h
@@ -59,7 +59,7 @@ private:
 
 private:
 
-	static constexpr uint32 HeaderSizeBitCount = 20U;
+	static constexpr uint32 HeaderSizeBitCount = 24U;
 	static constexpr uint32 MaxRpcSizeInBits = (1U << HeaderSizeBitCount) - 1U;
 
 private:
