dataflow : fix bounding box computation facade causing potential crashes
- use geometryIndex instead of transform index to access the bounding box attribute
- compute the bounds from all the valid geometries instead of querying the roots ( which may not have any bounds )

diff --git a/Engine/Source/Runtime/Experimental/Chaos/Private/GeometryCollection/Facades/CollectionBoundsFacade.cpp b/Engine/Source/Runtime/Experimental/Chaos/Private/GeometryCollection/Facades/CollectionBoundsFacade.cpp
--- a/Engine/Source/Runtime/Experimental/Chaos/Private/GeometryCollection/Facades/CollectionBoundsFacade.cpp
+++ b/Engine/Source/Runtime/Experimental/Chaos/Private/GeometryCollection/Facades/CollectionBoundsFacade.cpp
@@ -102,47 +102,26 @@ namespace GeometryCollection::Facades
 		return Centroids;
 	}
 
-	FBox FBoundsFacade::GetBoundingBox() const
-	{
-		FBox BoundingBox;
-
-		if (IsValid())
-		{
-			const TArray<int32>& RootIndices = Chaos::Facades::FCollectionHierarchyFacade::GetRootIndices(ParentAttribute);
-
-			const TManagedArray<FBox>& BoundingBoxArr = BoundingBoxAttribute.Get();
-
-			BoundingBox.Init();
-
-			for (int32 Idx : RootIndices)
-			{
-				BoundingBox += BoundingBoxArr[Idx];
-			}
-		}
-
-		return BoundingBox;
-	}
-
 	FBox FBoundsFacade::GetBoundingBoxInCollectionSpace() const
 	{
 		FBox BoundingBox;
+		BoundingBox.Init();
 
 		if (IsValid())
 		{
-			const TArray<int32>& RootIndices = Chaos::Facades::FCollectionHierarchyFacade::GetRootIndices(ParentAttribute);
-
-			const TManagedArray<FBox>& BoundingBoxArr = BoundingBoxAttribute.Get();
-
 			GeometryCollection::Facades::FCollectionTransformFacade TransformFacade(ConstCollection);
+			TArray<FTransform> CollectionSpaceTransforms = TransformFacade.ComputeCollectionSpaceTransforms();
 
-			BoundingBox.Init();
-
-			for (int32 Idx : RootIndices)
+			for (int32 TransformIndex = 0; TransformIndex < CollectionSpaceTransforms.Num(); ++TransformIndex)
 			{
-				const FTransform CollectionSpaceTransform = TransformFacade.ComputeCollectionSpaceTransform(Idx);
-				const FBox BoundingBoxInCollectionSpace = BoundingBoxArr[Idx].TransformBy(CollectionSpaceTransform);
-
-				BoundingBox += BoundingBoxInCollectionSpace;
+				const int32 GeoIndex = TransformToGeometryIndexAttribute[TransformIndex];
+				if (BoundingBoxAttribute.IsValidIndex(GeoIndex))
+				{
+					const FBox& GeoBoundingBox = BoundingBoxAttribute[GeoIndex];
+					const FTransform CollectionSpaceTransform = CollectionSpaceTransforms[TransformIndex];
+					const FBox BoundingBoxInCollectionSpace = GeoBoundingBox.TransformBy(CollectionSpaceTransform);
+					BoundingBox += BoundingBoxInCollectionSpace;
+				}
 			}
 		}
 
diff --git a/Engine/Source/Runtime/Experimental/Chaos/Public/GeometryCollection/Facades/CollectionBoundsFacade.h b/Engine/Source/Runtime/Experimental/Chaos/Public/GeometryCollection/Facades/CollectionBoundsFacade.h
--- a/Engine/Source/Runtime/Experimental/Chaos/Public/GeometryCollection/Facades/CollectionBoundsFacade.h
+++ b/Engine/Source/Runtime/Experimental/Chaos/Public/GeometryCollection/Facades/CollectionBoundsFacade.h
@@ -47,9 +47,6 @@ namespace GeometryCollection::Facades
 		/** Centroids (Centers of BoundingBoxes) access */
 		TArray<FVector> GetCentroids() const;
 
-		/** BoundingBox for the whole collection */
-		FBox GetBoundingBox() const;
-
 		/** BoundingBox for the whole collection in Collection space */
 		FBox GetBoundingBoxInCollectionSpace() const;
 
diff --git a/Engine/Source/Runtime/Experimental/Chaos/Public/GeometryCollection/ManagedArrayAccessor.h b/Engine/Source/Runtime/Experimental/Chaos/Public/GeometryCollection/ManagedArrayAccessor.h
--- a/Engine/Source/Runtime/Experimental/Chaos/Public/GeometryCollection/ManagedArrayAccessor.h
+++ b/Engine/Source/Runtime/Experimental/Chaos/Public/GeometryCollection/ManagedArrayAccessor.h
@@ -54,6 +54,12 @@ public:
 
 	bool IsPersistent() const { return ConstCollection.IsAttributePersistent(Name, Group); }
 
+	bool IsValidIndex(int32 Index) const 
+	{ 
+		check(IsValid());
+		return (Index >= 0) && (Index < ConstAttributeArray->Num());
+	}
+
 	int32 AddElements(int32 NumElements) const 
 	{ 
 		check(!IsConst());
