geometry collection : make scale on removal read write for blueprints

diff --git a/Engine/Source/Runtime/Experimental/GeometryCollectionEngine/Public/GeometryCollection/GeometryCollectionObject.h b/Engine/Source/Runtime/Experimental/GeometryCollectionEngine/Public/GeometryCollection/GeometryCollectionObject.h
--- a/Engine/Source/Runtime/Experimental/GeometryCollectionEngine/Public/GeometryCollection/GeometryCollectionObject.h
+++ b/Engine/Source/Runtime/Experimental/GeometryCollectionEngine/Public/GeometryCollection/GeometryCollectionObject.h
@@ -649,7 +649,7 @@ public:
 #endif
 
 	/** When enabled, particle will scale down (shrink) when using being removed ( using both remove on sleep or remove on break ) - Enabled by default */
-	UPROPERTY(EditAnywhere, BlueprintReadOnly, Category = Removal)
+	UPROPERTY(EditAnywhere, BlueprintReadWrite, Category = Removal)
 	bool bScaleOnRemoval;
 
 	/** Remove particle from simulation and dissolve rendered geometry once sleep threshold has been exceeded. */
