Sequencer: Constraints: Don't remove a constraint when compensating since that can cause the channel to get deleted while we are adding it.

diff --git a/Engine/Plugins/Animation/ControlRig/Source/ControlRigEditor/Private/Sequencer/ControlConstraintChannelInterface.cpp b/Engine/Plugins/Animation/ControlRig/Source/ControlRigEditor/Private/Sequencer/ControlConstraintChannelInterface.cpp
--- a/Engine/Plugins/Animation/ControlRig/Source/ControlRigEditor/Private/Sequencer/ControlConstraintChannelInterface.cpp
+++ b/Engine/Plugins/Animation/ControlRig/Source/ControlRigEditor/Private/Sequencer/ControlConstraintChannelInterface.cpp
@@ -106,7 +106,8 @@ bool FControlConstraintChannelInterface::SmartConstraintKey(
 			const bool bNeedsCompensation = InConstraint->NeedsCompensation();
 				
 			TGuardValue<bool> CompensateGuard(FMovieSceneConstraintChannelHelper::bDoNotCompensate, true);
-				
+			TGuardValue<bool> RemoveConstraintGuard(FConstraintsManagerController::bDoNotRemoveConstraint, true);
+
 			UControlRig* ControlRig = ControlHandle->ControlRig.Get();
 			const FName& ControlName = ControlHandle->ControlName;
 				
diff --git a/Engine/Source/Editor/MovieSceneTools/Private/Constraints/ComponentConstraintChannelInterface.cpp b/Engine/Source/Editor/MovieSceneTools/Private/Constraints/ComponentConstraintChannelInterface.cpp
--- a/Engine/Source/Editor/MovieSceneTools/Private/Constraints/ComponentConstraintChannelInterface.cpp
+++ b/Engine/Source/Editor/MovieSceneTools/Private/Constraints/ComponentConstraintChannelInterface.cpp
@@ -99,6 +99,8 @@ bool FComponentConstraintChannelInterface::SmartConstraintKey(
 			//new for compensation
 
 			TGuardValue<bool> CompensateGuard(FMovieSceneConstraintChannelHelper::bDoNotCompensate, true);
+			TGuardValue<bool> RemoveConstraintGuard(FConstraintsManagerController::bDoNotRemoveConstraint, true);
+
 			// store the frames to compensate
 			const TArrayView<FMovieSceneDoubleChannel*> Channels = ComponentHandle->GetDoubleChannels(TransformSection);
 			TArray<FFrameNumber> FramesToCompensate;
diff --git a/Engine/Source/Runtime/Experimental/Animation/Constraints/Private/ConstraintsManager.cpp b/Engine/Source/Runtime/Experimental/Animation/Constraints/Private/ConstraintsManager.cpp
--- a/Engine/Source/Runtime/Experimental/Animation/Constraints/Private/ConstraintsManager.cpp
+++ b/Engine/Source/Runtime/Experimental/Animation/Constraints/Private/ConstraintsManager.cpp
@@ -175,9 +175,12 @@ void UConstraintsManager::OnActorDestroyed(AActor* InActor)
 		if (!IndicesToRemove.IsEmpty())
 		{
 			FConstraintsManagerController& Controller = FConstraintsManagerController::Get(InActor->GetWorld());
-			for (int32 Index = IndicesToRemove.Num()-1; Index >= 0; --Index)
+			if (FConstraintsManagerController::bDoNotRemoveConstraint == false)
 			{
-				Controller.RemoveConstraint(Index,/*do not compensate*/ true);
+				for (int32 Index = IndicesToRemove.Num() - 1; Index >= 0; --Index)
+				{
+					Controller.RemoveConstraint(Index,/*do not compensate*/ true);
+				}
 			}
 		}
 	}
@@ -338,6 +341,13 @@ void UConstraintsManager::Dump() const
 	}
 }
 
+/**
+ * FConstraintsManagerController
+ **/
+
+
+bool FConstraintsManagerController::bDoNotRemoveConstraint = false;
+
 FConstraintsManagerController& FConstraintsManagerController::Get(UWorld* InWorld)
 {
 	static FConstraintsManagerController Singleton;
@@ -514,6 +524,10 @@ int32 FConstraintsManagerController::GetConstraintIndex(const FName& InConstrain
 	
 bool FConstraintsManagerController::RemoveConstraint(UTickableConstraint* InConstraint, bool bDoNotCompensate) const
 {
+	if (FConstraintsManagerController::bDoNotRemoveConstraint == false)
+	{
+		return false;
+	}
 	const UConstraintsManager* Manager = FindManager();
 	if (!Manager)
 	{
diff --git a/Engine/Source/Runtime/Experimental/Animation/Constraints/Public/ConstraintsManager.h b/Engine/Source/Runtime/Experimental/Animation/Constraints/Public/ConstraintsManager.h
--- a/Engine/Source/Runtime/Experimental/Animation/Constraints/Public/ConstraintsManager.h
+++ b/Engine/Source/Runtime/Experimental/Animation/Constraints/Public/ConstraintsManager.h
@@ -209,7 +209,10 @@ DECLARE_MULTICAST_DELEGATE_TwoParams(FConstraintsManagerNotifyDelegate, EConstra
 class CONSTRAINTS_API FConstraintsManagerController
 {
 public:
-	/** @todo document */
+	/** Static to control if we should remove constraints or not, may not want to do some when
+	compensating since we may be deleting spawnables while doing so*/
+	static bool bDoNotRemoveConstraint;
+	/** Get and Set Controller acive in this world*/
 	static FConstraintsManagerController& Get(UWorld* InWorld);
 
 	/** Allocates a new constraint with the constraints manager as the owner. */
