Fix build time crash in geometry component render data.
Also updating DDC key because that caused the crash could also potentially build corrupt data.

diff --git a/Engine/Source/Runtime/Experimental/GeometryCollectionEngine/Private/GeometryCollection/DerivedDataGeometryCollectionCooker.cpp b/Engine/Source/Runtime/Experimental/GeometryCollectionEngine/Private/GeometryCollection/DerivedDataGeometryCollectionCooker.cpp
--- a/Engine/Source/Runtime/Experimental/GeometryCollectionEngine/Private/GeometryCollection/DerivedDataGeometryCollectionCooker.cpp
+++ b/Engine/Source/Runtime/Experimental/GeometryCollectionEngine/Private/GeometryCollection/DerivedDataGeometryCollectionCooker.cpp
@@ -121,7 +121,7 @@ bool FDerivedDataGeometryCollectionRenderDataCooker::Build(TArray<uint8>& OutDat
 
 const TCHAR* FDerivedDataGeometryCollectionRenderDataCooker::GetVersionString() const
 {
-	const TCHAR* VersionString = TEXT("C3F0D73CE6094A1C84E57D75C9131F55");
+	const TCHAR* VersionString = TEXT("043C044F08614AAF8A56BA8070BD0D5A");
 
 	static FString CachedNaniteVersionString;
 	if (CachedNaniteVersionString.IsEmpty())
diff --git a/Engine/Source/Runtime/Experimental/GeometryCollectionEngine/Private/GeometryCollection/GeometryCollectionRenderData.cpp b/Engine/Source/Runtime/Experimental/GeometryCollectionEngine/Private/GeometryCollection/GeometryCollectionRenderData.cpp
--- a/Engine/Source/Runtime/Experimental/GeometryCollectionEngine/Private/GeometryCollection/GeometryCollectionRenderData.cpp
+++ b/Engine/Source/Runtime/Experimental/GeometryCollectionEngine/Private/GeometryCollection/GeometryCollectionRenderData.cpp
@@ -243,7 +243,7 @@ void BuildMeshDataFromGeometryCollection(FGeometryCollection& InCollection, FGeo
 					{
 						UV = FVector2f::ZeroVector;
 					}
-					BuildVertexData.UVs[TexCoord].Emplace(UV);
+					BuildVertexData.UVs[TexCoord][DestVertexStart + VertexIndex] = UV;
 				}
 
 				const int32 BoneIndex = BoneMapArray[VertexStart + VertexIndex];
