[RewindData]

Don't reset frames saved on rewind to allow for potential subsequent rewinds to the same frame or earlier.

diff --git a/Engine/Source/Runtime/Experimental/Chaos/Private/RewindData.cpp b/Engine/Source/Runtime/Experimental/Chaos/Private/RewindData.cpp
--- a/Engine/Source/Runtime/Experimental/Chaos/Private/RewindData.cpp
+++ b/Engine/Source/Runtime/Experimental/Chaos/Private/RewindData.cpp
@@ -322,8 +322,7 @@ bool FRewindData::RewindToFrame(int32 Frame)
 
 	CurFrame = Frame;
 	bNeedsSave = false;
-	FramesSaved = 0; //can't rewind before this point. This simplifies saving the state at head
-
+	
 	return true;
 }
 
