Iris
- Support 64 bit NetRefHandle ID assignment

diff --git a/Engine/Source/Runtime/Experimental/Iris/Core/Private/Iris/ReplicationSystem/NetRefHandle.cpp b/Engine/Source/Runtime/Experimental/Iris/Core/Private/Iris/ReplicationSystem/NetRefHandle.cpp
--- a/Engine/Source/Runtime/Experimental/Iris/Core/Private/Iris/ReplicationSystem/NetRefHandle.cpp
+++ b/Engine/Source/Runtime/Experimental/Iris/Core/Private/Iris/ReplicationSystem/NetRefHandle.cpp
@@ -11,10 +11,10 @@ FString FNetRefHandle::ToString() const
 {
 	FString Result;
 #if UE_NET_ALLOW_MULTIPLE_REPLICATION_SYSTEMS
-	const uint32 ReplicationSystemIdToDisplay = ReplicationSystemId - 1U;
-	Result = FString::Printf(TEXT("NetRefHandle (Id=%u):(RepSystemId=%u)"), GetId(), ReplicationSystemIdToDisplay);
+	const uint64 ReplicationSystemIdToDisplay = ReplicationSystemId - 1;
+	Result = FString::Printf(TEXT("NetRefHandle (Id=%" UINT64_FMT "):(RepSystemId=%" UINT64_FMT ")"), GetId(), ReplicationSystemIdToDisplay);
 #else
-	Result = FString::Printf(TEXT("NetRefHandle (Id=%u)"), GetId());
+	Result = FString::Printf(TEXT("NetRefHandle (Id=%" UINT64_FMT ")"), GetId());
 #endif
 	return Result;
 }
@@ -25,9 +25,9 @@ FStringBuilderBase& operator<<(FStringBuilderBase& Builder, const UE::Net::FNetR
 { 	
 #if UE_NET_ALLOW_MULTIPLE_REPLICATION_SYSTEMS
 	const uint32 ReplicationSystemIdToDisplay = NetRefHandle.GetReplicationSystemId() - 1U;
-	Builder.Appendf(TEXT("NetRefHandle (Id=%u):(RepSystemId=%u)"), NetRefHandle.GetId(), ReplicationSystemIdToDisplay);
+	Builder.Appendf(TEXT("NetRefHandle (Id=%" UINT64_FMT "):(RepSystemId=%u)"), NetRefHandle.GetId(), ReplicationSystemIdToDisplay);
 #else
-	Builder.Appendf(TEXT("NetRefHandle (Id=%u)"), NetRefHandle.GetId());
+	Builder.Appendf(TEXT("NetRefHandle (Id=%" UINT64_FMT ")"), NetRefHandle.GetId());
 #endif
 	return Builder;
 }
@@ -36,9 +36,9 @@ FAnsiStringBuilderBase& operator<<(FAnsiStringBuilderBase& Builder, const UE::Ne
 {
 #if UE_NET_ALLOW_MULTIPLE_REPLICATION_SYSTEMS
 	const uint32 ReplicationSystemIdToDisplay = NetRefHandle.GetReplicationSystemId() - 1U;
-	Builder.Appendf("NetRefHandle (Id=%u):(RepSystemId=%u)", NetRefHandle.GetId(), ReplicationSystemIdToDisplay);
+	Builder.Appendf("NetRefHandle (Id=%" UINT64_FMT "):(RepSystemId=%u)", NetRefHandle.GetId(), ReplicationSystemIdToDisplay);
 #else
-	Builder.Appendf("NetRefHandle (Id=%u)", NetRefHandle.GetId());
+	Builder.Appendf("NetRefHandle (Id=%" UINT64_FMT ")", NetRefHandle.GetId());
 #endif
 	return Builder;
 }
diff --git a/Engine/Source/Runtime/Experimental/Iris/Core/Private/Iris/ReplicationSystem/NetRefHandleManager.cpp b/Engine/Source/Runtime/Experimental/Iris/Core/Private/Iris/ReplicationSystem/NetRefHandleManager.cpp
--- a/Engine/Source/Runtime/Experimental/Iris/Core/Private/Iris/ReplicationSystem/NetRefHandleManager.cpp
+++ b/Engine/Source/Runtime/Experimental/Iris/Core/Private/Iris/ReplicationSystem/NetRefHandleManager.cpp
@@ -28,8 +28,8 @@ FNetRefHandleManager::FNetRefHandleManager(FReplicationProtocolManager& InReplic
 , ObjectsWithDependentObjectsInternalIndices(MaxActiveObjectCount)
 , DestroyedStartupObjectInternalIndices(MaxActiveObjectCount)
 , WantToBeDormantInternalIndices(MaxActiveObjectCount)
-, NextStaticHandleIndex(1U) // Index 0 is always reserved, for both static and dynamic handles
-, NextDynamicHandleIndex(1U)
+, NextStaticHandleIndex(1) // Index 0 is always reserved, for both static and dynamic handles
+, NextDynamicHandleIndex(1)
 , ReplicationProtocolManager(InReplicationProtocolManager)
 {
 	static_assert(InvalidInternalIndex == 0, "FNetRefHandleManager::InvalidInternalIndex has an unexpected value");
@@ -46,12 +46,12 @@ FNetRefHandleManager::FNetRefHandleManager(FReplicationProtocolManager& InReplic
 	ReplicatedObjectData[InvalidInternalIndex] = FReplicatedObjectData();
 }
 
-uint32 FNetRefHandleManager::GetNextNetRefHandleId(uint32 HandleId) const
+uint64 FNetRefHandleManager::GetNextNetRefHandleId(uint64 HandleId) const
 {
-	// Since we use the lowest bit in the index to indicate if the handle is static or dynamic we can not use all bits as the index
-	constexpr uint32 NetHandleIdIndexBitMask = (1U << (FNetRefHandle::IdBits - 1U)) - 1U;
+	// Since we use the lowest bit in the index to indicate if the handle is static or dynamic we cannot use all bits as the index
+	constexpr uint64 NetHandleIdIndexBitMask = (1ULL << (FNetRefHandle::IdBits - 1)) - 1;
 
-	uint32 NextHandleId = (HandleId + 1) & NetHandleIdIndexBitMask;
+	uint64 NextHandleId = (HandleId + 1) & NetHandleIdIndexBitMask;
 	if (NextHandleId == 0)
 	{
 		++NextHandleId;
@@ -150,9 +150,9 @@ const FReplicationInstanceProtocol* FNetRefHandleManager::DetachInstanceProtocol
 
 FNetRefHandle FNetRefHandleManager::AllocateNetRefHandle(bool bIsStatic)
 {
-	uint32& NextHandleId = bIsStatic ? NextStaticHandleIndex : NextDynamicHandleIndex;
+	uint64& NextHandleId = bIsStatic ? NextStaticHandleIndex : NextDynamicHandleIndex;
 
-	const uint32 NewHandleId = MakeNetRefHandleId(NextHandleId, bIsStatic);
+	const uint64 NewHandleId = MakeNetRefHandleId(NextHandleId, bIsStatic);
 	FNetRefHandle NewHandle = MakeNetRefHandle(NewHandleId, ReplicationSystemId);
 
 	// Verify that the handle is free
@@ -745,12 +745,12 @@ void FNetRefHandleManager::AddReferencedObjects(FReferenceCollector& Collector)
 	Collector.AddReferencedObjects(ReplicatedInstances);
 }
 
-uint32 FNetRefHandleManager::MakeNetRefHandleId(uint32 Id, bool bIsStatic)
+uint64 FNetRefHandleManager::MakeNetRefHandleId(uint64 Id, bool bIsStatic)
 {
 	return (Id << 1U) | (bIsStatic ? 1U : 0U);
 }
 
-FNetRefHandle FNetRefHandleManager::MakeNetRefHandle(uint32 Id, uint32 ReplicationSystemId)
+FNetRefHandle FNetRefHandleManager::MakeNetRefHandle(uint64 Id, uint32 ReplicationSystemId)
 {
 	check((Id & FNetRefHandle::IdMask) == Id);
 	check(ReplicationSystemId < FNetRefHandle::MaxReplicationSystemId);
@@ -763,7 +763,7 @@ FNetRefHandle FNetRefHandleManager::MakeNetRefHandle(uint32 Id, uint32 Replicati
 	return Handle;
 }
 
-FNetRefHandle FNetRefHandleManager::MakeNetRefHandleFromId(uint32 Id)
+FNetRefHandle FNetRefHandleManager::MakeNetRefHandleFromId(uint64 Id)
 {
 	check((Id & FNetRefHandle::IdMask) == Id);
 
diff --git a/Engine/Source/Runtime/Experimental/Iris/Core/Private/Iris/ReplicationSystem/NetRefHandleManager.h b/Engine/Source/Runtime/Experimental/Iris/Core/Private/Iris/ReplicationSystem/NetRefHandleManager.h
--- a/Engine/Source/Runtime/Experimental/Iris/Core/Private/Iris/ReplicationSystem/NetRefHandleManager.h
+++ b/Engine/Source/Runtime/Experimental/Iris/Core/Private/Iris/ReplicationSystem/NetRefHandleManager.h
@@ -120,8 +120,8 @@ public:
 	// Return true if this is a scopable index
 	bool IsScopableIndex(FInternalNetRefIndex InternalIndex) const { return ScopableInternalIndices.GetBit(InternalIndex); }
 
-	static FNetRefHandle MakeNetRefHandle(uint32 Id, uint32 ReplicationSystemId);
-	static FNetRefHandle MakeNetRefHandleFromId(uint32 Id);
+	static FNetRefHandle MakeNetRefHandle(uint64 Id, uint32 ReplicationSystemId);
+	static FNetRefHandle MakeNetRefHandleFromId(uint64 Id);
 
 	// Returns a valid handle if the wanted handle can be allocated
 	FNetRefHandle AllocateNetRefHandle(bool bIsStatic);
@@ -270,8 +270,8 @@ private:
 	FInternalNetRefIndex InternalCreateNetObject(const FNetRefHandle NetRefHandle, const FNetHandle GlobalHandle, const FReplicationProtocol* ReplicationProtocol);
 	void InternalDestroyNetObject(FInternalNetRefIndex InternalIndex);
 
-	static uint32 MakeNetRefHandleId(uint32 Seed, bool bIsStatic);
-	uint32 GetNextNetRefHandleId(uint32 HandleIndex) const;
+	static uint64 MakeNetRefHandleId(uint64 Seed, bool bIsStatic);
+	uint64 GetNextNetRefHandleId(uint64 HandleIndex) const;
 
 	// Get the next free internal index, returns InvalidInternalIndex if a free one cannot be found
 	FInternalNetRefIndex GetNextFreeInternalIndex() const;
@@ -344,8 +344,8 @@ private:
 	TArray<UObject*> ReplicatedInstances;
 
 	// Assign handles
-	uint32 NextStaticHandleIndex;
-	uint32 NextDynamicHandleIndex;
+	uint64 NextStaticHandleIndex;
+	uint64 NextDynamicHandleIndex;
 
 	FNetDependencyData SubObjects;
 
diff --git a/Engine/Source/Runtime/Experimental/Iris/Core/Private/Iris/ReplicationSystem/ReplicationReader.cpp b/Engine/Source/Runtime/Experimental/Iris/Core/Private/Iris/ReplicationSystem/ReplicationReader.cpp
--- a/Engine/Source/Runtime/Experimental/Iris/Core/Private/Iris/ReplicationSystem/ReplicationReader.cpp
+++ b/Engine/Source/Runtime/Experimental/Iris/Core/Private/Iris/ReplicationSystem/ReplicationReader.cpp
@@ -181,7 +181,7 @@ void FReplicationReader::Deinit()
 // Read incomplete handle
 FNetRefHandle FReplicationReader::ReadNetRefHandleId(FNetBitStreamReader& Reader) const
 {
-	const uint32 NetId = Reader.ReadBits(FNetRefHandle::IdBits);
+	const uint64 NetId = ReadPackedUint64(&Reader);
 	return FNetRefHandleManager::MakeNetRefHandleFromId(NetId);
 }
 	
diff --git a/Engine/Source/Runtime/Experimental/Iris/Core/Private/Iris/ReplicationSystem/ReplicationWriter.cpp b/Engine/Source/Runtime/Experimental/Iris/Core/Private/Iris/ReplicationSystem/ReplicationWriter.cpp
--- a/Engine/Source/Runtime/Experimental/Iris/Core/Private/Iris/ReplicationSystem/ReplicationWriter.cpp
+++ b/Engine/Source/Runtime/Experimental/Iris/Core/Private/Iris/ReplicationSystem/ReplicationWriter.cpp
@@ -512,7 +512,7 @@ const FReplicationWriter::FReplicationInfo& FReplicationWriter::GetReplicationIn
 
 void FReplicationWriter::WriteNetRefHandleId(FNetBitStreamWriter& Writer, FNetRefHandle Handle)
 {
-	Writer.WriteBits(Handle.GetId(), FNetRefHandle::IdBits);
+	WritePackedUint64(&Writer, Handle.GetId());
 }
 
 uint32 FReplicationWriter::GetDefaultFlushFlags() const
diff --git a/Engine/Source/Runtime/Experimental/Iris/Core/Private/Iris/Serialization/ObjectNetSerializer.cpp b/Engine/Source/Runtime/Experimental/Iris/Core/Private/Iris/Serialization/ObjectNetSerializer.cpp
--- a/Engine/Source/Runtime/Experimental/Iris/Core/Private/Iris/Serialization/ObjectNetSerializer.cpp
+++ b/Engine/Source/Runtime/Experimental/Iris/Core/Private/Iris/Serialization/ObjectNetSerializer.cpp
@@ -22,7 +22,7 @@ void WriteNetRefHandle(FNetSerializationContext& Context, const FNetRefHandle Ha
 	if (Handle.IsValid())
 	{
 		Writer->WriteBool(true);
-		WritePackedUint32(Writer, Handle.GetId());
+		WritePackedUint64(Writer, Handle.GetId());
 	}
 	else
 	{
@@ -36,7 +36,7 @@ FNetRefHandle ReadNetRefHandle(FNetSerializationContext& Context)
 
 	if (Reader->ReadBool())
 	{
-		const uint32 NetId = ReadPackedUint32(Reader);
+		const uint64 NetId = ReadPackedUint64(Reader);
 		if (!Reader->IsOverflown())
 		{
 			return Private::FNetRefHandleManager::MakeNetRefHandleFromId(NetId);
diff --git a/Engine/Source/Runtime/Experimental/Iris/Core/Public/Iris/ReplicationSystem/NetRefHandle.h b/Engine/Source/Runtime/Experimental/Iris/Core/Public/Iris/ReplicationSystem/NetRefHandle.h
--- a/Engine/Source/Runtime/Experimental/Iris/Core/Public/Iris/ReplicationSystem/NetRefHandle.h
+++ b/Engine/Source/Runtime/Experimental/Iris/Core/Public/Iris/ReplicationSystem/NetRefHandle.h
@@ -25,14 +25,14 @@ class FNetRefHandle
 {
 public:
 	enum { Invalid = 0 };
-	enum { IdBits = 32 };
+	enum { IdBits = 60 };
 	enum { ReplicationSystemIdBits = 4 };
 
 public:
 	FNetRefHandle() : Value(Invalid) {}
 
-	uint32 GetId() const { return Id; }
-	uint32 GetReplicationSystemId() const { check(ReplicationSystemId != 0U); return ReplicationSystemId - 1U; }
+	uint64 GetId() const { return Id; }
+	uint32 GetReplicationSystemId() const { check(ReplicationSystemId != 0); return (uint32)(ReplicationSystemId - 1); }
 	bool IsValid() const { return Value != Invalid; }
 
 	bool IsCompleteHandle() const { return Value != Invalid && ReplicationSystemId != 0U; }
@@ -52,16 +52,16 @@ private:
 	friend uint32 GetTypeHash(const FNetRefHandle& Handle);
 	friend Private::FNetRefHandleManager;
 
-	static constexpr uint32 StaticIdMask = 1U;
-	static constexpr uint32 IdMask = ~0U;
-	static constexpr uint32 MaxReplicationSystemId = (1 << ReplicationSystemIdBits) - 1;
+	static constexpr uint64 StaticIdMask = 1;
+	static constexpr uint64 IdMask = (1ULL << IdBits) - 1;
+	static constexpr uint64 MaxReplicationSystemId = (1ULL << ReplicationSystemIdBits) - 1;
 
 	union 
 	{
 		struct
 		{
-			uint32 Id : IdBits;										// Id, lowest bit indicates if the handle is static or dynamic
-			uint32 ReplicationSystemId : ReplicationSystemIdBits;	// ReplicationSystemId, when running in pie, we track the owning instance
+			uint64 Id : IdBits;										// Id, lowest bit indicates if the handle is static or dynamic
+			uint64 ReplicationSystemId : ReplicationSystemIdBits;	// ReplicationSystemId, when running in pie, we track the owning instance
 		};
 		uint64 Value;
 	};
@@ -74,7 +74,8 @@ FORCEINLINE uint32 GetTypeHash(const FNetRefHandle& Handle)
 
 FORCEINLINE uint32 GetObjectIdForNetTrace(const FNetRefHandle& Handle)
 {
-	return Handle.GetId();
+	// Until NetTrace is fixed: https://jira.it.epicgames.com/browse/UE-179190
+	return (uint32)Handle.GetId();
 }
 
 }
