Fix issues with regenerating geometry collection using the destruction builder script
- Fixed broken rendering in thumbnail after re-exporting a geometry collection a different layout / number of transforms
- Loaded components ( including the one in blueprints ) are now properly updated when the geometry collection asset is updated
- Fixed the thumbnail renderer crashing after re-exporting an asset and opening its editor
- Added check for size of arrays passed to the geometry group based bounds calculation

diff --git a/Engine/Plugins/Experimental/GeometryCollectionPlugin/Source/GeometryCollectionNodes/Private/Dataflow/Nodes/GeometryCollectionAssetNodes.cpp b/Engine/Plugins/Experimental/GeometryCollectionPlugin/Source/GeometryCollectionNodes/Private/Dataflow/Nodes/GeometryCollectionAssetNodes.cpp
--- a/Engine/Plugins/Experimental/GeometryCollectionPlugin/Source/GeometryCollectionNodes/Private/Dataflow/Nodes/GeometryCollectionAssetNodes.cpp
+++ b/Engine/Plugins/Experimental/GeometryCollectionPlugin/Source/GeometryCollectionNodes/Private/Dataflow/Nodes/GeometryCollectionAssetNodes.cpp
@@ -53,6 +53,13 @@ void FGeometryCollectionTerminalDataflowNode::SetAssetValue(TObjectPtr<UObject>
 			const bool bHasInternalMaterial = false; // with data flow there's no assumption of internal materials
 			CollectionAsset->ResetFrom(InCollection, InMaterials, false);
 			CollectionAsset->AutoInstanceMeshes = InInstancedMeshes;
+
+#if WITH_EDITOR
+			// make sure we rebuild the render data when we are done setting everything 
+			CollectionAsset->RebuildRenderData();
+			// also make sure all components using it are getting a notification about it
+			CollectionAsset->PropagateTransformUpdateToComponents();
+#endif
 		}
 	}
 }
diff --git a/Engine/Source/Runtime/Experimental/GeometryCollectionEngine/Private/GeometryCollection/GeometryCollectionComponent.cpp b/Engine/Source/Runtime/Experimental/GeometryCollectionEngine/Private/GeometryCollection/GeometryCollectionComponent.cpp
--- a/Engine/Source/Runtime/Experimental/GeometryCollectionEngine/Private/GeometryCollection/GeometryCollectionComponent.cpp
+++ b/Engine/Source/Runtime/Experimental/GeometryCollectionEngine/Private/GeometryCollection/GeometryCollectionComponent.cpp
@@ -502,6 +502,11 @@ namespace
 		FBox BoundingBox(ForceInit);
 		if (bChaos_BoxCalcBounds_ISPC_Enabled && !bGeometryCollectionSingleThreadedBoundsCalculation)
 		{
+			ensure(BoundingBoxes.Num() > 0);
+			ensure(TransformIndices.Num() == TransformIndices.Num());
+			ensure(TransformToGeometryIndex.Num() > 0);
+			ensure(TransformToGeometryIndex.Num() == GlobalMatrices.Num());
+			
 #if INTEL_ISPC
 			ispc::BoxCalcBoundsFromGeometryGroup(
 				(int32*)&TransformToGeometryIndex[0],
@@ -2484,6 +2489,9 @@ void UGeometryCollectionComponent::ResetDynamicCollection()
 			RemoveOnBreakDynamicFacade.SetAttributeValues(RemoveOnBreakFacade);
 		}
 
+		DynamicCollection->MakeDirty();
+		MarkRenderStateDirty();
+		MarkRenderDynamicDataDirty();
 		SetRenderStateDirty();
 	}
 
@@ -3184,31 +3192,29 @@ void UGeometryCollectionComponent::BuildInitialFilterData()
 	}
 }
 
-void UGeometryCollectionComponent::SetRestCollection(const UGeometryCollection* RestCollectionIn)
+void UGeometryCollectionComponent::SetRestCollection(const UGeometryCollection* RestCollectionIn, bool bApplyAssetDefaults)
 {
 	//UE_LOG(UGCC_LOG, Log, TEXT("GeometryCollectionComponent[%p]::SetRestCollection()"), this);
 	if (RestCollectionIn)
 	{
 		RestCollection = RestCollectionIn;
 
-		RestTransforms = RestCollection->GetGeometryCollection()->Transform.GetConstArray();
-
-		CalculateGlobalMatrices();
-		CalculateLocalBounds();
+		ResetDynamicCollection();
 
 		if (!IsEmbeddedGeometryValid())
 		{
 			InitializeEmbeddedGeometry();
 		}
 
-		// initialize the component per level damage threshold from the asset defaults 
-		DamageThreshold = RestCollection->DamageThreshold;
-		bUseSizeSpecificDamageThreshold = RestCollection->bUseSizeSpecificDamageThreshold;
+		if (bApplyAssetDefaults)
+		{
+			// initialize the component per level damage threshold from the asset defaults 
+			DamageThreshold = RestCollection->DamageThreshold;
+			bUseSizeSpecificDamageThreshold = RestCollection->bUseSizeSpecificDamageThreshold;
 
-		// initialize the component damage progataion data from the asset defaults 
-		DamagePropagationData = RestCollection->DamagePropagationData;
-		
-		ResetDynamicCollection();
+			// initialize the component damage progataion data from the asset defaults 
+			DamagePropagationData = RestCollection->DamagePropagationData;
+		}
 	}
 }
 
diff --git a/Engine/Source/Runtime/Experimental/GeometryCollectionEngine/Private/GeometryCollection/GeometryCollectionObject.cpp b/Engine/Source/Runtime/Experimental/GeometryCollectionEngine/Private/GeometryCollection/GeometryCollectionObject.cpp
--- a/Engine/Source/Runtime/Experimental/GeometryCollectionEngine/Private/GeometryCollection/GeometryCollectionObject.cpp
+++ b/Engine/Source/Runtime/Experimental/GeometryCollectionEngine/Private/GeometryCollection/GeometryCollectionObject.cpp
@@ -1141,6 +1141,19 @@ void UGeometryCollection::PropagateMarkDirtyToComponents() const
 	}
 }
 
+void  UGeometryCollection::PropagateTransformUpdateToComponents() const
+{
+	for (TObjectIterator<UGeometryCollectionComponent> It(RF_ClassDefaultObject, false, EInternalObjectFlags::Garbage); It; ++It)
+	{
+		if (It->RestCollection == this)
+		{
+			// make sure to reset the rest collection to make sure the internal state of the components is up to date 
+			// but we do not apply asset default to avoid overriding the existing overrides
+			It->SetRestCollection(this, false /* bApplyAssetDefaults */);
+		}
+	}
+}
+
 TSharedPtr<FGeometryCollection, ESPMode::ThreadSafe> UGeometryCollection::GenerateMinimalGeometryCollection() const
 {
 	TMap<FName, TSet<FName>> SkipList;
diff --git a/Engine/Source/Runtime/Experimental/GeometryCollectionEngine/Public/GeometryCollection/GeometryCollectionComponent.h b/Engine/Source/Runtime/Experimental/GeometryCollectionEngine/Public/GeometryCollection/GeometryCollectionComponent.h
--- a/Engine/Source/Runtime/Experimental/GeometryCollectionEngine/Public/GeometryCollection/GeometryCollectionComponent.h
+++ b/Engine/Source/Runtime/Experimental/GeometryCollectionEngine/Public/GeometryCollection/GeometryCollectionComponent.h
@@ -638,7 +638,7 @@ public:
 
 	/** RestCollection */
 	UFUNCTION(BlueprintCallable, Category = "ChaosPhysics")
-	void SetRestCollection(const UGeometryCollection * RestCollectionIn);
+	void SetRestCollection(const UGeometryCollection * RestCollectionIn, bool bApplyAssetDefaults = true);
 
 	/** RestCollection */
 	UFUNCTION(BlueprintCallable, Category = "ChaosPhysics")
diff --git a/Engine/Source/Runtime/Experimental/GeometryCollectionEngine/Public/GeometryCollection/GeometryCollectionObject.h b/Engine/Source/Runtime/Experimental/GeometryCollectionEngine/Public/GeometryCollection/GeometryCollectionObject.h
--- a/Engine/Source/Runtime/Experimental/GeometryCollectionEngine/Public/GeometryCollection/GeometryCollectionObject.h
+++ b/Engine/Source/Runtime/Experimental/GeometryCollectionEngine/Public/GeometryCollection/GeometryCollectionObject.h
@@ -458,6 +458,9 @@ public:
 
 	/** Propogate render state dirty to components */
 	void PropagateMarkDirtyToComponents() const;
+
+	/** Propogate the fact that transform have been updated to all components */
+	void PropagateTransformUpdateToComponents() const;
 #endif
 
 	void InitResources();
