Fix overlap query with triangle mesh

diff --git a/Engine/Source/Programs/HeadlessChaos/Private/HeadlessChaosTestObbTree.cpp b/Engine/Source/Programs/HeadlessChaos/Private/HeadlessChaosTestObbTree.cpp
--- a/Engine/Source/Programs/HeadlessChaos/Private/HeadlessChaosTestObbTree.cpp
+++ b/Engine/Source/Programs/HeadlessChaos/Private/HeadlessChaosTestObbTree.cpp
@@ -15,7 +15,7 @@ namespace ChaosTest
 		{
 			// Original points to create OBB // TArray<FVec3f> PointsA{ FVec3f(0, 0, 0), FVec3f(5, 0, 0), FVec3f(0, 5, 0) };
 			FRigidTransform3 Transform = FRigidTransform3(FVec3(1.25, 1.25, 0.0), FQuat(FVec3(0.0, 0.0, 1.0), -3.1415926 / 4.0));
-			Private::FOBBVectorized ObbA(Transform, FVec3f(UE_INV_SQRT_2*5.0, UE_INV_SQRT_2 * 5.0 /2.0, 0.0f));
+			Private::FOBBVectorized ObbA(Transform, FVec3f(UE_INV_SQRT_2*5.0, UE_INV_SQRT_2 * 5.0 /2.0, 0.0f), FVec3f(1.0f, 1.0f, 1.0f));
 			{
 				TAABB<FRealSingle, 3> Bounds(FVec3f(-2.0f, 1.5f, 0.0f), FVec3f(-1.0f, 2.5f, 0.0f));
 				EXPECT_TRUE(ObbA.IntersectAABB(FAABBVectorized(Bounds)));
@@ -54,11 +54,37 @@ namespace ChaosTest
 			}
 		}
 		{
-			
-			FRigidTransform3 Transform = FRigidTransform3(FVec3(0.0, 0.0, 12.0), FQuat(FVec3(1.0, 0.0, 0.0), 3.1415926 / 4.0));
-			Private::FOBBVectorized ObbA(Transform.ToMatrixNoScale(), FVec3f(2.0f, 2.0f, 2.0f));
 			TAABB<FRealSingle, 3> Bounds(FVec3f(-10.0f, -10.0f, 0.0f), FVec3f(10.0f, 10.0f, 10.0f));
-			EXPECT_TRUE(ObbA.IntersectAABB(FAABBVectorized(Bounds)));
+			{
+				FRigidTransform3 Transform = FRigidTransform3(FVec3(0.0, 0.0, 12.0), FQuat(FVec3(1.0, 0.0, 0.0), 3.1415926 / 4.0));
+				Private::FOBBVectorized ObbA(Transform.ToMatrixNoScale(), FVec3f(2.0f, 2.0f, 2.0f), FVec3f(1.0f, 1.0f, 1.0f));
+				EXPECT_TRUE(ObbA.IntersectAABB(FAABBVectorized(Bounds)));
+			}
+			{
+				FRigidTransform3 Transform = FRigidTransform3(FVec3(0.0, 11.0, 0.0), FQuat::Identity);
+				Private::FOBBVectorized ObbA(Transform.ToMatrixNoScale(), FVec3f(2.0f, 2.0f, 2.0f), FVec3f(1.0f, 1.0f, 1.0f));
+				EXPECT_TRUE(ObbA.IntersectAABB(FAABBVectorized(Bounds)));
+			}
+			{
+				FRigidTransform3 Transform = FRigidTransform3(FVec3(0.0, 11.0, 0.0), FQuat::Identity);
+				Private::FOBBVectorized ObbA(Transform.ToMatrixNoScale(), FVec3f(2.0f, 2.0f, 2.0f), FVec3f(1.0f, 3.0f, 1.0f));
+				EXPECT_FALSE(ObbA.IntersectAABB(FAABBVectorized(Bounds)));
+			}
+			{
+				FRigidTransform3 Transform = FRigidTransform3(FVec3(0.0, 11.0, 0.0), FQuat::Identity);
+				Private::FOBBVectorized ObbA(Transform.ToMatrixNoScale(), FVec3f(2.0f, 2.0f, 2.0f), FVec3f(1.0f, 1.0f, 3.0f));
+				EXPECT_TRUE(ObbA.IntersectAABB(FAABBVectorized(Bounds)));
+			}
+			{
+				FRigidTransform3 Transform = FRigidTransform3(FVec3(0.0, 11.0, 0.0), FQuat(FVec3(1.0, 0.0, 0.0), 3.1415926 / 2.0));
+				Private::FOBBVectorized ObbA(Transform.ToMatrixNoScale(), FVec3f(2.0f, 2.0f, 2.0f), FVec3f(3.0f, 1.0f, 1.0f));
+				EXPECT_TRUE(ObbA.IntersectAABB(FAABBVectorized(Bounds)));
+			}
+			{
+				FRigidTransform3 Transform = FRigidTransform3(FVec3(0.0, 11.0, 0.0), FQuat(FVec3(1.0, 0.0, 0.0), -3.1415926 / 2.0));
+				Private::FOBBVectorized ObbA(Transform.ToMatrixNoScale(), FVec3f(2.0f, 2.0f, 2.0f), FVec3f(3.0f, 1.0f, 1.0f));
+				EXPECT_TRUE(ObbA.IntersectAABB(FAABBVectorized(Bounds)));
+			}
 		}
 	}
 
diff --git a/Engine/Source/Runtime/Experimental/Chaos/Private/Chaos/OBBVectorized.cpp b/Engine/Source/Runtime/Experimental/Chaos/Private/Chaos/OBBVectorized.cpp
--- a/Engine/Source/Runtime/Experimental/Chaos/Private/Chaos/OBBVectorized.cpp
+++ b/Engine/Source/Runtime/Experimental/Chaos/Private/Chaos/OBBVectorized.cpp
@@ -8,7 +8,7 @@
 namespace Chaos
 {
 
-	Private::FOBBVectorized::FOBBVectorized(const FRigidTransform3& Transform, const FVec3f& HalfExtentsIn)
+	Private::FOBBVectorized::FOBBVectorized(const FRigidTransform3& Transform, const FVec3f& HalfExtentsIn, const FVec3f& InvScaleIn)
 	{
 		TRotation<FRealSingle, 3> Rotation = Transform.GetRotation();
 		PMatrix<FRealSingle, 3, 3> Matrix = Rotation.ToMatrix();
@@ -18,45 +18,58 @@ namespace Chaos
 
 		FVec3 Translation = Transform.GetTranslation();
 		Position = MakeVectorRegisterFloatFromDouble(VectorLoadDouble3(&Translation.X));
+		InvScale = VectorLoadFloat3(&InvScaleIn.X);
+		Position = VectorMultiply(Position, InvScale);
 
 		HalfExtents = VectorLoadFloat3(&HalfExtentsIn.X);
+		HalfExtents = VectorAdd(HalfExtents, GlobalVectorConstants::KindaSmallNumber);
+
 		const VectorRegister4Float XHalfExtent = VectorReplicate(HalfExtents, 0);
 		const VectorRegister4Float YHalfExtent = VectorReplicate(HalfExtents, 1);
 		const VectorRegister4Float ZHalfExtent = VectorReplicate(HalfExtents, 2);
 
 		VectorRegister4Float FurthestPoint = VectorMultiplyAdd(VectorAbs(XAxis), XHalfExtent, VectorMultiplyAdd(VectorAbs(YAxis), YHalfExtent, VectorMultiply(VectorAbs(ZAxis), ZHalfExtent)));
+
+		FurthestPoint = VectorMultiply(FurthestPoint, InvScale);
 		MaxObb = VectorAdd(Position, FurthestPoint);
 		MinObb = VectorSubtract(Position, FurthestPoint);
+
+		const VectorRegister4Float Scale = VectorDivide(GlobalVectorConstants::FloatOne, InvScale);
+
+		XAxis = VectorMultiply(XAxis, Scale);
+		YAxis = VectorMultiply(YAxis, Scale);
+		ZAxis = VectorMultiply(ZAxis, Scale);
 	}
 
 	bool Private::FOBBVectorized::IntersectAABB(const FAABBVectorized& Bounds) const
 	{
-		VectorRegister4Float HasSeparationAxis = VectorBitwiseOr(VectorCompareGT(MinObb, Bounds.GetMax()), VectorCompareGT(Bounds.GetMin(), MaxObb));
+		const VectorRegister4Float HasSeparationAxis = VectorBitwiseOr(VectorCompareGT(MinObb, Bounds.GetMax()), VectorCompareGT(Bounds.GetMin(), MaxObb));
 
 		if (VectorMaskBits(HasSeparationAxis) != 0)
 		{
 			return false;
 		}
 
-		VectorRegister4Float HalfAabb = VectorMultiply(VectorSubtract(Bounds.GetMax(), Bounds.GetMin()), GlobalVectorConstants::FloatOneHalf);
-		VectorRegister4Float CenterAabb = VectorMultiply(VectorAdd(Bounds.GetMax(), Bounds.GetMin()), GlobalVectorConstants::FloatOneHalf);
+		const VectorRegister4Float HalfAabb = VectorMultiply(VectorSubtract(Bounds.GetMax(), Bounds.GetMin()), GlobalVectorConstants::FloatOneHalf);
+		const VectorRegister4Float CenterAabb = VectorMultiply(VectorAdd(Bounds.GetMax(), Bounds.GetMin()), GlobalVectorConstants::FloatOneHalf);
 		VectorRegister4Float DirCenter = VectorSubtract(CenterAabb, Position);
 
-		VectorRegister4Float DistCenterX = VectorAbs(VectorDot3(XAxis, DirCenter));
-		VectorRegister4Float DistCenterY = VectorAbs(VectorDot3(YAxis, DirCenter));
-		VectorRegister4Float DistCenterZ = VectorAbs(VectorDot3(ZAxis, DirCenter));
+
+		const VectorRegister4Float DistCenterX = VectorAbs(VectorDot3(XAxis, DirCenter));
+		const VectorRegister4Float DistCenterY = VectorAbs(VectorDot3(YAxis, DirCenter));
+		const VectorRegister4Float DistCenterZ = VectorAbs(VectorDot3(ZAxis, DirCenter));
 
 		VectorRegister4Float DistCenter = VectorUnpackLo(DistCenterX, DistCenterY);
 		DistCenter = VectorMoveLh(DistCenter, DistCenterZ);
-		
-		VectorRegister4Float DistCornerX = VectorDot3(VectorAbs(XAxis), HalfAabb);
-		VectorRegister4Float DistCornerY = VectorDot3(VectorAbs(YAxis), HalfAabb);
-		VectorRegister4Float DistCornerZ = VectorDot3(VectorAbs(ZAxis), HalfAabb);
+
+		const VectorRegister4Float DistCornerX = VectorDot3(VectorAbs(XAxis), HalfAabb);
+		const VectorRegister4Float DistCornerY = VectorDot3(VectorAbs(YAxis), HalfAabb);
+		const VectorRegister4Float DistCornerZ = VectorDot3(VectorAbs(ZAxis), HalfAabb);
 
 		VectorRegister4Float DistCorner = VectorUnpackLo(DistCornerX, DistCornerY);
 		DistCorner = VectorMoveLh(DistCorner, DistCornerZ);
 
-		uint32 MaskBitTrue = VectorMaskBits(VectorCompareGT(VectorSubtract(DistCenter, DistCorner), HalfExtents));
+		const uint32 MaskBitTrue = VectorMaskBits(VectorCompareGT(VectorSubtract(DistCenter, DistCorner), HalfExtents));
 
 		return !(MaskBitTrue & 7);
 	}
diff --git a/Engine/Source/Runtime/Experimental/Chaos/Private/Chaos/TriangleMeshImplicitObject.cpp b/Engine/Source/Runtime/Experimental/Chaos/Private/Chaos/TriangleMeshImplicitObject.cpp
--- a/Engine/Source/Runtime/Experimental/Chaos/Private/Chaos/TriangleMeshImplicitObject.cpp
+++ b/Engine/Source/Runtime/Experimental/Chaos/Private/Chaos/TriangleMeshImplicitObject.cpp
@@ -1056,12 +1056,11 @@ bool FTriangleMeshImplicitObject::Overlap(const FVec3& Point, const FReal Thickn
 template <bool IsSpherical, typename QueryGeomType>
 bool FTriangleMeshImplicitObject::OverlapGeomImp(const QueryGeomType& QueryGeom, const FRigidTransform3& QueryTM, const FReal Thickness, FMTDInfo* OutMTD, FVec3 TriMeshScale) const
 {
-
+	const FVec3 InvTriMeshScale = FVec3(FReal(1) / TriMeshScale.X, FReal(1) / TriMeshScale.Y, FReal(1) / TriMeshScale.Z);
 	const QueryGeomType& WorldScaleQueryGeom = ScaleGeomIntoWorldHelper(QueryGeom, TriMeshScale);
 
 	if (OutMTD)
 	{
-		const FVec3 InvTriMeshScale = FVec3(FReal(1) / TriMeshScale.X, FReal(1) / TriMeshScale.Y, FReal(1) / TriMeshScale.Z);
 		// IMPORTANT QueryTM comes with a invscaled translation so we need a version of the TM with world space translation to properly compute the bounds
 		FRigidTransform3 TriMeshToGeomNoScale{ QueryTM };
 		TriMeshToGeomNoScale.SetTranslation(TriMeshToGeomNoScale.GetTranslation() * TriMeshScale);
@@ -1120,14 +1119,13 @@ bool FTriangleMeshImplicitObject::OverlapGeomImp(const QueryGeomType& QueryGeom,
 	}
 	else
 	{
-		// Disabled optimization due to a bug UE-176682, TODO Re-enabled it 
-		// if (IsSpherical || FMath::Abs(QueryTM.GetRotation().W) > 0.9)
+		
+		FRigidTransform3 TriMeshToGeomNoScale{ QueryTM };
+		TriMeshToGeomNoScale.SetTranslation(TriMeshToGeomNoScale.GetTranslation() * TriMeshScale);
+		// If the query shape is axis aligned, meaning has no rotation or a 90 degree rotation use a middle phase AABB query
+		if (IsSpherical || FMath::Abs(QueryTM.GetRotation().W) > 0.9 || FMath::IsNearlyEqual(QueryTM.GetRotation().W, UE_HALF_SQRT_2, 0.01))
 		{
-			const FVec3 InvTriMeshScale = FVec3(FReal(1) / TriMeshScale.X, FReal(1) / TriMeshScale.Y, FReal(1) / TriMeshScale.Z);
 			// IMPORTANT QueryTM comes with a invscaled translation so we need a version of the TM with world space translation to properly compute the bounds
-			FRigidTransform3 TriMeshToGeomNoScale{ QueryTM };
-			TriMeshToGeomNoScale.SetTranslation(TriMeshToGeomNoScale.GetTranslation() * TriMeshScale);
-
 			// NOTE: BVH test is done in tri-mesh local space (whereas collision detection is done in world space because you can't non-uniformly scale all shapes)
 			FAABB3 QueryBounds = WorldScaleQueryGeom.CalculateTransformedBounds(TriMeshToGeomNoScale);
 			QueryBounds.ThickenSymmetrically(FVec3(Thickness));
@@ -1136,15 +1134,17 @@ bool FTriangleMeshImplicitObject::OverlapGeomImp(const QueryGeomType& QueryGeom,
 			ScaleTransformHelper(TriMeshScale, QueryTM, WorldScaleQueryTM);
 			return FastBVH.FindAllIntersectionsNoMTD(QueryBounds, WorldScaleQueryTM, WorldScaleQueryGeom, Thickness, TriMeshScale, this);
 		}
-		//else
-		//{
-		//	FAABB3 GeometryAABB = QueryGeom.BoundingBox();
-		//	GeometryAABB.ThickenSymmetrically(FVec3(Thickness));
-		//	Private::FOBBVectorized QueryObb(QueryTM, (GeometryAABB.Max() - GeometryAABB.Min()) * 0.5);
-		//	TRigidTransform<FReal, 3> WorldScaleQueryTM;
-		//	ScaleTransformHelper(TriMeshScale, QueryTM, WorldScaleQueryTM);
-		//	return FastBVH.FindAllIntersectionsNoMTD(QueryObb, WorldScaleQueryTM, WorldScaleQueryGeom, Thickness, TriMeshScale, this);
-		//}
+		else
+		{
+			FAABB3 GeometryAABB = WorldScaleQueryGeom.BoundingBox();
+			GeometryAABB.ThickenSymmetrically(FVec3(Thickness));
+			// The geometry is not necessarily in the center of its translation. (cf: a capsule can be create with two points away from its own transform)
+			TriMeshToGeomNoScale.SetTranslation(TriMeshToGeomNoScale.GetTranslation() + GeometryAABB.Center());
+			Private::FOBBVectorized QueryObb(TriMeshToGeomNoScale, (GeometryAABB.Max() - GeometryAABB.Min()) * 0.5, InvTriMeshScale);
+			TRigidTransform<FReal, 3> WorldScaleQueryTM;
+			ScaleTransformHelper(TriMeshScale, QueryTM, WorldScaleQueryTM);
+			return FastBVH.FindAllIntersectionsNoMTD(QueryObb, WorldScaleQueryTM, WorldScaleQueryGeom, Thickness, TriMeshScale, this);
+		}
 	}
 }
 
diff --git a/Engine/Source/Runtime/Experimental/Chaos/Public/Chaos/OBBVectorized.h b/Engine/Source/Runtime/Experimental/Chaos/Public/Chaos/OBBVectorized.h
--- a/Engine/Source/Runtime/Experimental/Chaos/Public/Chaos/OBBVectorized.h
+++ b/Engine/Source/Runtime/Experimental/Chaos/Public/Chaos/OBBVectorized.h
@@ -13,7 +13,7 @@ namespace Chaos
 		class FOBBVectorized
 		{
 		public:
-			FOBBVectorized(const FRigidTransform3& Transform, const FVec3f& HalfExtentsIn);
+			FOBBVectorized(const FRigidTransform3& Transform, const FVec3f& HalfExtentsIn, const FVec3f& InvScale);
 			bool IntersectAABB(const FAABBVectorized& Bounds) const;
 
 		private:
@@ -25,6 +25,7 @@ namespace Chaos
 
 			VectorRegister4Float MaxObb;
 			VectorRegister4Float MinObb;
+			VectorRegister4Float InvScale;
 		};
 	}
 }
\ No newline at end of file
