geometry collection : fix scale on removal not using the scale matrix properly ( was causing a crash because of resulting optimizations )
- make sure to handle the case where inner radius attriute has been stripped out

diff --git a/Engine/Source/Runtime/Experimental/GeometryCollectionEngine/Private/GeometryCollection/GeometryCollectionComponent.cpp b/Engine/Source/Runtime/Experimental/GeometryCollectionEngine/Private/GeometryCollection/GeometryCollectionComponent.cpp
--- a/Engine/Source/Runtime/Experimental/GeometryCollectionEngine/Private/GeometryCollection/GeometryCollectionComponent.cpp
+++ b/Engine/Source/Runtime/Experimental/GeometryCollectionEngine/Private/GeometryCollection/GeometryCollectionComponent.cpp
@@ -2908,7 +2908,7 @@ void UGeometryCollectionComponent::UpdateRemovalIfNeeded()
 							const FVector CenterOfMass = DynamicCollection->MassToLocal[Idx].GetTranslation();
 							const FVector ScaleCenter = LocalDown + CenterOfMass;
 							const FTransform ScaleTransform(FQuat::Identity, ScaleCenter * FVector::FReal(1.f - Scale), FVector(Scale));
-							DynamicCollection->Transform[Idx].SetScale3D(FVector::ZeroVector);
+							DynamicCollection->Transform[Idx] = ScaleTransform * DynamicCollection->Transform[Idx];
 						}
 					}
 				}
@@ -4999,44 +4999,44 @@ bool UGeometryCollectionComponent::CalculateInnerSphere(int32 TransformIndex, UE
 	// Approximates the inscribed sphere. Returns false if no such sphere exists, if for instance the index is to an embedded geometry. 
 
 	const TManagedArray<int32>& TransformToGeometryIndex = RestCollection->GetGeometryCollection()->TransformToGeometryIndex;
-	const TManagedArray<Chaos::FRealSingle>& InnerRadius = RestCollection->GetGeometryCollection()->InnerRadius;
 	const TManagedArray<TSet<int32>>& Children = RestCollection->GetGeometryCollection()->Children;
 	const TManagedArray<FTransform>& MassToLocal = RestCollection->GetGeometryCollection()->GetAttribute<FTransform>("MassToLocal", FGeometryCollection::TransformGroup);
 
-	if (RestCollection->GetGeometryCollection()->IsRigid(TransformIndex))
+	TManagedArrayAccessor<Chaos::FRealSingle> InnerRadiusAttribute(*RestCollection->GetGeometryCollection(), "InnerRadius", FGeometryCollection::GeometryGroup);
+	if (InnerRadiusAttribute.IsValid() && InnerRadiusAttribute.IsValidIndex(TransformIndex))
 	{
-		// Sphere in component space, centered on body's COM.
-		FVector COM = MassToLocal[TransformIndex].GetLocation();
-		SphereOut = UE::Math::TSphere<double>(COM, InnerRadius[TransformToGeometryIndex[TransformIndex]]);
-		return true;
-	}
-	else if (RestCollection->GetGeometryCollection()->IsClustered(TransformIndex))
-	{
-		// Recursively accumulate the cluster's child spheres. 
-		bool bSphereFound = false;
-		for (int32 ChildIndex: Children[TransformIndex])
+		if (RestCollection->GetGeometryCollection()->IsRigid(TransformIndex))
+		{
+			// Sphere in component space, centered on body's COM.
+			FVector COM = MassToLocal[TransformIndex].GetLocation();
+			SphereOut = UE::Math::TSphere<double>(COM, InnerRadiusAttribute[TransformToGeometryIndex[TransformIndex]]);
+			return true;
+		}
+		else if (RestCollection->GetGeometryCollection()->IsClustered(TransformIndex))
 		{
-			UE::Math::TSphere<double> LocalSphere;
-			if (CalculateInnerSphere(ChildIndex, LocalSphere))
+			// Recursively accumulate the cluster's child spheres. 
+			bool bSphereFound = false;
+			for (int32 ChildIndex : Children[TransformIndex])
 			{
-				if (!bSphereFound)
+				UE::Math::TSphere<double> LocalSphere;
+				if (CalculateInnerSphere(ChildIndex, LocalSphere))
 				{
-					bSphereFound = true;
-					SphereOut = LocalSphere;
-				}
-				else
-				{
-					SphereOut += LocalSphere;
+					if (!bSphereFound)
+					{
+						bSphereFound = true;
+						SphereOut = LocalSphere;
+					}
+					else
+					{
+						SphereOut += LocalSphere;
+					}
 				}
 			}
+			return bSphereFound;
 		}
-		return bSphereFound;
-	}
-	else
-	{
-		// Likely an embedded geometry, which doesn't count towards volume.
-		return false;
 	}
+	// Likely an embedded geometry or missing inner radius attribute , which doesn't count towards volume.
+	return false;
 }
 
 void UGeometryCollectionComponent::PostLoad()
