Chaos - don't allow ArrayCollectionArrays to resize on shrink.

The default shrink policy will resize when the slack space is > 50% of the size, or if the slack space is using 16KB. This was causing major hangs removing particles on scenes with ~1M particles.

diff --git a/Engine/Source/Runtime/Experimental/ChaosCore/Public/Chaos/ArrayCollectionArray.h b/Engine/Source/Runtime/Experimental/ChaosCore/Public/Chaos/ArrayCollectionArray.h
--- a/Engine/Source/Runtime/Experimental/ChaosCore/Public/Chaos/ArrayCollectionArray.h
+++ b/Engine/Source/Runtime/Experimental/ChaosCore/Public/Chaos/ArrayCollectionArray.h
@@ -19,6 +19,8 @@ class TArrayCollectionArray : public TArrayCollectionArrayBase, public TArray<T>
 	using TArray<T>::Emplace;
 
 public:
+	constexpr static bool bAllowShrinkOnRemove = false;
+
 	using TArray<T>::Num;
 
 	TArrayCollectionArray()
@@ -56,17 +58,17 @@ public:
 
 	void Resize(const int Num) override
 	{
-		SetNum(Num);
+		SetNum(Num, bAllowShrinkOnRemove);
 	}
 
 	FORCEINLINE void RemoveAt(const int Idx, const int Count) override
 	{
-		TArray<T>::RemoveAt(Idx, Count);
+		TArray<T>::RemoveAt(Idx, Count, bAllowShrinkOnRemove);
 	}
 
 	FORCEINLINE void RemoveAtSwap(const int Idx) override
 	{
-		TArray<T>::RemoveAtSwap(Idx);
+		TArray<T>::RemoveAtSwap(Idx, 1, bAllowShrinkOnRemove);
 	}
 
 	FORCEINLINE void MoveToOtherArray(const int Idx, TArrayCollectionArrayBase& Other)
@@ -74,7 +76,7 @@ public:
 		//todo: add developer check to make sure this is ok?
 		auto& OtherTArray = static_cast<TArrayCollectionArray<T>&>(Other);
 		OtherTArray.Emplace(MoveTemp(TArray<T>::operator [](Idx)));
-		TArray<T>::RemoveAtSwap(Idx);
+		TArray<T>::RemoveAtSwap(Idx, 1, bAllowShrinkOnRemove);
 	}
 
 	FORCEINLINE uint64 SizeOfElem() const override
