Chaos - minor solver shutdown cleanup

diff --git a/Engine/Source/Runtime/Experimental/Chaos/Private/PhysicsProxy/GeometryCollectionPhysicsProxy.cpp b/Engine/Source/Runtime/Experimental/Chaos/Private/PhysicsProxy/GeometryCollectionPhysicsProxy.cpp
--- a/Engine/Source/Runtime/Experimental/Chaos/Private/PhysicsProxy/GeometryCollectionPhysicsProxy.cpp
+++ b/Engine/Source/Runtime/Experimental/Chaos/Private/PhysicsProxy/GeometryCollectionPhysicsProxy.cpp
@@ -2449,6 +2449,7 @@ void FGeometryCollectionPhysicsProxy::OnRemoveFromSolver(Chaos::FPBDRigidsSolver
 			Chaos::FUniqueIdx UniqueIdx = Handle->UniqueIdx();
 			Evolution->DestroyParticle(Handle);
 			Evolution->ReleaseUniqueIdx(UniqueIdx);
+			SolverParticleHandles[i] = nullptr;
 		}
 	}
 
