Chaos - recover memory from particle arrays each frame if the TArray shrink policy deems it necessary.

This replaces the prior behaviour which was to run the shrink policy every time a particle was removed (and was disabled in a prior CL).

diff --git a/Engine/Source/Runtime/Experimental/Chaos/Private/PBDRigidsSolver.cpp b/Engine/Source/Runtime/Experimental/Chaos/Private/PBDRigidsSolver.cpp
--- a/Engine/Source/Runtime/Experimental/Chaos/Private/PBDRigidsSolver.cpp
+++ b/Engine/Source/Runtime/Experimental/Chaos/Private/PBDRigidsSolver.cpp
@@ -224,6 +224,9 @@ namespace Chaos
 		bool bChaosSolverCollisionEnabled = true;
 		FAutoConsoleVariableRef CVarChaosSolverCollisionDisable(TEXT("p.Chaos.Solver.Collision.Enabled"), bChaosSolverCollisionEnabled, TEXT("Enable/Disable collisions in the main scene."));
 
+		// Shrink particle arrays every frame to recove rmemory when a scene changes significantly
+		bool bChaosSolverShrinkArrays = true;
+		FAutoConsoleVariableRef CVarChaosSolverShrinkArrays(TEXT("p.Chaos.Solver.ShrinkArrays"), bChaosSolverShrinkArrays, TEXT("Enable/Disable particle array shrinking in the main scene"));
 
 		// Iteration count cvars
 		// These override the engine config if >= 0
@@ -476,6 +479,13 @@ namespace Chaos
 			// Todo(chaos) we should probably move all of the solver event reset here in the future
 			MSolver->GetEvolution()->GetRigidClustering().ResetAllEvents();
 
+			// Recover unused memory from particle arrays, based on the array shrink policy
+			if (bChaosSolverShrinkArrays)
+			{
+				QUICK_SCOPE_CYCLE_COUNTER(ShrinkParticleArrays);
+				MSolver->GetParticles().ShrinkArrays();
+			}
+
 			if (FRewindData* RewindData = MSolver->GetRewindData())
 			{
 				RewindData->FinishFrame();
@@ -994,21 +1004,16 @@ namespace Chaos
 			}
 		}
 
-		bool bResetCollisionConstraints=false;
+		MarshallingManager.GetCurrentPullData_Internal()->DirtyGeometryCollections.Reset();
 		for (auto Proxy : PendingDestroyGeometryCollectionPhysicsProxy)
 		{
 			// Removing the geometry collection from the solver a bit delayed. This lets the cluster union do its cleanup first before
 			// the geometry collection if they're all being destroyed at the same time.
 			Proxy->OnRemoveFromSolver(this);
 
-			MarshallingManager.GetCurrentPullData_Internal()->DirtyGeometryCollections.Reset();
-			bResetCollisionConstraints = true;
 			delete Proxy;
 		}
 		PendingDestroyGeometryCollectionPhysicsProxy.Reset();
-
-		if(bResetCollisionConstraints)
-			GetEvolution()->GetCollisionConstraints();
 	}
 
 	void FPBDRigidsSolver::PrepareAdvanceBy(const FReal DeltaTime)
diff --git a/Engine/Source/Runtime/Experimental/Chaos/Public/Chaos/PBDRigidsSOAs.h b/Engine/Source/Runtime/Experimental/Chaos/Public/Chaos/PBDRigidsSOAs.h
--- a/Engine/Source/Runtime/Experimental/Chaos/Public/Chaos/PBDRigidsSOAs.h
+++ b/Engine/Source/Runtime/Experimental/Chaos/Public/Chaos/PBDRigidsSOAs.h
@@ -230,6 +230,23 @@ public:
 		check(0);
 	}
 
+	void ShrinkArrays()
+	{
+		StaticParticles->ShrinkArrays();
+		StaticDisabledParticles->ShrinkArrays();
+
+		KinematicParticles->ShrinkArrays();;
+		KinematicDisabledParticles->ShrinkArrays();
+
+		DynamicDisabledParticles->ShrinkArrays();
+		DynamicParticles->ShrinkArrays();
+		DynamicKinematicParticles->ShrinkArrays();
+
+		ClusteredParticles->ShrinkArrays();
+
+		GeometryCollectionParticles->ShrinkArrays();
+	}
+
 	void UpdateDirtyViews()
 	{
 		// @todo(chaos): this should only refresh views that may have changed
diff --git a/Engine/Source/Runtime/Experimental/ChaosCore/Public/Chaos/ArrayCollection.h b/Engine/Source/Runtime/Experimental/ChaosCore/Public/Chaos/ArrayCollection.h
--- a/Engine/Source/Runtime/Experimental/ChaosCore/Public/Chaos/ArrayCollection.h
+++ b/Engine/Source/Runtime/Experimental/ChaosCore/Public/Chaos/ArrayCollection.h
@@ -22,6 +22,17 @@ public:
 		}
 	}
 
+	void ShrinkArrays()
+	{
+		for (int32 Index = 0; Index < MArrays.Num(); Index++)
+		{
+			if (MArrays[Index] != nullptr)
+			{
+				MArrays[Index]->Shrink();
+			}
+		}
+	}
+
 	int32 AddArray(TArrayCollectionArrayBase* Array)
 	{
 		int32 Index = MArrays.Find(nullptr);
diff --git a/Engine/Source/Runtime/Experimental/ChaosCore/Public/Chaos/ArrayCollectionArray.h b/Engine/Source/Runtime/Experimental/ChaosCore/Public/Chaos/ArrayCollectionArray.h
--- a/Engine/Source/Runtime/Experimental/ChaosCore/Public/Chaos/ArrayCollectionArray.h
+++ b/Engine/Source/Runtime/Experimental/ChaosCore/Public/Chaos/ArrayCollectionArray.h
@@ -17,6 +17,7 @@ class TArrayCollectionArray : public TArrayCollectionArrayBase, public TArray<T>
 	using TArray<T>::RemoveAt;
 	using TArray<T>::RemoveAtSwap;
 	using TArray<T>::Emplace;
+	using TArray<T>::Shrink;
 
 public:
 	constexpr static bool bAllowShrinkOnRemove = false;
@@ -56,6 +57,11 @@ public:
 		return NewArray;
 	}
 
+	void Shrink() override
+	{
+		TArray<T>::Shrink();
+	}
+
 	void Resize(const int Num) override
 	{
 		SetNum(Num, bAllowShrinkOnRemove);
diff --git a/Engine/Source/Runtime/Experimental/ChaosCore/Public/Chaos/ArrayCollectionArrayBase.h b/Engine/Source/Runtime/Experimental/ChaosCore/Public/Chaos/ArrayCollectionArrayBase.h
--- a/Engine/Source/Runtime/Experimental/ChaosCore/Public/Chaos/ArrayCollectionArrayBase.h
+++ b/Engine/Source/Runtime/Experimental/ChaosCore/Public/Chaos/ArrayCollectionArrayBase.h
@@ -8,6 +8,7 @@ namespace Chaos
 class TArrayCollectionArrayBase
 {
   public:
+	virtual void Shrink() = 0;
 	virtual void Resize(const int Num) = 0;
 	virtual void RemoveAt(const int Idx, const int Count) = 0;
 	virtual void RemoveAtSwap(const int Idx) = 0;
