Chaos - fix bounds and cull distance for moving bodies

This fixes the issue where a dynamic body drops onto another resting on the ground, giving it velocity which results in the predicted positions being separated by more than CullDistance at the start of the next frame.

diff --git a/Engine/Source/Runtime/Experimental/Chaos/Private/Chaos/Collision/ParticlePairMidPhase.cpp b/Engine/Source/Runtime/Experimental/Chaos/Private/Chaos/Collision/ParticlePairMidPhase.cpp
--- a/Engine/Source/Runtime/Experimental/Chaos/Private/Chaos/Collision/ParticlePairMidPhase.cpp
+++ b/Engine/Source/Runtime/Experimental/Chaos/Private/Chaos/Collision/ParticlePairMidPhase.cpp
@@ -315,6 +315,8 @@ namespace Chaos
 			// pointer because of Unions - see FMultiShapePairCollisionDetector)
 			Constraint->SetShapeWorldTransforms(ShapeWorldTransform0, ShapeWorldTransform1);
 
+			Constraint->SetCullDistance(CullDistance);
+
 			// If the constraint was not used last frame, it needs to be reset, otherwise we will try to reuse
 			if (!bWasUpdatedLastTick || (Constraint->GetManifoldPoints().Num() == 0))
 			{
@@ -685,12 +687,6 @@ namespace Chaos
 		// CCD may be temporarily disabled by the user or because we are moving slowly.
 		Constraint->SetCCDSweepEnabled(MidPhase.IsCCD() && bEnableSweep);
 
-		const FRigidTransform3 ParticleTransform0 = FConstGenericParticleHandle(InParticle0)->GetTransformPQ();
-		const FRigidTransform3 ParticleTransform1 = FConstGenericParticleHandle(InParticle1)->GetTransformPQ();
-		const FRigidTransform3 ShapeWorldTransform0 = InShapeRelativeTransform0 * ParticleTransform0;
-		const FRigidTransform3 ShapeWorldTransform1 = InShapeRelativeTransform1 * ParticleTransform1;
-		Constraint->SetShapeWorldTransforms(ShapeWorldTransform0, ShapeWorldTransform1);
-
 		NewConstraints.Add(Constraint);
 		return Constraint;
 	}
@@ -789,6 +785,8 @@ namespace Chaos
 		// pointer because of Unions - see FMultiShapePairCollisionDetector)
 		Constraint->SetShapeWorldTransforms(ShapeWorldTransform0, ShapeWorldTransform1);
 
+		Constraint->SetCullDistance(CullDistance);
+
 		// If the constraint was not used last frame, it needs to be reset, otherwise we will try to reuse
 		if (!bWasUpdatedLastTick || (Constraint->GetManifoldPoints().Num() == 0))
 		{
@@ -1144,12 +1142,26 @@ namespace Chaos
 			// CullDistance is scaled by the size of the dynamic objects.
 			FReal CullDistance = InCullDistance * CullDistanceScale;
 
-			// Extend cull distance for sweep-enabled CCD collision
+			// If CCD is enabled, did we move far enough to require a sweep?
 			const bool bUseSweep = Flags.bIsCCD && ShouldEnableCCD(Dt);
-			if (bUseSweep)
+
+			// Extend cull distance based on velocity
+			const FReal VMaxDt = (FConstGenericParticleHandle(GetParticle0())->V() - FConstGenericParticleHandle(GetParticle1())->V()).GetAbsMax() * Dt;
+			if (!bUseSweep)
+			{
+				// Normal (non sweep) mode: we increase CullDistance based on velocity up to a limit
+				// NOTE: This somewhat matches the bounds expansion in FPBDRigidsEvolutionGBF::Integrate
+				const FReal VelocityBoundsMultiplier = Context.GetSettings().BoundsVelocityInflation;
+				const FReal MaxVelocityBoundsExpansion = Context.GetSettings().MaxVelocityBoundsExpansion;
+				if ((VelocityBoundsMultiplier > 0) && (MaxVelocityBoundsExpansion > 0))
+				{
+					CullDistance += FMath::Min(VelocityBoundsMultiplier * VMaxDt, MaxVelocityBoundsExpansion);
+				}
+			}
+			else
 			{
-				const FReal VMax = (FConstGenericParticleHandle(GetParticle0())->V() - FConstGenericParticleHandle(GetParticle1())->V()).GetAbsMax();
-				CullDistance += VMax * Dt;
+				// CCD (sweept) mode: we increase CullDistance based on velocity, with no limits
+				CullDistance += VMaxDt;
 			}
 
 			// Run collision detection on all potentially colliding shape pairs
diff --git a/Engine/Source/Runtime/Experimental/Chaos/Private/PBDRigidsSolver.cpp b/Engine/Source/Runtime/Experimental/Chaos/Private/PBDRigidsSolver.cpp
--- a/Engine/Source/Runtime/Experimental/Chaos/Private/PBDRigidsSolver.cpp
+++ b/Engine/Source/Runtime/Experimental/Chaos/Private/PBDRigidsSolver.cpp
@@ -267,6 +267,12 @@ namespace Chaos
 		Chaos::FRealSingle ChaosSolverCullDistance = -1.0f;
 		FAutoConsoleVariableRef CVarChaosSolverCullDistance(TEXT("p.Chaos.Solver.Collision.CullDistance"), ChaosSolverCullDistance, TEXT("Override cull distance (if >= 0)"));
 
+		// @todo(chaos): move to physics project settings and set these to -1 when we are settled on values...
+		Chaos::FRealSingle ChaosSolverVelocityBoundsMultiplier = 1.0f;
+		Chaos::FRealSingle ChaosSolverMaxVelocityBoundsExpansion = 3.0f;
+		FAutoConsoleVariableRef CVarChaosSolverVelocityBoundsMultiplier(TEXT("p.Chaos.Solver.Collision.VelocityBoundsMultiplier"), ChaosSolverVelocityBoundsMultiplier, TEXT("Override velocity bounds multiplier (if >= 0)"));
+		FAutoConsoleVariableRef CVarChaosSolverMaxVelocityBoundsExpansion(TEXT("p.Chaos.Solver.Collision.MaxVelocityBoundsExpansion"), ChaosSolverMaxVelocityBoundsExpansion, TEXT("Override max velocity bounds expansion (if >= 0)"));
+
 		Chaos::FRealSingle ChaosSolverMaxPushOutVelocity = -1.0f;
 		FAutoConsoleVariableRef CVarChaosSolverMaxPushOutVelocity(TEXT("p.Chaos.Solver.Collision.MaxPushOutVelocity"), ChaosSolverMaxPushOutVelocity, TEXT("Override max pushout velocity (if >= 0)"));
 
@@ -1076,6 +1082,10 @@ namespace Chaos
 			{
 				SetCollisionCullDistance(ChaosSolverCullDistance);
 			}
+			if ((ChaosSolverVelocityBoundsMultiplier > 0.0f) && (ChaosSolverMaxVelocityBoundsExpansion > 0.0f))
+			{
+				SetVelocityBoundsExpansion(ChaosSolverVelocityBoundsMultiplier, ChaosSolverMaxVelocityBoundsExpansion);
+			}
 			if (ChaosSolverMaxPushOutVelocity >= 0.0f)
 			{
 				SetCollisionMaxPushOutVelocity(ChaosSolverMaxPushOutVelocity);
diff --git a/Engine/Source/Runtime/Experimental/Chaos/Public/Chaos/Collision/CollisionContext.h b/Engine/Source/Runtime/Experimental/Chaos/Public/Chaos/Collision/CollisionContext.h
--- a/Engine/Source/Runtime/Experimental/Chaos/Public/Chaos/Collision/CollisionContext.h
+++ b/Engine/Source/Runtime/Experimental/Chaos/Public/Chaos/Collision/CollisionContext.h
@@ -17,6 +17,7 @@ namespace Chaos
 		FCollisionDetectorSettings()
 			: BoundsExpansion(0)
 			, BoundsVelocityInflation(0)
+			, MaxVelocityBoundsExpansion(0)
 			, bFilteringEnabled(true)
 			, bDeferNarrowPhase(false)
 			, bAllowManifolds(true)
@@ -31,6 +32,9 @@ namespace Chaos
 		// Shape bounds in the broadphase are expanded by this multiple of velocity
 		FReal BoundsVelocityInflation;
 
+		// We only allow the bounds to grow from velocity by this much
+		FReal MaxVelocityBoundsExpansion;
+
 		// Whether to check the shape query flags in the narrow phase (e.g., Rigid Body nodes have already performed filtering prior to collision detection)
 		bool bFilteringEnabled;
 
diff --git a/Engine/Source/Runtime/Experimental/Chaos/Public/Chaos/PBDCollisionConstraints.h b/Engine/Source/Runtime/Experimental/Chaos/Public/Chaos/PBDCollisionConstraints.h
--- a/Engine/Source/Runtime/Experimental/Chaos/Public/Chaos/PBDCollisionConstraints.h
+++ b/Engine/Source/Runtime/Experimental/Chaos/Public/Chaos/PBDCollisionConstraints.h
@@ -293,6 +293,12 @@ public:
 		DetectorSettings.BoundsExpansion = InCullDistance;
 	}
 
+	void SetVelocityBoundsExpansion(const FReal BoundsVelocityMultiplier, const FReal MaxVelocityBoundsExpansion)
+	{
+		DetectorSettings.BoundsVelocityInflation = BoundsVelocityMultiplier;
+		DetectorSettings.MaxVelocityBoundsExpansion = MaxVelocityBoundsExpansion;
+	}
+
 protected:
 	FPBDCollisionConstraint& GetConstraint(int32 Index);
 
diff --git a/Engine/Source/Runtime/Experimental/Chaos/Public/Chaos/PBDRigidsEvolutionGBF.h b/Engine/Source/Runtime/Experimental/Chaos/Public/Chaos/PBDRigidsEvolutionGBF.h
--- a/Engine/Source/Runtime/Experimental/Chaos/Public/Chaos/PBDRigidsEvolutionGBF.h
+++ b/Engine/Source/Runtime/Experimental/Chaos/Public/Chaos/PBDRigidsEvolutionGBF.h
@@ -194,6 +194,8 @@ namespace Chaos
 			CHAOS_SCOPED_TIMER(Integrate);
 
 			const FReal BoundsThickness = GetCollisionConstraints().GetDetectorSettings().BoundsExpansion;
+			const FReal VelocityBoundsMultiplier = GetCollisionConstraints().GetDetectorSettings().BoundsVelocityInflation;
+			const FReal MaxVelocityBoundsExpansion = GetCollisionConstraints().GetDetectorSettings().MaxVelocityBoundsExpansion;
 			const FReal MaxAngularSpeedSq = CVars::HackMaxAngularVelocity * CVars::HackMaxAngularVelocity;
 			const FReal MaxSpeedSq = CVars::HackMaxVelocity * CVars::HackMaxVelocity;
 			InParticles.ParallelFor([&](auto& GeomParticle, int32 Index) 
@@ -276,11 +278,23 @@ namespace Chaos
 
 					Particle.SetTransformPQCom(PCoM, QCoM);
 
+					// We need to expand the bounds back along velocity otherwise we can miss collisions when a box
+					// lands on another box, imparting velocity to the lower box and causing the boxes to be
+					// separated by more than Cull Distance at collision detection time.
+					FVec3 VelocityBoundsDelta = FVec3(0);
+					if ((VelocityBoundsMultiplier > 0) && (MaxVelocityBoundsExpansion > 0))
+					{
+						VelocityBoundsDelta = (-VelocityBoundsMultiplier * Dt) * Particle.V();
+						
+						// Box clamp to avoid sqrt
+						VelocityBoundsDelta.BoundToCube(MaxVelocityBoundsExpansion);
+					}
+
 					if (!Particle.CCDEnabled())
 					{
 						// Expand bounds about P/Q by a small amount. This can still result in missed collisions, especially
 						// when we have joints that pull the body back to X/R, if P-X is greater than the BoundsThickness
-						Particle.UpdateWorldSpaceState(FRigidTransform3(Particle.P(), Particle.Q()), FVec3(BoundsThickness));
+						Particle.UpdateWorldSpaceStateSwept(FRigidTransform3(Particle.P(), Particle.Q()), FVec3(BoundsThickness), VelocityBoundsDelta);
 					}
 					else
 					{
@@ -307,7 +321,7 @@ namespace Chaos
 						}
 						else
 						{
-							Particle.UpdateWorldSpaceState(FRigidTransform3(Particle.P(), Particle.Q()), FVec3(BoundsThickness));
+							Particle.UpdateWorldSpaceStateSwept(FRigidTransform3(Particle.P(), Particle.Q()), FVec3(BoundsThickness), VelocityBoundsDelta);
 						}
 					}
 				}
diff --git a/Engine/Source/Runtime/Experimental/Chaos/Public/PBDRigidsSolver.h b/Engine/Source/Runtime/Experimental/Chaos/Public/PBDRigidsSolver.h
--- a/Engine/Source/Runtime/Experimental/Chaos/Public/PBDRigidsSolver.h
+++ b/Engine/Source/Runtime/Experimental/Chaos/Public/PBDRigidsSolver.h
@@ -181,6 +181,7 @@ namespace Chaos
 		void SetVelocityIterations(const int32 InNumIterations) { GetEvolution()->SetNumVelocityIterations(InNumIterations); }
 		void SetProjectionIterations(const int32 InNumIterations) { GetEvolution()->SetNumProjectionIterations(InNumIterations); }
 		void SetCollisionCullDistance(const FReal InCullDistance) { GetEvolution()->GetCollisionConstraints().SetCullDistance(InCullDistance); }
+		void SetVelocityBoundsExpansion(const FReal BoundsVelocityMultiplier, const FReal MaxBoundsVelocityExpansion) { GetEvolution()->GetCollisionConstraints().SetVelocityBoundsExpansion(BoundsVelocityMultiplier, MaxBoundsVelocityExpansion); }
 		void SetCollisionMaxPushOutVelocity(const FReal InMaxPushOutVelocity) { GetEvolution()->GetCollisionConstraints().SetMaxPushOutVelocity(InMaxPushOutVelocity); }
 
 		/**/
