set geometry collection particle debug name

diff --git a/Engine/Source/Runtime/Experimental/Chaos/Private/PhysicsProxy/GeometryCollectionPhysicsProxy.cpp b/Engine/Source/Runtime/Experimental/Chaos/Private/PhysicsProxy/GeometryCollectionPhysicsProxy.cpp
--- a/Engine/Source/Runtime/Experimental/Chaos/Private/PhysicsProxy/GeometryCollectionPhysicsProxy.cpp
+++ b/Engine/Source/Runtime/Experimental/Chaos/Private/PhysicsProxy/GeometryCollectionPhysicsProxy.cpp
@@ -1291,6 +1291,10 @@ void FGeometryCollectionPhysicsProxy::InitializeBodiesPT(Chaos::FPBDRigidsSolver
 			SetSleepingState(*RigidsSolver);
 		}
 
+#if !(UE_BUILD_SHIPPING || UE_BUILD_TEST)
+		int32 HandleIndex = 0;
+#endif
+
 		// apply various features on the handles 
 		const bool bEnableGravity = Parameters.EnableGravity && !DisableGeometryCollectionGravity;
 		for (Chaos::FPBDRigidParticleHandle* Handle: SolverParticleHandles)
@@ -1303,7 +1307,15 @@ void FGeometryCollectionPhysicsProxy::InitializeBodiesPT(Chaos::FPBDRigidsSolver
 				Handle->SetInertiaConditioningEnabled(Parameters.UseInertiaConditioning);
 				Handle->SetLinearEtherDrag(Parameters.LinearDamping);
 				Handle->SetAngularEtherDrag(Parameters.AngularDamping);
+
+#if !(UE_BUILD_SHIPPING || UE_BUILD_TEST)
+				Handle->SetDebugName(MakeShared<FString, ESPMode::ThreadSafe>(FString::Printf(TEXT("%s-%d"), *Parameters.Name, HandleIndex)));
+#endif
 			}
+
+#if !(UE_BUILD_SHIPPING || UE_BUILD_TEST)
+			++HandleIndex;
+#endif
 		}
 		// call DirtyParticle to make sure the acceleration structure is up to date with all the changes happening here
 		DirtyAllParticles(*RigidsSolver);
