Geometry collection : ISM Pool warning because of the mesh group has not been created yet by the GC

diff --git a/Engine/Source/Runtime/Experimental/GeometryCollectionEngine/Private/GeometryCollection/GeometryCollectionComponent.cpp b/Engine/Source/Runtime/Experimental/GeometryCollectionEngine/Private/GeometryCollection/GeometryCollectionComponent.cpp
--- a/Engine/Source/Runtime/Experimental/GeometryCollectionEngine/Private/GeometryCollection/GeometryCollectionComponent.cpp
+++ b/Engine/Source/Runtime/Experimental/GeometryCollectionEngine/Private/GeometryCollection/GeometryCollectionComponent.cpp
@@ -4377,7 +4377,7 @@ void UGeometryCollectionComponent::UnregisterFromISMPool()
 
 void UGeometryCollectionComponent::AddAutoInstancesToISMPool()
 {
-	if (ISMPoolAutoInstancesMeshIds.Num() > 0)
+	if (ISMPoolAutoInstancesMeshIds.Num() > 0 || ISMPoolMeshGroupIndex == -1)
 	{
 		return;
 	}
