fixed issue where the initialization and/or replication order of various pieces within a cluster union could cause incorrect visuals and/or incorrect collision behavior on the client

- Ensure that the replicated proxy component only starts to handle replicated data if its data is "valid".
- Force the cluster union on the client to only accept position initialization and child to parent changes from an authoritative source (the server).
- Add missing calls to dirty the cluster union particle and to update its world space state when its XR is changed from replication.

diff --git a/Engine/Source/Runtime/Engine/Private/PhysicsEngine/ClusterUnionComponent.cpp b/Engine/Source/Runtime/Engine/Private/PhysicsEngine/ClusterUnionComponent.cpp
--- a/Engine/Source/Runtime/Engine/Private/PhysicsEngine/ClusterUnionComponent.cpp
+++ b/Engine/Source/Runtime/Engine/Private/PhysicsEngine/ClusterUnionComponent.cpp
@@ -113,7 +113,7 @@ void UClusterUnionComponent::AddComponentToCluster(UPrimitiveComponent* InCompon
 
 void UClusterUnionComponent::RemoveComponentFromCluster(UPrimitiveComponent* InComponent)
 {
-	if (!InComponent || !PhysicsProxy || !ensure(InComponent->HasValidPhysicsState()))
+	if (!InComponent || !PhysicsProxy)
 	{
 		return;
 	}
@@ -181,9 +181,16 @@ void UClusterUnionComponent::OnCreatePhysicsState()
 	// TODO: Expose these parameters via the component.
 	Chaos::FClusterCreationParameters Parameters{ 0.3f, 100, false, false };
 	Parameters.ConnectionMethod = Chaos::FClusterCreationParameters::EConnectionMethod::DelaunayTriangulation;
+	Parameters.bUseExistingChildToParent = !GetOwner()->HasAuthority();
 
 	FChaosUserData::Set<UPrimitiveComponent>(&PhysicsUserData, this);
-	PhysicsProxy = new Chaos::FClusterUnionPhysicsProxy{this, Parameters, static_cast<void*>(&PhysicsUserData), GetOwner()->GetUniqueID(), GetUniqueID()};
+
+	Chaos::FClusterUnionInitData InitData;
+	InitData.UserData = static_cast<void*>(&PhysicsUserData);
+	InitData.ActorId = GetOwner()->GetUniqueID();
+	InitData.ComponentId = GetUniqueID();
+	InitData.bNeedsClusterXRInitialization = GetOwner()->HasAuthority();
+	PhysicsProxy = new Chaos::FClusterUnionPhysicsProxy{ this, Parameters, InitData };
 	PhysicsProxy->Initialize_External();
 	if (FPhysScene_Chaos* Scene = GetChaosScene())
 	{
diff --git a/Engine/Source/Runtime/Engine/Private/PhysicsEngine/ClusterUnionReplicatedProxyComponent.cpp b/Engine/Source/Runtime/Engine/Private/PhysicsEngine/ClusterUnionReplicatedProxyComponent.cpp
--- a/Engine/Source/Runtime/Engine/Private/PhysicsEngine/ClusterUnionReplicatedProxyComponent.cpp
+++ b/Engine/Source/Runtime/Engine/Private/PhysicsEngine/ClusterUnionReplicatedProxyComponent.cpp
@@ -106,16 +106,17 @@ void UClusterUnionReplicatedProxyComponent::PostRepNotifies()
 {
 	UActorComponent::PostRepNotifies();
 
+	const bool bIsValid = ParentClusterUnion && ChildClusteredComponent && !ParticleBoneIds.IsEmpty();
+	if (!bIsValid)
+	{
+		return;
+	}
+
 	// These three properties should only get set once when the component is created.
 	const bool bIsInitialReplication = bNetUpdateParentClusterUnion || bNetUpdateChildClusteredComponent || bNetUpdateParticleBoneIds;
-	const bool bIsValid = ParentClusterUnion && ChildClusteredComponent && !ParticleBoneIds.IsEmpty();
 	if (bIsInitialReplication)
 	{
-		if (bIsValid)
-		{
-			ParentClusterUnion->AddComponentToCluster(ChildClusteredComponent, ParticleBoneIds);
-		}
-
+		ParentClusterUnion->AddComponentToCluster(ChildClusteredComponent, ParticleBoneIds);
 		bNetUpdateParentClusterUnion = false;
 		bNetUpdateChildClusteredComponent = false;
 		bNetUpdateParticleBoneIds = false;
diff --git a/Engine/Source/Runtime/Experimental/Chaos/Private/Chaos/ClusterUnionManager.cpp b/Engine/Source/Runtime/Experimental/Chaos/Private/Chaos/ClusterUnionManager.cpp
--- a/Engine/Source/Runtime/Experimental/Chaos/Private/Chaos/ClusterUnionManager.cpp
+++ b/Engine/Source/Runtime/Experimental/Chaos/Private/Chaos/ClusterUnionManager.cpp
@@ -101,8 +101,18 @@ namespace Chaos
 
 		for (FPBDRigidParticleHandle* Child : Union.ChildParticles)
 		{
-			const FRigidTransform3 ChildWorldTM(Child->X(), Child->R());
-			FRigidTransform3 Frame = ChildWorldTM.GetRelativeTransform(ClusterWorldTM);
+			FRigidTransform3 Frame = FRigidTransform3::Identity;
+			
+			if (FPBDRigidClusteredParticleHandle* ClusterChild = Child->CastToClustered(); ClusterChild && Union.Parameters.bUseExistingChildToParent)
+			{
+				Frame = ClusterChild->ChildToParent();
+			}
+			else
+			{
+				const FRigidTransform3 ChildWorldTM(Child->X(), Child->R());
+				Frame = ChildWorldTM.GetRelativeTransform(ClusterWorldTM);
+			}
+
 			if (Child->Geometry())
 			{
 				Objects.Add(TUniquePtr<FImplicitObject>(new TImplicitObjectTransformed<FReal, 3>(Child->Geometry(), Frame)));
@@ -317,6 +327,7 @@ namespace Chaos
 			}
 
 			MEvolution.SetPhysicsMaterial(Cluster->InternalCluster, MEvolution.GetPhysicsMaterial(FinalParticlesToAdd[0]));
+			Cluster->bNeedsXRInitialization = false;
 		}
 
 		if (Cluster->InternalCluster->Disabled())
@@ -426,9 +437,10 @@ namespace Chaos
 		TSet<FPBDRigidParticleHandle*> FullChildrenSet(ClusterUnion.ChildParticles);
 
 		const FRigidTransform3 ForceMassOrientation{ ClusterUnion.InternalCluster->X(), ClusterUnion.InternalCluster->R() };
-		UpdateClusterMassProperties(ClusterUnion.InternalCluster, FullChildrenSet, ClusterInertia, bRecomputeMassOrientation ? nullptr : &ForceMassOrientation);
+		UpdateClusterMassProperties(ClusterUnion.InternalCluster, FullChildrenSet, ClusterInertia, (bRecomputeMassOrientation || ClusterUnion.bNeedsXRInitialization) ? nullptr : &ForceMassOrientation);
 		UpdateKinematicProperties(ClusterUnion.InternalCluster, MClustering.GetChildrenMap(), MEvolution);
 
+		MEvolution.InvalidateParticle(ClusterUnion.InternalCluster);
 		// The recreation of the geometry must happen after the call to UpdateClusterMassProperties.
 		// Creating the geometry requires knowing the relative frame between the parent cluster and the child clusters. The
 		// parent transform is not set properly for a new empty cluster until UpdateClusterMassProperties is called for the first time.
diff --git a/Engine/Source/Runtime/Experimental/Chaos/Private/Chaos/PBDRigidClusteringAlgo.cpp b/Engine/Source/Runtime/Experimental/Chaos/Private/Chaos/PBDRigidClusteringAlgo.cpp
--- a/Engine/Source/Runtime/Experimental/Chaos/Private/Chaos/PBDRigidClusteringAlgo.cpp
+++ b/Engine/Source/Runtime/Experimental/Chaos/Private/Chaos/PBDRigidClusteringAlgo.cpp
@@ -333,8 +333,18 @@ namespace Chaos
 
 			for (FPBDRigidParticleHandle* Child : Children)
 			{
-				const FRigidTransform3 ChildWorldTM(Child->X(), Child->R());
-				FRigidTransform3 Frame = ChildWorldTM.GetRelativeTransform(ClusterWorldTM);
+				FRigidTransform3 Frame = FRigidTransform3::Identity;
+
+				if (FPBDRigidClusteredParticleHandle* ClusterChild = Child->CastToClustered(); ClusterChild && Parameters.bUseExistingChildToParent)
+				{
+					Frame = ClusterChild->ChildToParent();
+				}
+				else
+				{
+					const FRigidTransform3 ChildWorldTM(Child->X(), Child->R());
+					Frame = ChildWorldTM.GetRelativeTransform(ClusterWorldTM);
+				}
+
 				FPBDRigidParticleHandle* UsedGeomChild = Child;
 				if (Child->Geometry())
 				{
diff --git a/Engine/Source/Runtime/Experimental/Chaos/Private/PhysicsProxy/ClusterUnionPhysicsProxy.cpp b/Engine/Source/Runtime/Experimental/Chaos/Private/PhysicsProxy/ClusterUnionPhysicsProxy.cpp
--- a/Engine/Source/Runtime/Experimental/Chaos/Private/PhysicsProxy/ClusterUnionPhysicsProxy.cpp
+++ b/Engine/Source/Runtime/Experimental/Chaos/Private/PhysicsProxy/ClusterUnionPhysicsProxy.cpp
@@ -1,5 +1,4 @@
 // Copyright Epic Games, Inc. All Rights Reserved.
-
 #include "PhysicsProxy/ClusterUnionPhysicsProxy.h"
 
 #include "Chaos/ClusterUnionManager.h"
@@ -69,12 +68,10 @@ namespace Chaos
 		}
 	}
 
-	FClusterUnionPhysicsProxy::FClusterUnionPhysicsProxy(UObject* InOwner, const FClusterCreationParameters& InParameters, void* InUserData, uint32 InActorId, uint32 InComponentId)
+	FClusterUnionPhysicsProxy::FClusterUnionPhysicsProxy(UObject* InOwner, const FClusterCreationParameters& InParameters, const FClusterUnionInitData& InInitData)
 		: Base(InOwner)
 		, ClusterParameters(InParameters)
-		, UserData(InUserData)
-		, ActorId(InActorId)
-		, ComponentId(InComponentId)
+		, InitData(InInitData)
 	{
 	}
 
@@ -85,7 +82,7 @@ namespace Chaos
 		check(Particle_External != nullptr);
 
 		Particle_External->SetProxy(this);
-		Particle_External->SetUserData(UserData);
+		Particle_External->SetUserData(InitData.UserData);
 
 		// NO DIRTY FLAGS ALLOWED. We must strictly manage the dirty flags on the particle.
 		// Setting the particle's XR on the particle will set the XR dirty flag but that isn't
@@ -116,8 +113,8 @@ namespace Chaos
 
 		FClusterUnionCreationParameters ClusterUnionParameters;
 		ClusterUnionParameters.UniqueIndex = &UniqueIndex;
-		ClusterUnionParameters.ActorId = ActorId;
-		ClusterUnionParameters.ComponentId = ComponentId;
+		ClusterUnionParameters.ActorId = InitData.ActorId;
+		ClusterUnionParameters.ComponentId = InitData.ComponentId;
 
 		ClusterUnionIndex = ClusterUnionManager.CreateNewClusterUnion(ClusterParameters, ClusterUnionParameters);
 		if (FClusterUnion* ClusterUnion = ClusterUnionManager.FindClusterUnion(ClusterUnionIndex); ensure(ClusterUnion != nullptr))
@@ -125,6 +122,13 @@ namespace Chaos
 			Particle_Internal = ClusterUnion->InternalCluster;
 			Particle_Internal->SetPhysicsProxy(this);
 			Particle_Internal->GTGeometryParticle() = Particle_External.Get();
+
+#if !(UE_BUILD_SHIPPING || UE_BUILD_TEST)
+			Particle_Internal->SetDebugName(MakeShared<FString, ESPMode::ThreadSafe>(FString::Printf(TEXT("%s"), *GetOwner()->GetName())));
+#endif
+
+			// On the client, we'd rather wait for the server to initialize the particle properly.
+			ClusterUnion->bNeedsXRInitialization = InitData.bNeedsClusterXRInitialization;
 		}
 	}
 
@@ -250,11 +254,19 @@ namespace Chaos
 		}
 
 		FPBDRigidsEvolutionGBF& Evolution = *RigidsSolver.GetEvolution();
-		FClusterUnionManager& ClusterUnionManager = Evolution.GetRigidClustering().GetClusterUnionManager();
 
 		if (const FParticlePositionRotation* NewXR = ParticleData.FindClusterXR(Manager, DataIdx))
 		{
 			Evolution.SetParticleTransform(Particle_Internal, NewXR->X(), NewXR->R(), true);
+
+			// Particle needs to actually be marked dirty! Less so to pull state from the PT back to the GT for the cluster union
+			// but this is primarily for proxies within the cluster union. They (i.e. GCs) need to see this change to the particle's transform.
+			Evolution.GetParticles().MarkTransientDirtyParticle(Particle_Internal);
+
+			const FRigidTransform3 WorldTransform{ Particle_Internal->X(), Particle_Internal->R() };
+			Particle_Internal->UpdateWorldSpaceState(WorldTransform, FVec3(0));
+
+			Evolution.DirtyParticle(*Particle_Internal);
 		}
 
 		if (const FParticleVelocities* NewVelocities = ParticleData.FindClusterVelocities(Manager, DataIdx))
@@ -534,13 +546,14 @@ namespace Chaos
 					}
 
 					const int32 ParticleIndex = ClusterUnion->ChildParticles.Find(Particle->CastToRigidParticle());
-					if (ParticleIndex != INDEX_NONE)
+					if (ParticleIndex != INDEX_NONE && ClusterUnion->InternalCluster)
 					{
 						if (FPBDRigidClusteredParticleHandle* ChildHandle = ClusterUnion->ChildParticles[ParticleIndex]->CastToClustered())
 						{
 							ChildHandle->SetChildToParent(RelativeTransform);
-							Evolution.GetParticles().MarkTransientDirtyParticle(ChildHandle);
-							Evolution.GetParticles().MarkTransientDirtyParticle(Particle);
+							ClusterUnionManager.UpdateAllClusterUnionProperties(*ClusterUnion, false);
+							Evolution.GetParticles().MarkTransientDirtyParticle(ClusterUnion->InternalCluster);
+							Evolution.DirtyParticle(*ClusterUnion->InternalCluster);
 						}
 					}
 				}
diff --git a/Engine/Source/Runtime/Experimental/Chaos/Public/Chaos/ClusterCreationParameters.h b/Engine/Source/Runtime/Experimental/Chaos/Public/Chaos/ClusterCreationParameters.h
--- a/Engine/Source/Runtime/Experimental/Chaos/Public/Chaos/ClusterCreationParameters.h
+++ b/Engine/Source/Runtime/Experimental/Chaos/Public/Chaos/ClusterCreationParameters.h
@@ -34,6 +34,7 @@ namespace Chaos
 			, Chaos::TPBDRigidClusteredParticleHandle<Chaos::FReal,3>* ClusterParticleHandleIn = nullptr
 			, const FVec3& ScaleIn = FVec3::OneVector
 			, bool bIsAnchoredIn = false
+			, bool bInUseExistingChildToParent = false
 		)
 			: CoillisionThicknessPercent(CoillisionThicknessPercentIn)
 			, MaxNumConnections(MaxNumConnectionsIn)
@@ -46,6 +47,7 @@ namespace Chaos
 			, ClusterParticleHandle(ClusterParticleHandleIn)
 			, Scale(ScaleIn)
 			, bIsAnchored(bIsAnchoredIn)
+			, bUseExistingChildToParent(bInUseExistingChildToParent)
 		{}
 
 		FReal CoillisionThicknessPercent;
@@ -59,5 +61,7 @@ namespace Chaos
 		Chaos::FPBDRigidClusteredParticleHandle* ClusterParticleHandle;
 		FVec3 Scale;
 		bool bIsAnchored;
+		// Whether or not to use the particle's existing child to parent transform instead of trying to recompute it.
+		bool bUseExistingChildToParent;
 	};
 }
diff --git a/Engine/Source/Runtime/Experimental/Chaos/Public/Chaos/ClusterUnionManager.h b/Engine/Source/Runtime/Experimental/Chaos/Public/Chaos/ClusterUnionManager.h
--- a/Engine/Source/Runtime/Experimental/Chaos/Public/Chaos/ClusterUnionManager.h
+++ b/Engine/Source/Runtime/Experimental/Chaos/Public/Chaos/ClusterUnionManager.h
@@ -51,6 +51,9 @@ namespace Chaos
 
 		// Other parameters that aren't related to clusters generally but are related to info we need about the cluster union.
 		FClusterUnionCreationParameters ClusterUnionParameters;
+
+		// Whether or not the position/rotation needs to be computed the first time a particle is added.
+		bool bNeedsXRInitialization = true;
 	};
 
 	
@@ -93,6 +96,11 @@ namespace Chaos
 
 		FClusterUnionIndex FindClusterUnionIndexFromParticle(FPBDRigidParticleHandle* Particle);
 
+		// Update the cluster union's properties after its set of particle changes.
+		void UpdateAllClusterUnionProperties(FClusterUnion& ClusterUnion, bool bRecomputeMassOrientation);
+
+		// Returns all cluster unions. Really meant only to be used for debugging.
+		const TMap<FClusterUnionIndex, FClusterUnion>& GetAllClusterUnions() const { return ClusterUnions; }
 	private:
 		FRigidClustering& MClustering;
 		FPBDRigidsEvolutionGBF& MEvolution;
@@ -140,9 +148,6 @@ namespace Chaos
 
 		// Forcefully recreate the shared geometry on a cluster. Potentially expensive so ideally should be used rarely.
 		TSharedPtr<FImplicitObject, ESPMode::ThreadSafe> ForceRecreateClusterUnionSharedGeometry(const FClusterUnion& Union);
-
-		// Update the cluster union's properties after its set of particle changes.
-		void UpdateAllClusterUnionProperties(FClusterUnion& ClusterUnion, bool bRecomputeMassOrientation);
 	};
 
 }
\ No newline at end of file
diff --git a/Engine/Source/Runtime/Experimental/Chaos/Public/PhysicsProxy/ClusterUnionPhysicsProxy.h b/Engine/Source/Runtime/Experimental/Chaos/Public/PhysicsProxy/ClusterUnionPhysicsProxy.h
--- a/Engine/Source/Runtime/Experimental/Chaos/Public/PhysicsProxy/ClusterUnionPhysicsProxy.h
+++ b/Engine/Source/Runtime/Experimental/Chaos/Public/PhysicsProxy/ClusterUnionPhysicsProxy.h
@@ -23,6 +23,14 @@ namespace Chaos
 		FTransform ChildToParent;
 	};
 
+	struct FClusterUnionInitData
+	{
+		void* UserData;
+		uint32 ActorId = INDEX_NONE;
+		uint32 ComponentId = INDEX_NONE;
+		bool bNeedsClusterXRInitialization = true;
+	};
+
 	/**
 	 * Extra data that needs to be synced between the PT and GT for cluster unions.
 	 */
@@ -34,7 +42,7 @@ namespace Chaos
 		// Data on every child particle in the cluster union.
 		TArray<FClusterUnionChildData> ChildParticles;
 	};
-
+	
 
 	class CHAOS_API FClusterUnionPhysicsProxy : public TPhysicsProxy<FClusterUnionPhysicsProxy, void, FClusterUnionProxyTimestamp>
 	{
@@ -45,7 +53,7 @@ namespace Chaos
 		using FInternalParticle = TPBDRigidClusteredParticleHandle<Chaos::FReal, 3>;
 
 		FClusterUnionPhysicsProxy() = delete;
-		FClusterUnionPhysicsProxy(UObject* InOwner, const FClusterCreationParameters& InParameters, void* InUserData, uint32 InActorId, uint32 InComponentId);
+		FClusterUnionPhysicsProxy(UObject* InOwner, const FClusterCreationParameters& InParameters, const FClusterUnionInitData& InInitData);
 
 		// Add physics objects to the cluster union. Should only be called from the game thread.
 		void AddPhysicsObjects_External(const TArray<FPhysicsObjectHandle>& Objects);
@@ -106,9 +114,7 @@ namespace Chaos
 	private:
 		bool bIsInitializedOnPhysicsThread = false;
 		FClusterCreationParameters ClusterParameters;
-		void* UserData = nullptr;
-		uint32 ActorId;
-		uint32 ComponentId;
+		const FClusterUnionInitData InitData;
 
 		FPhysicsObjectUniquePtr PhysicsObject;
 		TUniquePtr<FExternalParticle> Particle_External;
