the bNeedsXRInitialization parameter on cluster unions should not be overridable within the cluster union code

diff --git a/Engine/Source/Runtime/Experimental/Chaos/Private/Chaos/ClusterUnionManager.cpp b/Engine/Source/Runtime/Experimental/Chaos/Private/Chaos/ClusterUnionManager.cpp
--- a/Engine/Source/Runtime/Experimental/Chaos/Private/Chaos/ClusterUnionManager.cpp
+++ b/Engine/Source/Runtime/Experimental/Chaos/Private/Chaos/ClusterUnionManager.cpp
@@ -437,7 +437,7 @@ namespace Chaos
 		TSet<FPBDRigidParticleHandle*> FullChildrenSet(ClusterUnion.ChildParticles);
 
 		const FRigidTransform3 ForceMassOrientation{ ClusterUnion.InternalCluster->X(), ClusterUnion.InternalCluster->R() };
-		UpdateClusterMassProperties(ClusterUnion.InternalCluster, FullChildrenSet, ClusterInertia, (bRecomputeMassOrientation || ClusterUnion.bNeedsXRInitialization) ? nullptr : &ForceMassOrientation);
+		UpdateClusterMassProperties(ClusterUnion.InternalCluster, FullChildrenSet, ClusterInertia, (bRecomputeMassOrientation && ClusterUnion.bNeedsXRInitialization) ? nullptr : &ForceMassOrientation);
 		UpdateKinematicProperties(ClusterUnion.InternalCluster, MClustering.GetChildrenMap(), MEvolution);
 
 		MEvolution.InvalidateParticle(ClusterUnion.InternalCluster);
