Chaos - fix debug draw stall when restoring from background if bThrottleCPUWhenNotForeground is enabled

diff --git a/Engine/Source/Runtime/Experimental/ChaosSolverEngine/Private/Chaos/ChaosDebugDrawComponent.cpp b/Engine/Source/Runtime/Experimental/ChaosSolverEngine/Private/Chaos/ChaosDebugDrawComponent.cpp
--- a/Engine/Source/Runtime/Experimental/ChaosSolverEngine/Private/Chaos/ChaosDebugDrawComponent.cpp
+++ b/Engine/Source/Runtime/Experimental/ChaosSolverEngine/Private/Chaos/ChaosDebugDrawComponent.cpp
@@ -186,6 +186,7 @@ void DebugDrawChaos(AActor* DebugDrawActor, const TArray<Chaos::FLatentDrawComma
 
 UChaosDebugDrawComponent::UChaosDebugDrawComponent()
 	: bInPlay(false)
+	, LastRenderTime(-1.0)
 {
 	// We must tick after anything that uses Chaos Debug Draw and also after the Line Batcher Component
 	PrimaryComponentTick.bAllowTickOnDedicatedServer = false;
@@ -211,6 +212,7 @@ void UChaosDebugDrawComponent::BeginPlay()
 	SetTickableWhenPaused(true);
 
 	bInPlay = true;
+	LastRenderTime = -1.0;
 
 #if CHAOS_DEBUG_DRAW
 	Chaos::FDebugDrawQueue::GetInstance().SetConsumerActive(this, bInPlay);
@@ -289,7 +291,14 @@ void UChaosDebugDrawComponent::TickComponent(float DeltaTime, enum ELevelTick Ti
 			FDebugDrawQueue::GetInstance().ExtractAllElements(DrawCommands);
 		}
 
-		DebugDrawChaos(GetOwner(), DrawCommands, World->IsPaused());
+		// Don't enqueue more lines if the renderer hasn't run since the last time we were ticked.
+		// This prevents us from accumulated lines in the line batcher when in the background, for example.
+		if (LastRenderTime < World->LastRenderTime)
+		{
+			LastRenderTime = World->LastRenderTime;
+
+			DebugDrawChaos(GetOwner(), DrawCommands, World->IsPaused());
+		}
 	}
 #endif
 }
diff --git a/Engine/Source/Runtime/Experimental/ChaosSolverEngine/Public/Chaos/ChaosDebugDrawComponent.h b/Engine/Source/Runtime/Experimental/ChaosSolverEngine/Public/Chaos/ChaosDebugDrawComponent.h
--- a/Engine/Source/Runtime/Experimental/ChaosSolverEngine/Public/Chaos/ChaosDebugDrawComponent.h
+++ b/Engine/Source/Runtime/Experimental/ChaosSolverEngine/Public/Chaos/ChaosDebugDrawComponent.h
@@ -33,6 +33,7 @@ private:
 	static void CreateDebugDrawActor(UWorld* World);
 
 	bool bInPlay;
+	double LastRenderTime;
 
 #if CHAOS_DEBUG_DRAW
 	TArray<Chaos::FLatentDrawCommand> DrawCommands;
