[Perf][Solver] - ~80% of steady state solver tick calling empty field objects

diff --git a/Engine/Source/Runtime/Experimental/Chaos/Private/PBDRigidsSolver.cpp b/Engine/Source/Runtime/Experimental/Chaos/Private/PBDRigidsSolver.cpp
--- a/Engine/Source/Runtime/Experimental/Chaos/Private/PBDRigidsSolver.cpp
+++ b/Engine/Source/Runtime/Experimental/Chaos/Private/PBDRigidsSolver.cpp
@@ -352,10 +352,11 @@ namespace Chaos
 					MSolver->FieldParameterUpdateCallback(PositionTarget, TargetedParticles);
 				}
 
-				for (FGeometryCollectionPhysicsProxy* GeoclObj : MSolver->GetGeometryCollectionPhysicsProxies_Internal())
+				for (FGeometryCollectionPhysicsProxy* GeoclObj : MSolver->GetGeometryCollectionPhysicsProxiesField_Internal())
 				{
 					GeoclObj->FieldParameterUpdateCallback(MSolver);
 				}
+				MSolver->GetGeometryCollectionPhysicsProxiesField_Internal().Reset();
 
 				MSolver->GetEvolution()->GetBroadPhase().GetIgnoreCollisionManager().ProcessPendingQueues(*MSolver);
 			}
@@ -404,10 +405,11 @@ namespace Chaos
 						MSolver->FieldForcesUpdateCallback();
 					}
 
-					for (FGeometryCollectionPhysicsProxy* GeoCollectionObj : MSolver->GetGeometryCollectionPhysicsProxies_Internal())
+					for (FGeometryCollectionPhysicsProxy* GeoCollectionObj : MSolver->GetGeometryCollectionPhysicsProxiesField_Internal())
 					{
 						GeoCollectionObj->FieldForcesUpdateCallback(MSolver);
 					}
+					MSolver->GetGeometryCollectionPhysicsProxiesField_Internal().Reset();
 
 					if(FRewindData* RewindData = MSolver->GetRewindData())
 					{
diff --git a/Engine/Source/Runtime/Experimental/Chaos/Private/PhysicsProxy/GeometryCollectionPhysicsProxy.cpp b/Engine/Source/Runtime/Experimental/Chaos/Private/PhysicsProxy/GeometryCollectionPhysicsProxy.cpp
--- a/Engine/Source/Runtime/Experimental/Chaos/Private/PhysicsProxy/GeometryCollectionPhysicsProxy.cpp
+++ b/Engine/Source/Runtime/Experimental/Chaos/Private/PhysicsProxy/GeometryCollectionPhysicsProxy.cpp
@@ -1052,6 +1052,7 @@ void FGeometryCollectionPhysicsProxy::InitializeBodiesPT(Chaos::FPBDRigidsSolver
 			FFieldSystemMetaDataProcessingResolution* ResolutionData = new FFieldSystemMetaDataProcessingResolution(EFieldResolutionType::Field_Resolution_Maximum);
 
 			Cmd.MetaData.Add(FFieldSystemMetaData::EMetaType::ECommandData_ProcessingResolution, TUniquePtr<FFieldSystemMetaDataProcessingResolution>(ResolutionData));
+			RigidsSolver->GetGeometryCollectionPhysicsProxiesField_Internal().Add(this);
 			Commands.Add(Cmd);
 		}
 		Parameters.InitializationCommands.Empty();
@@ -4297,7 +4298,7 @@ void BuildSimulationData(Chaos::FErrorReporter& ErrorReporter, FGeometryCollecti
 void FGeometryCollectionPhysicsProxy::FieldParameterUpdateCallback(Chaos::FPBDRigidsSolver* RigidSolver, const bool bUpdateViews)
 {
 	SCOPE_CYCLE_COUNTER(STAT_ParamUpdateField_Object);
-
+	ensure(RigidSolver);
 	// We are updating the Collection from the InitializeBodiesPT, so we need the PT collection
 	FGeometryDynamicCollection& Collection = PhysicsThreadCollection;
 	Chaos::FPBDPositionConstraints PositionTarget;
@@ -4305,7 +4306,7 @@ void FGeometryCollectionPhysicsProxy::FieldParameterUpdateCallback(Chaos::FPBDRi
 
 	// Process Particle-Collection commands
 	int32 NumCommands = Commands.Num();
-	if (NumCommands && RigidSolver && !RigidSolver->IsShuttingDown() && Collection.Transform.Num())
+	if (NumCommands && !RigidSolver->IsShuttingDown() && Collection.Transform.Num())
 	{
 		TArray<int32> CommandsToRemove;
 		CommandsToRemove.Reserve(NumCommands);
@@ -4494,9 +4495,10 @@ void FGeometryCollectionPhysicsProxy::FieldParameterUpdateCallback(Chaos::FPBDRi
 void FGeometryCollectionPhysicsProxy::FieldForcesUpdateCallback(Chaos::FPBDRigidsSolver* RigidSolver)
 {
 	SCOPE_CYCLE_COUNTER(STAT_ForceUpdateField_Object);
+	ensure(RigidSolver);
 
 	const int32 NumCommands = Commands.Num();
-	if (NumCommands && RigidSolver && !RigidSolver->IsShuttingDown())
+	if (NumCommands && !RigidSolver->IsShuttingDown())
 	{
 		TArray<int32> CommandsToRemove;
 		CommandsToRemove.Reserve(NumCommands);
diff --git a/Engine/Source/Runtime/Experimental/Chaos/Public/PBDRigidsSolver.h b/Engine/Source/Runtime/Experimental/Chaos/Public/PBDRigidsSolver.h
--- a/Engine/Source/Runtime/Experimental/Chaos/Public/PBDRigidsSolver.h
+++ b/Engine/Source/Runtime/Experimental/Chaos/Public/PBDRigidsSolver.h
@@ -253,6 +253,11 @@ namespace Chaos
 			return GeometryCollectionPhysicsProxies_Internal;
 		}
 
+		TArray<FGeometryCollectionPhysicsProxy*>& GetGeometryCollectionPhysicsProxiesField_Internal()
+		{
+			return GeometryCollectionPhysicsProxiesField_Internal;
+		}
+
 		const TArray<FJointConstraintPhysicsProxy*>& GetJointConstraintPhysicsProxies_Internal() const
 		{
 			return JointConstraintPhysicsProxies_Internal;
@@ -359,6 +364,7 @@ namespace Chaos
 		TSharedPtr<FCriticalSection> MCurrentLock;
 		TSparseArray< FSingleParticlePhysicsProxy* > SingleParticlePhysicsProxies_PT;
 		TArray< FGeometryCollectionPhysicsProxy* > GeometryCollectionPhysicsProxies_Internal; // PT
+		TArray< FGeometryCollectionPhysicsProxy* > GeometryCollectionPhysicsProxiesField_Internal; // PT
 		TArray< FClusterUnionPhysicsProxy* > ClusterUnionPhysicsProxies_Internal; // PT
 		TArray< FJointConstraintPhysicsProxy* > JointConstraintPhysicsProxies_Internal; // PT
 		TArray< FCharacterGroundConstraintProxy* > CharacterGroundConstraintProxies_Internal;
diff --git a/Engine/Source/Runtime/Experimental/Chaos/Public/PhysicsProxy/GeometryCollectionPhysicsProxy.h b/Engine/Source/Runtime/Experimental/Chaos/Public/PhysicsProxy/GeometryCollectionPhysicsProxy.h
--- a/Engine/Source/Runtime/Experimental/Chaos/Public/PhysicsProxy/GeometryCollectionPhysicsProxy.h
+++ b/Engine/Source/Runtime/Experimental/Chaos/Public/PhysicsProxy/GeometryCollectionPhysicsProxy.h
@@ -245,8 +245,12 @@ public:
 	/** Enqueue a field \p Command to be processed by \c ProcessCommands() or 
 	 * \c FieldForcesUpdateCallback(). 
 	 */
-	void BufferCommand(Chaos::FPBDRigidsSolver* , const FFieldSystemCommand& Command)
-	{ Commands.Add(Command); }
+	void BufferCommand(Chaos::FPBDRigidsSolver* RigidsSolver, const FFieldSystemCommand& Command)
+	{ 
+		check(RigidsSolver != nullptr);
+		RigidsSolver->GetGeometryCollectionPhysicsProxiesField_Internal().Add(this);
+		Commands.Add(Command); 
+	}
 
 	static void InitializeSharedCollisionStructures(Chaos::FErrorReporter& ErrorReporter, FGeometryCollection& RestCollection, const FSharedSimulationParameters& SharedParams);
 
