Iris
- Fixed issue with SubObjects not marked as having dirty state was considered for serialization.

diff --git a/Engine/Source/Runtime/Experimental/Iris/Core/Private/Iris/ReplicationSystem/ReplicationWriter.cpp b/Engine/Source/Runtime/Experimental/Iris/Core/Private/Iris/ReplicationSystem/ReplicationWriter.cpp
--- a/Engine/Source/Runtime/Experimental/Iris/Core/Private/Iris/ReplicationSystem/ReplicationWriter.cpp
+++ b/Engine/Source/Runtime/Experimental/Iris/Core/Private/Iris/ReplicationSystem/ReplicationWriter.cpp
@@ -327,13 +327,14 @@ void FReplicationWriter::QueueNetObjectAttachments(FInternalNetRefIndex OwnerInt
 		return;
 	}
 
-	ObjectsWithDirtyChanges.SetBit(OwnerInternalIndex);
+	ObjectsWithDirtyChanges.SetBit(SubObjectInternalIndex);
 
 	FReplicationInfo& TargetInfo = GetReplicationInfo(TargetIndex);
 	TargetInfo.HasAttachments = 1;
 
 	if (OwnerInternalIndex != SubObjectInternalIndex)
 	{
+		ObjectsWithDirtyChanges.SetBit(OwnerInternalIndex);
 		FReplicationInfo& OwnerInfo = GetReplicationInfo(OwnerInternalIndex);
 		OwnerInfo.HasDirtySubObjects = 1;
 	}
@@ -575,6 +576,7 @@ void FReplicationWriter::SetPendingDestroyOrSubObjectPendingDestroyState(uint32
 				OwnerInfo.HasDirtySubObjects = 1U;
 
 				SetState(InternalIndex, EReplicatedObjectState::SubObjectPendingDestroy);
+				ObjectsWithDirtyChanges.SetBit(InternalIndex);
 				ObjectsPendingDestroy.SetBit(InternalIndex);
 				Info.SubObjectPendingDestroy = 1U;
 				return;
@@ -1271,6 +1273,8 @@ void FReplicationWriter::HandleDroppedRecord<FReplicationWriter::EReplicatedObje
 			ObjectsPendingDestroy.ClearBit(InternalIndex);
 			SetState(InternalIndex, EReplicatedObjectState::Created);
 
+			ObjectsWithDirtyChanges.SetBitValue(InternalIndex, Info.HasDirtyChangeMask);
+
 			FReplicationInfo& OwnerInfo = ReplicatedObjects[ObjectData.SubObjectRootIndex];
 			if (OwnerInfo.GetState() < EReplicatedObjectState::PendingDestroy)
 			{
@@ -1338,6 +1342,7 @@ void FReplicationWriter::HandleDroppedRecord<FReplicationWriter::EReplicatedObje
 				OwnerInfo.HasDirtySubObjects = 1U;
 
 				SetState(InternalIndex, EReplicatedObjectState::SubObjectPendingDestroy);
+				ObjectsWithDirtyChanges.SetBit(InternalIndex);
 				ObjectsPendingDestroy.SetBit(InternalIndex);
 			}
 		}
@@ -2102,6 +2107,11 @@ FReplicationWriter::EWriteObjectStatus FReplicationWriter::WriteObjectInBatch(FN
 
 		for (FInternalNetRefIndex SubObjectInternalIndex : SubObjectsToReplicate)
 		{
+			if (!ObjectsWithDirtyChanges.GetBit(SubObjectInternalIndex))
+			{
+				continue;
+			}
+
 			const int BatchObjectInfoCount = OutBatchInfo.ObjectInfos.Num();
 			EWriteObjectStatus SubObjectWriteStatus = WriteObjectInBatch(Context, SubObjectInternalIndex, WriteObjectFlags, OutBatchInfo);
 			if (!IsWriteObjectSuccess(SubObjectWriteStatus))
@@ -2918,6 +2928,7 @@ UDataStream::EWriteResult FReplicationWriter::Write(FNetSerializationContext& Co
 
 	// Report packet stats
 	UE_NET_TRACE_PACKET_STATSCOUNTER(Parameters.ReplicationSystem->GetId(), Parameters.ConnectionId, ReplicationWriter.WrittenObjectCount, WrittenObjectCount, ENetTraceVerbosity::Trace);
+	UE_NET_TRACE_PACKET_STATSCOUNTER(Parameters.ReplicationSystem->GetId(), Parameters.ConnectionId, ReplicationWriter.FailedToWriteSmallObjectCount, WriteContext.FailedToWriteSmallObjectCount, ENetTraceVerbosity::Trace);
 	UE_NET_TRACE_PACKET_STATSCOUNTER(Parameters.ReplicationSystem->GetId(), Parameters.ConnectionId, ReplicationWriter.RemainingObjectsPendingWriteCount, WriteContext.ScheduledObjectCount - WriteContext.CurrentIndex, ENetTraceVerbosity::Trace);
 	UE_NET_TRACE_PACKET_STATSCOUNTER(Parameters.ReplicationSystem->GetId(), Parameters.ConnectionId, ReplicationWriter.ScheduledObjectCount, WriteContext.ScheduledObjectCount, ENetTraceVerbosity::Trace);
 
