Fix for CIS issue in ChaosVDRecording.cpp

(Removed module name from include path)

diff --git a/Engine/Source/Runtime/Experimental/ChaosVisualDebugger/Private/ChaosVDRecording.cpp b/Engine/Source/Runtime/Experimental/ChaosVisualDebugger/Private/ChaosVDRecording.cpp
--- a/Engine/Source/Runtime/Experimental/ChaosVisualDebugger/Private/ChaosVDRecording.cpp
+++ b/Engine/Source/Runtime/Experimental/ChaosVisualDebugger/Private/ChaosVDRecording.cpp
@@ -1,6 +1,6 @@
 // Copyright Epic Games, Inc. All Rights Reserved.
 
-#include "ChaosVisualDebugger/Public/ChaosVDRecording.h"
+#include "ChaosVDRecording.h"
 
 FChaosVDFrameData& FChaosVDRecording::AddFrame()
 {
