Iris - Changes to how we calculate some net stats and added more data points.

diff --git a/Engine/Source/Runtime/Experimental/Iris/Core/Private/Iris/ReplicationSystem/ReplicationWriter.cpp b/Engine/Source/Runtime/Experimental/Iris/Core/Private/Iris/ReplicationSystem/ReplicationWriter.cpp
--- a/Engine/Source/Runtime/Experimental/Iris/Core/Private/Iris/ReplicationSystem/ReplicationWriter.cpp
+++ b/Engine/Source/Runtime/Experimental/Iris/Core/Private/Iris/ReplicationSystem/ReplicationWriter.cpp
@@ -1804,6 +1804,12 @@ FReplicationWriter::EWriteObjectStatus FReplicationWriter::WriteObjectInBatch(FN
 	if (bNeedToFilterChangeMask)
 	{
 		ApplyFilterToChangeMask(OutBatchInfo.ParentInternalIndex, InternalIndex, Info, ObjectData.Protocol, ReplicatedObjectStateBuffer, bIsInitialState);
+#if UE_NET_IRIS_CSV_STATS && CSV_PROFILER
+		if (!bIsInitialState && Info.HasDirtyChangeMask)
+		{
+			WriteContext.Stats.AddNumberOfReplicatedObjectStatesMaskedOut(1U);
+		}
+#endif
 	}
 
 	const bool bIsObjectIndexForAttachment = IsObjectIndexForOOBAttachment(InternalIndex);
@@ -2191,6 +2197,7 @@ int FReplicationWriter::PrepareAndSendHugeObjectPayload(FNetSerializationContext
 	const uint32 PastHeaderPos = HugeObjectWriter.GetPosBits();
 
 	FBatchInfo BatchInfo;
+	BatchInfo.Type = EBatchInfoType::Internal;
 	BatchInfo.ParentInternalIndex = InternalIndex;
 	uint32 WriteObjectFlags = WriteObjectFlag_State;
 	// Get the creation going as quickly as possible.
@@ -2266,6 +2273,7 @@ int FReplicationWriter::PrepareAndSendHugeObjectPayload(FNetSerializationContext
 		FNetBitStreamRollbackScope RollbackScope(Writer);
 		
 		FBatchInfo HugeObjectBatchInfo;
+		HugeObjectBatchInfo.Type = EBatchInfoType::HugeObject;
 		HugeObjectBatchInfo.ParentInternalIndex = FNetRefHandleManager::InvalidInternalIndex;
 		const uint32 WriteHugeObjectFlags = WriteObjectFlag_Attachments | WriteObjectFlag_HugeObject;
 		const EWriteObjectStatus HugeObjectStatus = WriteObjectInBatch(Context, ObjectIndexForOOBAttachment, WriteHugeObjectFlags, HugeObjectBatchInfo);
@@ -2328,6 +2336,7 @@ int FReplicationWriter::WriteObjectBatch(FNetSerializationContext& Context, uint
 
 		WriteBitStreamInfo.BatchStartPos = Writer.GetPosBits();
 		FBatchInfo BatchInfo;
+		BatchInfo.Type = (InternalIndex == ObjectIndexForOOBAttachment ? (WriteObjectFlags & WriteObjectFlag_HugeObject ? EBatchInfoType::HugeObject : EBatchInfoType::OOBAttachment) : EBatchInfoType::Object);
 		BatchInfo.ParentInternalIndex = InternalIndex;
 		const EWriteObjectStatus WriteObjectStatus = WriteObjectInBatch(Context, InternalIndex, WriteObjectFlags, BatchInfo);
 		if (IsWriteObjectSuccess(WriteObjectStatus))
@@ -2432,6 +2441,8 @@ int FReplicationWriter::WriteDestructionInfo(FNetSerializationContext& Context,
 
 		// Reset scheduling priority
 		SchedulingPriorities[InternalIndex] = 0.0f;
+
+		WriteContext.Stats.AddNumberOfReplicatedDestructionInfos(1U);
 	}
 
 	return Writer.IsOverflown() ? -1 :  1;
@@ -2578,6 +2589,11 @@ int FReplicationWriter::HandleObjectBatchSuccess(const FBatchInfo& BatchInfo, FR
 {
 	int WrittenObjectCount = 0;
 
+	const bool bTrackObjectStats = BatchInfo.Type != EBatchInfoType::Internal;
+	uint32 ObjectCount = 0;
+	uint32 AttachmentCount = 0;
+	uint32 DeltaCompressedObjectCount = 0;
+
 	OutRecord.ObjectReplicationRecords.Reserve(BatchInfo.ObjectInfos.Num());
 	for (const FBatchObjectInfo& BatchObjectInfo : BatchInfo.ObjectInfos)
 	{
@@ -2612,6 +2628,8 @@ int FReplicationWriter::HandleObjectBatchSuccess(const FBatchInfo& BatchInfo, FR
 			Attachments.CommitReplicationRecord(BatchObjectInfo.AttachmentType, BatchObjectInfo.InternalIndex, BatchObjectInfo.AttachmentRecord);
 		}
 
+		AttachmentCount += BatchObjectInfo.bSentAttachments;
+
 		// Update transmission record.
 		if (BatchObjectInfo.bSentState)
 		{
@@ -2621,6 +2639,12 @@ int FReplicationWriter::HandleObjectBatchSuccess(const FBatchInfo& BatchInfo, FR
 
 			// The object no longer has any dirty state, but may still have attachments that didn't fit
 			ChangeMask.Reset();
+
+			++ObjectCount;
+			if (Info.LastAckedBaselineIndex != FDeltaCompressionBaselineManager::InvalidBaselineIndex)
+			{
+				++DeltaCompressedObjectCount;
+			}
 		}
 		else if (BatchObjectInfo.AttachmentRecord != 0 || BatchObjectInfo.bSentTearOff || BatchObjectInfo.bSentDestroySubObject)
 		{
@@ -2667,13 +2691,25 @@ int FReplicationWriter::HandleObjectBatchSuccess(const FBatchInfo& BatchInfo, FR
 			WriteContext.ObjectsWrittenThisPacket.SetBit(BatchObjectInfo.InternalIndex);
 		}
 
+	}
+
 #if UE_NET_IRIS_CSV_STATS && CSV_PROFILER
-		if (BatchObjectInfo.bSentState && (Info.LastAckedBaselineIndex != FDeltaCompressionBaselineManager::InvalidBaselineIndex))
+	if (bTrackObjectStats)
+	{
+		FNetSendStats& NetStats = WriteContext.Stats;
+
+		// We count RootObjects if anything is sent in an object batch, even if it's just subobjects or attachments. This is to mimic UReplicationGraph::ReplicateSingleActor stats.
+		if (BatchInfo.Type == EBatchInfoType::Object)
 		{
-			++WriteContext.DeltaCompressedObjectCount;
+			if (ObjectCount || AttachmentCount)
+			{
+				NetStats.AddNumberOfReplicatedRootObjects(1U);
+			}
 		}
-#endif
+		NetStats.AddNumberOfReplicatedObjects(ObjectCount);
+		NetStats.AddNumberOfDeltaCompressedReplicatedObjects(DeltaCompressedObjectCount);
 	}
+#endif
 
 	return WrittenObjectCount;
 }
@@ -2753,9 +2789,7 @@ UDataStream::EWriteResult FReplicationWriter::BeginWrite()
 	WriteContext.bHasHugeObjectToSend = bHasUnsentHugeObject;
 	WriteContext.bHasOOBAttachmentsToSend = bHasUnsentOOBAttachments;
 	WriteContext.CurrentIndex = 0U;
-	WriteContext.DeltaCompressedObjectCount = 0U;
 	WriteContext.FailedToWriteSmallObjectCount = 0U;
-	WriteContext.DeltaCompressedObjectCount = 0U;
 	WriteContext.SortedObjectCount = 0U;
 
 	// Reset dependent object array
@@ -2768,13 +2802,9 @@ UDataStream::EWriteResult FReplicationWriter::BeginWrite()
 	WriteContext.ScheduledObjectInfos = reinterpret_cast<FScheduleObjectInfo*>(FMemory::Malloc(sizeof(FScheduleObjectInfo) * Parameters.MaxActiveReplicatedObjectCount));
 	WriteContext.ScheduledObjectCount = ScheduleObjects(WriteContext.ScheduledObjectInfos);
 
-	// Init CSV stats
-#if UE_NET_IRIS_CSV_STATS && CSV_PROFILER
-	{
-		FNetSendStats& Stats = WriteContext.Stats;
-		Stats.Reset();
-	}
-#endif
+	// Clear net stats. Used for CVS and Network Insights stats.
+	WriteContext.Stats.Reset();
+	WriteContext.Stats.SetNumberOfRootObjectsScheduledForReplication(WriteContext.ScheduledObjectCount);
 
 	WriteContext.bIsValid = true;
 
@@ -2791,10 +2821,6 @@ void FReplicationWriter::EndWrite()
 		// Update stats
 		{
 			FNetSendStats& Stats = WriteContext.Stats;
-
-			Stats.SetNumberOfObjectsScheduledForReplication(WriteContext.ScheduledObjectCount);
-			Stats.SetNumberOfReplicatedObjects(WriteContext.CurrentIndex - WriteContext.FailedToWriteSmallObjectCount);
-			Stats.SetNumberOfDeltaCompressedReplicatedObjects(WriteContext.DeltaCompressedObjectCount);
 			if (HugeObjectContext.SendStatus == EHugeObjectSendStatus::Sending)
 			{
 				Stats.SetNumberOfActiveHugeObjects(1U);
diff --git a/Engine/Source/Runtime/Experimental/Iris/Core/Private/Iris/ReplicationSystem/ReplicationWriter.h b/Engine/Source/Runtime/Experimental/Iris/Core/Private/Iris/ReplicationSystem/ReplicationWriter.h
--- a/Engine/Source/Runtime/Experimental/Iris/Core/Private/Iris/ReplicationSystem/ReplicationWriter.h
+++ b/Engine/Source/Runtime/Experimental/Iris/Core/Private/Iris/ReplicationSystem/ReplicationWriter.h
@@ -219,9 +219,6 @@ private:
 		// The index into the scheduled objects array to attempt to replicate next.
 		uint32 CurrentIndex;
 
-		// The number of replicated objects that were serialized with delta compression.
-		uint32 DeltaCompressedObjectCount;
-
 		// How many objects that were attempted to be replicated but which ultimately didn't fit in the packet.
 		uint32 FailedToWriteSmallObjectCount;
 
@@ -250,10 +247,21 @@ private:
 		uint32 bSentDestroySubObject : 1;
 	};
 
+	enum class EBatchInfoType : uint32
+	{
+		Object,
+		HugeObject,
+		OOBAttachment,
+		Internal,
+		// Currently unused as destruction infos don't create a BatchInfo.
+		DestructionInfo,
+	};
+
 	struct FBatchInfo
 	{
 		TArray<FBatchObjectInfo, TInlineAllocator<16>> ObjectInfos;
 		uint32 ParentInternalIndex;
+		EBatchInfoType Type;
 	};
 
 	struct FObjectRecord
diff --git a/Engine/Source/Runtime/Experimental/Iris/Core/Public/Iris/Stats/NetStats.cpp b/Engine/Source/Runtime/Experimental/Iris/Core/Public/Iris/Stats/NetStats.cpp
--- a/Engine/Source/Runtime/Experimental/Iris/Core/Public/Iris/Stats/NetStats.cpp
+++ b/Engine/Source/Runtime/Experimental/Iris/Core/Public/Iris/Stats/NetStats.cpp
@@ -14,10 +14,13 @@ void FNetSendStats::Accumulate(const FNetSendStats& Other)
 {
 	FScopeLock Lock(&CS);
 
-	Stats.ScheduledForReplicationObjectCount += Other.Stats.ScheduledForReplicationObjectCount;
+	Stats.ScheduledForReplicationRootObjectCount += Other.Stats.ScheduledForReplicationRootObjectCount;
+	Stats.ReplicatedRootObjectCount += Other.Stats.ReplicatedRootObjectCount;
 	Stats.ReplicatedObjectCount += Other.Stats.ReplicatedObjectCount;
+	Stats.ReplicatedDestructionInfoCount += Other.Stats.ReplicatedDestructionInfoCount;
 	Stats.DeltaCompressedObjectCount += Other.Stats.DeltaCompressedObjectCount;
 	Stats.ReplicationWasteObjectCount += Other.Stats.ReplicationWasteObjectCount;
+	Stats.ReplicatedObjectStatesMaskedOut += Other.Stats.ReplicatedObjectStatesMaskedOut;
 	Stats.ActiveHugeObjectCount += Other.Stats.ActiveHugeObjectCount;
 	Stats.HugeObjectsWaitingForAckCount += Other.Stats.HugeObjectsWaitingForAckCount;
 	Stats.HugeObjectsStallingCount += Other.Stats.HugeObjectsStallingCount;
@@ -42,15 +45,21 @@ void FNetSendStats::ReportCsvStats()
 	if (Stats.ReplicatingConnectionCount > 0)
 	{
 		const float ConnectionCountFloat = float(Stats.ReplicatingConnectionCount);
-		CSV_CUSTOM_STAT(Iris, AvgScheduledForReplicationObjectCount, Stats.ScheduledForReplicationObjectCount/ConnectionCountFloat, ECsvCustomStatOp::Set);
+		CSV_CUSTOM_STAT(Iris, AvgScheduledForReplicationRootObjectCount, Stats.ScheduledForReplicationRootObjectCount/ConnectionCountFloat, ECsvCustomStatOp::Set);
+		CSV_CUSTOM_STAT(Iris, AvgReplicatedRootObjectCount, Stats.ReplicatedRootObjectCount/ConnectionCountFloat, ECsvCustomStatOp::Set);
 		CSV_CUSTOM_STAT(Iris, AvgReplicatedObjectCount, Stats.ReplicatedObjectCount/ConnectionCountFloat, ECsvCustomStatOp::Set);
+		CSV_CUSTOM_STAT(Iris, AvgReplicatedDestructionInfoCount, Stats.ReplicatedDestructionInfoCount/ConnectionCountFloat, ECsvCustomStatOp::Set);
 		CSV_CUSTOM_STAT(Iris, AvgReplicationWasteObjectCount, Stats.ReplicationWasteObjectCount/ConnectionCountFloat, ECsvCustomStatOp::Set);
+		CSV_CUSTOM_STAT(Iris, AvgReplicatedObjectStatesMaskedOut, Stats.ReplicatedObjectStatesMaskedOut/ConnectionCountFloat, ECsvCustomStatOp::Set);
 		CSV_CUSTOM_STAT(Iris, AvgDeltaCompressedObjectCount, Stats.DeltaCompressedObjectCount/ConnectionCountFloat, ECsvCustomStatOp::Set);
 	}
 	else
 	{
-		CSV_CUSTOM_STAT(Iris, AvgScheduledForReplicationObjectCount, 0, ECsvCustomStatOp::Set);
+		CSV_CUSTOM_STAT(Iris, AvgScheduledForReplicationRootObjectCount, 0, ECsvCustomStatOp::Set);
+		CSV_CUSTOM_STAT(Iris, AvgReplicatedRootObjectCount, 0, ECsvCustomStatOp::Set);
 		CSV_CUSTOM_STAT(Iris, AvgReplicatedObjectCount, 0, ECsvCustomStatOp::Set);
+		CSV_CUSTOM_STAT(Iris, AvgReplicatedDestructionInfoCount, 0, ECsvCustomStatOp::Set);
+		CSV_CUSTOM_STAT(Iris, AvgReplicatedHugeObjectCount, 0, ECsvCustomStatOp::Set);
 		CSV_CUSTOM_STAT(Iris, AvgReplicationWasteObjectCount, 0, ECsvCustomStatOp::Set);
 		CSV_CUSTOM_STAT(Iris, AvgDeltaCompressedObjectCount, 0, ECsvCustomStatOp::Set);
 	}
diff --git a/Engine/Source/Runtime/Experimental/Iris/Core/Public/Iris/Stats/NetStats.h b/Engine/Source/Runtime/Experimental/Iris/Core/Public/Iris/Stats/NetStats.h
--- a/Engine/Source/Runtime/Experimental/Iris/Core/Public/Iris/Stats/NetStats.h
+++ b/Engine/Source/Runtime/Experimental/Iris/Core/Public/Iris/Stats/NetStats.h
@@ -21,16 +21,28 @@ public:
 	FNetSendStats& operator=(const FNetSendStats&) = delete;
 
 	/** Set number of objects scheduled for replication. */
-	void SetNumberOfObjectsScheduledForReplication(uint32 Count);
+	void SetNumberOfRootObjectsScheduledForReplication(uint32 Count);
 
-	/** Set number of replicated objects. */
-	void SetNumberOfReplicatedObjects(uint32 Count);
+	/** Add number of replicated root objects. */
+	void AddNumberOfReplicatedRootObjects(uint32 Count);
 
-	/** Get the number of replicated objects. */
-	uint32 GetNumberOfReplicatedObjects() const;
+	/** Add number of replicated objects, including subobjects. */
+	void AddNumberOfReplicatedObjects(uint32 Count);
+
+	/** Add number of replicated destruction infos. */
+	void AddNumberOfReplicatedDestructionInfos(uint32 Count);
+
+	/** Add number of replicated objects, including subobjects, using delta compression. */
+	void AddNumberOfDeltaCompressedReplicatedObjects(uint32 Count);
+
+	/** Add number of replicated object states masked out such that no state is replicated for the object. The object may still replicate attachments. */
+	void AddNumberOfReplicatedObjectStatesMaskedOut(uint32 Count);
+
+	/** Get the number of replicated root objects. */
+	uint32 GetNumberOfReplicatedRootObjects() const;
 
-	/** Set number of replicated objects that used delta compression. */
-	void SetNumberOfDeltaCompressedReplicatedObjects(uint32 Count);
+	/** Get the number of replicated objects, including subobjects. */
+	uint32 GetNumberOfReplicatedObjects() const;
 
 	/** Set the number of huge objects in sending or waiting to be acked. */
 	void SetNumberOfActiveHugeObjects(uint32 Count);
@@ -64,10 +76,13 @@ private:
 		double HugeObjectStallingTimeInSeconds = 0;
 		double ReplicationWasteTimeInSeconds = 0;
 
-		int32 ScheduledForReplicationObjectCount = 0;
+		int32 ScheduledForReplicationRootObjectCount = 0;
+		int32 ReplicatedRootObjectCount = 0;
 		int32 ReplicatedObjectCount = 0;
+		int32 ReplicatedDestructionInfoCount = 0;
 		int32 DeltaCompressedObjectCount = 0;
 		int32 ReplicationWasteObjectCount = 0;
+		int32 ReplicatedObjectStatesMaskedOut = 0;
 		int32 ActiveHugeObjectCount = 0;
 		int32 HugeObjectsWaitingForAckCount = 0;
 		int32 HugeObjectsStallingCount = 0;
@@ -78,24 +93,39 @@ private:
 	FStats Stats;
 };
 
-inline void FNetSendStats::SetNumberOfObjectsScheduledForReplication(uint32 Count)
+inline void FNetSendStats::SetNumberOfRootObjectsScheduledForReplication(uint32 Count)
 {
-	Stats.ScheduledForReplicationObjectCount = Count;
+	Stats.ScheduledForReplicationRootObjectCount = Count;
 }
 
-inline void FNetSendStats::SetNumberOfReplicatedObjects(uint32 Count)
+inline void FNetSendStats::AddNumberOfReplicatedRootObjects(uint32 Count)
 {
-	Stats.ReplicatedObjectCount = Count;
+	Stats.ReplicatedRootObjectCount += Count;
 }
 
-inline uint32 FNetSendStats::GetNumberOfReplicatedObjects() const
+inline void FNetSendStats::AddNumberOfReplicatedObjects(uint32 Count)
 {
-	return Stats.ReplicatedObjectCount;
+	Stats.ReplicatedObjectCount += Count;
 }
 
-inline void FNetSendStats::SetNumberOfDeltaCompressedReplicatedObjects(uint32 Count)
+inline void FNetSendStats::AddNumberOfDeltaCompressedReplicatedObjects(uint32 Count)
 {
-	Stats.DeltaCompressedObjectCount = Count;
+	Stats.DeltaCompressedObjectCount += Count;
+}
+
+inline void FNetSendStats::AddNumberOfReplicatedObjectStatesMaskedOut(uint32 Count)
+{
+	Stats.ReplicatedObjectStatesMaskedOut += Count;
+}
+
+inline void FNetSendStats::AddNumberOfReplicatedDestructionInfos(uint32 Count)
+{
+	Stats.ReplicatedDestructionInfoCount += Count;
+}
+
+inline uint32 FNetSendStats::GetNumberOfReplicatedObjects() const
+{
+	return Stats.ReplicatedObjectCount;
 }
 
 inline void FNetSendStats::SetNumberOfActiveHugeObjects(uint32 Count)
