Geometry collection : fix issue where removal may not work properly
- Reorganize the UpdateRemovalIfNeeded method in a way that avoid the optimizer to break the runtime logic and stop the non-scale case logic block from running properly
  this seems to be coming from the condition around 1.0-GetDecay() been optimized out to get rid of the subtraction but it's not clear exactly why that fails yet
- Fixed a static analysis warning where compoaring the same variable in an ensure

diff --git a/Engine/Source/Runtime/Experimental/GeometryCollectionEngine/Private/GeometryCollection/GeometryCollectionComponent.cpp b/Engine/Source/Runtime/Experimental/GeometryCollectionEngine/Private/GeometryCollection/GeometryCollectionComponent.cpp
--- a/Engine/Source/Runtime/Experimental/GeometryCollectionEngine/Private/GeometryCollection/GeometryCollectionComponent.cpp
+++ b/Engine/Source/Runtime/Experimental/GeometryCollectionEngine/Private/GeometryCollection/GeometryCollectionComponent.cpp
@@ -503,7 +503,7 @@ namespace
 		if (bChaos_BoxCalcBounds_ISPC_Enabled && !bGeometryCollectionSingleThreadedBoundsCalculation)
 		{
 			ensure(BoundingBoxes.Num() > 0);
-			ensure(TransformIndices.Num() == TransformIndices.Num());
+			ensure(TransformIndices.Num() == BoundingBoxes.Num());
 			ensure(TransformToGeometryIndex.Num() > 0);
 			ensure(TransformToGeometryIndex.Num() == GlobalMatrices.Num());
 			
@@ -2879,53 +2879,38 @@ void UGeometryCollectionComponent::UpdateRemovalIfNeeded()
 		{
 			const FTransform ZeroScaleTransform(FQuat::Identity, FVector::Zero(), FVector(0, 0, 0));
 
-			if (RestCollection->bScaleOnRemoval)
+			const FTransform InverseComponentTransform = GetComponentTransform().Inverse();
+			for (int32 Idx = 0; Idx < GetTransformArray().Num(); ++Idx)
 			{
-				const FTransform InverseComponentTransform = GetComponentTransform().Inverse();
-				for (int32 Idx = 0; Idx < GetTransformArray().Num(); ++Idx)
+				// only update values if the decay has changed 
+				const float Decay = DecayFacade.GetDecay(Idx);
+				if (Decay > 0.f && Decay <= 1.f)
 				{
-					// only update values if the decay has changed 
-					const float Decay = DecayFacade.GetDecay(Idx);
-					if (Decay > 0.f && Decay <= 1.f)
+					const float Scale = 1.0 - Decay;
+					if (Scale < UE_SMALL_NUMBER)
 					{
-						const float Scale = 1.0 - Decay;
-						if (Scale < UE_SMALL_NUMBER)
+						DynamicCollection->Transform[Idx].SetScale3D(FVector::ZeroVector);
+					}
+					// do not try to get this condition out of the loop as this may cause some optimizer related issues
+					else if (RestCollection->bScaleOnRemoval) 
+					{
+						float ShrinkRadius = 0.0f;
+						UE::Math::TSphere<double> AccumulatedSphere;
+						// todo(chaos) : find a faster way to do that ( precompute the data ? )
+						if (CalculateInnerSphere(Idx, AccumulatedSphere))
 						{
-							DynamicCollection->Transform[Idx].SetScale3D(FVector::ZeroVector);
+							ShrinkRadius = -AccumulatedSphere.W;
 						}
-						else
-						{
-							float ShrinkRadius = 0.0f;
-							UE::Math::TSphere<double> AccumulatedSphere;
-							// todo(chaos) : find a faster way to do that ( precompute the data ? )
-							if (CalculateInnerSphere(Idx, AccumulatedSphere))
-							{
-								ShrinkRadius = -AccumulatedSphere.W;
-							}
 
-							const FQuat LocalRotation = (InverseComponentTransform * FTransform(GlobalMatrices[Idx]).Inverse()).GetRotation();
-							const FVector LocalDown = LocalRotation.RotateVector(FVector(0.f, 0.f, ShrinkRadius));
-							const FVector CenterOfMass = DynamicCollection->MassToLocal[Idx].GetTranslation();
-							const FVector ScaleCenter = LocalDown + CenterOfMass;
-							const FTransform ScaleTransform(FQuat::Identity, ScaleCenter * FVector::FReal(1.f - Scale), FVector(Scale));
-							DynamicCollection->Transform[Idx] = ScaleTransform * DynamicCollection->Transform[Idx];
-						}
+						const FQuat LocalRotation = (InverseComponentTransform * FTransform(GlobalMatrices[Idx]).Inverse()).GetRotation();
+						const FVector LocalDown = LocalRotation.RotateVector(FVector(0.f, 0.f, ShrinkRadius));
+						const FVector CenterOfMass = DynamicCollection->MassToLocal[Idx].GetTranslation();
+						const FVector ScaleCenter = LocalDown + CenterOfMass;
+						const FTransform ScaleTransform(FQuat::Identity, ScaleCenter * FVector::FReal(1.f - Scale), FVector(Scale));
+						DynamicCollection->Transform[Idx] = ScaleTransform * DynamicCollection->Transform[Idx];
 					}
 				}
 			}
-			else
-			{
-				// if we do not use scale we still need to make the piece to disappear when decay has reached 1
-				for (int32 Idx = 0; Idx < GetTransformArray().Num(); ++Idx)
-				{
-					const float Scale = (1.0 - DecayFacade.GetDecay(Idx));
-					if (Scale < UE_SMALL_NUMBER)
-					{
-						DynamicCollection->Transform[Idx].SetScale3D(FVector::ZeroVector);
-					}
-				}
-			}
-
 		}
 	}
 }
