Chaos : Memory optimization
- Change clustered particle SOA following array from double to float : CollisionImpulses, ExternalStrains, Strains ( as well as strin on the connectivityEdge structure )
- when creating the connection graph from the connection attribute , reserve memory so only use the required memory

diff --git a/Engine/Source/Runtime/Experimental/Chaos/Private/Chaos/ChaosDebugDraw.cpp b/Engine/Source/Runtime/Experimental/Chaos/Private/Chaos/ChaosDebugDraw.cpp
--- a/Engine/Source/Runtime/Experimental/Chaos/Private/Chaos/ChaosDebugDraw.cpp
+++ b/Engine/Source/Runtime/Experimental/Chaos/Private/Chaos/ChaosDebugDraw.cpp
@@ -1446,9 +1446,9 @@ namespace Chaos
 
 			//if (bChaosDebugDrawClusterConnectionGraph)
 			{
-				FReal MaxStrain = -FLT_MAX;
+				FRealSingle MaxStrain = TNumericLimits<FRealSingle>::Lowest();
 				int32 NumConnections = 0;
-				const TArrayCollectionArray<FReal>& Strain = Clustering.GetStrainArray();
+				const TArrayCollectionArray<FRealSingle>& Strain = Clustering.GetStrainArray();
 
 				for (auto& ActiveCluster : Clustering.GetChildrenMap())
 				{
diff --git a/Engine/Source/Runtime/Experimental/Chaos/Private/Chaos/PBDRigidClustering.cpp b/Engine/Source/Runtime/Experimental/Chaos/Private/Chaos/PBDRigidClustering.cpp
--- a/Engine/Source/Runtime/Experimental/Chaos/Private/Chaos/PBDRigidClustering.cpp
+++ b/Engine/Source/Runtime/Experimental/Chaos/Private/Chaos/PBDRigidClustering.cpp
@@ -173,7 +173,7 @@ namespace Chaos
 		}
 		if (ChildrenSet.Num())
 		{
-			NewParticle->SetInternalStrains(NewParticle->GetInternalStrains() / static_cast<FReal>(ChildrenSet.Num()));
+			NewParticle->SetInternalStrains(NewParticle->GetInternalStrains() / static_cast<FRealSingle>(ChildrenSet.Num()));
 			UpdateTopLevelParticle(NewParticle);
 		}
 
@@ -241,7 +241,7 @@ namespace Chaos
 		// build it all at the same time. The parent cluster's internal strain should be the average of all the child strains.
 		// The easy way to compute the new average is the multiply the old average by the number of old elements, add in the new strains,
 		// and then divide by the new total number of elements.
-		Cluster->SetInternalStrains(Cluster->GetInternalStrains() * OldNumChildren);
+		Cluster->SetInternalStrains(Cluster->GetInternalStrains() * static_cast<FRealSingle>(OldNumChildren));
 		Cluster->ClusterIds().NumChildren = Children.Num();
 
 		UpdateClusterParticlePropertiesFromChildren(Cluster, InChildren, ChildToParentMap);
@@ -304,7 +304,7 @@ namespace Chaos
 
 		if (Cluster->ClusterIds().NumChildren > 0)
 		{
-			Cluster->SetInternalStrains(Cluster->GetInternalStrains() / static_cast<FReal>(Cluster->ClusterIds().NumChildren));
+			Cluster->SetInternalStrains(Cluster->GetInternalStrains() / static_cast<FRealSingle>(Cluster->ClusterIds().NumChildren));
 		}
 	}
 
@@ -631,13 +631,13 @@ namespace Chaos
 		return NewClusters;
 	}
 
-	void FRigidClustering::SetInternalStrain(FPBDRigidClusteredParticleHandle* Particle, FReal Strain)
+	void FRigidClustering::SetInternalStrain(FPBDRigidClusteredParticleHandle* Particle, FRealSingle Strain)
 	{
 		Particle->SetInternalStrains(Strain);
 		UpdateTopLevelParticle(Particle);
 	}
 
-	void FRigidClustering::SetExternalStrain(FPBDRigidClusteredParticleHandle* Particle, FReal Strain)
+	void FRigidClustering::SetExternalStrain(FPBDRigidClusteredParticleHandle* Particle, FRealSingle Strain)
 	{
 		Particle->SetExternalStrains(Strain);
 		UpdateTopLevelParticle(Particle);
@@ -727,7 +727,7 @@ namespace Chaos
 		bool bFoundFirstRelease = false;
 
 		// only used for propagation
-		TMap<FPBDRigidParticleHandle*, FReal> AppliedStrains;
+		TMap<FPBDRigidParticleHandle*, FRealSingle> AppliedStrains;
 
 		for (int32 ChildIdx = Children.Num() - 1; ChildIdx >= 0; --ChildIdx)
 		{
@@ -739,7 +739,7 @@ namespace Chaos
 			}
 
 			// @todo(chaos) eventually should get rid of collision impulse array and only use external strain
-			const FReal MaxAppliedStrain = FMath::Max(Child->CollisionImpulses(), Child->GetExternalStrain());
+			const FRealSingle MaxAppliedStrain = FMath::Max(Child->CollisionImpulses(), Child->GetExternalStrain());
 			if ((MaxAppliedStrain >= Child->GetInternalStrains()) || bForceRelease)
 			{
 				if (!bFoundFirstRelease)
@@ -780,25 +780,25 @@ namespace Chaos
 			{
 				FPBDRigidClusteredParticleHandle* ClusteredChild = AppliedStrain.Key->CastToClustered();
 
-				const FReal AppliedStrainValue = AppliedStrain.Value;
-				FReal PropagatedStrainPerConnection = 0.0f;
+				const FRealSingle AppliedStrainValue = AppliedStrain.Value;
+				FRealSingle PropagatedStrainPerConnection = 0.0f;
 
 				// @todo(chaos) : may not be optimal, but good enough for now
 				if (BreakDamagePropagationFactor > 0 && ActivatedChildren.Contains(AppliedStrain.Key))
 				{
 					// break damage propagation case: we only look at the broken pieces and propagate the strain remainder 
-					const FReal RemainingStrain = (AppliedStrainValue - ClusteredChild->GetInternalStrains());
+					const FRealSingle RemainingStrain = (AppliedStrainValue - ClusteredChild->GetInternalStrains());
 					if (RemainingStrain > 0)
 					{
-						const FReal AdjustedRemainingStrain = (FReal)BreakDamagePropagationFactor * RemainingStrain;
+						const FRealSingle AdjustedRemainingStrain = BreakDamagePropagationFactor * RemainingStrain;
 						// todo(chaos) : could do better and have something weighted on distance with a falloff maybe ?  
-						PropagatedStrainPerConnection = AdjustedRemainingStrain / ClusteredChild->ConnectivityEdges().Num();
+						PropagatedStrainPerConnection = AdjustedRemainingStrain / static_cast<FRealSingle>(ClusteredChild->ConnectivityEdges().Num());
 					}
 				}
 				else if (ShockDamagePropagationFactor > 0)
 				{
 					// shock damage propagation case : for all the non broken pieces, propagate the actual applied strain 
-					PropagatedStrainPerConnection = (FReal)ShockDamagePropagationFactor * AppliedStrainValue;
+					PropagatedStrainPerConnection = ShockDamagePropagationFactor * AppliedStrainValue;
 				}
 
 				if (PropagatedStrainPerConnection > 0)
@@ -894,7 +894,7 @@ namespace Chaos
 					{
 						if (ensure(!ClusterHandle || ClusteredChildHandle->ClusterIds().Id == ClusterHandle))
 						{
-							SetExternalStrain(ClusteredChildHandle, TNumericLimits<FReal>::Max());
+							SetExternalStrain(ClusteredChildHandle, TNumericLimits<FRealSingle>::Max());
 							ClusterHandle = ClusteredChildHandle->ClusterIds().Id;
 						}
 						else
@@ -1273,8 +1273,8 @@ namespace Chaos
 			{
 				const FVec3 ContactWorldLocation = ContactHandle->GetContact().CalculateWorldContactLocation();
 				
-				const FReal AccumulatedImpulse = ContactHandle->GetAccumulatedImpulse().Size();
-				if (AccumulatedImpulse > UE_SMALL_NUMBER)
+				const FRealSingle AccumulatedImpulse = static_cast<FRealSingle>(ContactHandle->GetAccumulatedImpulse().Size());
+				if (AccumulatedImpulse > UE_SMALL_NUMBER && FMath::IsFinite(AccumulatedImpulse))
 				{
 					// gather propagation information from the parent proxy
 					bool bUseDamagePropagation = false;
@@ -1465,7 +1465,7 @@ namespace Chaos
 		ClusterUnionManager.HandleRemoveOperationWithClusterLookup({ ClusteredParticle }, true);
 
 		// max strain will allow to unconditionally release the children when strain is evaluated
-		constexpr FReal MaxStrain = TNumericLimits<FReal>::Max();
+		constexpr FRealSingle MaxStrain = TNumericLimits<FRealSingle>::Max();
 		if (TArray<FPBDRigidParticleHandle*>* ChildrenHandles = GetChildrenMap().Find(ClusteredParticle))
 		{
 			for (FPBDRigidParticleHandle* ChildHandle: *ChildrenHandles)
@@ -1490,7 +1490,7 @@ namespace Chaos
 	{
 		bool bCrumbledAnyCluster = false;
 		// max strain will allow to unconditionally release the children when strain is evaluated
-		constexpr FReal MaxStrain = TNumericLimits<FReal>::Max();
+		constexpr FRealSingle MaxStrain = TNumericLimits<FRealSingle>::Max();
 
 		// we should probably have a way to retrieve all the active clusters per proxy instead of having to do this iteration
 		for (FPBDRigidClusteredParticleHandle* ClusteredHandle : GetTopLevelClusterParents())
@@ -1558,7 +1558,7 @@ namespace Chaos
 		check(ClusteredChild1 && ClusteredChild2);
 		if (ClusteredChild1 == ClusteredChild2)
 			return;
-		const FReal AvgStrain = (ClusteredChild1->GetInternalStrains() + ClusteredChild2->GetInternalStrains()) * (FReal)0.5;
+		const FRealSingle AvgStrain = (ClusteredChild1->GetInternalStrains() + ClusteredChild2->GetInternalStrains()) * 0.5f;
 		TArray<TConnectivityEdge<FReal>>& Edges1 = ClusteredChild1->ConnectivityEdges();
 		TArray<TConnectivityEdge<FReal>>& Edges2 = ClusteredChild2->ConnectivityEdges();
 		if (//Edges1.Num() < Parameters.MaxNumConnections && 
diff --git a/Engine/Source/Runtime/Experimental/Chaos/Private/PhysicsProxy/GeometryCollectionPhysicsProxy.cpp b/Engine/Source/Runtime/Experimental/Chaos/Private/PhysicsProxy/GeometryCollectionPhysicsProxy.cpp
--- a/Engine/Source/Runtime/Experimental/Chaos/Private/PhysicsProxy/GeometryCollectionPhysicsProxy.cpp
+++ b/Engine/Source/Runtime/Experimental/Chaos/Private/PhysicsProxy/GeometryCollectionPhysicsProxy.cpp
@@ -1242,6 +1242,20 @@ void FGeometryCollectionPhysicsProxy::InitializeBodiesPT(Chaos::FPBDRigidsSolver
 			{
 				if (Connections)
 				{
+					// first reserve memory for the connection so we alolocate just the right amount of memory 
+					for (int32 TransformGroupIndex = 0; TransformGroupIndex < NumTransforms; ++TransformGroupIndex)
+					{
+						if (FClusterHandle* ClusteredParticle = SolverParticleHandles[TransformGroupIndex])
+						{
+							const TSet<int32>& Siblings = (*Connections)[TransformGroupIndex];
+							if (Siblings.Num())
+							{
+								ClusteredParticle->ConnectivityEdges().Reserve(Siblings.Num());
+							}
+						}
+					}
+
+					// now set the connections
 					for (int32 TransformGroupIndex = 0; TransformGroupIndex < NumTransforms; ++TransformGroupIndex)
 					{
 						if (FClusterHandle* ClusteredParticle = SolverParticleHandles[TransformGroupIndex]) 
@@ -2131,7 +2145,7 @@ static void ApplyToChildrenAtPointWithRadiusAndPropagation_Internal(
 		if (ChildrenHandles.Num())
 		{
 			TArray<Chaos::FPBDRigidParticleHandle*> ProcessedHandles;
-			Chaos::FReal PropagationMultiplier = 1.0f;
+			Chaos::FRealSingle PropagationMultiplier = 1.0f;
 
 			int32 CurrentPropagationDepth = PropagationDepth;
 			while (CurrentPropagationDepth >= 0)
@@ -2169,7 +2183,7 @@ static void ApplyToChildrenAtPointWithRadiusAndPropagation_Internal(
 					}
 				}
 				
-				PropagationMultiplier *= static_cast<Chaos::FReal>(PropagationFactor);
+				PropagationMultiplier *= PropagationFactor;
 				--CurrentPropagationDepth;
 			}
 		}
@@ -2191,7 +2205,7 @@ void FGeometryCollectionPhysicsProxy::ApplyExternalStrain_External(FGeometryColl
 				ApplyToChildrenAtPointWithRadiusAndPropagation_Internal(
 					RBDSolver->GetEvolution()->GetRigidClustering(), *ClusteredHandle,
 					WorldLocation, Radius, PropagationDepth, PropagationFactor,
-					[StrainValue, &RigidClustering](Chaos::FPBDRigidClusteredParticleHandle* ClusteredChildHandle, Chaos::FReal PropagationMultiplier)
+					[StrainValue, &RigidClustering](Chaos::FPBDRigidClusteredParticleHandle* ClusteredChildHandle, Chaos::FRealSingle PropagationMultiplier)
 					{
 						RigidClustering.SetExternalStrain(ClusteredChildHandle, FMath::Max(ClusteredChildHandle->GetExternalStrain(), StrainValue * PropagationMultiplier));
 					});
@@ -2213,9 +2227,9 @@ void FGeometryCollectionPhysicsProxy::ApplyInternalStrain_External(FGeometryColl
 				ApplyToChildrenAtPointWithRadiusAndPropagation_Internal(
 					RBDSolver->GetEvolution()->GetRigidClustering(), *ClusteredHandle,
 					WorldLocation, Radius, PropagationDepth, PropagationFactor,
-					[StrainValue, RBDSolver](Chaos::FPBDRigidClusteredParticleHandle* ClusteredChildHandle, Chaos::FReal PropagationMultiplier)
+					[StrainValue, RBDSolver](Chaos::FPBDRigidClusteredParticleHandle* ClusteredChildHandle, Chaos::FRealSingle PropagationMultiplier)
 					{
-						const Chaos::FReal NewInternalStrain = ClusteredChildHandle->GetInternalStrains() - ((Chaos::FReal)StrainValue * PropagationMultiplier);
+						const Chaos::FRealSingle NewInternalStrain = ClusteredChildHandle->GetInternalStrains() - (StrainValue * PropagationMultiplier);
 						RBDSolver->GetEvolution()->GetRigidClustering().SetInternalStrain(ClusteredChildHandle, FMath::Max(0, NewInternalStrain));
 					});
 			}
diff --git a/Engine/Source/Runtime/Experimental/Chaos/Public/Chaos/PBDRigidClusteredParticles.h b/Engine/Source/Runtime/Experimental/Chaos/Public/Chaos/PBDRigidClusteredParticles.h
--- a/Engine/Source/Runtime/Experimental/Chaos/Public/Chaos/PBDRigidClusteredParticles.h
+++ b/Engine/Source/Runtime/Experimental/Chaos/Public/Chaos/PBDRigidClusteredParticles.h
@@ -46,7 +46,7 @@ struct TConnectivityEdge
 		, Strain(0.0)
 	{}
 
-	TConnectivityEdge(TPBDRigidParticleHandle<T, 3>* InSibling, const T InStrain)
+	TConnectivityEdge(TPBDRigidParticleHandle<T, 3>* InSibling, const FRealSingle InStrain)
 		: Sibling(InSibling)
 		, Strain(InStrain) 
 	{}
@@ -61,7 +61,7 @@ struct TConnectivityEdge
 	{ return Sibling == OtherSibling; }
 
 	TPBDRigidParticleHandle<T, 3>* Sibling;
-	T Strain;
+	Chaos::FRealSingle Strain;
 };
 typedef TConnectivityEdge<FReal> FConnectivityEdge;
 typedef TArray<FConnectivityEdge> FConnectivityEdgeArray;
@@ -217,14 +217,14 @@ class TPBDRigidClusteredParticles : public TPBDRigidParticles<T, d>
 	  TArrayCollectionArray<TSet<IPhysicsProxyBase*>> MPhysicsProxies;
 
 	  // Collision Impulses
-	  TArrayCollectionArray<T> MCollisionImpulses;
+	  TArrayCollectionArray<Chaos::FRealSingle> MCollisionImpulses;
 
 	  // external strains ( use by fields )
 	  // @todo(chaos) we should eventually merge MCollisionImpulses into MExternalStrains when Clustering code has been updated to not clear the impulses just before processing them 
-	  TArrayCollectionArray<T> MExternalStrains; 
+	  TArrayCollectionArray<Chaos::FRealSingle> MExternalStrains;
 
 	  // User set parameters
-	  TArrayCollectionArray<T> MStrains;
+	  TArrayCollectionArray<Chaos::FRealSingle> MStrains;
 
 	  TArrayCollectionArray<TArray<TConnectivityEdge<T>>> MConnectivityEdges;
   
diff --git a/Engine/Source/Runtime/Experimental/Chaos/Public/Chaos/PBDRigidClustering.h b/Engine/Source/Runtime/Experimental/Chaos/Public/Chaos/PBDRigidClustering.h
--- a/Engine/Source/Runtime/Experimental/Chaos/Public/Chaos/PBDRigidClustering.h
+++ b/Engine/Source/Runtime/Experimental/Chaos/Public/Chaos/PBDRigidClustering.h
@@ -250,8 +250,8 @@ public:
 	*    body in the simulation. This attribute is initialized during the creation of
 	*    the cluster body, can be updated during the evaluation of the simulation.
 	*/
-	TArrayCollectionArray<FReal>& GetStrainArray() { return MParticles.StrainsArray(); }
-	const TArrayCollectionArray<FReal>& GetStrainArray() const { return MParticles.StrainsArray(); }
+	TArrayCollectionArray<FRealSingle>& GetStrainArray() { return MParticles.StrainsArray(); }
+	const TArrayCollectionArray<FRealSingle>& GetStrainArray() const { return MParticles.StrainsArray(); }
 		
 	/**
 	*  GetParentToChildren
@@ -365,8 +365,8 @@ public:
 	FRigidEvolution& GetEvolution() { return MEvolution; }
 	const FRigidEvolution& GetEvolution() const { return MEvolution; }
 
-	void SetInternalStrain(FPBDRigidClusteredParticleHandle* Particle, FReal Strain);
-	void SetExternalStrain(FPBDRigidClusteredParticleHandle* Particle, FReal Strain);
+	void SetInternalStrain(FPBDRigidClusteredParticleHandle* Particle, FRealSingle Strain);
+	void SetExternalStrain(FPBDRigidClusteredParticleHandle* Particle, FRealSingle Strain);
 
 	static bool ShouldUnionsHaveCollisionParticles();
 
diff --git a/Engine/Source/Runtime/Experimental/Chaos/Public/Chaos/ParticleHandle.h b/Engine/Source/Runtime/Experimental/Chaos/Public/Chaos/ParticleHandle.h
--- a/Engine/Source/Runtime/Experimental/Chaos/Public/Chaos/ParticleHandle.h
+++ b/Engine/Source/Runtime/Experimental/Chaos/Public/Chaos/ParticleHandle.h
@@ -1281,31 +1281,31 @@ public:
 	TUniquePtr<FImplicitObjectUnionClustered>& ChildrenSpatial() { return PBDRigidClusteredParticles->ChildrenSpatial(ParticleIdx); }
 	void SetChildrenSpatial(TUniquePtr<FImplicitObjectUnion>& Obj) { PBDRigidClusteredParticles->ChildrenSpatial(ParticleIdx) = Obj; }
 
-	const T& CollisionImpulse() const { return PBDRigidClusteredParticles->CollisionImpulses(ParticleIdx); }
-	T& CollisionImpulse() { return PBDRigidClusteredParticles->CollisionImpulses(ParticleIdx); }
-	void SetCollisionImpulse(const T Value) { PBDRigidClusteredParticles->CollisionImpulses(ParticleIdx) = Value; }
-	const T& CollisionImpulses() const { return PBDRigidClusteredParticles->CollisionImpulses(ParticleIdx); }
-	T& CollisionImpulses() { return PBDRigidClusteredParticles->CollisionImpulses(ParticleIdx); }
-	void SetCollisionImpulses(const T Value) { PBDRigidClusteredParticles->CollisionImpulses(ParticleIdx) = Value; }
-	void ClearCollisionImpulse() { PBDRigidClusteredParticles->CollisionImpulses(ParticleIdx) = static_cast<T>(0); }
-
-	T GetExternalStrain() const { return PBDRigidClusteredParticles->ExternalStrains(ParticleIdx); }
+	const FRealSingle& CollisionImpulse() const { return PBDRigidClusteredParticles->CollisionImpulses(ParticleIdx); }
+	FRealSingle& CollisionImpulse() { return PBDRigidClusteredParticles->CollisionImpulses(ParticleIdx); }
+	void SetCollisionImpulse(const FRealSingle Value) { PBDRigidClusteredParticles->CollisionImpulses(ParticleIdx) = Value; }
+	const FRealSingle& CollisionImpulses() const { return PBDRigidClusteredParticles->CollisionImpulses(ParticleIdx); }
+	FRealSingle& CollisionImpulses() { return PBDRigidClusteredParticles->CollisionImpulses(ParticleIdx); }
+	void SetCollisionImpulses(const FRealSingle Value) { PBDRigidClusteredParticles->CollisionImpulses(ParticleIdx) = Value; }
+	void ClearCollisionImpulse() { PBDRigidClusteredParticles->CollisionImpulses(ParticleIdx) = static_cast<FRealSingle>(0); }
+
+	FRealSingle GetExternalStrain() const { return PBDRigidClusteredParticles->ExternalStrains(ParticleIdx); }
 	UE_DEPRECATED(5.2, "This method should not be used anymore. FRigidClustering::SetExternalStrain should be used instead.")
-	void SetExternalStrain(const T Value) { PBDRigidClusteredParticles->ExternalStrains(ParticleIdx) = Value; }
+	void SetExternalStrain(const FRealSingle Value) { PBDRigidClusteredParticles->ExternalStrains(ParticleIdx) = Value; }
 	UE_DEPRECATED(5.2, "This method should not be used anymore. FRigidClustering::SetExternalStrain with 0 strain should be used instead.")
-	void ClearExternalStrain() { PBDRigidClusteredParticles->ExternalStrains(ParticleIdx) = static_cast<T>(0); }
+	void ClearExternalStrain() { PBDRigidClusteredParticles->ExternalStrains(ParticleIdx) = static_cast<FRealSingle>(0); }
 	
-	const T& GetInternalStrains() const { return PBDRigidClusteredParticles->Strains(ParticleIdx); }
-	const T& Strain() const { return PBDRigidClusteredParticles->Strains(ParticleIdx); }
+	const FRealSingle& GetInternalStrains() const { return PBDRigidClusteredParticles->Strains(ParticleIdx); }
+	const FRealSingle& Strain() const { return PBDRigidClusteredParticles->Strains(ParticleIdx); }
 	UE_DEPRECATED(5.2, "This method should not be used anymore. FRigidClustering::SetInternalStrain should be used instead.")
-	T& Strain() { return PBDRigidClusteredParticles->Strains(ParticleIdx); }
+	FRealSingle& Strain() { return PBDRigidClusteredParticles->Strains(ParticleIdx); }
 	UE_DEPRECATED(5.2, "This method should not be used anymore. FRigidClustering::SetInternalStrain should be used instead.")
-	void SetStrain(const T Value) { PBDRigidClusteredParticles->Strains(ParticleIdx) = Value; }
-	const T& Strains() const { return PBDRigidClusteredParticles->Strains(ParticleIdx); }
+	void SetStrain(const FRealSingle Value) { PBDRigidClusteredParticles->Strains(ParticleIdx) = Value; }
+	const FRealSingle& Strains() const { return PBDRigidClusteredParticles->Strains(ParticleIdx); }
 	UE_DEPRECATED(5.2, "This method should not be used anymore. FRigidClustering::SetInternalStrain should be used instead.")
-	T& Strains() { return PBDRigidClusteredParticles->Strains(ParticleIdx); }
+	FRealSingle& Strains() { return PBDRigidClusteredParticles->Strains(ParticleIdx); }
 	UE_DEPRECATED(5.2, "This method should not be used anymore. FRigidClustering::SetInternalStrain should be used instead.")
-	void SetStrains(const T Value) { PBDRigidClusteredParticles->Strains(ParticleIdx) = Value; }
+	void SetStrains(const FRealSingle Value) { PBDRigidClusteredParticles->Strains(ParticleIdx) = Value; }
 	void SetMaximumInternalStrain() { PBDRigidClusteredParticles->Strains(ParticleIdx) = MaxStrain; }
 
 	const TArray<TConnectivityEdge<T>>& ConnectivityEdges() const { return PBDRigidClusteredParticles->ConnectivityEdges(ParticleIdx); }
@@ -1327,10 +1327,10 @@ public:
 	int32 TransientParticleIndex() const { return ParticleIdx; }
 
 private:
-	void SetInternalStrains(const T Value) { PBDRigidClusteredParticles->Strains(ParticleIdx) = Value; }
-	void SetExternalStrains(const T Value) { PBDRigidClusteredParticles->ExternalStrains(ParticleIdx) = Value; }
+	void SetInternalStrains(const FRealSingle Value) { PBDRigidClusteredParticles->Strains(ParticleIdx) = Value; }
+	void SetExternalStrains(const FRealSingle Value) { PBDRigidClusteredParticles->ExternalStrains(ParticleIdx) = Value; }
 	friend class FRigidClustering;
-	static constexpr Chaos::FReal MaxStrain = TNumericLimits<Chaos::FReal>::Max() - TNumericLimits<Chaos::FReal>::Min();
+	static constexpr Chaos::FRealSingle MaxStrain = TNumericLimits<Chaos::FRealSingle>::Max() - TNumericLimits<Chaos::FRealSingle>::Min();
 };
 
 template <typename T, int d, bool bPersistent = true>
diff --git a/Engine/Source/Runtime/Experimental/Chaos/Public/PhysicsProxy/FieldSystemProxyHelper.h b/Engine/Source/Runtime/Experimental/Chaos/Public/PhysicsProxy/FieldSystemProxyHelper.h
--- a/Engine/Source/Runtime/Experimental/Chaos/Public/PhysicsProxy/FieldSystemProxyHelper.h
+++ b/Engine/Source/Runtime/Experimental/Chaos/Public/PhysicsProxy/FieldSystemProxyHelper.h
@@ -503,7 +503,7 @@ namespace Chaos
 					{
 						if (Chaos::FPBDRigidClusteredParticleHandle* ClusteredParticle = ParticleHandles[Index.Sample]->CastToClustered())
 						{
-							const FReal CurrentExternalStrains = ClusteredParticle->GetExternalStrain();
+							const FRealSingle CurrentExternalStrains = ClusteredParticle->GetExternalStrain();
 							RigidClustering.SetExternalStrain(ClusteredParticle, FMath::Max(CurrentExternalStrains, ExternalStrainValue));
 						}
 					}
