optimize the replication of child to parent of cluster unions to do it in bulk and only re-create the shared geometry once

diff --git a/Engine/Source/Runtime/Engine/Private/PhysicsEngine/ClusterUnionComponent.cpp b/Engine/Source/Runtime/Engine/Private/PhysicsEngine/ClusterUnionComponent.cpp
--- a/Engine/Source/Runtime/Engine/Private/PhysicsEngine/ClusterUnionComponent.cpp
+++ b/Engine/Source/Runtime/Engine/Private/PhysicsEngine/ClusterUnionComponent.cpp
@@ -555,11 +555,16 @@ void UClusterUnionComponent::ForceSetChildToParent(UPrimitiveComponent* InCompon
 		return;
 	}
 
+	TArray< Chaos::FPhysicsObjectHandle> Objects;
+	Objects.Reserve(BoneIds.Num());
+
 	for (int32 Index = 0; Index < BoneIds.Num(); ++Index)
 	{
 		Chaos::FPhysicsObjectHandle Handle = InComponent->GetPhysicsObjectById(BoneIds[Index]);
-		PhysicsProxy->SetChildToParent_External(Handle, ChildToParent[Index]);
+		Objects.Add(Handle);
 	}
+
+	PhysicsProxy->BulkSetChildToParent_External(Objects, ChildToParent);
 }
 
 void UClusterUnionComponent::SetSimulatePhysics(bool bSimulate)
diff --git a/Engine/Source/Runtime/Experimental/Chaos/Private/PhysicsProxy/ClusterUnionPhysicsProxy.cpp b/Engine/Source/Runtime/Experimental/Chaos/Private/PhysicsProxy/ClusterUnionPhysicsProxy.cpp
--- a/Engine/Source/Runtime/Experimental/Chaos/Private/PhysicsProxy/ClusterUnionPhysicsProxy.cpp
+++ b/Engine/Source/Runtime/Experimental/Chaos/Private/PhysicsProxy/ClusterUnionPhysicsProxy.cpp
@@ -526,36 +526,46 @@ namespace Chaos
 
 	void FClusterUnionPhysicsProxy::SetChildToParent_External(FPhysicsObjectHandle Child, const FTransform& RelativeTransform)
 	{
-		if (!Solver || !ensure(Particle_External))
+		BulkSetChildToParent_External({ Child }, { RelativeTransform });
+	}
+
+	void FClusterUnionPhysicsProxy::BulkSetChildToParent_External(const TArray<FPhysicsObjectHandle>& Objects, const TArray<FTransform>& Transforms)
+	{
+		if (!Solver || !ensure(Particle_External) || !ensure(Objects.Num() == Transforms.Num()))
 		{
 			return;
 		}
 
 		Solver->EnqueueCommandImmediate(
-			[this, Child, RelativeTransform]() mutable
+			[this, Objects, Transforms]() mutable
 			{
 				FPBDRigidsEvolutionGBF& Evolution = *GetEvolution(this);
 				FClusterUnionManager& ClusterUnionManager = Evolution.GetRigidClustering().GetClusterUnionManager();
 
 				if (FClusterUnion* ClusterUnion = ClusterUnionManager.FindClusterUnion(ClusterUnionIndex))
 				{
-					FGeometryParticleHandle* Particle = Child->GetParticle<EThreadContext::Internal>();
-					if (!ensure(Particle))
+					for (int32 Index = 0; Index < Objects.Num() && Index < Transforms.Num(); ++Index)
 					{
-						return;
-					}
+						FPhysicsObjectHandle Child = Objects[Index];
+						FGeometryParticleHandle* Particle = Child->GetParticle<EThreadContext::Internal>();
+						if (!ensure(Particle))
+						{
+							return;
+						}
 
-					const int32 ParticleIndex = ClusterUnion->ChildParticles.Find(Particle->CastToRigidParticle());
-					if (ParticleIndex != INDEX_NONE && ClusterUnion->InternalCluster)
-					{
-						if (FPBDRigidClusteredParticleHandle* ChildHandle = ClusterUnion->ChildParticles[ParticleIndex]->CastToClustered())
+						const int32 ParticleIndex = ClusterUnion->ChildParticles.Find(Particle->CastToRigidParticle());
+						if (ParticleIndex != INDEX_NONE && ClusterUnion->InternalCluster)
 						{
-							ChildHandle->SetChildToParent(RelativeTransform);
-							ClusterUnionManager.UpdateAllClusterUnionProperties(*ClusterUnion, false);
-							Evolution.GetParticles().MarkTransientDirtyParticle(ClusterUnion->InternalCluster);
-							Evolution.DirtyParticle(*ClusterUnion->InternalCluster);
+							if (FPBDRigidClusteredParticleHandle* ChildHandle = ClusterUnion->ChildParticles[ParticleIndex]->CastToClustered())
+							{
+								ChildHandle->SetChildToParent(Transforms[Index]);
+							}
 						}
 					}
+
+					ClusterUnionManager.UpdateAllClusterUnionProperties(*ClusterUnion, false);
+					Evolution.GetParticles().MarkTransientDirtyParticle(ClusterUnion->InternalCluster);
+					Evolution.DirtyParticle(*ClusterUnion->InternalCluster);
 				}
 			}
 		);
diff --git a/Engine/Source/Runtime/Experimental/Chaos/Public/PhysicsProxy/ClusterUnionPhysicsProxy.h b/Engine/Source/Runtime/Experimental/Chaos/Public/PhysicsProxy/ClusterUnionPhysicsProxy.h
--- a/Engine/Source/Runtime/Experimental/Chaos/Public/PhysicsProxy/ClusterUnionPhysicsProxy.h
+++ b/Engine/Source/Runtime/Experimental/Chaos/Public/PhysicsProxy/ClusterUnionPhysicsProxy.h
@@ -92,6 +92,7 @@ namespace Chaos
 		void SetLinearVelocity_External(const FVector& V);
 		void SetAngularVelocity_External(const FVector& W);
 		void SetChildToParent_External(FPhysicsObjectHandle Child, const FTransform& RelativeTransform);
+		void BulkSetChildToParent_External(const TArray<FPhysicsObjectHandle>& Objects, const TArray<FTransform>& Transforms);
 
 		//
 		// These functions take care of marshaling data back and forth between the game thread
