breaking a child particle from its parent due to strain must also account for the situation where it's actually trying to break apart from a cluster union instead

diff --git a/Engine/Source/Runtime/Experimental/Chaos/Private/Chaos/PBDRigidClustering.cpp b/Engine/Source/Runtime/Experimental/Chaos/Private/Chaos/PBDRigidClustering.cpp
--- a/Engine/Source/Runtime/Experimental/Chaos/Private/Chaos/PBDRigidClustering.cpp
+++ b/Engine/Source/Runtime/Experimental/Chaos/Private/Chaos/PBDRigidClustering.cpp
@@ -1,5 +1,4 @@
 // Copyright Epic Games, Inc. All Rights Reserved.
-
 #include "Chaos/PBDRigidClustering.h"
 
 #include "Chaos/ErrorReporter.h"
@@ -729,6 +728,10 @@ namespace Chaos
 		// only used for propagation
 		TMap<FPBDRigidParticleHandle*, FRealSingle> AppliedStrains;
 
+		// We'll pass these particles to the cluster union manager to remove. This can't be done within the same loop
+		// since it'll be modifying the children array.
+		TArray<FPBDRigidParticleHandle*> DeferredRemoveFromClusterUnion;
+
 		for (int32 ChildIdx = Children.Num() - 1; ChildIdx >= 0; --ChildIdx)
 		{
 			FPBDRigidClusteredParticleHandle* Child = Children[ChildIdx]->CastToClustered();
@@ -748,21 +751,30 @@ namespace Chaos
 					bFoundFirstRelease = true;
 				}
 
+				// There's a possibility that the child is in a cluster union so we'd need to be able to remove the child particle from the cluster union as well.
+				const FClusterUnionIndex ClusterUnionIndex = ClusterUnionManager.FindClusterUnionIndexFromParticle(Child);
+				const bool bIsInClusterUnion = ClusterUnionIndex != INDEX_NONE;
+
 				//UE_LOG(LogTemp, Warning, TEXT("Releasing child %d from parent %p due to strain %.5f Exceeding internal strain %.5f (Source: %s)"), ChildIdx, ClusteredParticle, ChildStrain, Child->Strain(), bForceRelease ? TEXT("Forced by caller") : ExternalStrainMap ? TEXT("External") : TEXT("Collision"));
+				if (bIsInClusterUnion)
+				{
+					DeferredRemoveFromClusterUnion.Add(Child);
+				}
+				else
+				{
+					// The piece that hits just breaks off - we may want more control 
+					// by looking at the edges of this piece which would give us cleaner 
+					// breaks (this approach produces more rubble)
+					RemoveChildFromParent(Child, ClusteredParticle);
+					UpdateTopLevelParticle(Child);
+
+					// Remove from the children array without freeing memory yet. 
+					// We're looping over Children and it'd be silly to free the array
+					// 1 entry at a time.
+					Children.RemoveAtSwap(ChildIdx, 1, false);
+				}
 				
-				// The piece that hits just breaks off - we may want more control 
-				// by looking at the edges of this piece which would give us cleaner 
-				// breaks (this approach produces more rubble)
-				RemoveChildFromParent(Child, ClusteredParticle);
 				ActivatedChildren.Add(Child);
-
-				UpdateTopLevelParticle(Child);
-
-				// Remove from the children array without freeing memory yet. 
-				// We're looping over Children and it'd be silly to free the array
-				// 1 entry at a time.
-				Children.RemoveAtSwap(ChildIdx, 1, false);
-
 				SendBreakingEvent(Child, bParentCrumbled);
 			}
 			if (bUseDamagePropagation)
@@ -772,6 +784,11 @@ namespace Chaos
 			Child->SetExternalStrains(0.0);
 		}
 
+		if (!DeferredRemoveFromClusterUnion.IsEmpty())
+		{
+			ClusterUnionManager.HandleRemoveOperationWithClusterLookup(DeferredRemoveFromClusterUnion, true);
+		}
+
 		// if necessary propagate strain through the graph
 		// IMPORTANT: this assumes that the connectivity graph has not yet been updated from pieces that broke off
 		if (bUseDamagePropagation)
