Geometry collection : fix TangentU and TangentV  not being used in when processing the render data

diff --git a/Engine/Source/Runtime/Experimental/GeometryCollectionEngine/Private/GeometryCollection/GeometryCollectionRenderData.cpp b/Engine/Source/Runtime/Experimental/GeometryCollectionEngine/Private/GeometryCollection/GeometryCollectionRenderData.cpp
--- a/Engine/Source/Runtime/Experimental/GeometryCollectionEngine/Private/GeometryCollection/GeometryCollectionRenderData.cpp
+++ b/Engine/Source/Runtime/Experimental/GeometryCollectionEngine/Private/GeometryCollection/GeometryCollectionRenderData.cpp
@@ -227,8 +227,8 @@ void BuildMeshDataFromGeometryCollection(FGeometryCollection& InCollection, FGeo
 			ParallelFor(TEXT("GC:BuildVertices"), VertexCount, 500, [&](int32 VertexIndex)
 			{
 				BuildVertexData.Position[DestVertexStart + VertexIndex] = VertexArray[VertexStart + VertexIndex];
-				BuildVertexData.TangentX[DestVertexStart + VertexIndex] = FVector3f::ZeroVector;
-				BuildVertexData.TangentY[DestVertexStart + VertexIndex] = FVector3f::ZeroVector;
+				BuildVertexData.TangentX[DestVertexStart + VertexIndex] = TangentUArray[VertexStart + VertexIndex];
+				BuildVertexData.TangentY[DestVertexStart + VertexIndex] = TangentVArray[VertexStart + VertexIndex];
 				BuildVertexData.TangentZ[DestVertexStart + VertexIndex] = NormalArray[VertexStart + VertexIndex];
 
 				if (bHasColors)
