Chaos - really fix debug draw resource accumulation etc

diff --git a/Engine/Source/Runtime/Experimental/ChaosSolverEngine/Private/Chaos/ChaosDebugDrawComponent.cpp b/Engine/Source/Runtime/Experimental/ChaosSolverEngine/Private/Chaos/ChaosDebugDrawComponent.cpp
--- a/Engine/Source/Runtime/Experimental/ChaosSolverEngine/Private/Chaos/ChaosDebugDrawComponent.cpp
+++ b/Engine/Source/Runtime/Experimental/ChaosSolverEngine/Private/Chaos/ChaosDebugDrawComponent.cpp
@@ -53,8 +53,21 @@ FAutoConsoleVariableRef CVarArrowSize(TEXT("p.Chaos.DebugDraw.Mode"), bChaosDebu
 
 float CommandLifeTime(const Chaos::FLatentDrawCommand& Command, const bool bIsPaused)
 {
-	// @todo(chaos): remove lifetime from the system
-	return ((Command.LifeTime > 0.0f) && !bIsPaused) ? Command.LifeTime : 0.0f;
+	// The linebatch time handling is a bit awkward and we need to translate
+	// Linebatcher Lifetime < 0: eternal (regardless of bPersistent flag)
+	// Linebatcher Lifetime = 0: default lifetime (which is usually 1 second)
+	// Linebatcher Lifetime > 0: specified duration
+	// Whereas in our Command, 
+	// Command lifetime <= 0: 1 frame
+	// Command lifetime > 0: specified duration
+	if (Command.LifeTime <= 0)
+	{
+		// One frame - must be non-zero but also less than the next frame time which we don't know
+		// NOTE: this only works because UChaosDebugDrawComponent ticks after the line batcher
+		return UE_SMALL_NUMBER;
+	}
+	
+	return Command.LifeTime;
 }
 
 void DebugDrawChaos(AActor* DebugDrawActor, const TArray<Chaos::FLatentDrawCommand>& DrawCommands, bool bIsPaused)
@@ -186,7 +199,6 @@ void DebugDrawChaos(AActor* DebugDrawActor, const TArray<Chaos::FLatentDrawComma
 
 UChaosDebugDrawComponent::UChaosDebugDrawComponent()
 	: bInPlay(false)
-	, LastRenderTime(-1.0)
 {
 	// We must tick after anything that uses Chaos Debug Draw and also after the Line Batcher Component
 	PrimaryComponentTick.bAllowTickOnDedicatedServer = false;
@@ -212,7 +224,6 @@ void UChaosDebugDrawComponent::BeginPlay()
 	SetTickableWhenPaused(true);
 
 	bInPlay = true;
-	LastRenderTime = -1.0;
 
 #if CHAOS_DEBUG_DRAW
 	Chaos::FDebugDrawQueue::GetInstance().SetConsumerActive(this, bInPlay);
@@ -285,17 +296,10 @@ void UChaosDebugDrawComponent::TickComponent(float DeltaTime, enum ELevelTick Ti
 
 		FDebugDrawQueue::GetInstance().SetMaxCost(ChaosDebugDraw_MaxElements);
 
-		// Get the latest commands unless we are paused (in which case we redraw the previous ones)
-		if (!World->IsPaused())
+		const bool bIsPaused = World->IsPaused();
+		if (!bIsPaused)
 		{
 			FDebugDrawQueue::GetInstance().ExtractAllElements(DrawCommands);
-		}
-
-		// Don't enqueue more lines if the renderer hasn't run since the last time we were ticked.
-		// This prevents us from accumulated lines in the line batcher when in the background, for example.
-		if (LastRenderTime < World->LastRenderTime)
-		{
-			LastRenderTime = World->LastRenderTime;
 
 			DebugDrawChaos(GetOwner(), DrawCommands, World->IsPaused());
 		}
diff --git a/Engine/Source/Runtime/Experimental/ChaosSolverEngine/Public/Chaos/ChaosDebugDrawComponent.h b/Engine/Source/Runtime/Experimental/ChaosSolverEngine/Public/Chaos/ChaosDebugDrawComponent.h
--- a/Engine/Source/Runtime/Experimental/ChaosSolverEngine/Public/Chaos/ChaosDebugDrawComponent.h
+++ b/Engine/Source/Runtime/Experimental/ChaosSolverEngine/Public/Chaos/ChaosDebugDrawComponent.h
@@ -33,7 +33,6 @@ private:
 	static void CreateDebugDrawActor(UWorld* World);
 
 	bool bInPlay;
-	double LastRenderTime;
 
 #if CHAOS_DEBUG_DRAW
 	TArray<Chaos::FLatentDrawCommand> DrawCommands;
