Chaos - disallow collision manifold restore if the manifold was reset

Fixes issue where calling SetTransform on a convex body resting on another convex body would attempt to restore a manifold that had been reset and therefore had no collisions, leading to one frame with no conacts. (The manifold is reset on user SetTransform so that friction does not try to force the body back to its original position.)

diff --git a/Engine/Source/Runtime/Experimental/Chaos/Private/Chaos/Collision/PBDCollisionConstraint.cpp b/Engine/Source/Runtime/Experimental/Chaos/Private/Chaos/Collision/PBDCollisionConstraint.cpp
--- a/Engine/Source/Runtime/Experimental/Chaos/Private/Chaos/Collision/PBDCollisionConstraint.cpp
+++ b/Engine/Source/Runtime/Experimental/Chaos/Private/Chaos/Collision/PBDCollisionConstraint.cpp
@@ -674,9 +674,10 @@ namespace Chaos
 		ManifoldPointResults.Reset();
 		ExpectedNumManifoldPoints = 0;
 		Flags.bWasManifoldRestored = false;
+		Flags.bCanRestoreManifold = false;
 	}
 
-	bool FPBDCollisionConstraint::UpdateAndTryRestoreManifold()
+	bool FPBDCollisionConstraint::TryRestoreManifold()
 	{
 		const FCollisionTolerances Tolerances = FCollisionTolerances();//Chaos_Manifold_Tolerances;
 		const FReal ContactPositionTolerance = Tolerances.ContactPositionToleranceScale * CollisionTolerance;
@@ -692,6 +693,11 @@ namespace Chaos
 		ExpectedNumManifoldPoints = ManifoldPoints.Num();
 		Flags.bWasManifoldRestored = false;
 
+		if (!Flags.bCanRestoreManifold)
+		{
+			return false;
+		}
+
 		// If we have not moved or rotated much we may reuse some of the manifold points, as long as they have not moved far as well (see below)
 		bool bMovedBeyondTolerance = true;
 		if ((ShapePositionTolerance > 0) && (ShapeRotationThreshold > 0))
@@ -715,7 +721,6 @@ namespace Chaos
 
 		if (bMovedBeyondTolerance)
 		{
-			ResetActiveManifoldContacts();
 			return false;
 		}
 
@@ -779,7 +784,6 @@ namespace Chaos
 				else
 				{
 					// We want to remove a(nother) point, but we will never reuse the manifold now so throw it away
-					ResetActiveManifoldContacts();
 					return false;
 				}
 			}
diff --git a/Engine/Source/Runtime/Experimental/Chaos/Private/Chaos/Collision/ParticlePairMidPhase.cpp b/Engine/Source/Runtime/Experimental/Chaos/Private/Chaos/Collision/ParticlePairMidPhase.cpp
--- a/Engine/Source/Runtime/Experimental/Chaos/Private/Chaos/Collision/ParticlePairMidPhase.cpp
+++ b/Engine/Source/Runtime/Experimental/Chaos/Private/Chaos/Collision/ParticlePairMidPhase.cpp
@@ -325,29 +325,30 @@ namespace Chaos
 			}
 			
 			bool bWasManifoldRestored = false;
-			if (Context.GetSettings().bAllowManifoldReuse && Flags.bEnableManifoldUpdate && bWasUpdatedLastTick)
+			const bool bAllowManifoldRestore = Context.GetSettings().bAllowManifoldReuse && Flags.bEnableManifoldUpdate;
+			if (bAllowManifoldRestore && bWasUpdatedLastTick && Constraint->GetCanRestoreManifold())
 			{
 				// Update the existing manifold. We can re-use as-is if none of the points have moved much and the bodies have not moved much
 				// NOTE: this can succeed in "restoring" even if we have no manifold points
 				// NOTE: this uses the transforms from SetLastShapeWorldTransforms, so we can only do this if we were updated last tick
-				bWasManifoldRestored = Constraint->UpdateAndTryRestoreManifold();
-			}
-			else
-			{
-				// We are not trying to reuse manifold points, so reset them but leave stored data intact (for friction)
-				Constraint->ResetActiveManifoldContacts();
+				bWasManifoldRestored = Constraint->TryRestoreManifold();
 			}
 
 			if (!bWasManifoldRestored)
 			{
-				// We will be updating the manifold so update transforms used to check for movement in UpdateAndTryRestoreManifold on future ticks
-				Constraint->SetLastShapeWorldTransforms(ShapeWorldTransform0, ShapeWorldTransform1);
+				// We are not trying to (or chose not to) reuse manifold points, so reset them but leave stored data intact (for friction)
+				Constraint->ResetActiveManifoldContacts();
 
 				if (!Context.GetSettings().bDeferNarrowPhase)
 				{
 					// Run the narrow phase
 					Collisions::UpdateConstraint(*Constraint.Get(), ShapeWorldTransform0, ShapeWorldTransform1, Dt);
 				}
+
+				// We will be updating the manifold so update transforms used to check for movement in UpdateAndTryRestoreManifold on future ticks
+				// NOTE: We call this after Collisions::UpdateConstraint because it may reset the manifold and reset the bCanRestoreManifold flag.
+				// @todo(chaos): Collisions::UpdateConstraint does not need to reset the manifold - fix that
+				Constraint->SetLastShapeWorldTransforms(ShapeWorldTransform0, ShapeWorldTransform1);
 			}
 
 			// If we have a valid contact, add it to the active list
diff --git a/Engine/Source/Runtime/Experimental/Chaos/Public/Chaos/Collision/PBDCollisionConstraint.h b/Engine/Source/Runtime/Experimental/Chaos/Public/Chaos/Collision/PBDCollisionConstraint.h
--- a/Engine/Source/Runtime/Experimental/Chaos/Public/Chaos/Collision/PBDCollisionConstraint.h
+++ b/Engine/Source/Runtime/Experimental/Chaos/Public/Chaos/Collision/PBDCollisionConstraint.h
@@ -545,13 +545,17 @@ namespace Chaos
 			ShapeWorldTransforms[1] = InShapeWorldTransform1;
 		}
 
+		// Set the transforms when we last ran collision detection. This also sets the bCanRestoreManifold flag which
+		// allows the use of UpdateAndTryRestoreManifold on the next tick.
 		void SetLastShapeWorldTransforms(const FRigidTransform3& InShapeWorldTransform0, const FRigidTransform3& InShapeWorldTransform1)
 		{
 			LastShapeWorldPositionDelta = InShapeWorldTransform0.GetTranslation() - InShapeWorldTransform1.GetTranslation();
 			LastShapeWorldRotationDelta = InShapeWorldTransform0.GetRotation().Inverse() * InShapeWorldTransform1.GetRotation();
+			Flags.bCanRestoreManifold = true;
 		}
 
-		bool UpdateAndTryRestoreManifold();
+		bool GetCanRestoreManifold() const { return Flags.bCanRestoreManifold; }
+		bool TryRestoreManifold();
 		void ResetActiveManifoldContacts();
 		bool TryAddManifoldContact(const FContactPoint& ContactPoint);
 		bool TryInsertManifoldContact(const FContactPoint& ContactPoint);
@@ -696,6 +700,10 @@ namespace Chaos
 		const FPBDCollisionConstraintHandle* GetConstraintHandle() const { return this; }
 		FPBDCollisionConstraintHandle* GetConstraintHandle() { return this; }
 
+
+		UE_DEPRECATED(5.3, "Use TryRestoreManifold")
+		bool UpdateAndTryRestoreManifold() { return TryRestoreManifold(); }
+
 	protected:
 
 		FPBDCollisionConstraint(
@@ -764,18 +772,19 @@ namespace Chaos
 			FFlags() : Bits(0) {}
 			struct
 			{
-				uint32 bDisabled : 1;
-				uint32 bUseManifold : 1;
-				uint32 bUseIncrementalManifold : 1;
-				uint32 bWasManifoldRestored : 1;
-				uint32 bIsQuadratic0 : 1;
-				uint32 bIsQuadratic1 : 1;
-				uint32 bIsProbeUnmodified : 1;  // Is this constraint a probe pre-contact-modification
-				uint32 bIsProbe : 1;            // Is this constraint currently a probe
-				uint32 bCCDEnabled : 1;			// Is CCD enabled for the current tick
-				uint32 bCCDSweepEnabled: 1;		// If this is a CCD constraint, do we want to enable the sweep/rewind phase?
-				uint32 bModifierApplied : 1;	// Was a constraint modifier applied this tick
-				uint32 bMaterialSet : 1;		// Has the material been set (or does it need to be reset)
+				uint32 bDisabled : 1;					// Is this contact disabled (by the user or because cull distance is exceeded)
+				uint32 bUseManifold : 1;				// Should we use contact manifolds or single points (faster but poor behaviour)
+				uint32 bUseIncrementalManifold : 1;		// Do we need to run incremental collision detection (only LavelSets now)
+				uint32 bCanRestoreManifold : 1;			// Can we try to restore the manifold this frame (set folowing narrowphase, cleared when reset for some reason)
+				uint32 bWasManifoldRestored : 1;		// Did we restore the manifold this frame
+				uint32 bIsQuadratic0 : 1;				// Is the first shape a sphere or capsule
+				uint32 bIsQuadratic1 : 1;				// Is the second shape a sphere or capsule
+				uint32 bIsProbeUnmodified : 1;			// Is this constraint a probe pre-contact-modification
+				uint32 bIsProbe : 1;					// Is this constraint currently a probe
+				uint32 bCCDEnabled : 1;					// Is CCD enabled for the current tick
+				uint32 bCCDSweepEnabled: 1;				// If this is a CCD constraint, do we want to enable the sweep/rewind phase?
+				uint32 bModifierApplied : 1;			// Was a constraint modifier applied this tick
+				uint32 bMaterialSet : 1;				// Has the material been set (or does it need to be reset)
 			};
 			uint32 Bits;
 		};
