Fix static analysis issue in geometry collection physics proxy

diff --git a/Engine/Source/Runtime/Experimental/Chaos/Private/PhysicsProxy/GeometryCollectionPhysicsProxy.cpp b/Engine/Source/Runtime/Experimental/Chaos/Private/PhysicsProxy/GeometryCollectionPhysicsProxy.cpp
--- a/Engine/Source/Runtime/Experimental/Chaos/Private/PhysicsProxy/GeometryCollectionPhysicsProxy.cpp
+++ b/Engine/Source/Runtime/Experimental/Chaos/Private/PhysicsProxy/GeometryCollectionPhysicsProxy.cpp
@@ -4509,50 +4509,51 @@ void FGeometryCollectionPhysicsProxy::FieldParameterUpdateCallback(Chaos::FPBDRi
 void FGeometryCollectionPhysicsProxy::FieldForcesUpdateCallback(Chaos::FPBDRigidsSolver* RigidSolver)
 {
 	SCOPE_CYCLE_COUNTER(STAT_ForceUpdateField_Object);
-	ensure(RigidSolver);
-
-	const int32 NumCommands = Commands.Num();
-	if (NumCommands && !RigidSolver->IsShuttingDown())
+	if (ensure(RigidSolver))
 	{
-		TArray<int32> CommandsToRemove;
-		CommandsToRemove.Reserve(NumCommands);
+		const int32 NumCommands = Commands.Num();
+		if (NumCommands && !RigidSolver->IsShuttingDown())
+		{
+			TArray<int32> CommandsToRemove;
+			CommandsToRemove.Reserve(NumCommands);
 
-		EFieldResolutionType PrevResolutionType = EFieldResolutionType::Field_Resolution_Max;
-		EFieldFilterType PrevFilterType = EFieldFilterType::Field_Filter_Max;
-		EFieldObjectType PrevObjectType = EFieldObjectType::Field_Object_Max;
-		EFieldPositionType PrevPositionType = EFieldPositionType::Field_Position_Max;
+			EFieldResolutionType PrevResolutionType = EFieldResolutionType::Field_Resolution_Max;
+			EFieldFilterType PrevFilterType = EFieldFilterType::Field_Filter_Max;
+			EFieldObjectType PrevObjectType = EFieldObjectType::Field_Object_Max;
+			EFieldPositionType PrevPositionType = EFieldPositionType::Field_Position_Max;
 
-		for (int32 CommandIndex = 0; CommandIndex < NumCommands; CommandIndex++)
-		{
-			const FFieldSystemCommand& FieldCommand = Commands[CommandIndex];
-			if (IsForceFieldValid(FieldCommand))
+			for (int32 CommandIndex = 0; CommandIndex < NumCommands; CommandIndex++)
 			{
-				if (Chaos::BuildFieldSamplePoints(this, RigidSolver, FieldCommand, ExecutionDatas, PrevResolutionType, PrevFilterType, PrevObjectType, PrevPositionType))
+				const FFieldSystemCommand& FieldCommand = Commands[CommandIndex];
+				if (IsForceFieldValid(FieldCommand))
 				{
-					const Chaos::FReal TimeSeconds = RigidSolver->GetSolverTime() - FieldCommand.TimeCreation;
+					if (Chaos::BuildFieldSamplePoints(this, RigidSolver, FieldCommand, ExecutionDatas, PrevResolutionType, PrevFilterType, PrevObjectType, PrevPositionType))
+					{
+						const Chaos::FReal TimeSeconds = RigidSolver->GetSolverTime() - FieldCommand.TimeCreation;
 
-					FFieldContext FieldContext(
-						ExecutionDatas,
-						FieldCommand.MetaData,
-						TimeSeconds);
+						FFieldContext FieldContext(
+							ExecutionDatas,
+							FieldCommand.MetaData,
+							TimeSeconds);
 
-					TArray<Chaos::FGeometryParticleHandle*>& ParticleHandles = ExecutionDatas.ParticleHandles[(uint8)EFieldCommandHandlesType::InsideHandles];
+						TArray<Chaos::FGeometryParticleHandle*>& ParticleHandles = ExecutionDatas.ParticleHandles[(uint8)EFieldCommandHandlesType::InsideHandles];
 
-					if (FieldCommand.RootNode->Type() == FFieldNode<FVector>::StaticType())
-					{
-						TArray<FVector>& FinalResults = ExecutionDatas.VectorResults[(uint8)EFieldCommandResultType::FinalResult];
-						ResetResultsArray < FVector >(ExecutionDatas.SamplePositions.Num(), FinalResults, FVector::ZeroVector);
+						if (FieldCommand.RootNode->Type() == FFieldNode<FVector>::StaticType())
+						{
+							TArray<FVector>& FinalResults = ExecutionDatas.VectorResults[(uint8)EFieldCommandResultType::FinalResult];
+							ResetResultsArray < FVector >(ExecutionDatas.SamplePositions.Num(), FinalResults, FVector::ZeroVector);
 
-						Chaos::FieldVectorForceUpdate(RigidSolver, FieldCommand, ParticleHandles,
-							FieldContext, FinalResults);
+							Chaos::FieldVectorForceUpdate(RigidSolver, FieldCommand, ParticleHandles,
+								FieldContext, FinalResults);
+						}
 					}
+					CommandsToRemove.Add(CommandIndex);
 				}
-				CommandsToRemove.Add(CommandIndex);
 			}
-		}
-		for (int32 Index = CommandsToRemove.Num() - 1; Index >= 0; --Index)
-		{
-			Commands.RemoveAt(CommandsToRemove[Index]);
+			for (int32 Index = CommandsToRemove.Num() - 1; Index >= 0; --Index)
+			{
+				Commands.RemoveAt(CommandsToRemove[Index]);
+			}
 		}
 	}
 }
