Iris - Force enable use of PlayerState prioritizer by default. Without it PlayerStates will get static priority instead of limiting the number of PlayerStates replicated per net tick.

diff --git a/Engine/Config/BaseEngine.ini b/Engine/Config/BaseEngine.ini
--- a/Engine/Config/BaseEngine.ini
+++ b/Engine/Config/BaseEngine.ini
@@ -1209,7 +1209,8 @@ DefaultSpatialFilterName=Spatial
 ;+FilterConfigs=(ClassName=/Script/Engine.Pawn, DynamicFilterName=Spatial))
 
 ; PrioritizerConfigs
-+PrioritizerConfigs=(ClassName=/Script/Engine.PlayerState, PrioritizerName=PlayerState)
+; PlayerState is by default bAlwaysRelevant but the default player state prioritizer handles this.
++PrioritizerConfigs=(ClassName=/Script/Engine.PlayerState, PrioritizerName=PlayerState, bForceEnableOnAllInstances=true)
 
 ;DeltaCompressionConfigs=ClearArray
 ; Classes for which deltacompression is enabled
diff --git a/Engine/Source/Runtime/Experimental/Iris/Core/Private/Iris/ReplicationSystem/ObjectReplicationBridge.cpp b/Engine/Source/Runtime/Experimental/Iris/Core/Private/Iris/ReplicationSystem/ObjectReplicationBridge.cpp
--- a/Engine/Source/Runtime/Experimental/Iris/Core/Private/Iris/ReplicationSystem/ObjectReplicationBridge.cpp
+++ b/Engine/Source/Runtime/Experimental/Iris/Core/Private/Iris/ReplicationSystem/ObjectReplicationBridge.cpp
@@ -305,13 +305,15 @@ UE::Net::FNetRefHandle UObjectReplicationBridge::BeginReplication(UObject* Insta
 			}
 
 			// Set prioritizer
-			if (Params.StaticPriority > 0.0f)
+			const bool bRequireForceEnabled = Params.StaticPriority > 0.0f;
+			const FNetObjectPrioritizerHandle PrioritizerHandle = GetPrioritizer(Instance->GetClass(), bRequireForceEnabled);
+			// Set static priority if valid unless we have a force enabled prioritizer.
+			if (Params.StaticPriority > 0.0f && PrioritizerHandle == InvalidNetObjectPrioritizerHandle)
 			{
 				ReplicationSystem->SetStaticPriority(RefHandle, Params.StaticPriority);
 			}
 			else
 			{
-				const FNetObjectPrioritizerHandle PrioritizerHandle = GetPrioritizer(Instance->GetClass());
 				if (PrioritizerHandle != InvalidNetObjectPrioritizerHandle)
 				{
 					ReplicationSystem->SetPrioritizer(RefHandle, PrioritizerHandle);
@@ -1257,7 +1259,7 @@ UE::Net::FNetObjectFilterHandle UObjectReplicationBridge::GetDynamicFilter(const
 	return FilterHandle;
 }
 
-UE::Net::FNetObjectPrioritizerHandle UObjectReplicationBridge::GetPrioritizer(const UClass* Class)
+UE::Net::FNetObjectPrioritizerHandle UObjectReplicationBridge::GetPrioritizer(const UClass* Class, bool bRequireForceEnabled)
 {
 	using namespace UE::Net;
 
@@ -1266,29 +1268,24 @@ UE::Net::FNetObjectPrioritizerHandle UObjectReplicationBridge::GetPrioritizer(co
 		const FName ClassName = GetConfigClassPathName(Class);
 
 		// Try exact match first.
-		if (FNetObjectPrioritizerHandle* PrioritizerHandlePtr = ClassesWithPrioritizer.Find(ClassName))
+		if (const FClassPrioritizerInfo* PrioritizerInfo = ClassesWithPrioritizer.Find(ClassName))
 		{
-			return *PrioritizerHandlePtr;
+			const bool bUsePrioritizer = !bRequireForceEnabled || PrioritizerInfo->bForceEnable;
+			return bUsePrioritizer ? PrioritizerInfo->PrioritizerHandle : InvalidNetObjectPrioritizerHandle;
 		}
 
-		/**
-		 * Try to find superclass. If we find it and the classes are considered equal we copy the filter setting.
-		 * If it's not equal we check whether it can be spatialized or not and use the result of that.
-		 * In all cases we add the result to the mapping for faster lookup next time.
-		 */
+		// Try to find superclass with prioritizer config. If we find it we copy the config and add the result to the mapping for faster lookup next time.
 		for (const UClass* SuperClass = Class->GetSuperClass(); SuperClass != nullptr; SuperClass = SuperClass->GetSuperClass())
 		{
 			const FName SuperClassName = GetConfigClassPathName(SuperClass);
 
-			// Try to get exact match first.
-			if (FNetObjectPrioritizerHandle* PrioritizerHandlePtr = ClassesWithPrioritizer.Find(SuperClassName))
+			if (const FClassPrioritizerInfo* PrioritizerInfo = ClassesWithPrioritizer.Find(SuperClassName))
 			{
-				// Assume this class should use the same prioritizer
-				{
-					const FNetObjectPrioritizerHandle PrioritizerHandle = *PrioritizerHandlePtr;
-					ClassesWithPrioritizer.Add(ClassName, PrioritizerHandle);
-					return PrioritizerHandle;
-				}
+				// Copy info to this class
+				ClassesWithPrioritizer.Add(ClassName, *PrioritizerInfo);
+
+				const bool bUsePrioritizer = !bRequireForceEnabled || PrioritizerInfo->bForceEnable;
+				return bUsePrioritizer ? PrioritizerInfo->PrioritizerHandle : InvalidNetObjectPrioritizerHandle;
 			}
 		}
 	}
@@ -1363,8 +1360,10 @@ void UObjectReplicationBridge::LoadConfig()
 	{
 		for (const FObjectReplicationBridgePrioritizerConfig& PrioritizerConfig : BridgeConfig->GetPrioritizerConfigs())
 		{
-			const UE::Net::FNetObjectPrioritizerHandle PrioritizerHandle = ReplicationSystem->GetPrioritizerHandle(PrioritizerConfig.PrioritizerName);
-			ClassesWithPrioritizer.Add(PrioritizerConfig.ClassName, PrioritizerHandle);
+			FClassPrioritizerInfo PrioInfo;
+			PrioInfo.PrioritizerHandle = ReplicationSystem->GetPrioritizerHandle(PrioritizerConfig.PrioritizerName);
+			PrioInfo.bForceEnable = PrioritizerConfig.bForceEnableOnAllInstances;
+			ClassesWithPrioritizer.Add(PrioritizerConfig.ClassName, PrioInfo);
 		}
 	}
 
diff --git a/Engine/Source/Runtime/Experimental/Iris/Core/Public/Iris/ReplicationSystem/ObjectReplicationBridge.h b/Engine/Source/Runtime/Experimental/Iris/Core/Public/Iris/ReplicationSystem/ObjectReplicationBridge.h
--- a/Engine/Source/Runtime/Experimental/Iris/Core/Public/Iris/ReplicationSystem/ObjectReplicationBridge.h
+++ b/Engine/Source/Runtime/Experimental/Iris/Core/Public/Iris/ReplicationSystem/ObjectReplicationBridge.h
@@ -242,8 +242,8 @@ private:
 	/** Retrieves the dynamic filter to set for the given class. Will return an invalid handle if no dynamic filter should be set. */
 	UE::Net::FNetObjectFilterHandle GetDynamicFilter(const UClass* Class);
 
-	/** Retrieves the prioritizer to set for the given class. Will return an invalid handle if no prioritizer should be set. */
-	UE::Net::FNetObjectPrioritizerHandle GetPrioritizer(const UClass* Class);
+	/** Retrieves the prioritizer to set for the given class. If bRequireForceEnabled the config needs to have bForceEnableOnAllInstances set in order for this method to return the configured prioritizer. Returns an invalid handle if no prioritizer should be set. */
+	UE::Net::FNetObjectPrioritizerHandle GetPrioritizer(const UClass* Class, bool bRequireForceEnabled);
 
 	/** Returns true if instances of this class should be delta compressed */
 	bool ShouldClassBeDeltaCompressed(const UClass* Class);
@@ -259,6 +259,13 @@ private:
 
 	TMap<const UClass*, FName> ConfigClassPathNameCache;
 
+	// Prioritization
+	struct FClassPrioritizerInfo
+	{
+		UE::Net::FNetObjectPrioritizerHandle PrioritizerHandle;
+		uint32 bForceEnable : 1;
+	};
+
 	// Polling
 	struct FPollInfo
 	{
@@ -280,7 +287,7 @@ private:
 	TFunction<bool(const UClass*,const UClass*)> ShouldSubclassUseSameFilterFunction;
 
 	// Prioritizer mapping
-	TMap<FName, UE::Net::FNetObjectPrioritizerHandle> ClassesWithPrioritizer;
+	TMap<FName, FClassPrioritizerInfo> ClassesWithPrioritizer;
 
 	// Delta compression
 	TMap<FName, bool> ClassesWithDeltaCompression;
diff --git a/Engine/Source/Runtime/Experimental/Iris/Core/Public/Iris/ReplicationSystem/ObjectReplicationBridgeConfig.h b/Engine/Source/Runtime/Experimental/Iris/Core/Public/Iris/ReplicationSystem/ObjectReplicationBridgeConfig.h
--- a/Engine/Source/Runtime/Experimental/Iris/Core/Public/Iris/ReplicationSystem/ObjectReplicationBridgeConfig.h
+++ b/Engine/Source/Runtime/Experimental/Iris/Core/Public/Iris/ReplicationSystem/ObjectReplicationBridgeConfig.h
@@ -58,6 +58,10 @@ struct FObjectReplicationBridgePrioritizerConfig
 	/** The name of the prioritizer to set on the class instances. "Default" can be used to specify the default spatial prioritizer. */
 	UPROPERTY()
 	FName PrioritizerName;
+
+	/** Whether this prioritizer should be used for all instances of this class and subclasses, regardless of bAlwaysRelevant and bOnlyRelevantToOwner settings on instance. */
+	UPROPERTY()
+	bool bForceEnableOnAllInstances = false;
 };
 
 USTRUCT()
