Iris - Generic support for targeting structs with an arbitrary NetSerializer.

diff --git a/Engine/Source/Runtime/Experimental/Iris/Core/Private/Iris/ReplicationState/PropertyNetSerializerInfoRegistry.cpp b/Engine/Source/Runtime/Experimental/Iris/Core/Private/Iris/ReplicationState/PropertyNetSerializerInfoRegistry.cpp
--- a/Engine/Source/Runtime/Experimental/Iris/Core/Private/Iris/ReplicationState/PropertyNetSerializerInfoRegistry.cpp
+++ b/Engine/Source/Runtime/Experimental/Iris/Core/Private/Iris/ReplicationState/PropertyNetSerializerInfoRegistry.cpp
@@ -121,7 +121,7 @@ const FFieldClass* FNamedStructLastResortPropertyNetSerializerInfo::GetPropertyT
 bool FNamedStructLastResortPropertyNetSerializerInfo::IsSupported(const FProperty* Property) const
 {
 	const FStructProperty* StructProp = CastFieldChecked<const FStructProperty>(Property);
-	return StructProp->Struct->GetFName() == PropertyFName;
+	return IsSupportedStruct(StructProp->Struct->GetFName());
 }
 
 const FPropertyNetSerializerInfo* FPropertyNetSerializerInfoRegistry::FindSerializerInfo(const FProperty* Property)
@@ -185,8 +185,8 @@ const FPropertyNetSerializerInfo* FPropertyNetSerializerInfoRegistry::FindStruct
 		// as we allow multiple infos for the same class type we need to check if the info supports the Property
 		for (; Index < Registry.Num() && Registry[Index].Get<0>() == ClassType; ++Index)
 		{
-			const FNamedStructPropertyNetSerializerInfo* Info = static_cast<const FNamedStructPropertyNetSerializerInfo*>(Registry[Index].Get<1>());
-			if (Info->PropertyFName == Name)
+			const FPropertyNetSerializerInfo* Info = Registry[Index].Get<1>();
+			if (Info->IsSupportedStruct(Name))
 			{
 				return Info;
 			}
diff --git a/Engine/Source/Runtime/Experimental/Iris/Core/Public/Iris/ReplicationState/PropertyNetSerializerInfoRegistry.h b/Engine/Source/Runtime/Experimental/Iris/Core/Public/Iris/ReplicationState/PropertyNetSerializerInfoRegistry.h
--- a/Engine/Source/Runtime/Experimental/Iris/Core/Public/Iris/ReplicationState/PropertyNetSerializerInfoRegistry.h
+++ b/Engine/Source/Runtime/Experimental/Iris/Core/Public/Iris/ReplicationState/PropertyNetSerializerInfoRegistry.h
@@ -38,6 +38,15 @@ struct FPropertyNetSerializerInfo
 	virtual const FNetSerializerConfig* BuildNetSerializerConfig(void* NetSerializerConfigBuffer, const FProperty* Property) const { return nullptr; }
 	/** Custom replication fragments are currently only supported by structs with a custom NetSerializer. See UE_NET_IMPLEMENT_NAMED_STRUCT_NETSERIALIZER_WITH_CUSTOM_FRAGMENT_INFO. */
 	virtual CreateAndRegisterReplicationFragmentFunc GetCreateAndRegisterReplicationFragmentFunction() const { return nullptr; }
+
+	bool IsSupportedStruct(FName InStructName) const
+	{
+		return StructName == InStructName;
+	}
+
+protected:
+	// Used by named struct serializers.
+	FName StructName;
 };
 
 /**
@@ -99,14 +108,17 @@ struct TSimplePropertyNetSerializerInfo : public FPropertyNetSerializerInfo
 struct FNamedStructPropertyNetSerializerInfo : public TSimplePropertyNetSerializerInfo<FStructProperty>
 {
 	typedef TSimplePropertyNetSerializerInfo<FStructProperty> Super;
-	const FName PropertyFName;
 
-	FNamedStructPropertyNetSerializerInfo(const FName InPropertyFName, const FNetSerializer& InSerializer) : Super(InSerializer), PropertyFName(InPropertyFName) {}
+	FNamedStructPropertyNetSerializerInfo(const FName InPropertyFName, const FNetSerializer& InSerializer)
+	: Super(InSerializer)
+	{
+		StructName = InPropertyFName;
+	}
 
 	virtual bool IsSupported(const FProperty* Property) const override
 	{	
 		const FStructProperty* StructProp = CastFieldChecked<const FStructProperty>(Property);
-		return StructProp->Struct->GetFName() == PropertyFName;
+		return IsSupportedStruct(StructProp->Struct->GetFName());
 	};
 
 	virtual CreateAndRegisterReplicationFragmentFunc GetCreateAndRegisterReplicationFragmentFunction() const
@@ -135,15 +147,13 @@ struct FLastResortPropertyNetSerializerInfo : public FPropertyNetSerializerInfo
 struct FNamedStructLastResortPropertyNetSerializerInfo final : public FLastResortPropertyNetSerializerInfo
 {
 	FNamedStructLastResortPropertyNetSerializerInfo(const FName InPropertyFName)
-	: PropertyFName(InPropertyFName)
+	: FLastResortPropertyNetSerializerInfo()
 	{
+		StructName = InPropertyFName;
 	}
 
 	IRISCORE_API virtual const FFieldClass* GetPropertyTypeClass() const override;
 	IRISCORE_API virtual bool IsSupported(const FProperty* Property) const override;
-
-private:
-	const FName PropertyFName;
 };
 
 // Issue fatal error if matching trait found in UsedReplicationStateTraits is not set for the Serializer.
