Fix issue where updaterep data is called with a callback and the physics proxy is null
- make sure UpdateRepData checks for the physics proxy  being valid

diff --git a/Engine/Source/Runtime/Experimental/GeometryCollectionEngine/Private/GeometryCollection/GeometryCollectionComponent.cpp b/Engine/Source/Runtime/Experimental/GeometryCollectionEngine/Private/GeometryCollection/GeometryCollectionComponent.cpp
--- a/Engine/Source/Runtime/Experimental/GeometryCollectionEngine/Private/GeometryCollection/GeometryCollectionComponent.cpp
+++ b/Engine/Source/Runtime/Experimental/GeometryCollectionEngine/Private/GeometryCollection/GeometryCollectionComponent.cpp
@@ -1545,6 +1545,11 @@ void UGeometryCollectionComponent::UpdateRepData()
 		return;
 	}
 
+	if (PhysicsProxy == nullptr)
+	{
+		return;
+	}
+
 	AActor* Owner = GetOwner();
 	
 	// If we have no owner or our netmode means we never require replication then early out
