Give strain modifier API option to filter breaking strains by strain type.

diff --git a/Engine/Source/Runtime/Experimental/Chaos/Private/Chaos/StrainModification.cpp b/Engine/Source/Runtime/Experimental/Chaos/Private/Chaos/StrainModification.cpp
--- a/Engine/Source/Runtime/Experimental/Chaos/Private/Chaos/StrainModification.cpp
+++ b/Engine/Source/Runtime/Experimental/Chaos/Private/Chaos/StrainModification.cpp
@@ -25,7 +25,7 @@ int32 Chaos::FStrainedProxyModifier::GetNumRestBreakables() const
 	return 0;
 }
 
-int32 Chaos::FStrainedProxyModifier::GetNumBreakingStrains() const
+int32 Chaos::FStrainedProxyModifier::GetNumBreakingStrains(const uint8 StrainTypes) const
 {
 	// Make sure we have a proxy and rest-children
 	if (Proxy == nullptr) { return 0; }
@@ -45,8 +45,12 @@ int32 Chaos::FStrainedProxyModifier::GetNumBreakingStrains() const
 		// graph or something... this logic is currently copied from
 		// FRigidClustering::ReleaseClusterParticlesImpl, but should be tucked behind
 		// an interface somewhere.
-		const Chaos::FReal CollisionImpulses = ChildHandle->CollisionImpulses();
-		const Chaos::FReal ExternalStrain = ChildHandle->GetExternalStrain();
+		const Chaos::FReal CollisionImpulses
+			= (StrainTypes & Chaos::EStrainTypes::CollisionStrain)
+			? ChildHandle->CollisionImpulses() : FReal(0);
+		const Chaos::FReal ExternalStrain
+			= (StrainTypes & Chaos::EStrainTypes::ExternalStrain)
+			? ChildHandle->GetExternalStrain() : FReal(0);
 		const Chaos::FReal MaxAppliedStrain = FMath::Max(CollisionImpulses, ExternalStrain);
 		const Chaos::FReal InternalStrain = ChildHandle->GetInternalStrains();
 		if (MaxAppliedStrain >= InternalStrain)
diff --git a/Engine/Source/Runtime/Experimental/Chaos/Public/Chaos/StrainModification.h b/Engine/Source/Runtime/Experimental/Chaos/Public/Chaos/StrainModification.h
--- a/Engine/Source/Runtime/Experimental/Chaos/Public/Chaos/StrainModification.h
+++ b/Engine/Source/Runtime/Experimental/Chaos/Public/Chaos/StrainModification.h
@@ -16,6 +16,13 @@ namespace Chaos
 	class FStrainedProxyRange;
 	class FStrainModifierAccessor;
 
+	enum EStrainTypes : uint8
+	{
+		InternalStrain = 1 << 0,
+		ExternalStrain = 1 << 1,
+		CollisionStrain = 1 << 2
+	};
+
 	// FStrainedProxyModifier
 	//
 	// User-facing api for accessing the proxy of a strained cluster. Provides const access to the
@@ -47,7 +54,7 @@ namespace Chaos
 		int32 GetNumRestBreakables() const;
 
 		// Get the number of breaking strains amongst the level-1 strainables
-		int32 GetNumBreakingStrains() const;
+		int32 GetNumBreakingStrains(const uint8 StrainTypes = EStrainTypes::ExternalStrain | EStrainTypes::CollisionStrain) const;
 
 		// Clear strains for all strained cluster children
 		void ClearStrains();
