Iris
- Added CVar net.iris.ForceFullCopyAndQuantize to control if we do partial copy and quantize or only do it for dirty fragments.

diff --git a/Engine/Source/Runtime/Experimental/Iris/Core/Private/Iris/ReplicationSystem/ReplicationOperations.cpp b/Engine/Source/Runtime/Experimental/Iris/Core/Private/Iris/ReplicationSystem/ReplicationOperations.cpp
--- a/Engine/Source/Runtime/Experimental/Iris/Core/Private/Iris/ReplicationSystem/ReplicationOperations.cpp
+++ b/Engine/Source/Runtime/Experimental/Iris/Core/Private/Iris/ReplicationSystem/ReplicationOperations.cpp
@@ -51,7 +51,8 @@ static FNetDebugName s_LifetimeConditionDebugNames[] =
 };
 static_assert(COND_Max == 17, "s_LifetimeConditionDebugNames may need updating.");
 
-static void AppendMemberChangeMasks(FNetBitStreamWriter* ChangeMaskWriter, FNetBitStreamWriter* ConditionalChangeMaskWriter, uint8* ExternalStateBuffer, const FReplicationStateDescriptor* Descriptor);
+// Append changemask bits to ChangeMaskWriter and conditionally to the ConditionalChangeMaskWriter. If the state is dirty, the function will return true.
+static bool AppendMemberChangeMasks(FNetBitStreamWriter* ChangeMaskWriter, FNetBitStreamWriter* ConditionalChangeMaskWriter, uint8* ExternalStateBuffer, const FReplicationStateDescriptor* Descriptor);
 static void ResetMemberChangeMasks(uint8* ExternalStateBuffer, const FReplicationStateDescriptor* Descriptor);
 
 }
@@ -545,6 +546,45 @@ void FReplicationInstanceOperations::CopyAndQuantize(FNetSerializationContext& C
 	}
 }
 
+void FReplicationInstanceOperations::CopyAndQuantizeIfDirty(FNetSerializationContext& Context, uint8* DstObjectStateBuffer, FNetBitStreamWriter* OutChangeMaskWriter, const FReplicationInstanceProtocol* InstanceProtocol, const FReplicationProtocol* Protocol)
+{
+	check(InstanceProtocol && InstanceProtocol->FragmentCount == Protocol->ReplicationStateCount);
+
+	FReplicationInstanceProtocol::FFragmentData* FragmentData = InstanceProtocol->FragmentData;
+	const FReplicationStateDescriptor** ReplicationStateDescriptors = Protocol->ReplicationStateDescriptors;
+
+	// Set up conditional change mask writer.
+	FNetBitStreamWriter ConditionalChangeMaskWriter;
+	FNetBitStreamWriter* OutConditionalChangeMaskWriter = nullptr;
+	if (EnumHasAnyFlags(Protocol->ProtocolTraits, EReplicationProtocolTraits::HasConditionalChangeMask))
+	{
+		ConditionalChangeMaskWriter.InitBytes(DstObjectStateBuffer + Protocol->GetConditionalChangeMaskOffset(), FNetBitArrayView::CalculateRequiredWordCount(Protocol->ChangeMaskBitCount)*sizeof(FNetBitArrayView::StorageWordType));
+		OutConditionalChangeMaskWriter = &ConditionalChangeMaskWriter;
+	}
+
+	uint8* CurrentInternalStateBuffer = DstObjectStateBuffer;	
+	for (uint32 StateIt = 0; StateIt < Protocol->ReplicationStateCount; ++StateIt)
+	{
+		const FReplicationStateDescriptor* CurrentDescriptor = ReplicationStateDescriptors[StateIt];
+
+		CurrentInternalStateBuffer = Align(CurrentInternalStateBuffer, CurrentDescriptor->InternalAlignment);
+
+		// Append changemask data to bitstreams and quantize state if mask is dirty
+		if (Private::AppendMemberChangeMasks(OutChangeMaskWriter, OutConditionalChangeMaskWriter, FragmentData[StateIt].ExternalSrcBuffer, CurrentDescriptor))
+		{
+			FReplicationStateOperations::Quantize(Context, CurrentInternalStateBuffer, FragmentData[StateIt].ExternalSrcBuffer, CurrentDescriptor);
+		}
+		
+		CurrentInternalStateBuffer += CurrentDescriptor->InternalSize;
+	}
+
+	OutChangeMaskWriter->CommitWrites();
+	if (OutConditionalChangeMaskWriter != nullptr)
+	{
+		OutConditionalChangeMaskWriter->CommitWrites();
+	}
+}
+
 void FReplicationInstanceOperations::ResetDirtiness(const FReplicationInstanceProtocol* InstanceProtocol, const FReplicationProtocol* Protocol)
 {
 	check(InstanceProtocol && InstanceProtocol->FragmentCount == Protocol->ReplicationStateCount);
@@ -1115,7 +1155,7 @@ bool FReplicationProtocolOperations::IsEqualQuantizedState(FNetSerializationCont
 namespace UE::Net::Private
 {
 
-static void AppendMemberChangeMasks(FNetBitStreamWriter* ChangeMaskWriter, FNetBitStreamWriter* ConditionalChangeMaskWriter, uint8* StateBuffer, const FReplicationStateDescriptor* Descriptor)
+static bool AppendMemberChangeMasks(FNetBitStreamWriter* ChangeMaskWriter, FNetBitStreamWriter* ConditionalChangeMaskWriter, uint8* StateBuffer, const FReplicationStateDescriptor* Descriptor)
 {
 	FNetBitArrayView::StorageWordType* ChangeMaskStorage = reinterpret_cast<FNetBitArrayView::StorageWordType*>(StateBuffer + Descriptor->GetChangeMaskOffset());
 	const uint32 BitCount = Descriptor->ChangeMaskBitCount;
@@ -1154,6 +1194,9 @@ static void AppendMemberChangeMasks(FNetBitStreamWriter* ChangeMaskWriter, FNetB
 			ConditionalChangeMaskWriter->Seek(ConditionalChangeMaskWriter->GetPosBits() + BitCount);
 		}
 	}
+
+	FReplicationStateHeader& ReplicationStateHeader = Private::GetReplicationStateHeader(StateBuffer, Descriptor);
+	return FReplicationStateHeaderAccessor::GetIsStateDirty(ReplicationStateHeader);
 }
 
 static void ResetMemberChangeMasks(uint8* ExternalStateBuffer, const FReplicationStateDescriptor* Descriptor)
diff --git a/Engine/Source/Runtime/Experimental/Iris/Core/Private/Iris/ReplicationSystem/ReplicationOperationsInternal.cpp b/Engine/Source/Runtime/Experimental/Iris/Core/Private/Iris/ReplicationSystem/ReplicationOperationsInternal.cpp
--- a/Engine/Source/Runtime/Experimental/Iris/Core/Private/Iris/ReplicationSystem/ReplicationOperationsInternal.cpp
+++ b/Engine/Source/Runtime/Experimental/Iris/Core/Private/Iris/ReplicationSystem/ReplicationOperationsInternal.cpp
@@ -23,6 +23,12 @@
 namespace UE::Net::Private
 {
 
+static bool bCVarForceFullCopyAndQuantize = false;
+static FAutoConsoleVariableRef CVarForceFullCopyAndQuantize(
+	TEXT("net.iris.ForceFullCopyAndQuantize"),
+	bCVarForceFullCopyAndQuantize,
+	TEXT("When enabled a full copy and quantize will be used, if disabled we will only copy and quantize dirty state data."));
+
 void FReplicationInstanceOperationsInternal::BindInstanceProtocol(FNetHandle NetHandle, FReplicationInstanceProtocol* InstanceProtocol, const FReplicationProtocol* Protocol)
 {
 	FReplicationInstanceProtocol::FFragmentData* FragmentData = InstanceProtocol->FragmentData;
@@ -62,8 +68,12 @@ uint32 FReplicationInstanceOperationsInternal::CopyObjectStateData(FNetBitStream
 		// We cannot copy state data for zero sized objects or objects that no longer has an instance protocol.
 		if (Object.InstanceProtocol && Object.Protocol->InternalTotalSize > 0U)
 		{
+			// if the object was scopable prev frame we can do partial copy
 			bool bShouldPropagateChangedStates = Object.bShouldPropagateChangedStates;
 
+			// We do a full CopyAndQunatize if the cvar bCVarForceFullCopyAndQuantize is set or that the object is new
+			const bool bUseFullCopyAndQuantize = bCVarForceFullCopyAndQuantize || !NetRefHandleManager.GetPrevFrameScopableInternalIndices().GetBit(InternalIndex);
+
 			// Copy dirty state
 			{
 				// Add entry to cache
@@ -73,7 +83,14 @@ uint32 FReplicationInstanceOperationsInternal::CopyObjectStateData(FNetBitStream
 				ChangeMaskWriter.InitBytes(ChangeMaskData, ChangeMaskByteCount);
 
 				//UE_LOG(LogIris, Log, TEXT("Copying state data for ( InternalIndex: %u ) with NetRefHandle (Id=%u)"), InternalIndex, Object.RefHandle.GetId());
-				FReplicationInstanceOperations::CopyAndQuantize(SerializationContext, NetRefHandleManager.GetReplicatedObjectStateBufferNoCheck(InternalIndex), &ChangeMaskWriter, Object.InstanceProtocol, Object.Protocol);
+				if (bUseFullCopyAndQuantize)
+				{
+					FReplicationInstanceOperations::CopyAndQuantize(SerializationContext, NetRefHandleManager.GetReplicatedObjectStateBufferNoCheck(InternalIndex), &ChangeMaskWriter, Object.InstanceProtocol, Object.Protocol);					
+				}
+				else
+				{
+					FReplicationInstanceOperations::CopyAndQuantizeIfDirty(SerializationContext, NetRefHandleManager.GetReplicatedObjectStateBufferNoCheck(InternalIndex), &ChangeMaskWriter, Object.InstanceProtocol, Object.Protocol);
+				}
 
 				Info.bHasDirtyChangeMask = MakeNetBitArrayView(ChangeMaskData, ChangeMaskByteCount * 8U).FindFirstOne() != FNetBitArrayView::InvalidIndex;
 			}
diff --git a/Engine/Source/Runtime/Experimental/Iris/Core/Public/Iris/ReplicationSystem/ReplicationOperations.h b/Engine/Source/Runtime/Experimental/Iris/Core/Public/Iris/ReplicationSystem/ReplicationOperations.h
--- a/Engine/Source/Runtime/Experimental/Iris/Core/Public/Iris/ReplicationSystem/ReplicationOperations.h
+++ b/Engine/Source/Runtime/Experimental/Iris/Core/Public/Iris/ReplicationSystem/ReplicationOperations.h
@@ -78,12 +78,22 @@ struct FReplicationInstanceOperations
 	static IRISCORE_API bool PollAndRefreshCachedObjectReferences(const FReplicationInstanceProtocol* InstanceProtocol, EReplicationFragmentTraits RequiredTraits);
 
 	/**
-	 * Quantize the state for a replicated object with a given InstanceProtocol using the ReplicationProtocol. DstObjectStateBuffer needs to be cleared before calling this function.
+	 * Quantize the state for a replicated object with a given InstanceProtocol using the ReplicationProtocol. 
+	 * DstObjectStateBuffer needs to be in a valid state before calling this function.
 	 * Changemasks will be written to the ChangeMaskWriter. Dirtiness will not be reset.
 	 * @see ResetDirtiness
 	 */
 	static IRISCORE_API void CopyAndQuantize(FNetSerializationContext& Context, uint8* DstObjectStateBuffer, FNetBitStreamWriter* ChangeMaskWriter, const FReplicationInstanceProtocol* InstanceProtocol, const FReplicationProtocol* Protocol);
 
+	/**
+	 * Quantize the state for a replicated object with a given InstanceProtocol using the ReplicationProtocol.
+	 * DstObjectStateBuffer needs to be in a valid state before calling this function.
+	 * Changemasks will be written to the ChangeMaskWriter. Dirtiness will not be reset.
+	 * This variant will only Quantize States marked as dirty
+	 * @see ResetDirtiness
+	 */
+	static IRISCORE_API void CopyAndQuantizeIfDirty(FNetSerializationContext& Context, uint8* DstObjectStateBuffer, FNetBitStreamWriter* ChangeMaskWriter, const FReplicationInstanceProtocol* InstanceProtocol, const FReplicationProtocol* Protocol);
+
 	/** Resets dirty tracking stored with the protocol, such as changemasks and init state dirtiness. */
 	static IRISCORE_API void ResetDirtiness(const FReplicationInstanceProtocol* InstanceProtocol, const FReplicationProtocol* Protocol);
 
