Fix C++20 compile error. Make Chaos::FMidPhaseModifierParticleIterator's comparison operators const to address ambiguity issues.

diff --git a/Engine/Source/Runtime/Experimental/Chaos/Public/Chaos/MidPhaseModification.h b/Engine/Source/Runtime/Experimental/Chaos/Public/Chaos/MidPhaseModification.h
--- a/Engine/Source/Runtime/Experimental/Chaos/Public/Chaos/MidPhaseModification.h
+++ b/Engine/Source/Runtime/Experimental/Chaos/Public/Chaos/MidPhaseModification.h
@@ -108,12 +108,12 @@ namespace Chaos
 			return MidPhaseIndex < Particle->ParticleCollisions().Num();
 		}
 
-		bool operator==(const FMidPhaseModifierParticleIterator& Other)
+		bool operator==(const FMidPhaseModifierParticleIterator& Other) const
 		{
 			return MidPhaseIndex == Other.MidPhaseIndex;
 		}
 
-		bool operator!=(const FMidPhaseModifierParticleIterator& Other)
+		bool operator!=(const FMidPhaseModifierParticleIterator& Other) const
 		{
 			return !operator==(Other);
 		}
