stop the geometry collection proximity -> connection conversion adding self-links to the connection graph

diff --git a/Engine/Source/Runtime/Experimental/Chaos/Private/GeometryCollection/GeometryCollectionProximityUtility.cpp b/Engine/Source/Runtime/Experimental/Chaos/Private/GeometryCollection/GeometryCollectionProximityUtility.cpp
--- a/Engine/Source/Runtime/Experimental/Chaos/Private/GeometryCollection/GeometryCollectionProximityUtility.cpp
+++ b/Engine/Source/Runtime/Experimental/Chaos/Private/GeometryCollection/GeometryCollectionProximityUtility.cpp
@@ -585,7 +585,7 @@ void FGeometryCollectionProximityUtility::CopyProximityToConnectionGraph()
 				{
 					NbrBone = Collection->Parent[NbrBone];
 				}
-				if (NbrBone != INDEX_NONE && Collection->Parent[NbrBone] == BoneParent)
+				if (NbrBone != INDEX_NONE && NbrBone != BoneIdx && Collection->Parent[NbrBone] == BoneParent)
 				{
 					ConnectionsFacade.Connect(BoneIdx, NbrBone);
 				}
