Add a new RWLock for Chaos based on the Core FRwLock with support for recursion (CHAOS_SCENE_LOCK_FRWLOCK). This avoids an unnecessary spinlock and uses a hardware optimized RwLock implementation internally, which is generally faster than the existing implementation.

This also fixes deadlocks in some situations on certain platforms with simple scehdulers and limited cores.

Added a CHAOS_SCENE_LOCK_TYPE define allowing simplier overriding of the lock type. This can be overridden externally if necessary (e.g per platform or per game project target.cs).

diff --git a/Engine/Source/Runtime/Experimental/Chaos/Public/Framework/Threading.h b/Engine/Source/Runtime/Experimental/Chaos/Public/Framework/Threading.h
--- a/Engine/Source/Runtime/Experimental/Chaos/Public/Framework/Threading.h
+++ b/Engine/Source/Runtime/Experimental/Chaos/Public/Framework/Threading.h
@@ -10,6 +10,7 @@
 #include "ProfilingDebugging/CsvProfiler.h"
 #include "Async/ParallelFor.h"
 #include <atomic>
+#include "HAL/CriticalSection.h"
 
 #ifndef PHYSICS_THREAD_CONTEXT
 	#if (!UE_BUILD_SHIPPING && !UE_BUILD_TEST)
@@ -19,14 +20,23 @@
 	#endif
 #endif
 
-/** Controls whether we use a FIFO RW lock or an unfair RW lock for scene locking in Chaos */
-#define CHAOS_SCENE_LOCK_RWFIFO 1
+
 
 /**
- * Controls whether we use a non-yielding lock for the FIFO lock if CHAOS_SCENE_LOCK_RWFIFO is enabled
- * @see CHAOS_SCENE_LOCK_RWFIFO
+ * Scene lock types
+ * @see CHAOS_SCENE_LOCK_TYPE
  */ 
-#define CHAOS_SCENE_LOCK_RWFIFO_SPINLOCK 0
+#define CHAOS_SCENE_LOCK_SCENE_GUARD 0            // Unfair RW lock
+#define CHAOS_SCENE_LOCK_RWFIFO_SPINLOCK 1        // Fair RW spinlock, non yielding
+#define CHAOS_SCENE_LOCK_RWFIFO_CRITICALSECTION 2 // Fair RW spinlock, yielding
+#define CHAOS_SCENE_LOCK_FRWLOCK 3                // Recurrant RW lock based on FRwLock (uses platform sync primitives)
+#define CHAOS_SCENE_LOCK_SIMPLE_MUTEX 4           // Just a critical section (not an RWLock). Provided For profiling/debugging only. Not recommended
+
+/** Controls the scene lock type. See above. */ 
+#ifndef CHAOS_SCENE_LOCK_TYPE
+#define CHAOS_SCENE_LOCK_TYPE CHAOS_SCENE_LOCK_FRWLOCK
+#endif
+
 
 /**
  * \def CHAOS_SCENE_LOCK_CHECKS
@@ -688,6 +698,174 @@ FORCEINLINE void EnsureIsInGameThreadContext()
 		uint32 Count;
 	};
 
+
+	/**
+	 * A recursive readwrite lock that uses FRwLock internally (this uses an efficient platform specific implementation)
+	 */
+	class FPhysicsRwLock
+	{
+		struct FRwLockInfo
+		{
+			FRwLockInfo(void* TlsSlotValue)
+			{
+				ThreadReadDepth = uint32(uint64(TlsSlotValue));
+				ThreadWriteDepth = uint64(TlsSlotValue) >> 32;
+			}
+			void* GetTlsSlotValue()
+			{
+				uint64 ValueOut = uint64(ThreadReadDepth) | (uint64(ThreadWriteDepth) << 32);
+				return (void*)ValueOut;
+			}
+
+			uint32 ThreadReadDepth = 0;
+			uint32 ThreadWriteDepth = 0;
+		};
+
+	public:
+
+		FPhysicsRwLock()
+		{
+			TlsSlot = FPlatformTLS::AllocTlsSlot();
+		}
+
+		~FPhysicsRwLock()
+		{
+			if (FPlatformTLS::IsValidTlsSlot(TlsSlot))
+			{
+				FPlatformTLS::FreeTlsSlot(TlsSlot);
+			}
+		}
+
+		void ReadLock()
+		{
+			FRwLockInfo ThreadInfo(FPlatformTLS::GetTlsValue(TlsSlot));
+			ThreadInfo.ThreadReadDepth++;
+			FPlatformTLS::SetTlsValue(TlsSlot, ThreadInfo.GetTlsSlotValue());
+
+			if (ThreadInfo.ThreadReadDepth + ThreadInfo.ThreadWriteDepth == 1)
+			{
+				RwLock.ReadLock();
+			}
+
+#if PHYSICS_THREAD_CONTEXT
+			// Read lock means we can access game thread data, so set the right context
+			FPhysicsThreadContext::Get().IncGameThreadContext();
+#endif
+
+		}
+		void WriteLock()
+		{
+			FRwLockInfo ThreadInfo(FPlatformTLS::GetTlsValue(TlsSlot));
+			ThreadInfo.ThreadWriteDepth++;
+			FPlatformTLS::SetTlsValue(TlsSlot, ThreadInfo.GetTlsSlotValue());
+
+#if CHAOS_SCENE_LOCK_CHECKS
+			if (ThreadInfo.ThreadReadDepth > 0)
+			{
+				UE_LOG(LogChaos, Warning, TEXT("Attempt to upgrade a read lock to a write lock. This is not supported. Writes will be unsafe"))
+			}
+#endif
+			if (ThreadInfo.ThreadReadDepth + ThreadInfo.ThreadWriteDepth == 1)
+			{
+				RwLock.WriteLock();
+			}
+#if PHYSICS_THREAD_CONTEXT
+			// Write lock means we can access game thread data, so set the right context
+			FPhysicsThreadContext::Get().IncGameThreadContext();
+#endif
+		}
+		void ReadUnlock()
+		{
+			FRwLockInfo ThreadInfo(FPlatformTLS::GetTlsValue(TlsSlot));
+			ThreadInfo.ThreadReadDepth--;
+			FPlatformTLS::SetTlsValue(TlsSlot, ThreadInfo.GetTlsSlotValue());
+
+			if (ThreadInfo.ThreadWriteDepth + ThreadInfo.ThreadReadDepth == 0)
+			{
+				RwLock.ReadUnlock();
+			}
+
+#if PHYSICS_THREAD_CONTEXT
+			// Read lock is released, the gamethread context is gone
+			FPhysicsThreadContext::Get().DecGameThreadContext();
+#endif
+		}
+		void WriteUnlock()
+		{
+			FRwLockInfo ThreadInfo(FPlatformTLS::GetTlsValue(TlsSlot));
+			ThreadInfo.ThreadWriteDepth--;
+			FPlatformTLS::SetTlsValue(TlsSlot, ThreadInfo.GetTlsSlotValue());
+			if (ThreadInfo.ThreadWriteDepth + ThreadInfo.ThreadReadDepth == 0)
+			{
+				RwLock.WriteUnlock();
+			}
+
+#if PHYSICS_THREAD_CONTEXT
+			// Write lock is released, the gamethread context is gone
+			FPhysicsThreadContext::Get().DecGameThreadContext();
+#endif
+		}
+
+	private:
+		FRWLock RwLock;
+		uint32 TlsSlot;
+	};
+
+
+	/**
+	 * A simple mutex based lock based on FCriticalSection. Reads are exclusive
+	 */
+	class FPhysicsSimpleMutexLock
+	{
+	public:
+		FPhysicsSimpleMutexLock()
+		{
+		}
+
+		~FPhysicsSimpleMutexLock()
+		{
+		}
+
+		void ReadLock()
+		{
+			Cs.Lock();
+#if PHYSICS_THREAD_CONTEXT
+			// Read lock means we can access game thread data, so set the right context
+			FPhysicsThreadContext::Get().IncGameThreadContext();
+#endif
+		}
+		void WriteLock()
+		{
+			Cs.Lock();
+#if PHYSICS_THREAD_CONTEXT
+			// Write lock means we can access game thread data, so set the right context
+			FPhysicsThreadContext::Get().IncGameThreadContext();
+#endif
+		}
+		void ReadUnlock()
+		{
+			Cs.Unlock();
+
+#if PHYSICS_THREAD_CONTEXT
+			// Read lock is released, the gamethread context is gone
+			FPhysicsThreadContext::Get().DecGameThreadContext();
+#endif
+		}
+		void WriteUnlock()
+		{
+			Cs.Unlock();
+
+#if PHYSICS_THREAD_CONTEXT
+			// Write lock is released, the gamethread context is gone
+			FPhysicsThreadContext::Get().DecGameThreadContext();
+#endif
+		}
+
+	private:
+		FCriticalSection Cs;
+	};
+
+
 	/**
 	 * Implements a RAII scoped write lock around a generic mutex.
 	 */
@@ -746,14 +924,16 @@ FORCEINLINE void EnsureIsInGameThreadContext()
 		MutexType& Mutex;
 	};
 
-#if !CHAOS_SCENE_LOCK_RWFIFO
+#if CHAOS_SCENE_LOCK_TYPE == CHAOS_SCENE_LOCK_SCENE_GUARD
 	using FPhysSceneLock = FPhysicsSceneGuard;
-#else
-#if CHAOS_SCENE_LOCK_RWFIFO_SPINLOCK
+#elif CHAOS_SCENE_LOCK_TYPE == CHAOS_SCENE_LOCK_RWFIFO_SPINLOCK
 	using FPhysSceneLock = TRwFifoLock<FPhysSpinLock>;
-#else
+#elif CHAOS_SCENE_LOCK_TYPE == CHAOS_SCENE_LOCK_RWFIFO_CRITICALSECTION
 	using FPhysSceneLock = TRwFifoLock<FCriticalSection>;
-#endif
+#elif CHAOS_SCENE_LOCK_TYPE == CHAOS_SCENE_LOCK_FRWLOCK
+	using FPhysSceneLock = FPhysicsRwLock;
+#elif CHAOS_SCENE_LOCK_TYPE == CHAOS_SCENE_LOCK_SIMPLE_MUTEX
+	using FPhysSceneLock = FPhysicsSimpleMutexLock;
 #endif
 
 	// Stable types to use in calling code configured by the compiler switches above
