- Use ISMComponent events on instance removed and index reallocated in the ISMPool to handle both ISM & HISM reordering on remove instances (Remove vs SwapRemove differences)
- Re-enable HISM by default for GeometryCollection proxy and leaf rendering

#rnx

diff --git a/Engine/Source/Runtime/Experimental/GeometryCollectionEngine/Private/GeometryCollection/GeometryCollectionComponent.cpp b/Engine/Source/Runtime/Experimental/GeometryCollectionEngine/Private/GeometryCollection/GeometryCollectionComponent.cpp
--- a/Engine/Source/Runtime/Experimental/GeometryCollectionEngine/Private/GeometryCollection/GeometryCollectionComponent.cpp
+++ b/Engine/Source/Runtime/Experimental/GeometryCollectionEngine/Private/GeometryCollection/GeometryCollectionComponent.cpp
@@ -114,10 +114,10 @@ FAutoConsoleVariableRef CVarChaosGCUseISMPoolForNonFracturedParts(TEXT("p.Chaos.
 bool bChaos_GC_ForceAutoAssignISMPool = false;
 FAutoConsoleVariableRef CVarChaosGCForceAutoAssignISMPool(TEXT("p.Chaos.GC.ForceAutoAssignISMPool"), bChaos_GC_ForceAutoAssignISMPool, TEXT("When enabled, force assignement of ISMPool regardgless of the settings of the components"));
 
-bool bChaos_GC_UseHierarchicalISMForProxyMesh = false;
+bool bChaos_GC_UseHierarchicalISMForProxyMesh = true;
 FAutoConsoleVariableRef CVarChaosGCUseHierarchicalISMForProxyMesh(TEXT("p.Chaos.GC.UseHierarchicalISMForProxyMesh"), bChaos_GC_UseHierarchicalISMForProxyMesh, TEXT("When enabled along with ISM Pool, proxy mesh will prefer using HISM vs standard ISM"));
 
-bool bChaos_GC_UseHierarchicalISMForLeafMeshes = false;
+bool bChaos_GC_UseHierarchicalISMForLeafMeshes = true;
 FAutoConsoleVariableRef CVarChaosGCUseHierarchicalISMForLeafMeshes(TEXT("p.Chaos.GC.UseHierarchicalISMForLeafMeshes"), bChaos_GC_UseHierarchicalISMForLeafMeshes, TEXT("When enabled along with ISM Pool, leaf meshes will prefer using HISM vs standard ISM"));
 
 bool bChaos_GC_InitConstantDataUseParallelFor = true;
diff --git a/Engine/Source/Runtime/Experimental/GeometryCollectionEngine/Private/GeometryCollection/GeometryCollectionISMPoolComponent.cpp b/Engine/Source/Runtime/Experimental/GeometryCollectionEngine/Private/GeometryCollection/GeometryCollectionISMPoolComponent.cpp
--- a/Engine/Source/Runtime/Experimental/GeometryCollectionEngine/Private/GeometryCollection/GeometryCollectionISMPoolComponent.cpp
+++ b/Engine/Source/Runtime/Experimental/GeometryCollectionEngine/Private/GeometryCollection/GeometryCollectionISMPoolComponent.cpp
@@ -77,6 +77,8 @@ FGeometryCollectionISM::FGeometryCollectionISM(AActor* OwmingActor, const FGeome
 int32 FGeometryCollectionISM::AddInstanceGroup(int32 InstanceCount, TArrayView<const float> CustomDataFloats)
 {
 	const int32 InstanceGroupIndex = InstanceGroups.AddGroup(InstanceCount);
+
+	// When adding new group it will always have a single range
 	const FInstanceGroups::FInstanceGroupRange& NewInstanceGroup = InstanceGroups.GetGroup(InstanceGroupIndex);
 
 	ISMComponent->PreAllocateInstancesMemory(InstanceCount);
@@ -95,7 +97,7 @@ int32 FGeometryCollectionISM::AddInstanceGroup(int32 InstanceCount, TArrayView<c
 		{
 			for (int32 InstanceIndex = 0; InstanceIndex < InstanceCount; ++InstanceIndex)
 			{
-				ISMComponent->SetCustomData(NewInstanceGroup.Start + InstanceIndex, CustomDataFloats.Slice(InstanceIndex * NumCustomDataFloats, NumCustomDataFloats));
+				ISMComponent->SetCustomData(NewInstanceGroup.InstanceIdToIndex[InstanceIndex], CustomDataFloats.Slice(InstanceIndex * NumCustomDataFloats, NumCustomDataFloats));
 			}
 		}
 	}
@@ -105,6 +107,11 @@ int32 FGeometryCollectionISM::AddInstanceGroup(int32 InstanceCount, TArrayView<c
 
 FGeometryCollectionMeshInfo FGeometryCollectionISMPool::AddISM(UGeometryCollectionISMPoolComponent* OwningComponent, const FGeometryCollectionStaticMeshInstance& MeshInstance, int32 InstanceCount, TArrayView<const float> CustomDataFloats, bool bPreferHISM)
 {
+	if (!OnISMInstanceIndexUpdatedHandle.IsValid())
+	{
+		OnISMInstanceIndexUpdatedHandle = FInstancedStaticMeshDelegates::OnInstanceIndexUpdated.AddRaw(this, &FGeometryCollectionISMPool::OnISMInstanceIndexUpdated);
+	}
+
 	FGeometryCollectionMeshInfo Info;
 
 	FISMIndex* ISMIndex = MeshToISMIndex.Find(MeshInstance);
@@ -112,6 +119,7 @@ FGeometryCollectionMeshInfo FGeometryCollectionISMPool::AddISM(UGeometryCollecti
 	{
 		Info.ISMIndex = ISMs.Emplace(OwningComponent->GetOwner(), MeshInstance, bPreferHISM);
 		MeshToISMIndex.Add(MeshInstance, Info.ISMIndex);
+		ISMComponentToISMIndex.Add(ISMs[Info.ISMIndex].ISMComponent, Info.ISMIndex);
 	}
 	else
 	{
@@ -128,10 +136,30 @@ bool FGeometryCollectionISMPool::BatchUpdateInstancesTransforms(FGeometryCollect
 	{
 		FGeometryCollectionISM& ISM = ISMs[MeshInfo.ISMIndex];
 		const FInstanceGroups::FInstanceGroupRange& InstanceGroup = ISM.InstanceGroups.GetGroup(MeshInfo.InstanceGroupIndex);
-		ensure((StartInstanceIndex + NewInstancesTransforms.Num()) <= InstanceGroup.Count);
+		ensure((StartInstanceIndex + NewInstancesTransforms.Num()) <= InstanceGroup.Count());
 
-		const int32 StartIndex = InstanceGroup.Start + StartInstanceIndex;
-		return ISM.ISMComponent->BatchUpdateInstancesTransforms(StartIndex, NewInstancesTransforms, bWorldSpace, bMarkRenderStateDirty, bTeleport);
+		int32 StartIndex = InstanceGroup.InstanceIdToIndex[StartInstanceIndex];
+		int32 TransformIndex = 0;
+		int32 BatchCount = 1;
+		TArray<FTransform> BatchTransforms; //< Can't use TArrayView because blueprint function doesn't support that 
+		BatchTransforms.Reserve(NewInstancesTransforms.Num());
+		BatchTransforms.Add(NewInstancesTransforms[TransformIndex]);
+		for (int InstanceIndex = StartInstanceIndex + 1; InstanceIndex < NewInstancesTransforms.Num(); ++InstanceIndex)
+		{
+			// flush batch?
+			if (InstanceGroup.InstanceIdToIndex[InstanceIndex] != (StartIndex + BatchCount))
+			{
+				ISM.ISMComponent->BatchUpdateInstancesTransforms(StartIndex, BatchTransforms, bWorldSpace, bMarkRenderStateDirty, bTeleport);
+				StartIndex = InstanceGroup.InstanceIdToIndex[InstanceIndex];
+				BatchTransforms.SetNum(0, false);
+				BatchCount = 0;
+			}
+			
+			BatchTransforms.Add(NewInstancesTransforms[TransformIndex++]);
+			BatchCount++;
+		}
+
+		return ISM.ISMComponent->BatchUpdateInstancesTransforms(StartIndex, BatchTransforms, bWorldSpace, bMarkRenderStateDirty, bTeleport);
 	}
 	UE_LOG(LogChaos, Warning, TEXT("UGeometryCollectionISMPoolComponent : Invalid ISM Id (%d) when updating the transform "), MeshInfo.ISMIndex);
 	return false;
@@ -143,24 +171,39 @@ void FGeometryCollectionISMPool::RemoveISM(const FGeometryCollectionMeshInfo& Me
 	{
 		FGeometryCollectionISM& ISM = ISMs[MeshInfo.ISMIndex];
 		const FInstanceGroups::FInstanceGroupRange& InstanceGroup = ISM.InstanceGroups.GetGroup(MeshInfo.InstanceGroupIndex);
+		ISM.ISMComponent->RemoveInstances(InstanceGroup.InstanceIdToIndex);
+		ISM.InstanceGroups.RemoveGroup(MeshInfo.InstanceGroupIndex);
+	}
+}
 
-		// todo: be able to remove a range of instance may be useful
-		TArray<int32> InstancesToRemove;
-		InstancesToRemove.Reserve(InstanceGroup.Count);
-		const int32 StartIndex = InstanceGroup.Start;
-		const int32 LastIndex = InstanceGroup.Start + InstanceGroup.Count;
-		for (int32 InstanceIndex = StartIndex; InstanceIndex < LastIndex; InstanceIndex++)
+void FGeometryCollectionISMPool::OnISMInstanceIndexUpdated(UInstancedStaticMeshComponent* InComponent, TArrayView<const FInstancedStaticMeshDelegates::FInstanceIndexUpdateData> InIndexUpdates)
+{
+	FISMIndex* ISMIndex = ISMComponentToISMIndex.Find(InComponent);
+	if (!ISMIndex)
+	{
+		return;
+	}
+
+	FGeometryCollectionISM& ISM = ISMs[*ISMIndex];
+	check(ISM.ISMComponent == InComponent);
+
+	for (const FInstancedStaticMeshDelegates::FInstanceIndexUpdateData& IndexUpdateData : InIndexUpdates)
+	{
+		if (IndexUpdateData.Type == FInstancedStaticMeshDelegates::EInstanceIndexUpdateType::Removed)
 		{
-			InstancesToRemove.Add(InstanceIndex);
+			ISM.InstanceGroups.IndexRemoved(IndexUpdateData.Index);
+		}
+		else if (IndexUpdateData.Type == FInstancedStaticMeshDelegates::EInstanceIndexUpdateType::Relocated)
+		{
+			ISM.InstanceGroups.IndexReallocated(IndexUpdateData.OldIndex, IndexUpdateData.Index);
 		}
-		ISM.ISMComponent->RemoveInstances(InstancesToRemove);
-		ISM.InstanceGroups.RemoveGroup(MeshInfo.InstanceGroupIndex);
 	}
 }
 
 void FGeometryCollectionISMPool::Clear()
 {
 	MeshToISMIndex.Reset();
+	ISMComponentToISMIndex.Reset();
 	if (ISMs.Num() > 0)
 	{
 		if (AActor* OwningActor = ISMs[0].ISMComponent->GetOwner())
@@ -174,6 +217,12 @@ void FGeometryCollectionISMPool::Clear()
 		}
 		ISMs.Reset();
 	}
+
+	if (OnISMInstanceIndexUpdatedHandle.IsValid())
+	{
+		FInstancedStaticMeshDelegates::OnInstanceIndexUpdated.Remove(OnISMInstanceIndexUpdatedHandle);
+		OnISMInstanceIndexUpdatedHandle.Reset();
+	}
 }
 
 UGeometryCollectionISMPoolComponent::UGeometryCollectionISMPoolComponent(const FObjectInitializer& ObjectInitializer)
diff --git a/Engine/Source/Runtime/Experimental/GeometryCollectionEngine/Public/GeometryCollection/GeometryCollectionISMPoolComponent.h b/Engine/Source/Runtime/Experimental/GeometryCollectionEngine/Public/GeometryCollection/GeometryCollectionISMPoolComponent.h
--- a/Engine/Source/Runtime/Experimental/GeometryCollectionEngine/Public/GeometryCollection/GeometryCollectionISMPoolComponent.h
+++ b/Engine/Source/Runtime/Experimental/GeometryCollectionEngine/Public/GeometryCollection/GeometryCollectionISMPoolComponent.h
@@ -7,6 +7,7 @@
 #include "Components/SceneComponent.h"
 #include "Containers/Map.h"
 #include "Materials/MaterialInterface.h"
+#include "InstancedStaticMeshDelegates.h"
 #include "GeometryCollectionISMPoolComponent.generated.h"
 
 class AActor;
@@ -29,11 +30,29 @@ public:
 	struct FInstanceGroupRange
 	{
 		FInstanceGroupRange(int32 InStart, int32 InCount)
-			: Start(InStart)
-			, Count(InCount)
-		{}
-		int32 Start;
-		int32 Count;
+		{
+			InstanceIdToIndex.Reserve(InCount);
+			for (int Index = 0; Index < InCount; ++Index)
+			{
+				InstanceIdToIndex.Add(InStart + Index);
+			}
+		}
+
+		bool TryIndexReallocate(int32 OldIndex, int32 NewIndex)
+		{
+			for (int32 i = 0; i < InstanceIdToIndex.Num(); ++i)
+			{
+				if (InstanceIdToIndex[i] == OldIndex)
+				{
+					InstanceIdToIndex[i] = NewIndex;
+					return true;
+				}
+			}
+			return false;
+		}
+
+		int32 Count() const { return InstanceIdToIndex.Num(); }
+		TArray<int32> InstanceIdToIndex;
 	};
 
 	//FInstanceGroupId AddGroup(int32 Count)
@@ -74,20 +93,37 @@ public:
 	{
 		check(GroupRanges.Contains(GroupId));
 		const FInstanceGroupRange& GroupRangeToRemove = GroupRanges[GroupId];
+		InstancesCount -= GroupRangeToRemove.InstanceIdToIndex.Num();
+		GroupRanges.Remove(GroupId);
+	}
 
-		// we need now to shift all the groups above the remove one
-		for (TPair<FInstanceGroupId, FInstanceGroupRange>& GroupRange: GroupRanges)
+	void IndexRemoved(int32 IndexToRemove)
+	{
+		bool bFound = false;
+		for (TPair<FInstanceGroupId, FInstanceGroupRange>& GroupRange : GroupRanges)
 		{
-			if (GroupRange.Value.Start > GroupRangeToRemove.Start)
+			if (GroupRange.Value.TryIndexReallocate(IndexToRemove, INDEX_NONE))
 			{
-				GroupRange.Value.Start -= GroupRangeToRemove.Count;
+				bFound = true;
+				break;
 			}
 		}
-		// now remove the range safely
-		InstancesCount -= GroupRangeToRemove.Count;
-		GroupRanges.Remove(GroupId);
+		check(bFound);
 	}
 
+	void IndexReallocated(int32 OldIndex, int32 NewIndex)
+	{
+		bool bFound = false;
+		for (TPair<FInstanceGroupId, FInstanceGroupRange>& GroupRange : GroupRanges)
+		{
+			if (GroupRange.Value.TryIndexReallocate(OldIndex, NewIndex))
+			{
+				bFound = true;
+				break;
+			}
+		}
+		check(bFound);
+	}
 	const FInstanceGroupRange& GetGroup(int32 GroupIndex) const { return GroupRanges[GroupIndex]; };
 
 private:
@@ -182,13 +218,19 @@ struct FGeometryCollectionISMPool
 	FGeometryCollectionMeshInfo AddISM(UGeometryCollectionISMPoolComponent* OwningComponent, const FGeometryCollectionStaticMeshInstance& MeshInstance, int32 InstanceCount, TArrayView<const float> CustomDataFloats, bool bPreferHISM);
 	bool BatchUpdateInstancesTransforms(FGeometryCollectionMeshInfo& MeshInfo, int32 StartInstanceIndex, const TArray<FTransform>& NewInstancesTransforms, bool bWorldSpace, bool bMarkRenderStateDirty, bool bTeleport);
 	void RemoveISM(const FGeometryCollectionMeshInfo& MeshInfo);
+	
+	void OnISMInstanceIndexUpdated(UInstancedStaticMeshComponent* InComponent, TArrayView<const FInstancedStaticMeshDelegates::FInstanceIndexUpdateData> InIndexUpdates);
 
 	/** Clear all ISM components and associated data */
 	void Clear();
 
 	TMap<FGeometryCollectionStaticMeshInstance, FISMIndex> MeshToISMIndex;
+	TMap<UInstancedStaticMeshComponent*, FISMIndex> ISMComponentToISMIndex;
 	TArray<FGeometryCollectionISM> ISMs;
+	
 	// Todo : since ISMs index cannot change, we'll need a free list 
+
+	FDelegateHandle OnISMInstanceIndexUpdatedHandle;
 };
 
 
