ISMPool, fix bug where processing batched transform were shifted by 1 after index 0

diff --git a/Engine/Source/Runtime/Experimental/GeometryCollectionEngine/Private/GeometryCollection/GeometryCollectionISMPoolComponent.cpp b/Engine/Source/Runtime/Experimental/GeometryCollectionEngine/Private/GeometryCollection/GeometryCollectionISMPoolComponent.cpp
--- a/Engine/Source/Runtime/Experimental/GeometryCollectionEngine/Private/GeometryCollection/GeometryCollectionISMPoolComponent.cpp
+++ b/Engine/Source/Runtime/Experimental/GeometryCollectionEngine/Private/GeometryCollection/GeometryCollectionISMPoolComponent.cpp
@@ -143,7 +143,7 @@ bool FGeometryCollectionISMPool::BatchUpdateInstancesTransforms(FGeometryCollect
 		int32 BatchCount = 1;
 		TArray<FTransform> BatchTransforms; //< Can't use TArrayView because blueprint function doesn't support that 
 		BatchTransforms.Reserve(NewInstancesTransforms.Num());
-		BatchTransforms.Add(NewInstancesTransforms[TransformIndex]);
+		BatchTransforms.Add(NewInstancesTransforms[TransformIndex++]);
 		for (int InstanceIndex = StartInstanceIndex + 1; InstanceIndex < NewInstancesTransforms.Num(); ++InstanceIndex)
 		{
 			// flush batch?
