Animation Constraints to transact with Multi-User

diff --git a/Engine/Plugins/Developer/Concert/ConcertSync/ConcertSyncCore/Config/BaseConcertSyncCore.ini b/Engine/Plugins/Developer/Concert/ConcertSync/ConcertSyncCore/Config/BaseConcertSyncCore.ini
--- a/Engine/Plugins/Developer/Concert/ConcertSync/ConcertSyncCore/Config/BaseConcertSyncCore.ini
+++ b/Engine/Plugins/Developer/Concert/ConcertSync/ConcertSyncCore/Config/BaseConcertSyncCore.ini
@@ -3,6 +3,7 @@
 +IncludeObjectClassFilters=(ObjectOuterClass=Class'/Script/Engine.World',ObjectClasses=(Class'/Script/Engine.Layer',Class'/Script/Engine.Level',Class'/Script/Engine.LevelStreaming',Class'/Script/Engine.Actor',Class'/Script/Engine.ActorComponent',Class'/Script/Engine.Polys',Class'/Script/Engine.BrushBuilder',Class'/Script/Foliage.FoliageType'))
 +IncludeObjectClassFilters=(ObjectOuterClass=Class'/Script/MovieScene.MovieSceneSignedObject',ObjectClasses=(Class'/Script/Engine.Actor',Class'/Script/Engine.ActorComponent'))
 +IncludeObjectClassFilters=(ObjectOuterClass=Class'/Script/Engine.World',ObjectClasses=(Class'/Script/Engine.MaterialInstanceDynamic', '/Script/MediaAssets.MediaPlaylist'))
++IncludeObjectClassFilters=(ObjectOuterClass=Class'/Script/Engine.World',ObjectClasses=(Class'/Script/Constraints.TickableConstraint', '/Script/Constraints.TransformableHandle', '/Script/Constraints.ConstraintsActor', '/Script/Constraints.ConstraintsManager'))
 +ExcludeTransactionClassFilters=(ObjectOuterClass=None,ObjectClasses=(Class'/Script/Landscape.LandscapeProxy',Class'/Script/Landscape.LandscapeHeightfieldCollisionComponent',Class'/Script/Landscape.LandscapeComponent'))
 +AllowedTransientProperties="/Script/Engine.Actor:bHiddenEdTemporary"
 +AllowedTransientProperties="/Script/Engine.Actor:bHiddenEdLevel"
diff --git a/Engine/Source/Runtime/Experimental/Animation/Constraints/Private/ConstraintsActor.cpp b/Engine/Source/Runtime/Experimental/Animation/Constraints/Private/ConstraintsActor.cpp
--- a/Engine/Source/Runtime/Experimental/Animation/Constraints/Private/ConstraintsActor.cpp
+++ b/Engine/Source/Runtime/Experimental/Animation/Constraints/Private/ConstraintsActor.cpp
@@ -58,7 +58,7 @@ void AConstraintsActor::PostRegisterAllComponents()
 
 	if (ConstraintsManager == nullptr)
 	{
-		ConstraintsManager = NewObject<UConstraintsManager>(this);
+		ConstraintsManager = NewObject<UConstraintsManager>(this, NAME_None, RF_Transactional);
 	}
 	ConstraintsManager->Init(GetWorld());
 	RegisterConstraintsTickFunctions();
diff --git a/Engine/Source/Runtime/Experimental/Animation/Constraints/Private/ConstraintsManager.cpp b/Engine/Source/Runtime/Experimental/Animation/Constraints/Private/ConstraintsManager.cpp
--- a/Engine/Source/Runtime/Experimental/Animation/Constraints/Private/ConstraintsManager.cpp
+++ b/Engine/Source/Runtime/Experimental/Animation/Constraints/Private/ConstraintsManager.cpp
@@ -232,7 +232,6 @@ UConstraintsManager* UConstraintsManager::Get(UWorld* InWorld)
 #if WITH_EDITOR
 	ConstraintsActor->SetActorLabel(TEXT("Constraints Manager"));
 #endif // WITH_EDITOR
-	ConstraintsActor->ConstraintsManager->SetFlags(RF_Transactional);
 	return ConstraintsActor->ConstraintsManager;
 }
 
@@ -386,28 +385,7 @@ bool FConstraintsManagerController::DoesExistInAnyWorld(UTickableConstraint* InC
 
 UConstraintsManager* FConstraintsManagerController::GetManager() const
 {
-	if (!World)
-	{
-		return nullptr;
-	}
-	
-	// look for ConstraintsActor and return its manager
-	if (UConstraintsManager* Manager = FindManager())
-	{
-		return Manager;
-	}
-
-	// create a new ConstraintsActor
-	AConstraintsActor* ConstraintsActor = World->SpawnActor<AConstraintsActor>();
-#if WITH_EDITOR
-	ConstraintsActor->SetActorLabel(TEXT("Constraints Manager"));
-#endif // WITH_EDITOR
-	ConstraintsActor->ConstraintsManager = NewObject<UConstraintsManager>(ConstraintsActor);
-	ConstraintsActor->ConstraintsManager->Init(World);
-	// ULevel* Level = World->GetCurrentLevel();
-	// ConstraintsActor->ConstraintsManager->Level = Level;
-
-	return ConstraintsActor->ConstraintsManager;
+	return UConstraintsManager::Get(World);
 }
 
 UConstraintsManager* FConstraintsManagerController::FindManager() const
diff --git a/Engine/Source/Runtime/Experimental/Animation/Constraints/Private/TransformConstraint.cpp b/Engine/Source/Runtime/Experimental/Animation/Constraints/Private/TransformConstraint.cpp
--- a/Engine/Source/Runtime/Experimental/Animation/Constraints/Private/TransformConstraint.cpp
+++ b/Engine/Source/Runtime/Experimental/Animation/Constraints/Private/TransformConstraint.cpp
@@ -309,6 +309,12 @@ void UTickableTransformConstraint::PostEditChangeProperty(FPropertyChangedEvent&
 	}
 }
 
+void UTickableTransformConstraint::PostEditUndo()
+{
+	Super::PostEditUndo();
+	InitConstraint();
+}
+
 UTickableTransformConstraint::FOnConstraintChanged& UTickableTransformConstraint::GetOnConstraintChanged()
 {
 	return OnConstraintChanged;	
@@ -408,15 +414,7 @@ void UTickableTransformConstraint::OnActiveStateChanged() const
 void UTickableTransformConstraint::PostLoad()
 {
 	Super::PostLoad();
-	if (ConstraintTick.ConstraintFunctions.IsEmpty())
-	{
-		ConstraintTick.RegisterFunction(GetFunction() );
-	}
-
-	SetupDependencies();	
-	RegisterDelegates();
-
-	EnsurePrimaryDependency();
+	InitConstraint();
 }
 
 void UTickableTransformConstraint::PostDuplicate(bool bDuplicateForPIE)
@@ -608,6 +606,19 @@ FTickFunction* UTickableTransformConstraint::GetHandleTickFunction(const TObject
 	return InHandle->GetTickFunction();
 }
 
+void UTickableTransformConstraint::InitConstraint()
+{
+	if (ConstraintTick.ConstraintFunctions.IsEmpty())
+	{
+		ConstraintTick.RegisterFunction(GetFunction());
+	}
+
+	SetupDependencies();
+	RegisterDelegates();
+
+	EnsurePrimaryDependency();
+}
+
 /** 
  * UTickableTranslationConstraint
  **/
@@ -1364,8 +1375,7 @@ UTransformableComponentHandle* FTransformConstraintUtils::CreateHandleForSceneCo
 	UTransformableComponentHandle* ComponentHandle = nullptr;
 	if (InSceneComponent)
 	{
-		ComponentHandle = NewObject<UTransformableComponentHandle>(Outer);
-		ComponentHandle->SetFlags(RF_Transactional);
+		ComponentHandle = NewObject<UTransformableComponentHandle>(Outer, NAME_None, RF_Transactional);
 		ComponentHandle->Component = InSceneComponent;
 		ComponentHandle->SocketName = InSocketName;
 		InSceneComponent->SetMobility(EComponentMobility::Movable);
diff --git a/Engine/Source/Runtime/Experimental/Animation/Constraints/Private/TransformableHandle.cpp b/Engine/Source/Runtime/Experimental/Animation/Constraints/Private/TransformableHandle.cpp
--- a/Engine/Source/Runtime/Experimental/Animation/Constraints/Private/TransformableHandle.cpp
+++ b/Engine/Source/Runtime/Experimental/Animation/Constraints/Private/TransformableHandle.cpp
@@ -24,11 +24,6 @@ UTransformableHandle::~UTransformableHandle()
 	OnHandleModified.Clear();
 }
 
-void UTransformableHandle::PostLoad()
-{
-	Super::PostLoad();
-}
-
 UTransformableHandle::FHandleModifiedEvent& UTransformableHandle::HandleModified()
 {
 	return OnHandleModified;
diff --git a/Engine/Source/Runtime/Experimental/Animation/Constraints/Public/TransformConstraint.h b/Engine/Source/Runtime/Experimental/Animation/Constraints/Public/TransformConstraint.h
--- a/Engine/Source/Runtime/Experimental/Animation/Constraints/Public/TransformConstraint.h
+++ b/Engine/Source/Runtime/Experimental/Animation/Constraints/Public/TransformConstraint.h
@@ -151,6 +151,9 @@ protected:
 	/** Returns the handle's tick function (ensuring it lives in the same world). */
 	FTickFunction* GetHandleTickFunction(const TObjectPtr<UTransformableHandle>& InHandle) const;
 
+	/** (Re-)Registers the constraint function and (re-)binds the required delegates*/
+	void InitConstraint();
+
 #if WITH_EDITOR
 public:
 	/** Returns the constraint's label used for UI. */
@@ -162,6 +165,7 @@ public:
 
 	// UObject interface
 	virtual void PostEditChangeProperty(FPropertyChangedEvent& PropertyChangedEvent) override;
+	virtual void PostEditUndo() override;
 	// End of UObject interface
 
 	/** Returns a delegate that can be used to monitor for property changes. This can be used to monitor constraints changes
diff --git a/Engine/Source/Runtime/Experimental/Animation/Constraints/Public/TransformableHandle.h b/Engine/Source/Runtime/Experimental/Animation/Constraints/Public/TransformableHandle.h
--- a/Engine/Source/Runtime/Experimental/Animation/Constraints/Public/TransformableHandle.h
+++ b/Engine/Source/Runtime/Experimental/Animation/Constraints/Public/TransformableHandle.h
@@ -45,8 +45,6 @@ public:
 	DECLARE_EVENT_TwoParams(UTransformableHandle, FHandleModifiedEvent, UTransformableHandle*, EHandleEvent);
 	
 	virtual ~UTransformableHandle();
-
-	virtual void PostLoad() override;
 	
 	/** Sanity check to ensure the handle is safe to use. */
 	virtual bool IsValid() const PURE_VIRTUAL(IsValid, return false;);
