Make sure delegates are removed on component destruction and move delegate ownership to component as well

#rnx

diff --git a/Engine/Source/Runtime/Experimental/GeometryCollectionEngine/Private/GeometryCollection/GeometryCollectionISMPoolComponent.cpp b/Engine/Source/Runtime/Experimental/GeometryCollectionEngine/Private/GeometryCollection/GeometryCollectionISMPoolComponent.cpp
--- a/Engine/Source/Runtime/Experimental/GeometryCollectionEngine/Private/GeometryCollection/GeometryCollectionISMPoolComponent.cpp
+++ b/Engine/Source/Runtime/Experimental/GeometryCollectionEngine/Private/GeometryCollection/GeometryCollectionISMPoolComponent.cpp
@@ -107,11 +107,6 @@ int32 FGeometryCollectionISM::AddInstanceGroup(int32 InstanceCount, TArrayView<c
 
 FGeometryCollectionMeshInfo FGeometryCollectionISMPool::AddISM(UGeometryCollectionISMPoolComponent* OwningComponent, const FGeometryCollectionStaticMeshInstance& MeshInstance, int32 InstanceCount, TArrayView<const float> CustomDataFloats, bool bPreferHISM)
 {
-	if (!OnISMInstanceIndexUpdatedHandle.IsValid())
-	{
-		OnISMInstanceIndexUpdatedHandle = FInstancedStaticMeshDelegates::OnInstanceIndexUpdated.AddRaw(this, &FGeometryCollectionISMPool::OnISMInstanceIndexUpdated);
-	}
-
 	FGeometryCollectionMeshInfo Info;
 
 	FISMIndex* ISMIndex = MeshToISMIndex.Find(MeshInstance);
@@ -217,12 +212,6 @@ void FGeometryCollectionISMPool::Clear()
 		}
 		ISMs.Reset();
 	}
-
-	if (OnISMInstanceIndexUpdatedHandle.IsValid())
-	{
-		FInstancedStaticMeshDelegates::OnInstanceIndexUpdated.Remove(OnISMInstanceIndexUpdatedHandle);
-		OnISMInstanceIndexUpdatedHandle.Reset();
-	}
 }
 
 UGeometryCollectionISMPoolComponent::UGeometryCollectionISMPoolComponent(const FObjectInitializer& ObjectInitializer)
@@ -230,6 +219,18 @@ UGeometryCollectionISMPoolComponent::UGeometryCollectionISMPoolComponent(const F
 {
 }
 
+void UGeometryCollectionISMPoolComponent::OnRegister()
+{
+	FInstancedStaticMeshDelegates::OnInstanceIndexUpdated.AddUObject(this, &UGeometryCollectionISMPoolComponent::OnISMInstanceIndexUpdated);
+	Super::OnRegister();
+}
+
+void UGeometryCollectionISMPoolComponent::OnUnregister()
+{
+	FInstancedStaticMeshDelegates::OnInstanceIndexUpdated.RemoveAll(this);
+	Super::OnUnregister();
+}
+
 UGeometryCollectionISMPoolComponent::FMeshGroupId  UGeometryCollectionISMPoolComponent::CreateMeshGroup()
 {
 	MeshGroups.Add(NextMeshGroupId);
diff --git a/Engine/Source/Runtime/Experimental/GeometryCollectionEngine/Public/GeometryCollection/GeometryCollectionISMPoolComponent.h b/Engine/Source/Runtime/Experimental/GeometryCollectionEngine/Public/GeometryCollection/GeometryCollectionISMPoolComponent.h
--- a/Engine/Source/Runtime/Experimental/GeometryCollectionEngine/Public/GeometryCollection/GeometryCollectionISMPoolComponent.h
+++ b/Engine/Source/Runtime/Experimental/GeometryCollectionEngine/Public/GeometryCollection/GeometryCollectionISMPoolComponent.h
@@ -229,8 +229,6 @@ struct FGeometryCollectionISMPool
 	TArray<FGeometryCollectionISM> ISMs;
 	
 	// Todo : since ISMs index cannot change, we'll need a free list 
-
-	FDelegateHandle OnISMInstanceIndexUpdatedHandle;
 };
 
 
@@ -247,6 +245,11 @@ public:
 	using FMeshGroupId = int32;
 	using FMeshId = int32;
 
+	//~ Begin UActorComponent Interface
+	virtual void OnRegister() override;
+	virtual void OnUnregister() override;
+	//~ End UActorComponent Interface
+
 	/** 
 	* Create an Mesh group which represent an arbitrary set of mesh with their instance 
 	* no resources are created until the meshes are added for this group 
@@ -265,12 +268,21 @@ public:
 
 	/** Get the ISM component related to a mesh id. */
 	UInstancedStaticMeshComponent* GetISMForMeshId(FMeshGroupId MeshGroupId, FMeshId MeshId) const;
+	
+	/** Instance Index updated on the InstancedStaticMeshComponent which might need to be handled by the pool instance groups */
+	void OnISMInstanceIndexUpdated(UInstancedStaticMeshComponent* InComponent, TArrayView<const FInstancedStaticMeshDelegates::FInstanceIndexUpdateData> InIndexUpdates)
+	{
+		// forward to the Pool
+		Pool.OnISMInstanceIndexUpdated(InComponent, InIndexUpdates);
+	}
 
 private:
 	uint32 NextMeshGroupId = 0;
 	TMap<FMeshGroupId, FGeometryCollectionMeshGroup> MeshGroups;
 	FGeometryCollectionISMPool Pool;
 
+	FDelegateHandle OnISMInstanceIndexUpdatedHandle;
+
 	// Expose internals for debug draw support.
 	friend class UGeometryCollectionISMPoolDebugDrawComponent;
 };
