Iris
- Added additional profiler scopes to provide more context over what is being polled and quantized

diff --git a/Engine/Source/Runtime/Experimental/Iris/Core/Private/Iris/ReplicationState/PropertyReplicationState.cpp b/Engine/Source/Runtime/Experimental/Iris/Core/Private/Iris/ReplicationState/PropertyReplicationState.cpp
--- a/Engine/Source/Runtime/Experimental/Iris/Core/Private/Iris/ReplicationState/PropertyReplicationState.cpp
+++ b/Engine/Source/Runtime/Experimental/Iris/Core/Private/Iris/ReplicationState/PropertyReplicationState.cpp
@@ -12,6 +12,7 @@
 #include "UObject/PropertyPortFlags.h"
 #include "Containers/StringFwd.h"
 #include "Iris/Core/IrisDebugging.h"
+#include "Iris/Core/IrisProfiler.h"
 
 DEFINE_LOG_CATEGORY_STATIC(LogIrisRepNotify, Warning, All);
 
@@ -224,6 +225,10 @@ bool FPropertyReplicationState::PollPropertyReplicationState(const void* RESTRIC
 		const FReplicationStateDescriptor* Descriptor = ReplicationStateDescriptor;
 		const uint8* SrcBuffer = reinterpret_cast<const uint8*>(SrcStateData);
 
+#if UE_IRIS_PROFILE_PROTOCOL_NAMES
+		IRIS_PROFILER_SCOPE_TEXT(ReplicationStateDescriptor->DebugName->Name);
+#endif
+
 		const FReplicationStateMemberDescriptor* MemberDescriptors = Descriptor->MemberDescriptors;
 		const FProperty** MemberProperties = Descriptor->MemberProperties;
 		const FReplicationStateMemberPropertyDescriptor* MemberPropertyDescriptors = Descriptor->MemberPropertyDescriptors;
@@ -252,6 +257,10 @@ void FPropertyReplicationState::PushPropertyReplicationState(void* RESTRICT DstD
 		const FReplicationStateDescriptor* Descriptor = ReplicationStateDescriptor;
 		uint8* DstBuffer = reinterpret_cast<uint8*>(DstData);
 
+#if UE_IRIS_PROFILE_PROTOCOL_NAMES
+		IRIS_PROFILER_SCOPE_TEXT(ReplicationStateDescriptor->DebugName->Name);
+#endif
+
 		const FReplicationStateMemberDescriptor* MemberDescriptors = Descriptor->MemberDescriptors;
 		const FProperty** MemberProperties = Descriptor->MemberProperties;
 		const FReplicationStateMemberPropertyDescriptor* MemberPropertyDescriptors = Descriptor->MemberPropertyDescriptors;
diff --git a/Engine/Source/Runtime/Experimental/Iris/Core/Private/Iris/ReplicationSystem/ObjectReplicationBridge.cpp b/Engine/Source/Runtime/Experimental/Iris/Core/Private/Iris/ReplicationSystem/ObjectReplicationBridge.cpp
--- a/Engine/Source/Runtime/Experimental/Iris/Core/Private/Iris/ReplicationSystem/ObjectReplicationBridge.cpp
+++ b/Engine/Source/Runtime/Experimental/Iris/Core/Private/Iris/ReplicationSystem/ObjectReplicationBridge.cpp
@@ -32,11 +32,6 @@
 #include "Iris/Serialization/InternalNetSerializationContext.h"
 #include "Iris/Serialization/NetBitStreamUtil.h"
 
-#ifndef UE_IRIS_PROFILE_PROTOCOL_NAMES
-	// Set this to true to see protocol names in profiling captures. The downside is the events will add an overhead on the cpu when cpu captures are occuring
-	#define UE_IRIS_PROFILE_PROTOCOL_NAMES !UE_BUILD_SHIPPING
-#endif
-
 #define UE_LOG_OBJECTREPLICATIONBRIDGE(Category, Format, ...)  UE_LOG(LogIrisBridge, Category, TEXT("ObjectReplicationBridge(%u)::") Format, GetReplicationSystem()->GetId(), ##__VA_ARGS__)
 
 static bool bUseFrequencyBasedPolling = true;
diff --git a/Engine/Source/Runtime/Experimental/Iris/Core/Private/Iris/ReplicationSystem/ReplicationOperations.cpp b/Engine/Source/Runtime/Experimental/Iris/Core/Private/Iris/ReplicationSystem/ReplicationOperations.cpp
--- a/Engine/Source/Runtime/Experimental/Iris/Core/Private/Iris/ReplicationSystem/ReplicationOperations.cpp
+++ b/Engine/Source/Runtime/Experimental/Iris/Core/Private/Iris/ReplicationSystem/ReplicationOperations.cpp
@@ -476,7 +476,7 @@ bool FReplicationInstanceOperations::PollAndRefreshCachedPropertyData(const FRep
 	{
 		const EReplicationFragmentTraits MaskedFragmentTraits = Fragments[StateIt]->GetTraits() & (EReplicationFragmentTraits::NeedsPoll | ExcludeTraits);
 		if (MaskedFragmentTraits == EReplicationFragmentTraits::NeedsPoll)
-		{			
+		{
 			bIsStateDirty |= Fragments[StateIt]->PollReplicatedState(PollOptions);
 		}
 	}
@@ -572,6 +572,9 @@ void FReplicationInstanceOperations::CopyAndQuantizeIfDirty(FNetSerializationCon
 		// Append changemask data to bitstreams and quantize state if mask is dirty
 		if (Private::AppendMemberChangeMasks(OutChangeMaskWriter, OutConditionalChangeMaskWriter, FragmentData[StateIt].ExternalSrcBuffer, CurrentDescriptor))
 		{
+#if UE_IRIS_PROFILE_PROTOCOL_NAMES
+			IRIS_PROFILER_SCOPE_TEXT(CurrentDescriptor->DebugName->Name);
+#endif
 			FReplicationStateOperations::Quantize(Context, CurrentInternalStateBuffer, FragmentData[StateIt].ExternalSrcBuffer, CurrentDescriptor);
 		}
 		
diff --git a/Engine/Source/Runtime/Experimental/Iris/Core/Private/Iris/ReplicationSystem/ReplicationOperationsInternal.cpp b/Engine/Source/Runtime/Experimental/Iris/Core/Private/Iris/ReplicationSystem/ReplicationOperationsInternal.cpp
--- a/Engine/Source/Runtime/Experimental/Iris/Core/Private/Iris/ReplicationSystem/ReplicationOperationsInternal.cpp
+++ b/Engine/Source/Runtime/Experimental/Iris/Core/Private/Iris/ReplicationSystem/ReplicationOperationsInternal.cpp
@@ -20,6 +20,8 @@
 #include "Iris/Serialization/InternalNetSerializers.h"
 #include "Math/NumericLimits.h"
 #include "HAL/IConsoleManager.h"
+#include "Iris/Core/IrisProfiler.h"
+#include "Net/Core/Trace/NetDebugName.h"
 
 namespace UE::Net::Private
 {
@@ -83,6 +85,10 @@ uint32 FReplicationInstanceOperationsInternal::CopyObjectStateData(FNetBitStream
 				ChangeMaskStorageType* ChangeMaskData = Cache.GetChangeMaskStorage(Info);
 				ChangeMaskWriter.InitBytes(ChangeMaskData, ChangeMaskByteCount);
 
+#if UE_IRIS_PROFILE_PROTOCOL_NAMES
+				IRIS_PROFILER_SCOPE_TEXT(Object.Protocol->DebugName->Name);
+#endif
+
 				//UE_LOG(LogIris, Log, TEXT("Copying state data for ( InternalIndex: %u ) with NetRefHandle (Id=%u)"), InternalIndex, Object.RefHandle.GetId());
 				if (bUseFullCopyAndQuantize)
 				{
diff --git a/Engine/Source/Runtime/Experimental/Iris/Core/Public/Iris/Core/IrisProfiler.h b/Engine/Source/Runtime/Experimental/Iris/Core/Public/Iris/Core/IrisProfiler.h
--- a/Engine/Source/Runtime/Experimental/Iris/Core/Public/Iris/Core/IrisProfiler.h
+++ b/Engine/Source/Runtime/Experimental/Iris/Core/Public/Iris/Core/IrisProfiler.h
@@ -10,6 +10,11 @@
 #	endif
 #endif
 
+#ifndef UE_IRIS_PROFILE_PROTOCOL_NAMES
+	// Set this to true to see protocol names in profiling captures. The downside is the events will add an overhead on the cpu when cpu captures are occuring
+	#define UE_IRIS_PROFILE_PROTOCOL_NAMES !UE_BUILD_SHIPPING
+#endif
+
 //#define IRIS_USE_SUPERLUMINAL
 
 #if IRIS_PROFILER_ENABLE
diff --git a/Engine/Source/Runtime/Experimental/Iris/Core/Public/Iris/ReplicationSystem/FastArrayReplicationFragment.h b/Engine/Source/Runtime/Experimental/Iris/Core/Public/Iris/ReplicationSystem/FastArrayReplicationFragment.h
--- a/Engine/Source/Runtime/Experimental/Iris/Core/Public/Iris/ReplicationSystem/FastArrayReplicationFragment.h
+++ b/Engine/Source/Runtime/Experimental/Iris/Core/Public/Iris/ReplicationSystem/FastArrayReplicationFragment.h
@@ -148,6 +148,10 @@ bool TFastArrayReplicationFragment<FastArrayItemType, FastArrayType>::PollAllSta
 	FastArrayType* SrcArraySerializer = GetFastArraySerializerFromOwner();
 	ItemArrayType* SrcWrappedArray = reinterpret_cast<ItemArrayType*>(reinterpret_cast<uint8*>(SrcArraySerializer) + WrappedArrayOffsetRelativeFastArraySerializerProperty);
 	
+#if UE_IRIS_PROFILE_PROTOCOL_NAMES
+	IRIS_PROFILER_SCOPE_TEXT(ReplicationStateDescriptor->DebugName->Name);
+#endif
+
 	// Lookup destination data
 	FastArrayType* DstArraySerializer = GetFastArraySerializerFromReplicationState();
 	ItemArrayType* DstWrappedArray = reinterpret_cast<ItemArrayType*>(reinterpret_cast<uint8*>(DstArraySerializer) + WrappedArrayOffsetRelativeFastArraySerializerProperty);
