[Chaos][Rewind]

Remove object from rewind data immediately when its unregistered.  Fixes crash in scenequery on GT when tracing against the pending delete particle.

diff --git a/Engine/Source/Runtime/Experimental/Chaos/Private/PBDRigidsSolver.cpp b/Engine/Source/Runtime/Experimental/Chaos/Private/PBDRigidsSolver.cpp
--- a/Engine/Source/Runtime/Experimental/Chaos/Private/PBDRigidsSolver.cpp
+++ b/Engine/Source/Runtime/Experimental/Chaos/Private/PBDRigidsSolver.cpp
@@ -647,6 +647,11 @@ namespace Chaos
 				});
 
 			FGeometryParticleHandle* Handle = Proxy->GetHandle_LowLevel();
+			if(MRewindData)
+			{
+				MRewindData->RemoveObject(Handle);
+			}
+			
 			Proxy->SetHandle(nullptr);
 			const int32 OffsetForRewind = MRewindData ? MRewindData->Capacity() : 0;
 			PendingDestroyPhysicsProxy.Add(FPendingDestroyInfo{Proxy, GetCurrentFrame() + OffsetForRewind, Handle, UniqueIdx});
