add geometry collection dataflow option to compute merged convex hulls for only a subset of the transforms

diff --git a/Engine/Plugins/Experimental/GeometryCollectionPlugin/Source/GeometryCollectionNodes/Private/Dataflow/GeometryCollectionUtilityNodes.cpp b/Engine/Plugins/Experimental/GeometryCollectionPlugin/Source/GeometryCollectionNodes/Private/Dataflow/GeometryCollectionUtilityNodes.cpp
--- a/Engine/Plugins/Experimental/GeometryCollectionPlugin/Source/GeometryCollectionNodes/Private/Dataflow/GeometryCollectionUtilityNodes.cpp
+++ b/Engine/Plugins/Experimental/GeometryCollectionPlugin/Source/GeometryCollectionNodes/Private/Dataflow/GeometryCollectionUtilityNodes.cpp
@@ -63,6 +63,7 @@ FGenerateClusterConvexHullsFromLeafHullsDataflowNode::FGenerateClusterConvexHull
 	RegisterInputConnection(&Collection);
 	RegisterInputConnection(&ConvexCount);
 	RegisterInputConnection(&ErrorTolerance);
+	RegisterInputConnection(&OptionalSelectionFilter);
 
 	RegisterOutputConnection(&Collection);
 }
@@ -77,12 +78,24 @@ void FGenerateClusterConvexHullsFromLeafHullsDataflowNode::Evaluate(Dataflow::FC
 		{
 			const int32 InConvexCount = GetValue(Context, &ConvexCount);
 			const double InErrorToleranceInCm = GetValue(Context, &ErrorTolerance);
-
-			FGeometryCollectionConvexUtility::GenerateClusterConvexHullsFromLeafHulls(
-				*GeomCollection,
-				InConvexCount,
-				InErrorToleranceInCm
-			);
+			if (IsConnected(&OptionalSelectionFilter))
+			{
+				const FDataflowTransformSelection& InOptionalSelectionFilter = GetValue<FDataflowTransformSelection>(Context, &OptionalSelectionFilter);
+				FGeometryCollectionConvexUtility::GenerateClusterConvexHullsFromLeafHulls(
+					*GeomCollection,
+					InConvexCount,
+					InErrorToleranceInCm,
+					InOptionalSelectionFilter.AsArray()
+				);
+			}
+			else
+			{
+				FGeometryCollectionConvexUtility::GenerateClusterConvexHullsFromLeafHulls(
+					*GeomCollection,
+					InConvexCount,
+					InErrorToleranceInCm
+				);
+			}
 
 			SetValue<FManagedArrayCollection>(Context, static_cast<const FManagedArrayCollection>(*GeomCollection), &Collection);
 		}
diff --git a/Engine/Plugins/Experimental/GeometryCollectionPlugin/Source/GeometryCollectionNodes/Public/Dataflow/GeometryCollectionUtilityNodes.h b/Engine/Plugins/Experimental/GeometryCollectionPlugin/Source/GeometryCollectionNodes/Public/Dataflow/GeometryCollectionUtilityNodes.h
--- a/Engine/Plugins/Experimental/GeometryCollectionPlugin/Source/GeometryCollectionNodes/Public/Dataflow/GeometryCollectionUtilityNodes.h
+++ b/Engine/Plugins/Experimental/GeometryCollectionPlugin/Source/GeometryCollectionNodes/Public/Dataflow/GeometryCollectionUtilityNodes.h
@@ -4,6 +4,7 @@
 
 #include "CoreMinimal.h"
 #include "Dataflow/DataflowEngine.h"
+#include "Dataflow/DataflowSelection.h"
 #include "GeometryCollection/ManagedArrayCollection.h"
 
 #include "GeometryCollectionUtilityNodes.generated.h"
@@ -84,6 +85,10 @@ public:
 	UPROPERTY(EditAnywhere, Category = "Convex", meta = (DataflowInput, UIMin = "0", UIMax = "100.", Units = cm))
 	double ErrorTolerance = 0.0;
 
+	/** Optional transform selection to compute cluster hulls on -- if not provided, all cluster hulls will be computed. */
+	UPROPERTY(meta = (DataflowInput, DataflowIntrinsic))
+	FDataflowTransformSelection OptionalSelectionFilter;
+
 	FGenerateClusterConvexHullsFromLeafHullsDataflowNode(const Dataflow::FNodeParameters& InParam, FGuid InGuid = FGuid::NewGuid());
 
 	virtual void Evaluate(Dataflow::FContext& Context, const FDataflowOutput* Out) const override;
diff --git a/Engine/Source/Runtime/Experimental/Chaos/Private/GeometryCollection/GeometryCollectionConvexUtility.cpp b/Engine/Source/Runtime/Experimental/Chaos/Private/GeometryCollection/GeometryCollectionConvexUtility.cpp
--- a/Engine/Source/Runtime/Experimental/Chaos/Private/GeometryCollection/GeometryCollectionConvexUtility.cpp
+++ b/Engine/Source/Runtime/Experimental/Chaos/Private/GeometryCollection/GeometryCollectionConvexUtility.cpp
@@ -1450,7 +1450,16 @@ FGeometryCollectionConvexUtility::FGeometryCollectionConvexData FGeometryCollect
 }
 
 
-void FGeometryCollectionConvexUtility::GenerateClusterConvexHullsFromLeafHulls(FGeometryCollection& Collection,	int32 ConvexCount, double ErrorToleranceInCm)
+void FGeometryCollectionConvexUtility::GenerateClusterConvexHullsFromLeafHulls(FGeometryCollection& Collection, int32 ConvexCount, double ErrorToleranceInCm, const TArrayView<const int32> TransformSubset)
+{
+	GenerateClusterConvexHullsFromLeafHullsInternal(Collection, ConvexCount, ErrorToleranceInCm, true, TransformSubset);
+}
+void FGeometryCollectionConvexUtility::GenerateClusterConvexHullsFromLeafHulls(FGeometryCollection& Collection, int32 ConvexCount, double ErrorToleranceInCm)
+{
+	GenerateClusterConvexHullsFromLeafHullsInternal(Collection, ConvexCount, ErrorToleranceInCm, false, TArrayView<const int32>());
+}
+
+void FGeometryCollectionConvexUtility::GenerateClusterConvexHullsFromLeafHullsInternal(FGeometryCollection& Collection,	int32 ConvexCount, double ErrorToleranceInCm, bool bOnlySubset, const TArrayView<const int32> TransformSubset)
 {
 	static FName ConvexGroupName("Convex");
 	static FName ConvexHullAttributeName("ConvexHull");
@@ -1472,14 +1481,35 @@ void FGeometryCollectionConvexUtility::GenerateClusterConvexHullsFromLeafHulls(F
 	const TManagedArray<TSet<int32>>& Proximity = Collection.GetAttribute<TSet<int32>>("Proximity", FGeometryCollection::GeometryGroup);
 
 	const TArray<FTransform> GlobalTransforms = TransformFacade.ComputeCollectionSpaceTransforms();
-	const TArray<int32> DepthFirstTransformIndices = HierarchyFacade.GetTransformArrayInDepthFirstOrder();
-
 	const TManagedArrayAccessor<int32> SimulationTypeAttribute(Collection, FGeometryCollection::SimulationTypeAttribute, FTransformCollection::TransformGroup);
+	TArray<int32> TransformsToProcess;
+	if (bOnlySubset)
+	{
+		TransformsToProcess.Append(TransformSubset);
+		if (TransformsToProcess.IsEmpty())
+		{
+			return;
+		}
+	}
+	else // process all
+	{
+		TransformsToProcess.Reserve(GlobalTransforms.Num());
+		for (int32 Idx = 0; Idx < GlobalTransforms.Num(); ++Idx)
+		{
+			TransformsToProcess.Add(Idx);
+		}
+	}
+
+	// Simulation types must be present
+	if (!SimulationTypeAttribute.IsValid())
+	{
+		return;
+	}
 
-	for (int32 TransformIndex: DepthFirstTransformIndices)
+	for (int32 TransformIndex : TransformsToProcess)
 	{
 		// only do this for clusters
-		const bool bIsCluster = (SimulationTypeAttribute.IsValid() && SimulationTypeAttribute.Get()[TransformIndex] == FGeometryCollection::ESimulationTypes::FST_Clustered);
+		const bool bIsCluster = (SimulationTypeAttribute.Get()[TransformIndex] == FGeometryCollection::ESimulationTypes::FST_Clustered);
 		if (bIsCluster)
 		{
 			// compute using the leaf nodes 
diff --git a/Engine/Source/Runtime/Experimental/Chaos/Public/GeometryCollection/GeometryCollectionConvexUtility.h b/Engine/Source/Runtime/Experimental/Chaos/Public/GeometryCollection/GeometryCollectionConvexUtility.h
--- a/Engine/Source/Runtime/Experimental/Chaos/Public/GeometryCollection/GeometryCollectionConvexUtility.h
+++ b/Engine/Source/Runtime/Experimental/Chaos/Public/GeometryCollection/GeometryCollectionConvexUtility.h
@@ -79,6 +79,7 @@ public:
 	static FGeometryCollectionConvexData CreateNonOverlappingConvexHullData(FGeometryCollection* GeometryCollection, double FractionAllowRemove = .3, double SimplificationDistanceThreshold = 0.0, double CanExceedFraction = .5, 
 		EConvexOverlapRemoval OverlapRemovalMethod = EConvexOverlapRemoval::All, double OverlapRemovalShrinkPercent = 0.0, UE::GeometryCollectionConvexUtility::FConvexHulls* ComputedLeafHullsToModify = nullptr);
 
+	static void GenerateClusterConvexHullsFromLeafHulls(FGeometryCollection& Collection, int32 ConvexCount, double ErrorToleranceInCm, const TArrayView<const int32> OptionalTransformSubset);
 	static void GenerateClusterConvexHullsFromLeafHulls(FGeometryCollection& Collection, int32 ConvexCount, double ErrorToleranceInCm);
 
 	/** Returns the convex hull of the vertices contained in the specified geometry. */
@@ -145,5 +146,8 @@ private:
 		TArray<FTransformedConvex>& InOutConvex);
 
 	static void CreateConvexHullAttributesIfNeeded(FManagedArrayCollection& GeometryCollection);
+
+	// Implementation for GenerateClusterConvexHullsFromLeafHulls, supporting the full-collection and subset cases
+	static void GenerateClusterConvexHullsFromLeafHullsInternal(FGeometryCollection& Collection, int32 ConvexCount, double ErrorToleranceInCm, bool bOnlySubset, const TArrayView<const int32> OptionalTransformSubset);
 };
 
