Chaos Cloth: bring pattern (warp/weft 2d panel data) into the solver. Use with new XPBDAnisotropicBendingConstraints. Add new DataFlowNode (SimulationXPBDConfigNode) to actually create the new constraint.

Pattern data is now fed through from the cloth collection into the FChaosClothSimulationLodModel and is accessible from the FClothSimulationMesh. It is processed into the cloth plugin as the new struct FChaosClothingPatternData which also builds helpful maps between the patterns and welded sim mesh.

The new XPBDAnisotropicBendingConstraints are just like the XPBDBendingConstraints except the stiffness depends on Warp, Weft, and Bias stiffness values.

diff --git a/Engine/Plugins/Experimental/ChaosCloth/Source/Chaos/Private/ChaosCloth/ChaosClothConstraints.cpp b/Engine/Plugins/Experimental/ChaosCloth/Source/Chaos/Private/ChaosCloth/ChaosClothConstraints.cpp
--- a/Engine/Plugins/Experimental/ChaosCloth/Source/Chaos/Private/ChaosCloth/ChaosClothConstraints.cpp
+++ b/Engine/Plugins/Experimental/ChaosCloth/Source/Chaos/Private/ChaosCloth/ChaosClothConstraints.cpp
@@ -1,10 +1,12 @@
 // Copyright Epic Games, Inc. All Rights Reserved.
 #include "ChaosCloth/ChaosClothConstraints.h"
 #include "ChaosCloth/ChaosWeightMapTarget.h"
+#include "ChaosCloth/ChaosClothingPatternData.h"
 #include "Chaos/PBDSpringConstraints.h"
 #include "Chaos/XPBDSpringConstraints.h"
 #include "Chaos/PBDBendingConstraints.h"
 #include "Chaos/XPBDBendingConstraints.h"
+#include "Chaos/XPBDAnisotropicBendingConstraints.h"
 #include "Chaos/PBDAxialSpringConstraints.h"
 #include "Chaos/XPBDAxialSpringConstraints.h"
 #include "Chaos/PBDVolumeConstraint.h"
@@ -101,7 +103,8 @@ void FClothConstraints::AddRules(
 	const TArray<TConstArrayView<FRealSingle>>& WeightMaps,
 	const TArray<TConstArrayView<TTuple<int32, int32, FRealSingle>>>& Tethers,
 	Softs::FSolverReal MeshScale,
-	bool bEnabled)
+	bool bEnabled,
+	const FClothingPatternData* PatternData)
 {
 	// Self collisions
 	CreateSelfCollisionConstraints(ConfigProperties, TriangleMesh);
@@ -110,7 +113,7 @@ void FClothConstraints::AddRules(
 	CreateEdgeConstraints(ConfigProperties, WeightMaps, TriangleMesh);
 
 	// Bending constraints
-	CreateBendingConstraints(ConfigProperties, WeightMaps, TriangleMesh);
+	CreateBendingConstraints(ConfigProperties, WeightMaps, TriangleMesh, PatternData);
 
 	// Area constraints
 	CreateAreaConstraints(ConfigProperties, WeightMaps, TriangleMesh);
@@ -228,88 +231,104 @@ void FClothConstraints::CreateEdgeConstraints(
 void FClothConstraints::CreateBendingConstraints(
 	const Softs::FCollectionPropertyConstFacade& ConfigProperties,
 	const TArray<TConstArrayView<FRealSingle>>& WeightMaps,
-	const FTriangleMesh& TriangleMesh)
+	const FTriangleMesh& TriangleMesh,
+	const FClothingPatternData* PatternData)
 {
-	const bool bUseBendingElements = ConfigProperties.GetValue<bool>(TEXT("UseBendingElements"));
+	if (PatternData && PatternData->PatternPositions.Num() && Softs::FXPBDAnisotropicBendingConstraints::IsEnabled(ConfigProperties))
+	{
+		const TConstArrayView<FRealSingle>& BendingStiffnessMultipliers = WeightMaps[(int32)EChaosWeightMapTarget::BendingStiffness];
+		const TConstArrayView<FRealSingle>& BucklingStiffnessMultipliers = WeightMaps[(int32)EChaosWeightMapTarget::BucklingStiffness];
+		const TConstArrayView<FRealSingle> DampingMultipliers;
+		XAnisoBendingElementConstraints = MakeShared<Softs::FXPBDAnisotropicBendingConstraints>(
+			Evolution->Particles(),
+			ParticleOffset, NumParticles,
+			TriangleMesh,
+			PatternData->WeldedFaceVertexPatternPositions,
+			BendingStiffnessMultipliers,
+			BendingStiffnessMultipliers,
+			BendingStiffnessMultipliers,
+			BucklingStiffnessMultipliers,
+			BucklingStiffnessMultipliers,
+			BucklingStiffnessMultipliers,
+			DampingMultipliers,
+			ConfigProperties,
+			/*bTrimKinematicConstraints =*/ true);
 
-	if (bUseBendingElements)
+		++NumConstraintInits;  // Uses init to update the property tables
+		++NumConstraintRules;
+	}
+	else if (Softs::FXPBDBendingConstraints::IsEnabled(ConfigProperties))
 	{
-		if (Softs::FXPBDBendingConstraints::IsEnabled(ConfigProperties))
-		{
-			TArray<Chaos::TVec4<int32>> BendingElements = TriangleMesh.GetUniqueAdjacentElements();
-			const TConstArrayView<FRealSingle>& BendingStiffnessMultipliers = WeightMaps[(int32)EChaosWeightMapTarget::BendingStiffness];
-			const TConstArrayView<FRealSingle>& BucklingStiffnessMultipliers = WeightMaps[(int32)EChaosWeightMapTarget::BucklingStiffness];
-			const TConstArrayView<FRealSingle> DampingMultipliers;
-
-			XBendingElementConstraints = MakeShared<Softs::FXPBDBendingConstraints>(
-				Evolution->Particles(),
-				ParticleOffset, NumParticles,
-				MoveTemp(BendingElements),
-				BendingStiffnessMultipliers,
-				BucklingStiffnessMultipliers,
-				DampingMultipliers,
-				ConfigProperties,
-				/*bTrimKinematicConstraints =*/ true);
+		TArray<Chaos::TVec4<int32>> BendingElements = TriangleMesh.GetUniqueAdjacentElements();
+		const TConstArrayView<FRealSingle>& BendingStiffnessMultipliers = WeightMaps[(int32)EChaosWeightMapTarget::BendingStiffness];
+		const TConstArrayView<FRealSingle>& BucklingStiffnessMultipliers = WeightMaps[(int32)EChaosWeightMapTarget::BucklingStiffness];
+		const TConstArrayView<FRealSingle> DampingMultipliers;
 
-			++NumConstraintInits;  // Uses init to update the property tables
-			++NumConstraintRules;
-		}
-		else if (Softs::FPBDBendingConstraints::IsEnabled(ConfigProperties))
-		{
-			TArray<Chaos::TVec4<int32>> BendingElements = TriangleMesh.GetUniqueAdjacentElements();
-			const TConstArrayView<FRealSingle>& BendingStiffnessMultipliers = WeightMaps[(int32)EChaosWeightMapTarget::BendingStiffness];
-			const TConstArrayView<FRealSingle>& BucklingStiffnessMultipliers = WeightMaps[(int32)EChaosWeightMapTarget::BucklingStiffness];
-
-			BendingElementConstraints = MakeShared<Softs::FPBDBendingConstraints>(
-				Evolution->Particles(),
-				ParticleOffset, NumParticles,
-				MoveTemp(BendingElements),
-				BendingStiffnessMultipliers,
-				BucklingStiffnessMultipliers,
-				ConfigProperties,
-				/*bTrimKinematicConstraints =*/ true);
+		XBendingElementConstraints = MakeShared<Softs::FXPBDBendingConstraints>(
+			Evolution->Particles(),
+			ParticleOffset, NumParticles,
+			MoveTemp(BendingElements),
+			BendingStiffnessMultipliers,
+			BucklingStiffnessMultipliers,
+			DampingMultipliers,
+			ConfigProperties,
+			/*bTrimKinematicConstraints =*/ true);
 
-			++NumConstraintInits;  // Uses init to update the property tables
-			++NumConstraintRules;
-		}
+		++NumConstraintInits;  // Uses init to update the property tables
+		++NumConstraintRules;
 	}
-	else
+	else if (Softs::FPBDBendingConstraints::IsEnabled(ConfigProperties))
 	{
-		if (Softs::FXPBDBendingSpringConstraints::IsEnabled(ConfigProperties))
-		{
-			const TArray<Chaos::TVec2<int32>> CrossEdges = TriangleMesh.GetUniqueAdjacentPoints();
-			const TConstArrayView<FRealSingle>& BendingStiffnessMultipliers = WeightMaps[(int32)EChaosWeightMapTarget::BendingStiffness];
-			const TConstArrayView<FRealSingle> DampingMultipliers;
-
-			XBendingConstraints = MakeShared<Softs::FXPBDBendingSpringConstraints>(
-				Evolution->Particles(),
-				ParticleOffset, NumParticles,
-				CrossEdges,
-				BendingStiffnessMultipliers,
-				DampingMultipliers,
-				ConfigProperties,
-				/*bTrimKinematicConstraints =*/ true);
+		TArray<Chaos::TVec4<int32>> BendingElements = TriangleMesh.GetUniqueAdjacentElements();
+		const TConstArrayView<FRealSingle>& BendingStiffnessMultipliers = WeightMaps[(int32)EChaosWeightMapTarget::BendingStiffness];
+		const TConstArrayView<FRealSingle>& BucklingStiffnessMultipliers = WeightMaps[(int32)EChaosWeightMapTarget::BucklingStiffness];
 
-			++NumConstraintInits;  // Uses init to update the property tables
-			++NumConstraintRules;
-		}
-		else if (Softs::FPBDBendingSpringConstraints::IsEnabled(ConfigProperties))
-		{
-			const TArray<Chaos::TVec2<int32>> CrossEdges = TriangleMesh.GetUniqueAdjacentPoints();
-			const TConstArrayView<FRealSingle>& BendingStiffnessMultipliers = WeightMaps[(int32)EChaosWeightMapTarget::BendingStiffness];
+		BendingElementConstraints = MakeShared<Softs::FPBDBendingConstraints>(
+			Evolution->Particles(),
+			ParticleOffset, NumParticles,
+			MoveTemp(BendingElements),
+			BendingStiffnessMultipliers,
+			BucklingStiffnessMultipliers,
+			ConfigProperties,
+			/*bTrimKinematicConstraints =*/ true);
 
-			BendingConstraints = MakeShared<Softs::FPBDBendingSpringConstraints>(
-				Evolution->Particles(),
-				ParticleOffset,
-				NumParticles,
-				CrossEdges,
-				BendingStiffnessMultipliers,
-				ConfigProperties,
-				/*bTrimKinematicConstraints =*/ true);
+		++NumConstraintInits;  // Uses init to update the property tables
+		++NumConstraintRules;
+	}
+	else if (Softs::FXPBDBendingSpringConstraints::IsEnabled(ConfigProperties))
+	{
+		const TArray<Chaos::TVec2<int32>> CrossEdges = TriangleMesh.GetUniqueAdjacentPoints();
+		const TConstArrayView<FRealSingle>& BendingStiffnessMultipliers = WeightMaps[(int32)EChaosWeightMapTarget::BendingStiffness];
+		const TConstArrayView<FRealSingle> DampingMultipliers;
 
-			++NumConstraintInits;  // Uses init to update the property tables
-			++NumConstraintRules;
-		}
+		XBendingConstraints = MakeShared<Softs::FXPBDBendingSpringConstraints>(
+			Evolution->Particles(),
+			ParticleOffset, NumParticles,
+			CrossEdges,
+			BendingStiffnessMultipliers,
+			DampingMultipliers,
+			ConfigProperties,
+			/*bTrimKinematicConstraints =*/ true);
+
+		++NumConstraintInits;  // Uses init to update the property tables
+		++NumConstraintRules;
+	}
+	else if (Softs::FPBDBendingSpringConstraints::IsEnabled(ConfigProperties))
+	{
+		const TArray<Chaos::TVec2<int32>> CrossEdges = TriangleMesh.GetUniqueAdjacentPoints();
+		const TConstArrayView<FRealSingle>& BendingStiffnessMultipliers = WeightMaps[(int32)EChaosWeightMapTarget::BendingStiffness];
+
+		BendingConstraints = MakeShared<Softs::FPBDBendingSpringConstraints>(
+			Evolution->Particles(),
+			ParticleOffset,
+			NumParticles,
+			CrossEdges,
+			BendingStiffnessMultipliers,
+			ConfigProperties,
+			/*bTrimKinematicConstraints =*/ true);
+
+		++NumConstraintInits;  // Uses init to update the property tables
+		++NumConstraintRules;
 	}
 }
 
@@ -612,6 +631,20 @@ void FClothConstraints::CreateRules()
 			XBendingElementConstraints->Apply(Particles, Dt);
 		};
 	}
+	if (XAnisoBendingElementConstraints)
+	{
+		ConstraintInits[ConstraintInitIndex++] =
+			[this](Softs::FSolverParticles& Particles, const Softs::FSolverReal Dt)
+		{
+			XAnisoBendingElementConstraints->Init(Particles);
+			XAnisoBendingElementConstraints->ApplyProperties(Dt, Evolution->GetIterations());
+		};
+		ConstraintRules[ConstraintRuleIndex++] =
+			[this](Softs::FSolverParticles& Particles, const Softs::FSolverReal Dt)
+		{
+			XAnisoBendingElementConstraints->Apply(Particles, Dt);
+		};
+	}
 	if (XAreaConstraints)
 	{
 		ConstraintInits[ConstraintInitIndex++] =
diff --git a/Engine/Plugins/Experimental/ChaosCloth/Source/Chaos/Private/ChaosCloth/ChaosClothingPatternData.cpp b/Engine/Plugins/Experimental/ChaosCloth/Source/Chaos/Private/ChaosCloth/ChaosClothingPatternData.cpp
new file mode 100644
--- /dev/null
+++ b/Engine/Plugins/Experimental/ChaosCloth/Source/Chaos/Private/ChaosCloth/ChaosClothingPatternData.cpp
@@ -0,0 +1,125 @@
+// Copyright Epic Games, Inc. All Rights Reserved.
+#include "ChaosCloth/ChaosClothingPatternData.h"
+#include "Containers/ArrayView.h"
+#include "Chaos/TriangleMesh.h"
+
+namespace Chaos {
+
+FClothingPatternData::FClothingPatternData(
+	int32 InNumParticles,
+	const TConstArrayView<uint32>& InIndices,
+	const TConstArrayView<FVector2f>& InPatternPositions,
+	const TConstArrayView<uint32>& InPatternIndices,
+	const TConstArrayView<uint32>& InPatternToWeldedIndices)
+	: PatternPositions(InPatternPositions)
+	, PatternToWeldedIndices(InPatternToWeldedIndices)
+{
+	GenerateDerivedPatternData(InNumParticles, InIndices, InPatternIndices);
+}
+
+static void InitializeTriangleMesh(FTriangleMesh& TriangleMesh, const TConstArrayView<uint32>& Indices, const int32 NumParticles)
+{
+	checkSlow(Indices.Num() % 3 == 0);
+	const int32 NumElements = Indices.Num() / 3;
+	TArray<TVec3<int32>> Elements;
+	Elements.Reserve(NumElements);
+
+	for (int32 i = 0; i < NumElements; ++i)
+	{
+		const int32 Index = 3 * i;
+		Elements.Add(
+			{ static_cast<int32>(Indices[Index]),
+			 static_cast<int32>(Indices[Index + 1]),
+			 static_cast<int32>(Indices[Index + 2]) });
+	}
+
+	TriangleMesh.Init(MoveTemp(Elements), 0, NumParticles - 1);  // Init with the Offset to avoid discrepancies since the Triangle Mesh only relies on the used indices
+}
+
+void FClothingPatternData::GenerateDerivedPatternData(
+	int32 NumParticles,
+	const TConstArrayView<uint32>& Indices,
+	const TConstArrayView<uint32>& PatternIndices)
+{
+	const int32 NumPatternParticles = PatternPositions.Num();
+	InitializeTriangleMesh(PatternTriangleMesh, PatternIndices, NumPatternParticles);
+
+	if (NumPatternParticles)
+	{
+		FTriangleMesh NoOffsetWeldedMesh;
+		InitializeTriangleMesh(NoOffsetWeldedMesh, Indices, NumParticles);
+		const TArray<TVec3<int32>>& WeldedElements = NoOffsetWeldedMesh.GetElements();
+
+		auto SortedElement = [](const TVec3<int32>& Element)
+		{
+			return TVec3<int32>(Element.Min(), Element.Mid(), Element.Max());
+		};
+
+		TArray<TVec3<int32>> SortedWeldedElements;
+		SortedWeldedElements.SetNumUninitialized(WeldedElements.Num());
+		for (int32 ElemIdx = 0; ElemIdx < WeldedElements.Num(); ++ElemIdx)
+		{
+			SortedWeldedElements[ElemIdx] = SortedElement(WeldedElements[ElemIdx]);
+		}
+
+		TConstArrayView<TArray<int32>> WeldedPointToTriangleMap = NoOffsetWeldedMesh.GetPointToTriangleMap();
+		const TArray<TVec3<int32>>& PatternElements = PatternTriangleMesh.GetElements();
+
+		TArray<TVec3<int32>> WeldedToPatternFaceVertexIndices;
+		WeldedToPatternFaceVertexIndices.Init(TVec3<int32>(INDEX_NONE), NoOffsetWeldedMesh.GetNumElements());
+		for (const TVec3<int32>& PatternElement : PatternElements)
+		{
+			// Find equivalent element in the welded mesh
+			const TVec3<int32> WeldedPatternElement = TVec3<int32>(PatternToWeldedIndices[PatternElement[0]], PatternToWeldedIndices[PatternElement[1]], PatternToWeldedIndices[PatternElement[2]]);
+			const TVec3<int32> SortedWeldedPatternElement = SortedElement(WeldedPatternElement);
+			const TArray<int32>& PossibleTriangles = WeldedPointToTriangleMap[WeldedPatternElement[0]];
+
+			int32 TriangleIndex = INDEX_NONE;
+			for (int32 Triangle : PossibleTriangles)
+			{
+				const TVec3<int32>& TriangleElements = SortedWeldedElements[Triangle];
+				if (TriangleElements == SortedWeldedPatternElement)
+				{
+					TriangleIndex = Triangle;
+					break;
+				}
+			}
+			check(TriangleIndex != INDEX_NONE);
+
+			const TVec3<int32>& WeldedElement = WeldedElements[TriangleIndex];
+			auto MatchIndex = [&WeldedPatternElement, &WeldedElement](int32 WeldedIndex)
+			{
+				check(WeldedPatternElement[0] == WeldedElement[WeldedIndex] || WeldedPatternElement[1] == WeldedElement[WeldedIndex] || WeldedPatternElement[2] == WeldedElement[WeldedIndex]); // TODO: checkSlow
+				return WeldedPatternElement[0] == WeldedElement[WeldedIndex] ? 0 : WeldedPatternElement[1] == WeldedElement[WeldedIndex] ? 1 : 2;
+			};
+
+			TVec3<int32>& FaceVertex = WeldedToPatternFaceVertexIndices[TriangleIndex];
+			check(FaceVertex == TVec3<int32>(INDEX_NONE)); // TODO: checkSlow
+			FaceVertex[0] = PatternElement[MatchIndex(0)];
+			FaceVertex[1] = PatternElement[MatchIndex(1)];
+			FaceVertex[2] = PatternElement[MatchIndex(2)];
+		}
+		WeldedFaceVertexPatternPositions.SetNumUninitialized(WeldedElements.Num());
+		for (int32 ElemIdx = 0; ElemIdx < WeldedElements.Num(); ++ElemIdx)
+		{
+			const TVec3<int32>& FaceVertex = WeldedToPatternFaceVertexIndices[ElemIdx];
+
+			// TODO: checkSlow
+			const TVec3<int32>& WeldedElement = WeldedElements[ElemIdx];
+			check(!(FaceVertex == TVec3<int32>(INDEX_NONE)));
+			check(PatternToWeldedIndices[FaceVertex[0]] == WeldedElement[0]);
+			check(PatternToWeldedIndices[FaceVertex[1]] == WeldedElement[1]);
+			check(PatternToWeldedIndices[FaceVertex[2]] == WeldedElement[2]);
+
+			TVec3<FVec2f>& FaceVertexPositions = WeldedFaceVertexPatternPositions[ElemIdx];
+			FaceVertexPositions[0] = PatternPositions[FaceVertex[0]];
+			FaceVertexPositions[1] = PatternPositions[FaceVertex[1]];
+			FaceVertexPositions[2] = PatternPositions[FaceVertex[2]];
+		}
+	}
+	else
+	{
+		WeldedFaceVertexPatternPositions.Reset();
+	}
+}
+}  // End namespace Chaos
diff --git a/Engine/Plugins/Experimental/ChaosCloth/Source/Chaos/Private/ChaosCloth/ChaosClothingSimulationCloth.cpp b/Engine/Plugins/Experimental/ChaosCloth/Source/Chaos/Private/ChaosCloth/ChaosClothingSimulationCloth.cpp
--- a/Engine/Plugins/Experimental/ChaosCloth/Source/Chaos/Private/ChaosCloth/ChaosClothingSimulationCloth.cpp
+++ b/Engine/Plugins/Experimental/ChaosCloth/Source/Chaos/Private/ChaosCloth/ChaosClothingSimulationCloth.cpp
@@ -4,6 +4,7 @@
 #include "ChaosCloth/ChaosClothingSimulationMesh.h"
 #include "ChaosCloth/ChaosClothingSimulationCollider.h"
 #include "ChaosCloth/ChaosClothingSimulationConfig.h"
+#include "ChaosCloth/ChaosClothingPatternData.h"
 #include "ChaosCloth/ChaosWeightMapTarget.h"
 #include "ChaosCloth/ChaosClothPrivate.h"
 #include "Chaos/CollectionPropertyFacade.h"
@@ -46,6 +47,8 @@ struct FClothingSimulationCloth::FLODData
 	const TArray<TConstArrayView<FRealSingle>> WeightMaps;
 	const TArray<TConstArrayView<TTuple<int32, int32, FRealSingle>>> Tethers;
 
+	const FClothingPatternData PatternData;
+
 	// Per Solver data
 	struct FSolverData
 	{
@@ -55,12 +58,15 @@ struct FClothingSimulationCloth::FLODData
 	TMap<FClothingSimulationSolver*, FSolverData> SolverData;
 
 	// Stats
-	int32 NumKinenamicParticles;
-	int32 NumDynammicParticles;
+	int32 NumKinematicParticles;
+	int32 NumDynamicParticles;
 
 	FLODData(
 		int32 InNumParticles,
 		const TConstArrayView<uint32>& InIndices,
+		const TConstArrayView<FVector2f>& InPatternPositions,
+		const TConstArrayView<uint32>& InPatternIndices,
+		const TConstArrayView<uint32>& InPatternToWeldedIndices,
 		const TArray<TConstArrayView<FRealSingle>>& InWeightMaps,
 		const TArray<TConstArrayView<TTuple<int32, int32, FRealSingle>>>& InTethers);
 
@@ -79,12 +85,16 @@ struct FClothingSimulationCloth::FLODData
 FClothingSimulationCloth::FLODData::FLODData(
 	int32 InNumParticles,
 	const TConstArrayView<uint32>& InIndices,
+	const TConstArrayView<FVector2f>& InPatternPositions,
+	const TConstArrayView<uint32>& InPatternIndices,
+	const TConstArrayView<uint32>& InPatternToWeldedIndices,
 	const TArray<TConstArrayView<FRealSingle>>& InWeightMaps,
 	const TArray<TConstArrayView<TTuple<int32, int32, FRealSingle>>>& InTethers)
 	: NumParticles(InNumParticles)
 	, Indices(InIndices)
 	, WeightMaps(InWeightMaps)
 	, Tethers(InTethers)
+	, PatternData(NumParticles, Indices, InPatternPositions, InPatternIndices, InPatternToWeldedIndices)
 {
 }
 
@@ -175,21 +185,21 @@ PRAGMA_ENABLE_DEPRECATION_WARNINGS
 
 	// Create constraints
 	const bool bEnabled = false;  // Set constraint disabled by default
-	ClothConstraints.AddRules(ConfigProperties, TriangleMesh, WeightMaps, Tethers, MeshScale, bEnabled);
+	ClothConstraints.AddRules(ConfigProperties, TriangleMesh, WeightMaps, Tethers, MeshScale, bEnabled, &PatternData);
 
 	// Update LOD stats
 	const TConstArrayView<Softs::FSolverReal> InvMasses(Solver->GetParticleInvMasses(Offset), NumParticles);
-	NumKinenamicParticles = 0;
-	NumDynammicParticles = 0;
+	NumKinematicParticles = 0;
+	NumDynamicParticles = 0;
 	for (int32 Index = 0; Index < NumParticles; ++Index)
 	{
 		if (InvMasses[Index] == (Softs::FSolverReal)0.)
 		{
-			++NumKinenamicParticles;
+			++NumKinematicParticles;
 		}
 		else
 		{
-			++NumDynammicParticles;
+			++NumDynamicParticles;
 		}
 	}
 }
@@ -359,7 +369,6 @@ PRAGMA_ENABLE_DEPRECATION_WARNINGS
 	if (InBendingStiffness[0] > 0.f || InBendingStiffness[1] > 0.f ||
 		(bInUseBendingElements && (InBucklingStiffness[0] > 0.f || InBucklingStiffness[1] > 0.f)))
 	{
-		Properties.AddValue(TEXT("UseBendingElements"), bInUseBendingElements);
 		if (bInUseBendingElements)
 		{
 			const int32 BendingElementStiffnessIndex = Properties.AddProperty(TEXT("BendingElementStiffness"), bEnable, bAnimatable);
@@ -519,6 +528,9 @@ void FClothingSimulationCloth::SetMesh(FClothingSimulationMesh* InMesh)
 		LODData.Add(MakeUnique<FLODData>(
 			Mesh->GetNumPoints(Index),
 			Mesh->GetIndices(Index),
+			Mesh->GetPatternPositions(Index),
+			Mesh->GetPatternIndices(Index),
+			Mesh->GetPatternToWeldedIndices(Index),
 			Mesh->GetWeightMaps(Index),
 			Mesh->GetTethers(Index, bUseGeodesicTethers)));
 	}
@@ -842,8 +854,8 @@ PRAGMA_ENABLE_DEPRECATION_WARNINGS
 		{
 			// Enable new LOD's particles
 			LODData[LODIndex]->Enable(Solver, true);
-			NumActiveKinematicParticles = LODData[LODIndex]->NumKinenamicParticles;
-			NumActiveDynamicParticles = LODData[LODIndex]->NumDynammicParticles;
+			NumActiveKinematicParticles = LODData[LODIndex]->NumKinematicParticles;
+			NumActiveDynamicParticles = LODData[LODIndex]->NumDynamicParticles;
 
 			// Wrap new LOD based on previous LOD if possible (can only do 1 level LOD at a time, and if previous LOD exists)
 PRAGMA_DISABLE_DEPRECATION_WARNINGS  // TODO: CHAOS_IS_CLOTHINGSIMULATIONMESH_ABSTRACT
diff --git a/Engine/Plugins/Experimental/ChaosCloth/Source/Chaos/Private/ChaosCloth/ChaosClothingSimulationConfig.cpp b/Engine/Plugins/Experimental/ChaosCloth/Source/Chaos/Private/ChaosCloth/ChaosClothingSimulationConfig.cpp
--- a/Engine/Plugins/Experimental/ChaosCloth/Source/Chaos/Private/ChaosCloth/ChaosClothingSimulationConfig.cpp
+++ b/Engine/Plugins/Experimental/ChaosCloth/Source/Chaos/Private/ChaosCloth/ChaosClothingSimulationConfig.cpp
@@ -73,7 +73,6 @@ namespace Chaos
 			if (ClothConfig->BendingStiffnessWeighted.Low > 0.f || ClothConfig->BendingStiffnessWeighted.High > 0.f ||
 				(ClothConfig->bUseBendingElements && (ClothConfig->BucklingStiffnessWeighted.Low > 0.f || ClothConfig->BucklingStiffnessWeighted.High > 0.f)))
 			{
-				Properties->AddValue(TEXT("UseBendingElements"), ClothConfig->bUseBendingElements);
 				if (ClothConfig->bUseBendingElements)
 				{
 					const int32 BendingElementStiffnessIndex = Properties->AddProperty(TEXT("BendingElementStiffness"), bEnable, bAnimatable);
diff --git a/Engine/Plugins/Experimental/ChaosCloth/Source/Chaos/Private/ChaosCloth/ChaosClothingSimulationMesh.cpp b/Engine/Plugins/Experimental/ChaosCloth/Source/Chaos/Private/ChaosCloth/ChaosClothingSimulationMesh.cpp
--- a/Engine/Plugins/Experimental/ChaosCloth/Source/Chaos/Private/ChaosCloth/ChaosClothingSimulationMesh.cpp
+++ b/Engine/Plugins/Experimental/ChaosCloth/Source/Chaos/Private/ChaosCloth/ChaosClothingSimulationMesh.cpp
@@ -100,6 +100,11 @@ int32 FClothingSimulationMesh::GetNumPoints(int32 LODIndex) const
 	return IsValidLODIndex(LODIndex) ? Asset->LodData[LODIndex].PhysicalMeshData.Vertices.Num() : 0;
 }
 
+int32 FClothingSimulationMesh::GetNumPatternPoints(int32 LODIndex) const
+{
+	return 0;
+}
+
 TConstArrayView<FVector3f> FClothingSimulationMesh::GetPositions(int32 LODIndex) const
 {
 	if (IsValidLODIndex(LODIndex))
@@ -111,6 +116,11 @@ TConstArrayView<FVector3f> FClothingSimulationMesh::GetPositions(int32 LODIndex)
 	return TConstArrayView<FVector3f>();
 }
 
+TConstArrayView<FVector2f> FClothingSimulationMesh::GetPatternPositions(int32 LODIndex) const
+{
+	return TConstArrayView<FVector2f>();
+}
+
 TConstArrayView<FVector3f> FClothingSimulationMesh::GetNormals(int32 LODIndex) const
 {
 	if (IsValidLODIndex(LODIndex))
@@ -129,6 +139,16 @@ TConstArrayView<uint32> FClothingSimulationMesh::GetIndices(int32 LODIndex) cons
 		TConstArrayView<uint32>();
 }
 
+TConstArrayView<uint32> FClothingSimulationMesh::GetPatternIndices(int32 LODIndex) const
+{
+	return TConstArrayView<uint32>();
+}
+
+TConstArrayView<uint32> FClothingSimulationMesh::GetPatternToWeldedIndices(int32 LODIndex) const
+{
+	return TConstArrayView<uint32>();
+}
+
 TArray<TConstArrayView<FRealSingle>> FClothingSimulationMesh::GetWeightMaps(int32 LODIndex) const
 {
 	TArray<TConstArrayView<FRealSingle>> WeightMaps;
diff --git a/Engine/Plugins/Experimental/ChaosCloth/Source/Chaos/Private/ChaosCloth/ChaosClothingSimulationSkeletalMesh.cpp b/Engine/Plugins/Experimental/ChaosCloth/Source/Chaos/Private/ChaosCloth/ChaosClothingSimulationSkeletalMesh.cpp
--- a/Engine/Plugins/Experimental/ChaosCloth/Source/Chaos/Private/ChaosCloth/ChaosClothingSimulationSkeletalMesh.cpp
+++ b/Engine/Plugins/Experimental/ChaosCloth/Source/Chaos/Private/ChaosCloth/ChaosClothingSimulationSkeletalMesh.cpp
@@ -59,6 +59,11 @@ PRAGMA_ENABLE_DEPRECATION_WARNINGS
 		return IsValidLODIndex(LODIndex) ? Asset->LodData[LODIndex].PhysicalMeshData.Vertices.Num() : 0;
 	}
 
+	int32 FClothingSimulationSkeletalMesh::GetNumPatternPoints(int32 LODIndex) const
+	{
+		return 0;
+	}
+
 	TConstArrayView<FVector3f> FClothingSimulationSkeletalMesh::GetPositions(int32 LODIndex) const
 	{
 		if (IsValidLODIndex(LODIndex))
@@ -70,6 +75,11 @@ PRAGMA_ENABLE_DEPRECATION_WARNINGS
 		return TConstArrayView<FVector3f>();
 	}
 
+	TConstArrayView<FVector2f> FClothingSimulationSkeletalMesh::GetPatternPositions(int32 LODIndex) const
+	{
+		return TConstArrayView<FVector2f>();
+	}
+
 	TConstArrayView<FVector3f> FClothingSimulationSkeletalMesh::GetNormals(int32 LODIndex) const
 	{
 		if (IsValidLODIndex(LODIndex))
@@ -88,6 +98,16 @@ PRAGMA_ENABLE_DEPRECATION_WARNINGS
 			TConstArrayView<uint32>();
 	}
 
+	TConstArrayView<uint32> FClothingSimulationSkeletalMesh::GetPatternIndices(int32 LODIndex) const
+	{
+		return TConstArrayView<uint32>();
+	}
+
+	TConstArrayView<uint32> FClothingSimulationSkeletalMesh::GetPatternToWeldedIndices(int32 LODIndex) const
+	{
+		return TConstArrayView<uint32>();
+	}
+
 	TArray<TConstArrayView<FRealSingle>> FClothingSimulationSkeletalMesh::GetWeightMaps(int32 LODIndex) const
 	{
 		TArray<TConstArrayView<FRealSingle>> WeightMaps;
diff --git a/Engine/Plugins/Experimental/ChaosCloth/Source/Chaos/Public/ChaosCloth/ChaosClothConstraints.h b/Engine/Plugins/Experimental/ChaosCloth/Source/Chaos/Public/ChaosCloth/ChaosClothConstraints.h
--- a/Engine/Plugins/Experimental/ChaosCloth/Source/Chaos/Public/ChaosCloth/ChaosClothConstraints.h
+++ b/Engine/Plugins/Experimental/ChaosCloth/Source/Chaos/Public/ChaosCloth/ChaosClothConstraints.h
@@ -10,6 +10,9 @@
 
 namespace Chaos
 {
+
+	struct FClothingPatternData;
+
 	class FClothConstraints final
 	{
 	public:
@@ -44,7 +47,8 @@ namespace Chaos
 			const TArray<TConstArrayView<FRealSingle>>& WeightMaps,
 			const TArray<TConstArrayView<TTuple<int32, int32, FRealSingle>>>& Tethers,
 			Softs::FSolverReal MeshScale,
-			bool bEnabled);
+			bool bEnabled,
+			const FClothingPatternData* PatternData = nullptr);
 
 		void Update(
 			const Softs::FCollectionPropertyConstFacade& ConfigProperties,
@@ -181,7 +185,8 @@ namespace Chaos
 		void CreateBendingConstraints(
 			const Softs::FCollectionPropertyConstFacade& ConfigProperties,
 			const TArray<TConstArrayView<FRealSingle>>& WeightMaps,
-			const FTriangleMesh& TriangleMesh);
+			const FTriangleMesh& TriangleMesh,
+			const FClothingPatternData* PatternData);
 		void CreateAreaConstraints(
 			const Softs::FCollectionPropertyConstFacade& ConfigProperties,
 			const TArray<TConstArrayView<FRealSingle>>& WeightMaps,
@@ -221,6 +226,7 @@ namespace Chaos
 		TSharedPtr<Softs::FXPBDBendingSpringConstraints> XBendingConstraints;
 		TSharedPtr<Softs::FPBDBendingConstraints> BendingElementConstraints;
 		TSharedPtr<Softs::FXPBDBendingConstraints> XBendingElementConstraints;
+		TSharedPtr<Softs::FXPBDAnisotropicBendingConstraints> XAnisoBendingElementConstraints;
 		TSharedPtr<Softs::FPBDAreaSpringConstraints> AreaConstraints;
 		TSharedPtr<Softs::FXPBDAreaSpringConstraints> XAreaConstraints;
 		TSharedPtr<Softs::FPBDLongRangeConstraints> LongRangeConstraints; 
diff --git a/Engine/Plugins/Experimental/ChaosCloth/Source/Chaos/Public/ChaosCloth/ChaosClothingPatternData.h b/Engine/Plugins/Experimental/ChaosCloth/Source/Chaos/Public/ChaosCloth/ChaosClothingPatternData.h
new file mode 100644
--- /dev/null
+++ b/Engine/Plugins/Experimental/ChaosCloth/Source/Chaos/Public/ChaosCloth/ChaosClothingPatternData.h
@@ -0,0 +1,31 @@
+// Copyright Epic Games, Inc. All Rights Reserved.
+#pragma once
+
+#include "Chaos/Core.h"
+#include "Chaos/TriangleMesh.h"
+#include "Containers/ContainersFwd.h"
+namespace Chaos
+{
+
+struct FClothingPatternData
+{
+	FTriangleMesh PatternTriangleMesh;
+
+	const TConstArrayView<FVector2f> PatternPositions;
+	const TConstArrayView<uint32> PatternToWeldedIndices;
+	TArray<TVec3<FVec2f>> WeldedFaceVertexPatternPositions; // This will be empty if there are no PatternPositions
+
+	FClothingPatternData(
+		int32 InNumParticles,
+		const TConstArrayView<uint32>& InIndices,
+		const TConstArrayView<FVector2f>& InPatternPositions,
+		const TConstArrayView<uint32>& InPatternIndices,
+		const TConstArrayView<uint32>& InPatternToWeldedIndices);
+
+private:
+	void GenerateDerivedPatternData(
+		int32 InNumParticles,
+		const TConstArrayView<uint32>& InIndices,
+		const TConstArrayView<uint32>& InPatternIndices);
+};
+} // namespace Chaos
diff --git a/Engine/Plugins/Experimental/ChaosCloth/Source/Chaos/Public/ChaosCloth/ChaosClothingSimulationMesh.h b/Engine/Plugins/Experimental/ChaosCloth/Source/Chaos/Public/ChaosCloth/ChaosClothingSimulationMesh.h
--- a/Engine/Plugins/Experimental/ChaosCloth/Source/Chaos/Public/ChaosCloth/ChaosClothingSimulationMesh.h
+++ b/Engine/Plugins/Experimental/ChaosCloth/Source/Chaos/Public/ChaosCloth/ChaosClothingSimulationMesh.h
@@ -58,9 +58,13 @@ namespace Chaos
 		virtual int32 GetOwnerLODIndex(int32 LODIndex) const;
 		virtual bool IsValidLODIndex(int32 LODIndex) const;
 		virtual int32 GetNumPoints(int32 LODIndex) const;
+		virtual int32 GetNumPatternPoints(int32 LODIndex) const;
 		virtual TConstArrayView<FVector3f> GetPositions(int32 LODIndex) const;
+		virtual TConstArrayView<FVector2f> GetPatternPositions(int32 LODIndex) const;
 		virtual TConstArrayView<FVector3f> GetNormals(int32 LODIndex) const;
 		virtual TConstArrayView<uint32> GetIndices(int32 LODIndex) const;
+		virtual TConstArrayView<uint32> GetPatternIndices(int32 LODIndex) const;
+		virtual TConstArrayView<uint32> GetPatternToWeldedIndices(int32 LODIndex) const;
 		virtual TArray<TConstArrayView<FRealSingle>> GetWeightMaps(int32 LODIndex) const;
 		virtual TArray<TConstArrayView<TTuple<int32, int32, float>>> GetTethers(int32 LODIndex, bool bUseGeodesicTethers) const;
 		virtual int32 GetReferenceBoneIndex() const;
@@ -88,15 +92,27 @@ namespace Chaos
 		/* Return the number of points for the specified LOD, or 0 if the LOD is empty or invalid. */
 		virtual int32 GetNumPoints(int32 LODIndex) const = 0;
 
+		/* Return the number of pattern points (2d, unwelded) for the specified LOD, or 0 if patterns are not supported or the LOD is empty or invalid. */
+		virtual int32 GetNumPatternPoints(int32 LODIndex) const = 0;
+
 		/* Return the source mesh positions (pre-skinning). */
 		virtual TConstArrayView<FVector3f> GetPositions(int32 LODIndex) const = 0;
 
+		/* Return the source mesh 2d pattern positions. */
+		virtual TConstArrayView<FVector2f> GetPatternPositions(int32 LODIndex) const = 0;
+
 		/* Return the source mesh normals (pre-skinning). */
 		virtual TConstArrayView<FVector3f> GetNormals(int32 LODIndex) const = 0;
 
 		/* Return the specified LOD's triangle indices for this mesh. */
 		virtual TConstArrayView<uint32> GetIndices(int32 LODIndex) const = 0;
 
+		/* Return the specified LOD's pattern (unwelded) triangle indices for this mesh, or empty array if patterns are not supported. */
+		virtual TConstArrayView<uint32> GetPatternIndices(int32 LODIndex) const = 0;
+
+		/* Return the specified LOD's map from pattern (unwelded) vertices to (welded) vertices, or empty array if patterns are not supported. */
+		virtual TConstArrayView<uint32> GetPatternToWeldedIndices(int32 LODIndex) const = 0;
+
 		/* Return the specified LOD's weightmap. */
 		virtual TArray<TConstArrayView<FRealSingle>> GetWeightMaps(int32 LODIndex) const = 0;
 
diff --git a/Engine/Plugins/Experimental/ChaosCloth/Source/Chaos/Public/ChaosCloth/ChaosClothingSimulationSkeletalMesh.h b/Engine/Plugins/Experimental/ChaosCloth/Source/Chaos/Public/ChaosCloth/ChaosClothingSimulationSkeletalMesh.h
--- a/Engine/Plugins/Experimental/ChaosCloth/Source/Chaos/Public/ChaosCloth/ChaosClothingSimulationSkeletalMesh.h
+++ b/Engine/Plugins/Experimental/ChaosCloth/Source/Chaos/Public/ChaosCloth/ChaosClothingSimulationSkeletalMesh.h
@@ -38,9 +38,13 @@ PRAGMA_ENABLE_DEPRECATION_WARNINGS
 		virtual int32 GetOwnerLODIndex(int32 LODIndex) const override;
 		virtual bool IsValidLODIndex(int32 LODIndex) const override;
 		virtual int32 GetNumPoints(int32 LODIndex) const override;
+		virtual int32 GetNumPatternPoints(int32 LODIndex) const override;
 		virtual TConstArrayView<FVector3f> GetPositions(int32 LODIndex) const override;
+		virtual TConstArrayView<FVector2f> GetPatternPositions(int32 LODIndex) const override;
 		virtual TConstArrayView<FVector3f> GetNormals(int32 LODIndex) const override;
 		virtual TConstArrayView<uint32> GetIndices(int32 LODIndex) const override;
+		virtual TConstArrayView<uint32> GetPatternIndices(int32 LODIndex) const override;
+		virtual TConstArrayView<uint32> GetPatternToWeldedIndices(int32 LODIndex) const override;
 		virtual TArray<TConstArrayView<FRealSingle>> GetWeightMaps(int32 LODIndex) const override;
 		virtual TArray<TConstArrayView<TTuple<int32, int32, float>>> GetTethers(int32 LODIndex, bool bUseGeodesicTethers) const override;
 		virtual int32 GetReferenceBoneIndex() const override;
diff --git a/Engine/Plugins/Experimental/ChaosClothAsset/Source/ChaosClothAsset/Private/ChaosClothAsset/CollectionClothLodFacade.cpp b/Engine/Plugins/Experimental/ChaosClothAsset/Source/ChaosClothAsset/Private/ChaosClothAsset/CollectionClothLodFacade.cpp
--- a/Engine/Plugins/Experimental/ChaosClothAsset/Source/ChaosClothAsset/Private/ChaosClothAsset/CollectionClothLodFacade.cpp
+++ b/Engine/Plugins/Experimental/ChaosClothAsset/Source/ChaosClothAsset/Private/ChaosClothAsset/CollectionClothLodFacade.cpp
@@ -780,7 +780,7 @@ namespace UE::Chaos::ClothAsset
 			LodIndex);
 	}
 
-	void FCollectionClothLodConstFacade::BuildSimulationMesh(TArray<FVector3f>& Positions, TArray<FVector3f>& Normals, TArray<uint32>& Indices, TArray<int32>& WeldingMap) const
+	void FCollectionClothLodConstFacade::BuildSimulationMesh(TArray<FVector3f>& Positions, TArray<FVector3f>& Normals, TArray<uint32>& Indices, TArray<int32>& WeldingMap, TArray<FVector2f>& PatternsPositions, TArray<uint32>& PatternsIndices, TArray<uint32>& PatternToWeldedIndices) const
 	{
 		const int32 NumSimVertices = GetNumSimVertices();
 
@@ -870,14 +870,14 @@ namespace UE::Chaos::ClothAsset
 		}
 
 		// Fill up the vertex arrays
+		PatternsPositions = GetSimPosition();
 		Positions.SetNumUninitialized(NumWeldedVertices);
 		Normals.SetNumUninitialized(NumWeldedVertices);
 
 		const TConstArrayView<FVector3f> SimRestPosition = GetSimRestPosition();
 		const TConstArrayView<FVector3f> SimRestNormal = GetSimRestNormal();
 
-		TArray<uint32> WeldedIndices;
-		WeldedIndices.SetNumUninitialized(NumSimVertices);
+		PatternToWeldedIndices.SetNumUninitialized(NumSimVertices);
 
 		uint32 WeldedIndex = 0;
 		for (int32 VertexIndex = 0; VertexIndex < NumSimVertices; ++VertexIndex)
@@ -886,25 +886,28 @@ namespace UE::Chaos::ClothAsset
 			{
 				Positions[WeldedIndex] = SimRestPosition[VertexIndex];
 				Normals[WeldedIndex] = SimRestNormal[VertexIndex];
-				WeldedIndices[VertexIndex] = WeldedIndex++;
+				PatternToWeldedIndices[VertexIndex] = WeldedIndex++;
 			}
 			else
 			{
-				WeldedIndices[VertexIndex] = WeldedIndices[WeldingMap[VertexIndex]];
+				PatternToWeldedIndices[VertexIndex] = PatternToWeldedIndices[WeldingMap[VertexIndex]];
 			}
 		}
 
 		// Fill up the face array
 		const int32 NumSimFaces = GetNumSimFaces();
 		Indices.SetNumUninitialized(NumSimFaces * 3);
-
+		PatternsIndices.SetNumUninitialized(NumSimFaces * 3);
 		const TConstArrayView<FIntVector3> SimIndices = GetSimIndices();
 
 		for (int32 FaceIndex = 0; FaceIndex < NumSimFaces; ++FaceIndex)
 		{
-			Indices[FaceIndex * 3 + 0] = WeldedIndices[SimIndices[FaceIndex][0]];
-			Indices[FaceIndex * 3 + 1] = WeldedIndices[SimIndices[FaceIndex][1]];
-			Indices[FaceIndex * 3 + 2] = WeldedIndices[SimIndices[FaceIndex][2]];
+			Indices[FaceIndex * 3 + 0] = PatternToWeldedIndices[SimIndices[FaceIndex][0]];
+			Indices[FaceIndex * 3 + 1] = PatternToWeldedIndices[SimIndices[FaceIndex][1]];
+			Indices[FaceIndex * 3 + 2] = PatternToWeldedIndices[SimIndices[FaceIndex][2]];
+			PatternsIndices[FaceIndex * 3 + 0] = SimIndices[FaceIndex][0];
+			PatternsIndices[FaceIndex * 3 + 1] = SimIndices[FaceIndex][1];
+			PatternsIndices[FaceIndex * 3 + 2] = SimIndices[FaceIndex][2];
 		}
 	}
 
diff --git a/Engine/Plugins/Experimental/ChaosClothAsset/Source/ChaosClothAsset/Public/ChaosClothAsset/CollectionClothLodFacade.h b/Engine/Plugins/Experimental/ChaosClothAsset/Source/ChaosClothAsset/Public/ChaosClothAsset/CollectionClothLodFacade.h
--- a/Engine/Plugins/Experimental/ChaosClothAsset/Source/ChaosClothAsset/Public/ChaosClothAsset/CollectionClothLodFacade.h
+++ b/Engine/Plugins/Experimental/ChaosClothAsset/Source/ChaosClothAsset/Public/ChaosClothAsset/CollectionClothLodFacade.h
@@ -101,7 +101,8 @@ namespace UE::Chaos::ClothAsset
 		int32 GetLodIndex() const { return LodIndex; }
 
 		/** Return the welded simulation mesh for this LOD. */
-		void BuildSimulationMesh(TArray<FVector3f>& Positions, TArray<FVector3f>& Normals, TArray<uint32>& Indices, TArray<int32>& WeldingMap) const;
+		void BuildSimulationMesh(TArray<FVector3f>& Positions, TArray<FVector3f>& Normals, TArray<uint32>& Indices, TArray<int32>& WeldingMap,
+			TArray<FVector2f>& PatternsPositions, TArray<uint32>& PatternsIndices, TArray<uint32>& PatternToWeldedIndices) const;
 
 	protected:
 		friend class FCollectionClothConstFacade;
diff --git a/Engine/Plugins/Experimental/ChaosClothAsset/Source/ChaosClothAssetEngine/Private/ChaosClothAsset/ClothSimulationMesh.cpp b/Engine/Plugins/Experimental/ChaosClothAsset/Source/ChaosClothAssetEngine/Private/ChaosClothAsset/ClothSimulationMesh.cpp
--- a/Engine/Plugins/Experimental/ChaosClothAsset/Source/ChaosClothAssetEngine/Private/ChaosClothAsset/ClothSimulationMesh.cpp
+++ b/Engine/Plugins/Experimental/ChaosClothAsset/Source/ChaosClothAssetEngine/Private/ChaosClothAsset/ClothSimulationMesh.cpp
@@ -42,11 +42,21 @@ PRAGMA_ENABLE_DEPRECATION_WARNINGS
 		return ClothSimulationModel.GetPositions(LODIndex).Num();
 	}
 
+	int32 FClothSimulationMesh::GetNumPatternPoints(int32 LODIndex) const
+	{
+		return ClothSimulationModel.GetPatternPositions(LODIndex).Num();
+	}
+
 	TConstArrayView<FVector3f> FClothSimulationMesh::GetPositions(int32 LODIndex) const
 	{
 		return TConstArrayView<FVector3f>(ClothSimulationModel.GetPositions(LODIndex));
 	}
 
+	TConstArrayView<FVector2f> FClothSimulationMesh::GetPatternPositions(int32 LODIndex) const
+	{
+		return TConstArrayView<FVector2f>(ClothSimulationModel.GetPatternPositions(LODIndex));
+	}
+
 	TConstArrayView<FVector3f> FClothSimulationMesh::GetNormals(int32 LODIndex) const
 	{
 		return TConstArrayView<FVector3f>(ClothSimulationModel.GetNormals(LODIndex));
@@ -57,6 +67,16 @@ PRAGMA_ENABLE_DEPRECATION_WARNINGS
 		return TConstArrayView<uint32>(ClothSimulationModel.GetIndices(LODIndex));
 	}
 
+	TConstArrayView<uint32> FClothSimulationMesh::GetPatternIndices(int32 LODIndex) const
+	{
+		return TConstArrayView<uint32>(ClothSimulationModel.GetPatternIndices(LODIndex));
+	}
+
+	TConstArrayView<uint32> FClothSimulationMesh::GetPatternToWeldedIndices(int32 LODIndex) const
+	{
+		return TConstArrayView<uint32>(ClothSimulationModel.GetPatternToWeldedIndices(LODIndex));
+	}
+
 	TArray<TConstArrayView<::Chaos::FRealSingle>> FClothSimulationMesh::GetWeightMaps(int32 LODIndex) const
 	{
 		constexpr int32 MaxNumWeightMaps = 15; // TODO: Refactor how weight maps are used in base simulation classes
diff --git a/Engine/Plugins/Experimental/ChaosClothAsset/Source/ChaosClothAssetEngine/Private/ChaosClothAsset/ClothSimulationMesh.h b/Engine/Plugins/Experimental/ChaosClothAsset/Source/ChaosClothAssetEngine/Private/ChaosClothAsset/ClothSimulationMesh.h
--- a/Engine/Plugins/Experimental/ChaosClothAsset/Source/ChaosClothAssetEngine/Private/ChaosClothAsset/ClothSimulationMesh.h
+++ b/Engine/Plugins/Experimental/ChaosClothAsset/Source/ChaosClothAssetEngine/Private/ChaosClothAsset/ClothSimulationMesh.h
@@ -31,9 +31,13 @@ PRAGMA_ENABLE_DEPRECATION_WARNINGS
 		virtual int32 GetOwnerLODIndex(int32 LODIndex) const override;
 		virtual bool IsValidLODIndex(int32 LODIndex) const override;
 		virtual int32 GetNumPoints(int32 LODIndex) const override;
+		virtual int32 GetNumPatternPoints(int32 LODIndex) const override;
 		virtual TConstArrayView<FVector3f> GetPositions(int32 LODIndex) const override;
+		virtual TConstArrayView<FVector2f> GetPatternPositions(int32 LODIndex) const override;
 		virtual TConstArrayView<FVector3f> GetNormals(int32 LODIndex) const override;
 		virtual TConstArrayView<uint32> GetIndices(int32 LODIndex) const override;
+		virtual TConstArrayView<uint32> GetPatternIndices(int32 LODIndex) const override;
+		virtual TConstArrayView<uint32> GetPatternToWeldedIndices(int32 LODIndex) const override;
 		virtual TArray<TConstArrayView<::Chaos::FRealSingle>> GetWeightMaps(int32 LODIndex) const override;
 		virtual TArray<TConstArrayView<TTuple<int32, int32, float>>> GetTethers(int32 LODIndex, bool bUseGeodesicTethers) const override;
 		virtual int32 GetReferenceBoneIndex() const override;
diff --git a/Engine/Plugins/Experimental/ChaosClothAsset/Source/ChaosClothAssetEngine/Private/ChaosClothAsset/ClothSimulationModel.cpp b/Engine/Plugins/Experimental/ChaosClothAsset/Source/ChaosClothAssetEngine/Private/ChaosClothAsset/ClothSimulationModel.cpp
--- a/Engine/Plugins/Experimental/ChaosClothAsset/Source/ChaosClothAssetEngine/Private/ChaosClothAsset/ClothSimulationModel.cpp
+++ b/Engine/Plugins/Experimental/ChaosClothAsset/Source/ChaosClothAssetEngine/Private/ChaosClothAsset/ClothSimulationModel.cpp
@@ -103,6 +103,34 @@ namespace UE::Chaos::ClothAsset
 		}
 		return ReferenceBoneIndex;
 	}
+
+	// Welded weights are averaged. TODO: add other welding options such as Max and Min
+	static TArray<float> WeldWeightMap(const TConstArrayView<float>& UnweldedMap, const TArray<uint32>& PatternToWeldedIndices, const int32 NumWelded)
+	{
+		TArray<float> WeldedMap;
+		WeldedMap.SetNumZeroed(NumWelded);
+
+		TArray<float> NumInfluences;
+		NumInfluences.SetNumZeroed(NumWelded);
+		
+		check(UnweldedMap.Num() == PatternToWeldedIndices.Num());
+		for (int32 OrigIdx = 0; OrigIdx < UnweldedMap.Num(); ++OrigIdx)
+		{
+			const int32 WeldedIdx = PatternToWeldedIndices[OrigIdx];
+			NumInfluences[WeldedIdx] += 1.f;
+			WeldedMap[WeldedIdx] += UnweldedMap[OrigIdx];
+		}
+
+		for (int32 WeldedIdx = 0; WeldedIdx < NumWelded; ++WeldedIdx)
+		{
+			if (NumInfluences[WeldedIdx] > 0.f)
+			{
+				WeldedMap[WeldedIdx] /= NumInfluences[WeldedIdx];
+			}
+		}
+		return WeldedMap;
+	}
+
 }  // End namespace UE::Chaos::ClothAsset
 
 FChaosClothSimulationModel::FChaosClothSimulationModel(const TSharedPtr<const FManagedArrayCollection>& ClothCollection, const FReferenceSkeleton& ReferenceSkeleton)
@@ -120,15 +148,19 @@ FChaosClothSimulationModel::FChaosClothSimulationModel(const TSharedPtr<const FM
 		FChaosClothSimulationLodModel& LodModel = ClothSimulationLodModels[LodIndex];
 		const FCollectionClothLodConstFacade ClothLod = Cloth.GetLod(LodIndex);
 		TArray<int32> WeldingMap;
-		ClothLod.BuildSimulationMesh(LodModel.Positions, LodModel.Normals, LodModel.Indices, WeldingMap);
+		ClothLod.BuildSimulationMesh(LodModel.Positions, LodModel.Normals, LodModel.Indices, WeldingMap, LodModel.PatternPositions, LodModel.PatternIndices, LodModel.PatternToWeldedIndices);
 
 		static const FName MaxDistanceName = TEXT("MaxDistance");
 		const TConstArrayView<float> MaxDistanceValues = ClothLod.GetWeightMap(MaxDistanceName);
 
-		if (MaxDistanceValues.Num())
+		if (MaxDistanceValues.Num() == LodModel.Positions.Num())
 		{
 			LodModel.MaxDistance = MaxDistanceValues;
 		}
+		else if(MaxDistanceValues.Num() > 0)
+		{
+			LodModel.MaxDistance = WeldWeightMap(MaxDistanceValues, LodModel.PatternToWeldedIndices, LodModel.Positions.Num());
+		}
 		else
 		{
 			LodModel.MaxDistance.Init(1.0, LodModel.Positions.Num());
diff --git a/Engine/Plugins/Experimental/ChaosClothAsset/Source/ChaosClothAssetEngine/Private/ChaosClothAsset/ClothSimulationModel.h b/Engine/Plugins/Experimental/ChaosClothAsset/Source/ChaosClothAssetEngine/Private/ChaosClothAsset/ClothSimulationModel.h
--- a/Engine/Plugins/Experimental/ChaosClothAsset/Source/ChaosClothAssetEngine/Private/ChaosClothAsset/ClothSimulationModel.h
+++ b/Engine/Plugins/Experimental/ChaosClothAsset/Source/ChaosClothAssetEngine/Private/ChaosClothAsset/ClothSimulationModel.h
@@ -39,6 +39,18 @@ struct FChaosClothSimulationLodModel
 
 	UPROPERTY()
 	TArray<float> MaxDistance;
+
+	/** 2d pattern positions. */
+	UPROPERTY()
+	TArray<FVector2f> PatternPositions;
+
+	/** Pattern triangle indices. Indexing into PatternPositions. */
+	UPROPERTY()
+	TArray<uint32> PatternIndices;
+
+	/** Map from PatternsPositions indices to Positions indices. */
+	UPROPERTY()
+	TArray<uint32> PatternToWeldedIndices;
 };
 
 /**
@@ -76,7 +88,10 @@ struct FChaosClothSimulationModel
 	int32 GetNumVertices(int32 LodIndex) const { return IsValidLodIndex(LodIndex) ? ClothSimulationLodModels[LodIndex].Positions.Num() : 0; }
 
 	TConstArrayView<FVector3f> GetPositions(int32 LodIndex) const { return IsValidLodIndex(LodIndex) ? ClothSimulationLodModels[LodIndex].Positions : TConstArrayView<FVector3f>(); }
+	TConstArrayView<FVector2f> GetPatternPositions(int32 LodIndex) const { return IsValidLodIndex(LodIndex) ? ClothSimulationLodModels[LodIndex].PatternPositions : TConstArrayView<FVector2f>(); }
 	TConstArrayView<FVector3f> GetNormals(int32 LodIndex) const { return IsValidLodIndex(LodIndex) ? ClothSimulationLodModels[LodIndex].Normals : TConstArrayView<FVector3f>(); }
 	TConstArrayView<uint32> GetIndices(int32 LodIndex) const { return IsValidLodIndex(LodIndex) ? ClothSimulationLodModels[LodIndex].Indices : TConstArrayView<uint32>(); }
+	TConstArrayView<uint32> GetPatternIndices(int32 LodIndex) const { return IsValidLodIndex(LodIndex) ? ClothSimulationLodModels[LodIndex].PatternIndices : TConstArrayView<uint32>(); }
+	TConstArrayView<uint32> GetPatternToWeldedIndices(int32 LodIndex) const { return IsValidLodIndex(LodIndex) ? ClothSimulationLodModels[LodIndex].PatternToWeldedIndices : TConstArrayView<uint32>(); }
 	TConstArrayView<FClothVertBoneData> GetBoneData(int32 LodIndex) const { return IsValidLodIndex(LodIndex) ? ClothSimulationLodModels[LodIndex].BoneData : TConstArrayView<FClothVertBoneData>(); }
 };
diff --git a/Engine/Plugins/Experimental/ChaosClothAssetEditor/Source/ChaosClothAssetEditor/Private/ChaosClothAsset/DataflowNodes/DataflowNodes.cpp b/Engine/Plugins/Experimental/ChaosClothAssetEditor/Source/ChaosClothAssetEditor/Private/ChaosClothAsset/DataflowNodes/DataflowNodes.cpp
--- a/Engine/Plugins/Experimental/ChaosClothAssetEditor/Source/ChaosClothAssetEditor/Private/ChaosClothAsset/DataflowNodes/DataflowNodes.cpp
+++ b/Engine/Plugins/Experimental/ChaosClothAssetEditor/Source/ChaosClothAssetEditor/Private/ChaosClothAsset/DataflowNodes/DataflowNodes.cpp
@@ -10,6 +10,7 @@
 #include "ChaosClothAsset/DataflowNodes/CopySimulationToRenderMeshNode.h"
 #include "ChaosClothAsset/DataflowNodes/ReverseNormalsNode.h"
 #include "ChaosClothAsset/DataflowNodes/SimulationDefaultConfigNode.h"
+#include "ChaosClothAsset/DataflowNodes/SimulationXPBDConfigNode.h"
 #include "ChaosClothAsset/DataflowNodes/BindToRootBoneNode.h"
 #include "Dataflow/DataflowNodeFactory.h"
 #include "Framework/Notifications/NotificationManager.h"
@@ -30,6 +31,7 @@ namespace UE::Chaos::ClothAsset::DataflowNodes
 		DATAFLOW_NODE_REGISTER_CREATION_FACTORY(FChaosClothAssetCopySimulationToRenderMeshNode);
 		DATAFLOW_NODE_REGISTER_CREATION_FACTORY(FChaosClothAssetReverseNormalsNode);
 		DATAFLOW_NODE_REGISTER_CREATION_FACTORY(FChaosClothAssetSimulationDefaultConfigNode);
+		DATAFLOW_NODE_REGISTER_CREATION_FACTORY(FChaosClothAssetSimulationXPBDConfigNode);
 		DATAFLOW_NODE_REGISTER_CREATION_FACTORY(FChaosClothAssetBindToRootBoneNode);
 	}
 
diff --git a/Engine/Plugins/Experimental/ChaosClothAssetEditor/Source/ChaosClothAssetEditor/Private/ChaosClothAsset/DataflowNodes/SimulationXPBDConfigNode.cpp b/Engine/Plugins/Experimental/ChaosClothAssetEditor/Source/ChaosClothAssetEditor/Private/ChaosClothAsset/DataflowNodes/SimulationXPBDConfigNode.cpp
new file mode 100644
--- /dev/null
+++ b/Engine/Plugins/Experimental/ChaosClothAssetEditor/Source/ChaosClothAssetEditor/Private/ChaosClothAsset/DataflowNodes/SimulationXPBDConfigNode.cpp
@@ -0,0 +1,185 @@
+// Copyright Epic Games, Inc. All Rights Reserved.
+
+#include "ChaosClothAsset/DataflowNodes/SimulationXPBDConfigNode.h"
+#include "ChaosClothAsset/DataflowNodes/DataflowNodes.h"
+#include "ChaosCloth/ChaosClothConfig.h"
+#include "ChaosCloth/ChaosClothingSimulationConfig.h"
+#include "Chaos/CollectionPropertyFacade.h"
+
+#include UE_INLINE_GENERATED_CPP_BY_NAME(SimulationXPBDConfigNode)
+
+#define LOCTEXT_NAMESPACE "ChaosClothAssetSimulationXPBDConfigNode"
+
+FChaosClothAssetSimulationXPBDConfigNode::FChaosClothAssetSimulationXPBDConfigNode(const Dataflow::FNodeParameters& InParam, FGuid InGuid)
+	: FDataflowNode(InParam, InGuid)
+	, SharedSimulationConfig(NewObject<UChaosClothSharedSimConfig>())
+{
+	RegisterInputConnection(&Collection);
+
+	RegisterInputConnection(&bEnableXPBDStretchConstraints);
+	RegisterInputConnection(&StretchStiffness);
+	RegisterInputConnection(&StretchDampingRatio);
+
+	RegisterInputConnection(&bEnableXPBDBendConstraints);
+	RegisterInputConnection(&bEnableBendAnisotropy);
+	RegisterInputConnection(&BendingStiffness);
+	RegisterInputConnection(&BendingStiffnessWeft);
+	RegisterInputConnection(&BendingStiffnessBias);
+	RegisterInputConnection(&BendingDampingRatio);
+	RegisterInputConnection(&BucklingRatio);
+	RegisterInputConnection(&BucklingStiffness);
+	RegisterInputConnection(&BucklingStiffnessWeft);
+	RegisterInputConnection(&BucklingStiffnessBias);
+
+	RegisterInputConnection(&bEnableXPBDAreaConstraints);
+	RegisterInputConnection(&AreaStiffness);
+
+	RegisterOutputConnection(&Collection, &Collection);
+}
+
+void FChaosClothAssetSimulationXPBDConfigNode::Serialize(FArchive& Ar)
+{
+	if (!SharedSimulationConfig)
+	{
+		SharedSimulationConfig = NewObject<UChaosClothSharedSimConfig>();
+	}
+	SharedSimulationConfig->Serialize(Ar);
+}
+
+void FChaosClothAssetSimulationXPBDConfigNode::Evaluate(Dataflow::FContext& Context, const FDataflowOutput* Out) const
+{
+	using namespace Chaos;
+	using namespace Chaos::Softs;
+
+	if (Out->IsA<FManagedArrayCollection>(&Collection))
+	{
+		FManagedArrayCollection InCollection = GetValue<FManagedArrayCollection>(Context, &Collection);
+		const TSharedRef<FManagedArrayCollection> ClothCollection = MakeShared<FManagedArrayCollection>(MoveTemp(InCollection));
+
+		FCollectionPropertyMutableFacade Properties(ClothCollection);
+		Properties.DefineSchema();
+
+		// Set xpbd properties
+		SetStretchProperties(Properties);
+		SetBendingProperties(Properties);
+		SetAreaProperties(Properties);
+		
+		// Copy shared config properties
+		// Solver properties
+		if (SharedSimulationConfig)
+		{
+			constexpr bool bEnable = true;
+			constexpr bool bAnimatable = true;
+			Properties.AddValue(TEXT("NumIterations"), SharedSimulationConfig->IterationCount, bEnable, bAnimatable);
+			Properties.AddValue(TEXT("MaxNumIterations"), SharedSimulationConfig->MaxIterationCount, bEnable, bAnimatable);
+			Properties.AddValue(TEXT("NumSubsteps"), SharedSimulationConfig->SubdivisionCount, bEnable, bAnimatable);
+		}
+
+
+		SetValue<FManagedArrayCollection>(Context, *ClothCollection, &Collection);
+	}
+}
+
+void FChaosClothAssetSimulationXPBDConfigNode::AddReferencedObjects(FReferenceCollector& Collector)
+{
+	Collector.AddReferencedObject(SharedSimulationConfig);
+}
+
+FString FChaosClothAssetSimulationXPBDConfigNode::GetReferencerName() const
+{
+	return TEXT("FChaosClothAssetSimulationXPBDConfigNode");
+}
+
+void FChaosClothAssetSimulationXPBDConfigNode::SetStretchProperties(Chaos::Softs::FCollectionPropertyMutableFacade& Properties) const
+{
+	if (bEnableXPBDStretchConstraints)
+	{
+		const int32 StretchStiffnessIndex = Properties.AddProperty(TEXT("XPBDEdgeSpringStiffness"));
+		Properties.SetWeightedValue(StretchStiffnessIndex, StretchStiffness.Low, StretchStiffness.High);
+		Properties.SetStringValue(StretchStiffnessIndex, TEXT("EdgeStiffness"));
+
+		const int32 StretchDampingRatioIndex = Properties.AddProperty(TEXT("XPBDEdgeSpringDamping"));
+		Properties.SetWeightedValue(StretchDampingRatioIndex, StretchDampingRatio.Low, StretchDampingRatio.High);
+		Properties.SetStringValue(StretchDampingRatioIndex, TEXT("EdgeDamping"));
+	}
+	else
+	{
+		Properties.SetEnabled(TEXT("XPBDEdgeSpringStiffness"), false);
+	}
+}
+
+void FChaosClothAssetSimulationXPBDConfigNode::SetBendingProperties(Chaos::Softs::FCollectionPropertyMutableFacade& Properties) const
+{
+	if(bEnableXPBDBendConstraints)
+	{
+		if (bEnableBendAnisotropy)
+		{
+			const int32 BendingStiffnessWarpIndex = Properties.AddProperty(TEXT("XPBDBendingElementStiffnessWarp"));
+			Properties.SetWeightedValue(BendingStiffnessWarpIndex, BendingStiffness.Low, BendingStiffness.High);
+			Properties.SetStringValue(BendingStiffnessWarpIndex, TEXT("BendingStiffnessWarp"));
+			const int32 BendingStiffnessWeftIndex = Properties.AddProperty(TEXT("XPBDBendingElementStiffnessWeft"));
+			Properties.SetWeightedValue(BendingStiffnessWeftIndex, BendingStiffnessWeft.Low, BendingStiffnessWeft.High);
+			Properties.SetStringValue(BendingStiffnessWeftIndex, TEXT("BendingStiffnessWeft"));
+			const int32 BendingStiffnessBiasIndex = Properties.AddProperty(TEXT("XPBDBendingElementStiffnessBias"));
+			Properties.SetWeightedValue(BendingStiffnessBiasIndex, BendingStiffnessBias.Low, BendingStiffnessBias.High);
+			Properties.SetStringValue(BendingStiffnessBiasIndex, TEXT("BendingStiffnessBias"));
+
+			const int32 BendingDampingRatioIndex = Properties.AddProperty(TEXT("XPBDBendingElementDamping"));
+			Properties.SetWeightedValue(BendingDampingRatioIndex, BendingDampingRatio.Low, BendingDampingRatio.High);
+			Properties.SetStringValue(BendingDampingRatioIndex, TEXT("BendingDampingRatio"));
+
+			Properties.AddValue(TEXT("XPBDBucklingRatio"), BucklingRatio, true, true);
+
+			const int32 BucklingStiffnessWarpIndex = Properties.AddProperty(TEXT("XPBDBucklingStiffnessWarp"));
+			Properties.SetWeightedValue(BucklingStiffnessWarpIndex, BucklingStiffness.Low, BucklingStiffness.High);
+			Properties.SetStringValue(BucklingStiffnessWarpIndex, TEXT("BucklingStiffnessWarp"));
+			const int32 BucklingStiffnessWeftIndex = Properties.AddProperty(TEXT("XPBDBucklingStiffnessWeft"));
+			Properties.SetWeightedValue(BucklingStiffnessWeftIndex, BucklingStiffnessWeft.Low, BucklingStiffnessWeft.High);
+			Properties.SetStringValue(BucklingStiffnessWeftIndex, TEXT("BucklingStiffnessWeft"));
+			const int32 BucklingStiffnessBiasIndex = Properties.AddProperty(TEXT("XPBDBucklingStiffnessBias"));
+			Properties.SetWeightedValue(BucklingStiffnessBiasIndex, BucklingStiffnessBias.Low, BucklingStiffnessBias.High);
+			Properties.SetStringValue(BucklingStiffnessBiasIndex, TEXT("BucklingStiffnessBias"));
+
+			Properties.SetEnabled(TEXT("XPBDBendingElementStiffness"), false);
+		}
+		else
+		{
+			const int32 BendingStiffnessIndex = Properties.AddProperty(TEXT("XPBDBendingElementStiffness"));
+			Properties.SetWeightedValue(BendingStiffnessIndex, BendingStiffness.Low, BendingStiffness.High);
+			Properties.SetStringValue(BendingStiffnessIndex, TEXT("BendingStiffness"));
+
+			const int32 BendingDampingRatioIndex = Properties.AddProperty(TEXT("XPBDBendingElementDamping"));
+			Properties.SetWeightedValue(BendingDampingRatioIndex, BendingDampingRatio.Low, BendingDampingRatio.High);
+			Properties.SetStringValue(BendingDampingRatioIndex, TEXT("BendingDampingRatio"));
+
+			Properties.AddValue(TEXT("XPBDBucklingRatio"), BucklingRatio, true, true);
+
+			const int32 BucklingStiffnessIndex = Properties.AddProperty(TEXT("XPBDBucklingStiffness"));
+			Properties.SetWeightedValue(BucklingStiffnessIndex, BucklingStiffness.Low, BucklingStiffness.High);
+			Properties.SetStringValue(BucklingStiffnessIndex, TEXT("BucklingStiffness"));
+
+			Properties.SetEnabled(TEXT("XPBDBendingElementStiffnessWarp"), false);
+		}
+	}
+	else
+	{
+		Properties.SetEnabled(TEXT("XPBDBendingElementStiffness"), false);
+		Properties.SetEnabled(TEXT("XPBDBendingElementStiffnessWarp"), false);
+	}
+}
+
+void FChaosClothAssetSimulationXPBDConfigNode::SetAreaProperties(Chaos::Softs::FCollectionPropertyMutableFacade& Properties) const
+{
+	if (bEnableXPBDAreaConstraints)
+	{
+		const int32 AreaStiffnessIndex = Properties.AddProperty(TEXT("XPBDAreaSpringStiffness"));
+		Properties.SetWeightedValue(AreaStiffnessIndex, AreaStiffness.Low, AreaStiffness.High);
+		Properties.SetStringValue(AreaStiffnessIndex, TEXT("AreaStiffness"));
+	}
+	else
+	{
+		Properties.SetEnabled(TEXT("XPBDAreaSpringStiffness"), false);
+	}
+}
+
+#undef LOCTEXT_NAMESPACE
diff --git a/Engine/Plugins/Experimental/ChaosClothAssetEditor/Source/ChaosClothAssetEditor/Private/ChaosClothAsset/DataflowNodes/SimulationXPBDConfigNode.h b/Engine/Plugins/Experimental/ChaosClothAssetEditor/Source/ChaosClothAssetEditor/Private/ChaosClothAsset/DataflowNodes/SimulationXPBDConfigNode.h
new file mode 100644
--- /dev/null
+++ b/Engine/Plugins/Experimental/ChaosClothAssetEditor/Source/ChaosClothAssetEditor/Private/ChaosClothAsset/DataflowNodes/SimulationXPBDConfigNode.h
@@ -0,0 +1,93 @@
+// Copyright Epic Games, Inc. All Rights Reserved.
+
+#pragma once 
+
+#include "Dataflow/DataflowNode.h"
+#include "GeometryCollection/ManagedArrayCollection.h"
+#include "ChaosCloth/ChaosClothConfig.h"
+#include "SimulationXPBDConfigNode.generated.h"
+
+class UChaosClothSharedSimConfig;
+class UMaterialInterface;
+namespace Chaos::Softs
+{
+class FCollectionPropertyMutableFacade;
+}
+
+USTRUCT(Meta = (DataflowCloth))
+struct FChaosClothAssetSimulationXPBDConfigNode : public FDataflowNode, public FGCObject
+{
+	GENERATED_USTRUCT_BODY()
+	DATAFLOW_NODE_DEFINE_INTERNAL(FChaosClothAssetSimulationXPBDConfigNode, "SimulationXPBDConfig", "Cloth", "Cloth Simulation XPBD Config")
+	//DATAFLOW_NODE_RENDER_TYPE(FManagedArrayCollection::StaticType(), "Collection")  // TODO: Leave out the render type until there is something to render
+
+public:
+
+	UPROPERTY(Meta = (Dataflowinput, DataflowOutput, DisplayName = "Collection", DataflowPassthrough = "Collection"))
+	FManagedArrayCollection Collection;
+
+	/** Cloth Simulation Properties. */
+	UPROPERTY(EditAnywhere, Category = "Stretch")
+	bool bEnableXPBDStretchConstraints = true;
+
+	UPROPERTY(EditAnywhere, Category = "Stretch", DisplayName = "Stretch Stiffness [kg/s^2]", meta = (UIMin = "0", UIMax = "10000", ClampMin = "0", ClampMax = "10000000", EditCondition = "bEnableXPBDStretchConstraints"))
+	FChaosClothWeightedValue StretchStiffness = { 100.f, 100.f };
+	UPROPERTY(EditAnywhere, Category = "Stretch", DisplayName = "Stretch Damping Ratio", meta = (UIMin = "0", UIMax = "10", ClampMin = "0", ClampMax = "1000", EditCondition = "bEnableXPBDStretchConstraints"))
+	FChaosClothWeightedValue StretchDampingRatio = { 1.f, 1.f };
+
+	UPROPERTY(EditAnywhere, Category = "Bending")
+	bool bEnableXPBDBendConstraints = true;
+	UPROPERTY(EditAnywhere, Category = "Bending")
+	bool bEnableBendAnisotropy = false;
+
+	UPROPERTY(EditAnywhere, Category = "Bending", DisplayName = "Bending Stiffness (Warp) [kg cm/s^2 rad]", meta = (UIMin = "0", UIMax = "10000", ClampMin = "0", ClampMax = "10000000", EditCondition = "bEnableXPBDBendConstraints"))
+	FChaosClothWeightedValue BendingStiffness = { 100.f, 100.f };
+
+	UPROPERTY(EditAnywhere, Category = "Bending", DisplayName = "Bending Stiffness (Weft) [kg cm/s^2 rad]", meta = (UIMin = "0", UIMax = "10000", ClampMin = "0", ClampMax = "10000000", EditCondition = "bEnableXPBDBendConstraints && bEnableBendAnisotropy"))
+	FChaosClothWeightedValue BendingStiffnessWeft = { 100.f, 100.f };
+
+	UPROPERTY(EditAnywhere, Category = "Bending", DisplayName = "Bending Stiffness (Bias) [kg cm/s^2 rad]", meta = (UIMin = "0", UIMax = "10000", ClampMin = "0", ClampMax = "10000000", EditCondition = "bEnableXPBDBendConstraints && bEnableBendAnisotropy"))
+	FChaosClothWeightedValue BendingStiffnessBias = { 100.f, 100.f };
+
+	UPROPERTY(EditAnywhere, Category = "Bending", DisplayName = "Bending Damping Ratio", meta = (UIMin = "0", UIMax = "10", ClampMin = "0", ClampMax = "1000", EditCondition = "bEnableXPBDBendConstraints"))
+	FChaosClothWeightedValue BendingDampingRatio = { 1.f, 1.f };
+		
+	UPROPERTY(EditAnywhere, Category = "Bending", meta = (UIMin = "0", UIMax = "1", ClampMin = "0", ClampMax = "1", EditCondition = "bEnableXPBDBendConstraints"))
+	float BucklingRatio = 0.f;
+
+	UPROPERTY(EditAnywhere, Category = "Bending", DisplayName = "Buckling Stiffness (Warp) [kg cm/s^2 rad]", meta = (UIMin = "0", UIMax = "10000", ClampMin = "0", ClampMax = "10000000", EditCondition = "bEnableXPBDBendConstraints && bEnableXPBDBendConstraints"))
+	FChaosClothWeightedValue BucklingStiffness = { 100.f, 100.f };
+
+	UPROPERTY(EditAnywhere, Category = "Bending", DisplayName = "Buckling Stiffness (Weft) [kg cm/s^2 rad]", meta = (UIMin = "0", UIMax = "10000", ClampMin = "0", ClampMax = "10000000", EditCondition = "bEnableXPBDBendConstraints && bEnableBendAnisotropy"))
+	FChaosClothWeightedValue BucklingStiffnessWeft = { 100.f, 100.f };
+
+	UPROPERTY(EditAnywhere, Category = "Bending", DisplayName = "Buckling Stiffness (Bias) [kg cm/s^2 rad]", meta = (UIMin = "0", UIMax = "10000", ClampMin = "0", ClampMax = "10000000", EditCondition = "bEnableXPBDBendConstraints && bEnableBendAnisotropy"))
+	FChaosClothWeightedValue BucklingStiffnessBias = { 100.f, 100.f };
+
+	UPROPERTY(EditAnywhere, Category = "Area")
+	bool bEnableXPBDAreaConstraints = true;
+	
+	UPROPERTY(EditAnywhere, Category = "Area", DisplayName = "Area Stiffness [kg/s^2]", meta = (UIMin = "0", UIMax = "10000", ClampMin = "0", ClampMax = "10000000", EditCondition = "bEnableXPBDAreaConstraints"))
+	FChaosClothWeightedValue AreaStiffness = { 100.f, 100.f };
+
+	/** Cloth Shared Simulation Properties. */
+	UPROPERTY(EditAnywhere, Instanced, NoClear, Category = "Dataflow")
+	TObjectPtr<UChaosClothSharedSimConfig> SharedSimulationConfig;
+
+	FChaosClothAssetSimulationXPBDConfigNode(const Dataflow::FNodeParameters& InParam, FGuid InGuid = FGuid::NewGuid());
+
+	virtual void Serialize(FArchive& Ar) override;
+
+	virtual void Evaluate(Dataflow::FContext& Context, const FDataflowOutput* Out) const override;
+
+	//~ Begin FGCObject Interface
+	virtual void AddReferencedObjects(FReferenceCollector& Collector) override;
+	virtual FString GetReferencerName() const override;
+	//~ End FGCObject Interface
+
+private:
+	void SetStretchProperties(Chaos::Softs::FCollectionPropertyMutableFacade& Properties) const;
+	void SetBendingProperties(Chaos::Softs::FCollectionPropertyMutableFacade& Properties) const;
+	void SetAreaProperties(Chaos::Softs::FCollectionPropertyMutableFacade& Properties) const;
+
+};
diff --git a/Engine/Source/Runtime/Experimental/Chaos/Private/Chaos/XPBDAnisotropicBendingConstraints.cpp b/Engine/Source/Runtime/Experimental/Chaos/Private/Chaos/XPBDAnisotropicBendingConstraints.cpp
new file mode 100644
--- /dev/null
+++ b/Engine/Source/Runtime/Experimental/Chaos/Private/Chaos/XPBDAnisotropicBendingConstraints.cpp
@@ -0,0 +1,582 @@
+// Copyright Epic Games, Inc. All Rights Reserved.
+#include "Chaos/XPBDAnisotropicBendingConstraints.h"
+#include "Chaos/XPBDBendingConstraints.h"
+#include "Chaos/GraphColoring.h"
+#include "Chaos/Framework/Parallel.h"
+#include "Chaos/TriangleMesh.h"
+#include "ChaosStats.h"
+
+#if INTEL_ISPC
+#include "XPBDAnisotropicBendingConstraints.ispc.generated.h"
+#endif
+
+DECLARE_CYCLE_STAT(TEXT("Chaos XPBD Anisotropic Bending Constraint"), STAT_XPBD_AnisoBending, STATGROUP_Chaos);
+
+#if INTEL_ISPC && !UE_BUILD_SHIPPING
+static_assert(sizeof(ispc::FVector4f) == sizeof(Chaos::Softs::FPAndInvM), "sizeof(ispc::FVector4f) != sizeof(Chaos::Softs::FPAndInvM");
+static_assert(sizeof(ispc::FIntVector4) == sizeof(Chaos::TVec4<int32>), "sizeof(ispc::FIntVector4) != sizeof(Chaos::TVec4<int32>");
+#endif
+
+namespace Chaos::Softs {
+
+// @todo(chaos): the parallel threshold (or decision to run parallel) should probably be owned by the solver and passed to the constraint container
+extern int32 Chaos_XPBDBending_ParallelConstraintCount;
+
+FXPBDAnisotropicBendingConstraints::FXPBDAnisotropicBendingConstraints(const FSolverParticles& InParticles,
+	int32 ParticleOffset,
+	int32 ParticleCount,
+	const FTriangleMesh& TriangleMesh,
+	const TArray<TVec3<FVec2f>>& FaceVertexPatternPositions,
+	const TConstArrayView<FRealSingle>& StiffnessWarpMultipliers,
+	const TConstArrayView<FRealSingle>& StiffnessWeftMultipliers,
+	const TConstArrayView<FRealSingle>& StiffnessBiasMultipliers,
+	const TConstArrayView<FRealSingle>& BucklingStiffnessWarpMultipliers,
+	const TConstArrayView<FRealSingle>& BucklingStiffnessWeftMultipliers,
+	const TConstArrayView<FRealSingle>& BucklingStiffnessBiasMultipliers,
+	const TConstArrayView<FRealSingle>& DampingMultipliers,
+	const FCollectionPropertyConstFacade& PropertyCollection,
+	bool bTrimKinematicConstraints)
+	: Base(
+		InParticles,
+		ParticleOffset,
+		ParticleCount,
+		TriangleMesh.GetUniqueAdjacentElements(),
+		StiffnessWarpMultipliers,
+		BucklingStiffnessWarpMultipliers,
+		FSolverVec2(GetWeightedFloatXPBDBendingElementStiffnessWarp(PropertyCollection, MaxStiffness)),
+		(FSolverReal)GetXPBDBucklingRatio(PropertyCollection, 0.f),
+		FSolverVec2(GetWeightedFloatXPBDBucklingStiffnessWarp(PropertyCollection, MaxStiffness)),
+		bTrimKinematicConstraints,
+		MaxStiffness)
+	, StiffnessWeft(
+		FSolverVec2(GetWeightedFloatXPBDBendingElementStiffnessWeft(PropertyCollection, MaxStiffness)),
+		StiffnessWeftMultipliers,
+		TConstArrayView<TVec2<int32>>(ConstraintSharedEdges),
+		ParticleOffset,
+		ParticleCount,
+		FPBDStiffness::DefaultTableSize,
+		FPBDStiffness::DefaultParameterFitBase,
+		MaxStiffness)
+	, StiffnessBias(
+		FSolverVec2(GetWeightedFloatXPBDBendingElementStiffnessBias(PropertyCollection, MaxStiffness)),
+		StiffnessBiasMultipliers,
+		TConstArrayView<TVec2<int32>>(ConstraintSharedEdges),
+		ParticleOffset,
+		ParticleCount,
+		FPBDStiffness::DefaultTableSize,
+		FPBDStiffness::DefaultParameterFitBase,
+		MaxStiffness)
+	, BucklingStiffnessWeft(
+		FSolverVec2(GetWeightedFloatXPBDBucklingStiffnessWeft(PropertyCollection, MaxStiffness)),
+		BucklingStiffnessWeftMultipliers,
+		TConstArrayView<TVec2<int32>>(ConstraintSharedEdges),
+		ParticleOffset,
+		ParticleCount,
+		FPBDStiffness::DefaultTableSize,
+		FPBDStiffness::DefaultParameterFitBase,
+		MaxStiffness)
+	, BucklingStiffnessBias(
+		FSolverVec2(GetWeightedFloatXPBDBucklingStiffnessBias(PropertyCollection, MaxStiffness)),
+		BucklingStiffnessBiasMultipliers,
+		TConstArrayView<TVec2<int32>>(ConstraintSharedEdges),
+		ParticleOffset,
+		ParticleCount,
+		FPBDStiffness::DefaultTableSize,
+		FPBDStiffness::DefaultParameterFitBase,
+		MaxStiffness)
+	, DampingRatio(
+		FSolverVec2(GetWeightedFloatXPBDBendingElementDamping(PropertyCollection, MinDamping)).ClampAxes(MinDamping, MaxDamping),
+		DampingMultipliers,
+		TConstArrayView<TVec2<int32>>(ConstraintSharedEdges),
+		ParticleOffset,
+		ParticleCount)
+	, WarpWeftBiasBaseMultipliers(GenerateWarpWeftBiasBaseMultipliers(FaceVertexPatternPositions, TriangleMesh))
+{
+	Lambdas.Init((FSolverReal)0., Constraints.Num());
+	InitColor(InParticles, ParticleOffset, ParticleCount);
+}
+
+FXPBDAnisotropicBendingConstraints::FXPBDAnisotropicBendingConstraints(const FSolverParticles& InParticles,
+	int32 ParticleOffset,
+	int32 ParticleCount,
+	const FTriangleMesh& TriangleMesh,
+	const TArray<TVec3<FVec2f>>& FaceVertexPatternPositions,
+	const TConstArrayView<FRealSingle>& StiffnessWarpMultipliers,
+	const TConstArrayView<FRealSingle>& StiffnessWeftMultipliers,
+	const TConstArrayView<FRealSingle>& StiffnessBiasMultipliers,
+	const TConstArrayView<FRealSingle>& BucklingStiffnessWarpMultipliers,
+	const TConstArrayView<FRealSingle>& BucklingStiffnessWeftMultipliers,
+	const TConstArrayView<FRealSingle>& BucklingStiffnessBiasMultipliers,
+	const TConstArrayView<FRealSingle>& DampingMultipliers,
+	const FSolverVec2& InStiffnessWarp,
+	const FSolverVec2& InStiffnessWeft,
+	const FSolverVec2& InStiffnessBias,
+	const FSolverReal InBucklingRatio,
+	const FSolverVec2& InBucklingStiffnessWarp,
+	const FSolverVec2& InBucklingStiffnessWeft,
+	const FSolverVec2& InBucklingStiffnessBias,
+	const FSolverVec2& InDampingRatio,
+	bool bTrimKinematicConstraints)
+	: Base(
+		InParticles,
+		ParticleOffset,
+		ParticleCount,
+		TriangleMesh.GetUniqueAdjacentElements(),
+		StiffnessWarpMultipliers,
+		BucklingStiffnessWarpMultipliers,
+		InStiffnessWarp,
+		InBucklingRatio,
+		InBucklingStiffnessWarp,
+		bTrimKinematicConstraints,
+		MaxStiffness)
+	, StiffnessWeft(
+		InStiffnessWeft,
+		StiffnessWeftMultipliers,
+		TConstArrayView<TVec2<int32>>(ConstraintSharedEdges),
+		ParticleOffset,
+		ParticleCount,
+		FPBDStiffness::DefaultTableSize,
+		FPBDStiffness::DefaultParameterFitBase,
+		MaxStiffness)
+	, StiffnessBias(
+		InStiffnessBias,
+		StiffnessBiasMultipliers,
+		TConstArrayView<TVec2<int32>>(ConstraintSharedEdges),
+		ParticleOffset,
+		ParticleCount,
+		FPBDStiffness::DefaultTableSize,
+		FPBDStiffness::DefaultParameterFitBase,
+		MaxStiffness)
+	, BucklingStiffnessWeft(
+		InBucklingStiffnessWeft,
+		BucklingStiffnessWeftMultipliers,
+		TConstArrayView<TVec2<int32>>(ConstraintSharedEdges),
+		ParticleOffset,
+		ParticleCount,
+		FPBDStiffness::DefaultTableSize,
+		FPBDStiffness::DefaultParameterFitBase,
+		MaxStiffness)
+	, BucklingStiffnessBias(
+		InBucklingStiffnessBias,
+		BucklingStiffnessBiasMultipliers,
+		TConstArrayView<TVec2<int32>>(ConstraintSharedEdges),
+		ParticleOffset,
+		ParticleCount,
+		FPBDStiffness::DefaultTableSize,
+		FPBDStiffness::DefaultParameterFitBase,
+		MaxStiffness)
+	, DampingRatio(
+		InDampingRatio.ClampAxes((FSolverReal)0., (FSolverReal)1.),
+		DampingMultipliers,
+		TConstArrayView<TVec2<int32>>(ConstraintSharedEdges),
+		ParticleOffset,
+		ParticleCount)
+	, WarpWeftBiasBaseMultipliers(GenerateWarpWeftBiasBaseMultipliers(FaceVertexPatternPositions, TriangleMesh))
+{
+	Lambdas.Init((FSolverReal)0., Constraints.Num());
+	InitColor(InParticles, ParticleOffset, ParticleCount);
+}
+
+TArray<FSolverVec3> FXPBDAnisotropicBendingConstraints::GenerateWarpWeftBiasBaseMultipliers(const TArray<TVec3<FVec2f>>& FaceVertexPatternPositions, const FTriangleMesh& TriangleMesh) const
+{
+	// This is easier to do by faces. Average the multipliers for faces that share the same edge.
+	TMap<TVec2<int32> /*Edge*/, TArray<FSolverVec3> /*BaseMultipliers*/> EdgeToMultiplierMap;
+	auto SortedEdge = [](int32 P0, int32 P1) { return P0 <= P1 ? TVec2<int32>(P0, P1) : TVec2<int32>(P1, P0); };
+	auto Multiplier = [](const FVec2f& UV0, const FVec2f& UV1)
+	{
+		// Input UVs: X = Warp, Y = Weft direction
+		// Transform to X = Warp, Y = Weft, Z = Bias
+		const FSolverVec2 UVDiff = UV1 - UV0;
+		const FSolverVec2 UVDiffAbs = UVDiff.GetAbs();
+		FSolverVec3 UVDiffTransformed;
+		if (UVDiffAbs.X > UVDiffAbs.Y)
+		{
+			// Stiffness is blend between warp and bias
+			UVDiffTransformed  = FSolverVec3(UVDiffAbs.X - UVDiffAbs.Y, (FSolverReal)0.f, UVDiffAbs.Y);
+		}
+		else
+		{
+			// Stiffness is blend between weft and bias
+			UVDiffTransformed = FSolverVec3((FSolverReal)0.f, UVDiffAbs.Y - UVDiff.X, UVDiffAbs.X);
+		}
+		const FSolverReal Denom = UVDiffTransformed.X + UVDiffTransformed.Y + UVDiffTransformed.Z;
+		return Denom > UE_SMALL_NUMBER ? UVDiffTransformed / Denom : FSolverVec3(1.f, 0.f, 0.f); // Default to Warp if zero length
+	};
+
+	const TArray<TVec3<int32>>& Elements = TriangleMesh.GetElements();
+	for (int32 ElemIdx = 0; ElemIdx < Elements.Num(); ++ElemIdx)
+	{
+		const TVec3<int32>& Element = Elements[ElemIdx];
+		const TVec3<FVec2f>& UVs = FaceVertexPatternPositions[ElemIdx];
+		EdgeToMultiplierMap.FindOrAdd(SortedEdge(Element[0], Element[1])).Add(Multiplier(UVs[0], UVs[1]));
+		EdgeToMultiplierMap.FindOrAdd(SortedEdge(Element[1], Element[2])).Add(Multiplier(UVs[1], UVs[2]));
+		EdgeToMultiplierMap.FindOrAdd(SortedEdge(Element[2], Element[0])).Add(Multiplier(UVs[2], UVs[0]));
+	}
+
+	TArray<FSolverVec3> WarpWeftBaseMultiplierResult;
+	WarpWeftBaseMultiplierResult.SetNumUninitialized(Constraints.Num());
+	for (int32 ConstraintIdx = 0; ConstraintIdx < Constraints.Num(); ++ConstraintIdx)
+	{
+		const TArray<FSolverVec3>& EdgeMultipliers = EdgeToMultiplierMap.FindChecked(SortedEdge(ConstraintSharedEdges[ConstraintIdx][0], ConstraintSharedEdges[ConstraintIdx][1]));
+		FSolverVec3 AveragedEdgeMultiplier((FSolverReal)0.);
+		if (ensure(EdgeMultipliers.Num()))
+		{
+			for (const FSolverVec3& EdgeMult : EdgeMultipliers)
+			{
+				AveragedEdgeMultiplier += EdgeMult;
+			}
+			AveragedEdgeMultiplier /= (FSolverReal)EdgeMultipliers.Num();
+		}
+		WarpWeftBaseMultiplierResult[ConstraintIdx] = AveragedEdgeMultiplier;
+	}
+	return WarpWeftBaseMultiplierResult;
+}
+
+void FXPBDAnisotropicBendingConstraints::InitColor(const FSolverParticles& InParticles, const int32 ParticleOffset, const int32 ParticleCount)
+{
+	// In dev builds we always color so we can tune the system without restarting. See Apply()
+#if UE_BUILD_SHIPPING || UE_BUILD_TEST
+	if (Constraints.Num() > Chaos_XPBDBending_ParallelConstraintCount)
+#endif
+	{
+		const TArray<TArray<int32>> ConstraintsPerColor = FGraphColoring::ComputeGraphColoring(Constraints, InParticles, ParticleOffset, ParticleOffset + ParticleCount);
+
+		// Reorder constraints based on color so each array in ConstraintsPerColor contains contiguous elements.
+		TArray<TVec4<int32>> ReorderedConstraints; 
+		TArray<TVec2<int32>> ReorderedConstraintSharedEdges;
+		TArray<FSolverReal> ReorderedRestAngles;
+		TArray<FSolverVec3> ReorderedWarpWeftBiasBaseMultipliers;
+		TArray<int32> OrigToReorderedIndices; // used to reorder stiffness indices
+		ReorderedConstraints.SetNumUninitialized(Constraints.Num());
+		ReorderedConstraintSharedEdges.SetNumUninitialized(Constraints.Num());
+		ReorderedRestAngles.SetNumUninitialized(Constraints.Num());
+		ReorderedWarpWeftBiasBaseMultipliers.SetNumUninitialized(Constraints.Num());
+		OrigToReorderedIndices.SetNumUninitialized(Constraints.Num());
+
+		ConstraintsPerColorStartIndex.Reset(ConstraintsPerColor.Num() + 1);
+
+		int32 ReorderedIndex = 0;
+		for (const TArray<int32>& ConstraintsBatch : ConstraintsPerColor)
+		{
+			ConstraintsPerColorStartIndex.Add(ReorderedIndex);
+			for (const int32& BatchConstraint : ConstraintsBatch)
+			{
+				const int32 OrigIndex = BatchConstraint;
+				ReorderedConstraints[ReorderedIndex] = Constraints[OrigIndex];
+				ReorderedConstraintSharedEdges[ReorderedIndex] = ConstraintSharedEdges[OrigIndex];
+				ReorderedRestAngles[ReorderedIndex] = RestAngles[OrigIndex];
+				ReorderedWarpWeftBiasBaseMultipliers[ReorderedIndex] = WarpWeftBiasBaseMultipliers[OrigIndex];
+				OrigToReorderedIndices[OrigIndex] = ReorderedIndex;
+				++ReorderedIndex;
+			}
+		}
+		ConstraintsPerColorStartIndex.Add(ReorderedIndex);
+
+		Constraints = MoveTemp(ReorderedConstraints);
+		ConstraintSharedEdges = MoveTemp(ReorderedConstraintSharedEdges);
+		RestAngles = MoveTemp(ReorderedRestAngles);
+		WarpWeftBiasBaseMultipliers = MoveTemp(ReorderedWarpWeftBiasBaseMultipliers);
+		Stiffness.ReorderIndices(OrigToReorderedIndices);
+		StiffnessWeft.ReorderIndices(OrigToReorderedIndices);
+		StiffnessBias.ReorderIndices(OrigToReorderedIndices);
+		BucklingStiffness.ReorderIndices(OrigToReorderedIndices);
+		BucklingStiffnessWeft.ReorderIndices(OrigToReorderedIndices);
+		BucklingStiffnessBias.ReorderIndices(OrigToReorderedIndices);
+		DampingRatio.ReorderIndices(OrigToReorderedIndices);
+	}
+}
+
+void FXPBDAnisotropicBendingConstraints::ApplyHelper(FSolverParticles& Particles, const FSolverReal Dt, const int32 ConstraintIndex, const FSolverVec3& ExpStiffnessValues, 
+	const FSolverVec3& ExpBucklingStiffnessValues, const FSolverReal DampingRatioValue) const
+{
+	const TVec4<int32>& Constraint = Constraints[ConstraintIndex];
+	const int32 i1 = Constraint[0];
+	const int32 i2 = Constraint[1];
+	const int32 i3 = Constraint[2];
+	const int32 i4 = Constraint[3];
+	const FSolverVec3& WarpWeftBiasBaseMultiplier = WarpWeftBiasBaseMultipliers[ConstraintIndex];
+
+	const FSolverReal BiphasicStiffnessValue = IsBuckled[ConstraintIndex] ?
+		WarpWeftBiasBaseMultiplier.Dot(ExpStiffnessValues) : WarpWeftBiasBaseMultiplier.Dot(ExpBucklingStiffnessValues);
+
+	if (BiphasicStiffnessValue < MinStiffness || (Particles.InvM(i1) == (FSolverReal)0. && Particles.InvM(i2) == (FSolverReal)0. && Particles.InvM(i3) == (FSolverReal)0. && Particles.InvM(i4) == (FSolverReal)0.))
+	{
+		return;
+	}
+
+	const FSolverReal CombinedInvMass = Particles.InvM(i1) + Particles.InvM(i2) + Particles.InvM(i3) + Particles.InvM(i4);
+	const FSolverReal Damping = (FSolverReal)2.f * DampingRatioValue * FMath::Sqrt(BiphasicStiffnessValue / CombinedInvMass);
+
+	const FSolverReal Angle = CalcAngle(Particles.P(i1), Particles.P(i2), Particles.P(i3), Particles.P(i4));
+	const TStaticArray<FSolverVec3, 4> Grads = Base::GetGradients(Particles, ConstraintIndex);
+
+	const FSolverVec3 V1TimesDt = Particles.P(i1) - Particles.X(i1);
+	const FSolverVec3 V2TimesDt = Particles.P(i2) - Particles.X(i2);
+	const FSolverVec3 V3TimesDt = Particles.P(i3) - Particles.X(i3);
+	const FSolverVec3 V4TimesDt = Particles.P(i4) - Particles.X(i4);
+
+	FSolverReal& Lambda = Lambdas[ConstraintIndex];
+	const FSolverReal Alpha = (FSolverReal)1.f / (BiphasicStiffnessValue * Dt * Dt);
+	const FSolverReal Gamma = Alpha * Damping * Dt;
+
+	const FSolverReal DampingTerm = Gamma * (FSolverVec3::DotProduct(V1TimesDt, Grads[0]) + FSolverVec3::DotProduct(V2TimesDt, Grads[1]) + FSolverVec3::DotProduct(V3TimesDt, Grads[2]) + FSolverVec3::DotProduct(V4TimesDt, Grads[3]));
+
+	const FSolverReal Denom = ((FSolverReal)1.f + Gamma) * (Particles.InvM(i1) * Grads[0].SizeSquared() + Particles.InvM(i2) * Grads[1].SizeSquared() + Particles.InvM(i3) * Grads[2].SizeSquared() + Particles.InvM(i4) * Grads[3].SizeSquared()) + Alpha;
+	const FSolverReal DLambda = (Angle - RestAngles[ConstraintIndex] - Alpha * Lambda + DampingTerm) / Denom;
+
+	Particles.P(i1) -= DLambda * Particles.InvM(i1) * Grads[0];
+	Particles.P(i2) -= DLambda * Particles.InvM(i2) * Grads[1];
+	Particles.P(i3) -= DLambda * Particles.InvM(i3) * Grads[2];
+	Particles.P(i4) -= DLambda * Particles.InvM(i4) * Grads[3];
+	Lambda += DLambda;
+}
+
+void FXPBDAnisotropicBendingConstraints::Apply(FSolverParticles& Particles, const FSolverReal Dt) const
+{
+	TRACE_CPUPROFILER_EVENT_SCOPE(FXPBDAnisotropicBendingConstraints_Apply);
+	SCOPE_CYCLE_COUNTER(STAT_XPBD_AnisoBending);
+
+	const bool StiffnessHasWeightMap = Stiffness.HasWeightMap();
+	const bool StiffnessWeftHasWeightMap = StiffnessWeft.HasWeightMap();
+	const bool StiffnessBiasHasWeightMap = StiffnessBias.HasWeightMap();
+	const bool BucklingStiffnessHasWeightMap = BucklingStiffness.HasWeightMap();
+	const bool BucklingStiffnessWeftHasWeightMap = BucklingStiffnessWeft.HasWeightMap();
+	const bool BucklingStiffnessBiasHasWeightMap = BucklingStiffnessBias.HasWeightMap();
+	const bool DampingHasWeightMap = DampingRatio.HasWeightMap();
+
+	if (ConstraintsPerColorStartIndex.Num() > 0 && Constraints.Num() > Chaos_XPBDBending_ParallelConstraintCount)
+	{
+		const int32 ConstraintColorNum = ConstraintsPerColorStartIndex.Num() - 1;
+		if (!StiffnessHasWeightMap && !StiffnessWeftHasWeightMap && !StiffnessBiasHasWeightMap &&
+			!BucklingStiffnessHasWeightMap && !BucklingStiffnessWeftHasWeightMap && !BucklingStiffnessBiasHasWeightMap &&
+			!DampingHasWeightMap)
+		{
+			const FSolverVec3 ExpStiffnessValue((FSolverReal)Stiffness, (FSolverReal)StiffnessWeft, (FSolverReal)StiffnessBias);
+			const FSolverVec3 ExpBucklingValue((FSolverReal)BucklingStiffness, (FSolverReal)BucklingStiffnessWeft, (FSolverReal)BucklingStiffnessBias);
+			const FSolverReal DampingRatioValue = (FSolverReal)DampingRatio;
+
+			if (ExpStiffnessValue.Max() < MinStiffness && ExpBucklingValue.Max() < MinStiffness)
+			{
+				return;
+			}
+
+#if INTEL_ISPC
+			if (bRealTypeCompatibleWithISPC && bChaos_XPBDBending_ISPC_Enabled)
+			{
+				if (DampingRatioValue > 0)
+				{
+					for (int32 ConstraintColorIndex = 0; ConstraintColorIndex < ConstraintColorNum; ++ConstraintColorIndex)
+					{
+						const int32 ColorStart = ConstraintsPerColorStartIndex[ConstraintColorIndex];
+						const int32 ColorSize = ConstraintsPerColorStartIndex[ConstraintColorIndex + 1] - ColorStart;
+						ispc::ApplyXPBDAnisotropicBendingConstraintsWithDamping(
+							(ispc::FVector4f*)Particles.GetPAndInvM().GetData(),
+							(const ispc::FVector3f*)Particles.X().GetData(),
+							(ispc::FIntVector4*)&Constraints.GetData()[ColorStart],
+							&RestAngles.GetData()[ColorStart],
+							&IsBuckled.GetData()[ColorStart],
+							(const ispc::FVector3f*)&WarpWeftBiasBaseMultipliers.GetData()[ColorStart],
+							&Lambdas.GetData()[ColorStart],
+							Dt,
+							MinStiffness,
+							reinterpret_cast<const ispc::FVector3f&>(ExpStiffnessValue),
+							reinterpret_cast<const ispc::FVector3f&>(ExpBucklingValue),
+							DampingRatioValue,
+							ColorSize);
+					}
+				}
+				else
+				{
+					for (int32 ConstraintColorIndex = 0; ConstraintColorIndex < ConstraintColorNum; ++ConstraintColorIndex)
+					{
+						const int32 ColorStart = ConstraintsPerColorStartIndex[ConstraintColorIndex];
+						const int32 ColorSize = ConstraintsPerColorStartIndex[ConstraintColorIndex + 1] - ColorStart;
+						ispc::ApplyXPBDAnisotropicBendingConstraints(
+							(ispc::FVector4f*)Particles.GetPAndInvM().GetData(),
+							(ispc::FIntVector4*)&Constraints.GetData()[ColorStart],
+							&RestAngles.GetData()[ColorStart],
+							&IsBuckled.GetData()[ColorStart],
+							(const ispc::FVector3f*)&WarpWeftBiasBaseMultipliers.GetData()[ColorStart],
+							&Lambdas.GetData()[ColorStart],
+							Dt,
+							MinStiffness,
+							reinterpret_cast<const ispc::FVector3f&>(ExpStiffnessValue),
+							reinterpret_cast<const ispc::FVector3f&>(ExpBucklingValue),
+							ColorSize);
+					}
+				}
+			}
+			else
+#endif
+			{
+				for (int32 ConstraintColorIndex = 0; ConstraintColorIndex < ConstraintColorNum; ++ConstraintColorIndex)
+				{
+					const int32 ColorStart = ConstraintsPerColorStartIndex[ConstraintColorIndex];
+					const int32 ColorSize = ConstraintsPerColorStartIndex[ConstraintColorIndex + 1] - ColorStart;
+					PhysicsParallelFor(ColorSize, [&](const int32 Index)
+					{
+						const int32 ConstraintIndex = ColorStart + Index;
+						ApplyHelper(Particles, Dt, ConstraintIndex, ExpStiffnessValue, ExpBucklingValue, DampingRatioValue);
+					});
+				}
+			}
+		}
+		else  // Has weight maps
+		{
+#if INTEL_ISPC
+			if (bRealTypeCompatibleWithISPC && bChaos_XPBDBending_ISPC_Enabled)
+			{
+				if (DampingHasWeightMap || (FSolverReal)DampingRatio > 0)
+				{
+					for (int32 ConstraintColorIndex = 0; ConstraintColorIndex < ConstraintColorNum; ++ConstraintColorIndex)
+					{
+						const int32 ColorStart = ConstraintsPerColorStartIndex[ConstraintColorIndex];
+						const int32 ColorSize = ConstraintsPerColorStartIndex[ConstraintColorIndex + 1] - ColorStart;
+						ispc::ApplyXPBDAnisotropicBendingConstraintsWithDampingAndMaps(
+							(ispc::FVector4f*)Particles.GetPAndInvM().GetData(),
+							(const ispc::FVector3f*)Particles.X().GetData(),
+							(ispc::FIntVector4*)&Constraints.GetData()[ColorStart],
+							&RestAngles.GetData()[ColorStart],
+							&IsBuckled.GetData()[ColorStart],
+							(const ispc::FVector3f*)&WarpWeftBiasBaseMultipliers.GetData()[ColorStart],
+							&Lambdas.GetData()[ColorStart],
+							Dt,
+							MinStiffness,
+							StiffnessHasWeightMap,
+							&Stiffness.GetIndices().GetData()[ColorStart],
+							&Stiffness.GetTable().GetData()[0],
+							StiffnessWeftHasWeightMap,
+							&StiffnessWeft.GetIndices().GetData()[ColorStart],
+							&StiffnessWeft.GetTable().GetData()[0],
+							StiffnessBiasHasWeightMap,
+							&StiffnessBias.GetIndices().GetData()[ColorStart],
+							&StiffnessBias.GetTable().GetData()[0],
+							BucklingStiffnessHasWeightMap,
+							&BucklingStiffness.GetIndices().GetData()[ColorStart],
+							&BucklingStiffness.GetTable().GetData()[0],
+							BucklingStiffnessWeftHasWeightMap,
+							&BucklingStiffnessWeft.GetIndices().GetData()[ColorStart],
+							&BucklingStiffnessWeft.GetTable().GetData()[0],
+							BucklingStiffnessBiasHasWeightMap,
+							&BucklingStiffnessBias.GetIndices().GetData()[ColorStart],
+							&BucklingStiffnessBias.GetTable().GetData()[0],
+							DampingHasWeightMap,
+							&DampingRatio.GetIndices().GetData()[ColorStart],
+							&DampingRatio.GetTable().GetData()[0],
+							ColorSize);
+					}
+				}
+				else
+				{
+					for (int32 ConstraintColorIndex = 0; ConstraintColorIndex < ConstraintColorNum; ++ConstraintColorIndex)
+					{
+						const int32 ColorStart = ConstraintsPerColorStartIndex[ConstraintColorIndex];
+						const int32 ColorSize = ConstraintsPerColorStartIndex[ConstraintColorIndex + 1] - ColorStart;
+						ispc::ApplyXPBDAnisotropicBendingConstraintsWithMaps(
+							(ispc::FVector4f*)Particles.GetPAndInvM().GetData(),
+							(ispc::FIntVector4*)&Constraints.GetData()[ColorStart],
+							&RestAngles.GetData()[ColorStart],
+							&IsBuckled.GetData()[ColorStart],
+							(const ispc::FVector3f*)&WarpWeftBiasBaseMultipliers.GetData()[ColorStart],
+							&Lambdas.GetData()[ColorStart],
+							Dt,
+							MinStiffness,
+							StiffnessHasWeightMap,
+							&Stiffness.GetIndices().GetData()[ColorStart],
+							&Stiffness.GetTable().GetData()[0],
+							StiffnessWeftHasWeightMap,
+							&StiffnessWeft.GetIndices().GetData()[ColorStart],
+							&StiffnessWeft.GetTable().GetData()[0],
+							StiffnessBiasHasWeightMap,
+							&StiffnessBias.GetIndices().GetData()[ColorStart],
+							&StiffnessBias.GetTable().GetData()[0],
+							BucklingStiffnessHasWeightMap,
+							&BucklingStiffness.GetIndices().GetData()[ColorStart],
+							&BucklingStiffness.GetTable().GetData()[0],
+							BucklingStiffnessWeftHasWeightMap,
+							&BucklingStiffnessWeft.GetIndices().GetData()[ColorStart],
+							&BucklingStiffnessWeft.GetTable().GetData()[0],
+							BucklingStiffnessBiasHasWeightMap,
+							&BucklingStiffnessBias.GetIndices().GetData()[ColorStart],
+							&BucklingStiffnessBias.GetTable().GetData()[0],
+							ColorSize);
+					}
+				}
+			}
+			else
+#endif
+			{
+				const FSolverReal StiffnessNoMap = (FSolverReal)Stiffness;
+				const FSolverReal StiffnessWeftNoMap = (FSolverReal)StiffnessWeft;
+				const FSolverReal StiffnessBiasNoMap = (FSolverReal)StiffnessBias;
+				const FSolverReal BucklingStiffnessNoMap = (FSolverReal)BucklingStiffness;
+				const FSolverReal BucklingStiffnessWeftNoMap = (FSolverReal)BucklingStiffnessWeft;
+				const FSolverReal BucklingStiffnessBiasNoMap = (FSolverReal)BucklingStiffnessBias;
+				const FSolverReal DampingNoMap = (FSolverReal)DampingRatio;
+				for (int32 ConstraintColorIndex = 0; ConstraintColorIndex < ConstraintColorNum; ++ConstraintColorIndex)
+				{
+					const int32 ColorStart = ConstraintsPerColorStartIndex[ConstraintColorIndex];
+					const int32 ColorSize = ConstraintsPerColorStartIndex[ConstraintColorIndex + 1] - ColorStart;
+					PhysicsParallelFor(ColorSize, [&](const int32 Index)
+					{
+						const int32 ConstraintIndex = ColorStart + Index;
+						const FSolverReal ExpStiffnessValue = StiffnessHasWeightMap ? Stiffness[ConstraintIndex] : StiffnessNoMap;
+						const FSolverReal ExpStiffnessWeftValue = StiffnessWeftHasWeightMap ? StiffnessWeft[ConstraintIndex] : StiffnessWeftNoMap;
+						const FSolverReal ExpStiffnessBiasValue = StiffnessBiasHasWeightMap ? StiffnessBias[ConstraintIndex] : StiffnessBiasNoMap;
+						const FSolverReal ExpBucklingValue = BucklingStiffnessHasWeightMap ? BucklingStiffness[ConstraintIndex] : BucklingStiffnessNoMap;
+						const FSolverReal ExpBucklingWeftValue = BucklingStiffnessWeftHasWeightMap ? BucklingStiffnessWeft[ConstraintIndex] : BucklingStiffnessWeftNoMap;
+						const FSolverReal ExpBucklingBiasValue = BucklingStiffnessBiasHasWeightMap ? BucklingStiffnessBias[ConstraintIndex] : BucklingStiffnessBiasNoMap;
+						const FSolverReal DampingRatioValue = DampingHasWeightMap ? DampingRatio[ConstraintIndex] : DampingNoMap;
+						ApplyHelper(Particles, Dt, ConstraintIndex, FSolverVec3(ExpStiffnessValue, ExpStiffnessWeftValue, ExpStiffnessBiasValue),
+							FSolverVec3(ExpBucklingValue, ExpBucklingWeftValue, ExpBucklingBiasValue), DampingRatioValue);
+					});
+				}
+			}
+		}
+	}
+	else
+	{
+		if (!StiffnessHasWeightMap && !StiffnessWeftHasWeightMap && !StiffnessBiasHasWeightMap &&
+			!BucklingStiffnessHasWeightMap && !BucklingStiffnessWeftHasWeightMap && !BucklingStiffnessBiasHasWeightMap &&
+			!DampingHasWeightMap)
+		{
+			const FSolverVec3 ExpStiffnessValue((FSolverReal)Stiffness, (FSolverReal)StiffnessWeft, (FSolverReal)StiffnessBias);
+			const FSolverVec3 ExpBucklingValue((FSolverReal)BucklingStiffness, (FSolverReal)BucklingStiffnessWeft, (FSolverReal)BucklingStiffnessBias);
+			const FSolverReal DampingRatioValue = (FSolverReal)DampingRatio;
+
+			if (ExpStiffnessValue.Max() < MinStiffness && ExpBucklingValue.Max() < MinStiffness)
+			{
+				return;
+			}
+
+			for (int32 ConstraintIndex = 0; ConstraintIndex < Constraints.Num(); ++ConstraintIndex)
+			{
+				ApplyHelper(Particles, Dt, ConstraintIndex, ExpStiffnessValue, ExpBucklingValue, DampingRatioValue);
+			}
+		}
+		else
+		{
+			const FSolverReal StiffnessNoMap = (FSolverReal)Stiffness;
+			const FSolverReal StiffnessWeftNoMap = (FSolverReal)StiffnessWeft;
+			const FSolverReal StiffnessBiasNoMap = (FSolverReal)StiffnessBias;
+			const FSolverReal BucklingStiffnessNoMap = (FSolverReal)BucklingStiffness;
+			const FSolverReal BucklingStiffnessWeftNoMap = (FSolverReal)BucklingStiffnessWeft;
+			const FSolverReal BucklingStiffnessBiasNoMap = (FSolverReal)BucklingStiffnessBias;
+			const FSolverReal DampingNoMap = (FSolverReal)DampingRatio;
+			for (int32 ConstraintIndex = 0; ConstraintIndex < Constraints.Num(); ++ConstraintIndex)
+			{
+				const FSolverReal ExpStiffnessValue = StiffnessHasWeightMap ? Stiffness[ConstraintIndex] : StiffnessNoMap;
+				const FSolverReal ExpStiffnessWeftValue = StiffnessWeftHasWeightMap ? StiffnessWeft[ConstraintIndex] : StiffnessWeftNoMap;
+				const FSolverReal ExpStiffnessBiasValue = StiffnessBiasHasWeightMap ? StiffnessBias[ConstraintIndex] : StiffnessBiasNoMap;
+				const FSolverReal ExpBucklingValue = BucklingStiffnessHasWeightMap ? BucklingStiffness[ConstraintIndex] : BucklingStiffnessNoMap;
+				const FSolverReal ExpBucklingWeftValue = BucklingStiffnessWeftHasWeightMap ? BucklingStiffnessWeft[ConstraintIndex] : BucklingStiffnessWeftNoMap;
+				const FSolverReal ExpBucklingBiasValue = BucklingStiffnessBiasHasWeightMap ? BucklingStiffnessBias[ConstraintIndex] : BucklingStiffnessBiasNoMap;
+				const FSolverReal DampingRatioValue = DampingHasWeightMap ? DampingRatio[ConstraintIndex] : DampingNoMap;
+				ApplyHelper(Particles, Dt, ConstraintIndex, FSolverVec3(ExpStiffnessValue, ExpStiffnessWeftValue, ExpStiffnessBiasValue),
+					FSolverVec3(ExpBucklingValue, ExpBucklingWeftValue, ExpBucklingBiasValue), DampingRatioValue);
+			}
+		}
+	}
+}
+
+}  // End namespace Chaos::Softs
diff --git a/Engine/Source/Runtime/Experimental/Chaos/Private/Chaos/XPBDAnisotropicBendingConstraints.ispc b/Engine/Source/Runtime/Experimental/Chaos/Private/Chaos/XPBDAnisotropicBendingConstraints.ispc
new file mode 100644
--- /dev/null
+++ b/Engine/Source/Runtime/Experimental/Chaos/Private/Chaos/XPBDAnisotropicBendingConstraints.ispc
@@ -0,0 +1,450 @@
+// Copyright Epic Games, Inc. All Rights Reserved.
+
+#define EXPLICIT_VECTOR4 1
+
+#include "Math/Vector.isph"
+#include "Chaos/PBDSofts.isph"
+
+static inline FVector3f SafeDivide(const FVector3f& Numerator, const float Denominator)
+{
+	return VectorSelect(Denominator <= FLOAT_SMALL_NUMBER, FloatZeroVector, Numerator / Denominator);
+}
+
+static inline varying float SafeDivide(const varying float Numerator, const varying float Denominator)
+{
+	return select(Denominator <= FLOAT_SMALL_NUMBER, (uniform float)0, Numerator / Denominator);
+}
+
+// Returns Angle
+static inline float CalcGradientsAndAngle(const FVector3f& P1, const FVector3f& P2, const FVector3f& P3, const FVector3f& P4, FVector3f& Grad1, FVector3f& Grad2, FVector3f& Grad3, FVector3f& Grad4)
+{
+	const FVector3f SharedEdgeNormalized = VectorGetSafeNormal(P2 - P1);
+	const FVector3f P13CrossP23 = VectorCross(P1 - P3, P2 - P3);
+	const float Normal1Len = VectorSize(P13CrossP23);
+	const FVector3f Normal1 = SafeDivide(P13CrossP23, Normal1Len);
+	const FVector3f P24CrossP14 = VectorCross(P2 - P4, P1 - P4);
+	const float Normal2Len = VectorSize(P24CrossP14);
+	const FVector3f Normal2 = SafeDivide(P24CrossP14, Normal2Len);
+
+	const FVector3f N2CrossN1 = VectorCross(Normal2, Normal1);
+
+	const float CosPhi = clamp(VectorDot(Normal1, Normal2), (uniform float)(-1), (uniform float)(1));
+	const float SinPhi = clamp(VectorDot(N2CrossN1, SharedEdgeNormalized), (uniform float)(-1), (uniform float)(1));
+
+	const FVector3f DPhiDN1_OverNormal1Len = SafeDivide(CosPhi * VectorCross(SharedEdgeNormalized, Normal2) - SinPhi * Normal2, Normal1Len);
+	const FVector3f DPhiDN2_OverNormal2Len = SafeDivide(CosPhi * VectorCross(Normal1, SharedEdgeNormalized) - SinPhi * Normal1, Normal2Len);
+
+	const FVector3f DPhiDP13 = VectorCross(P2 - P3, DPhiDN1_OverNormal1Len);
+	const FVector3f DPhiDP23 = VectorCross(DPhiDN1_OverNormal1Len, P1 - P3);
+	const FVector3f DPhiDP24 = VectorCross(P1 - P4, DPhiDN2_OverNormal2Len);
+	const FVector3f DPhiDP14 = VectorCross(DPhiDN2_OverNormal2Len, P2 - P4);
+
+	Grad1 = DPhiDP13 + DPhiDP14;
+	Grad2 = DPhiDP23 + DPhiDP24;
+	Grad3 = -1 * DPhiDP13 - DPhiDP23;
+	Grad4 = -1 * DPhiDP14 - DPhiDP24;
+
+	return atan2(SinPhi, CosPhi);
+}
+
+
+export void ApplyXPBDAnisotropicBendingConstraints(uniform FVector4f PandInvM[],
+									const uniform FIntVector4 Constraints[],
+									const uniform float RestAngles[],
+									const uniform bool IsBuckled[],
+									const uniform FVector3f WarpWeftBiasBaseMultiplers[],
+									uniform float Lambdas[],
+									const uniform float Dt,
+									const uniform float MinStiffness,
+									const uniform FVector3f& Stiffness,
+									const uniform FVector3f& BucklingStiffness,
+									const uniform int32 NumConstraints)
+{
+
+	foreach(i = 0 ... NumConstraints)
+	{
+		const varying FIntVector4 Constraint = VectorLoad(&Constraints[extract(i,0)]);
+
+		const varying int32 i1 = Constraint.V[0];
+		const varying int32 i2 = Constraint.V[1];
+		const varying int32 i3 = Constraint.V[2];
+		const varying int32 i4 = Constraint.V[3];
+		
+		const varying FVector4f PandInvM1 = VectorGather(&PandInvM[i1]);
+		const varying FVector4f PandInvM2 = VectorGather(&PandInvM[i2]);
+		const varying FVector4f PandInvM3 = VectorGather(&PandInvM[i3]);
+		const varying FVector4f PandInvM4 = VectorGather(&PandInvM[i4]);
+
+		varying FVector3f P1, P2, P3, P4;
+		varying float M1, M2, M3, M4;
+
+		UnzipPandInvM(PandInvM1, P1, M1);
+		UnzipPandInvM(PandInvM2, P2, M2);
+		UnzipPandInvM(PandInvM3, P3, M3);
+		UnzipPandInvM(PandInvM4, P4, M4);
+
+		varying FVector3f Grad1, Grad2, Grad3, Grad4;
+		const varying float Angle = CalcGradientsAndAngle(P1, P2, P3, P4, Grad1, Grad2, Grad3, Grad4);
+		
+		const varying float RestAngle = RestAngles[i];
+		const varying FVector3f WarpWeftBiasBaseMultiplier = VectorGather(&WarpWeftBiasBaseMultiplers[i]);
+		const varying float StiffnessValue = select(IsBuckled[i], VectorDot(WarpWeftBiasBaseMultiplier, BucklingStiffness), VectorDot(WarpWeftBiasBaseMultiplier, Stiffness));
+		
+		const varying float Lambda = Lambdas[i];
+		const varying float Alpha  = 1.f / (StiffnessValue * Dt * Dt);
+		const varying float Denom = M1 * VectorSizeSquared(Grad1) + M2 * VectorSizeSquared(Grad2) + M3 * VectorSizeSquared(Grad3) + M4 * VectorSizeSquared(Grad4) + Alpha;
+
+		const varying float ComputedDLambda = (Angle - RestAngle - Alpha * Lambda) / Denom;
+		const varying float DLambda = select(StiffnessValue < MinStiffness, (uniform float)0, ComputedDLambda);
+		Lambdas[i] = Lambda + DLambda;
+
+		if (M1 > 0)
+		{
+			VectorScatter(&PandInvM[i1], SetVector4(P1 - DLambda * M1 * Grad1, M1));
+		}
+		if (M2 > 0)
+		{
+			VectorScatter(&PandInvM[i2], SetVector4(P2 - DLambda * M2 * Grad2, M2));
+		}
+		if (M3 > 0)
+		{
+			VectorScatter(&PandInvM[i3], SetVector4(P3 - DLambda * M3 * Grad3, M3));
+		}
+		if (M4 > 0)
+		{
+			VectorScatter(&PandInvM[i4], SetVector4(P4 - DLambda * M4 * Grad4, M4));
+		}
+	}
+}
+
+export void ApplyXPBDAnisotropicBendingConstraintsWithDamping(uniform FVector4f PandInvM[],
+									const uniform FVector3f X[],
+									const uniform FIntVector4 Constraints[],
+									const uniform float RestAngles[],
+									const uniform bool IsBuckled[],
+									const uniform FVector3f WarpWeftBiasBaseMultiplers[],
+									uniform float Lambdas[],
+									const uniform float Dt,
+									const uniform float MinStiffness,
+									const uniform FVector3f& Stiffness,
+									const uniform FVector3f& BucklingStiffness,
+									const uniform float DampingRatio,
+									const uniform int32 NumConstraints)
+{
+
+	foreach(i = 0 ... NumConstraints)
+	{
+		const varying FIntVector4 Constraint = VectorLoad(&Constraints[extract(i,0)]);
+
+		const varying int32 i1 = Constraint.V[0];
+		const varying int32 i2 = Constraint.V[1];
+		const varying int32 i3 = Constraint.V[2];
+		const varying int32 i4 = Constraint.V[3];
+		
+		const varying FVector4f PandInvM1 = VectorGather(&PandInvM[i1]);
+		const varying FVector4f PandInvM2 = VectorGather(&PandInvM[i2]);
+		const varying FVector4f PandInvM3 = VectorGather(&PandInvM[i3]);
+		const varying FVector4f PandInvM4 = VectorGather(&PandInvM[i4]);
+
+		varying FVector3f P1, P2, P3, P4;
+		varying float M1, M2, M3, M4;
+
+		UnzipPandInvM(PandInvM1, P1, M1);
+		UnzipPandInvM(PandInvM2, P2, M2);
+		UnzipPandInvM(PandInvM3, P3, M3);
+		UnzipPandInvM(PandInvM4, P4, M4);
+
+		varying FVector3f Grad1, Grad2, Grad3, Grad4;
+		const varying float Angle = CalcGradientsAndAngle(P1, P2, P3, P4, Grad1, Grad2, Grad3, Grad4);
+		
+		const varying float RestAngle = RestAngles[i];
+		const varying FVector3f WarpWeftBiasBaseMultiplier = VectorGather(&WarpWeftBiasBaseMultiplers[i]);
+		const varying float StiffnessValue = select(IsBuckled[i], VectorDot(WarpWeftBiasBaseMultiplier, BucklingStiffness), VectorDot(WarpWeftBiasBaseMultiplier, Stiffness));
+
+		const varying FVector3f X1 = VectorGather(&X[i1]);
+		const varying FVector3f X2 = VectorGather(&X[i2]);
+		const varying FVector3f X3 = VectorGather(&X[i3]);
+		const varying FVector3f X4 = VectorGather(&X[i4]);
+
+		const varying float CombinedInvMass = M1 + M2 + M3 + M4;
+		const varying float Damping = 2.f * DampingRatio * sqrt(StiffnessValue / CombinedInvMass);
+
+		const varying float Lambda = Lambdas[i];
+		const varying float Alpha  = 1.f / (StiffnessValue * Dt * Dt);
+		const varying float Gamma = Alpha * Damping * Dt;
+
+		const varying FVector3f V1TimesDt = P1 - X1;
+		const varying FVector3f V2TimesDt = P2 - X2;
+		const varying FVector3f V3TimesDt = P3 - X3;
+		const varying FVector3f V4TimesDt = P4 - X4;
+		const varying float DampingTerm = Gamma * (VectorDot(V1TimesDt, Grad1) + VectorDot(V2TimesDt, Grad2) + VectorDot(V3TimesDt, Grad3) + VectorDot(V4TimesDt, Grad4));
+
+		const varying float Denom = (1.f + Gamma)*(M1 * VectorSizeSquared(Grad1) + M2 * VectorSizeSquared(Grad2) + M3 * VectorSizeSquared(Grad3) + M4 * VectorSizeSquared(Grad4)) + Alpha;
+
+		const varying float ComputedDLambda = (Angle - RestAngle - Alpha * Lambda + DampingTerm) / Denom;
+		const varying float DLambda = select(StiffnessValue < MinStiffness, (uniform float)0, ComputedDLambda);
+		Lambdas[i] = Lambda + DLambda;
+
+		if (M1 > 0)
+		{
+			VectorScatter(&PandInvM[i1], SetVector4(P1 - DLambda * M1 * Grad1, M1));
+		}
+		if (M2 > 0)
+		{
+			VectorScatter(&PandInvM[i2], SetVector4(P2 - DLambda * M2 * Grad2, M2));
+		}
+		if (M3 > 0)
+		{
+			VectorScatter(&PandInvM[i3], SetVector4(P3 - DLambda * M3 * Grad3, M3));
+		}
+		if (M4 > 0)
+		{
+			VectorScatter(&PandInvM[i4], SetVector4(P4 - DLambda * M4 * Grad4, M4));
+		}
+	}
+}
+
+export void ApplyXPBDAnisotropicBendingConstraintsWithMaps(uniform FVector4f PandInvM[],
+											const uniform FIntVector4 Constraints[],
+											const uniform float RestAngles[],
+											const uniform bool IsBuckled[],
+											const uniform FVector3f WarpWeftBiasBaseMultiplers[],
+											uniform float Lambdas[],
+											const uniform float Dt,
+											const uniform float MinStiffness,
+											const uniform bool StiffnessHasMap,
+											const uniform uint8 StiffnessIndices[],
+											const uniform float StiffnessTable[],
+											const uniform bool StiffnessWeftHasMap,
+											const uniform uint8 StiffnessWeftIndices[],
+											const uniform float StiffnessWeftTable[],
+											const uniform bool StiffnessBiasHasMap,
+											const uniform uint8 StiffnessBiasIndices[],
+											const uniform float StiffnessBiasTable[],
+											const uniform bool BucklingStiffnessHasMap,
+											const uniform uint8 BucklingStiffnessIndices[],
+											const uniform float BucklingStiffnessTable[],
+											const uniform bool BucklingStiffnessWeftHasMap,
+											const uniform uint8 BucklingStiffnessWeftIndices[],
+											const uniform float BucklingStiffnessWeftTable[],
+											const uniform bool BucklingStiffnessBiasHasMap,
+											const uniform uint8 BucklingStiffnessBiasIndices[],
+											const uniform float BucklingStiffnessBiasTable[],
+											const uniform int32 NumConstraints)
+{
+
+	foreach(i = 0 ... NumConstraints)
+	{
+		const varying FIntVector4 Constraint = VectorLoad(&Constraints[extract(i,0)]);
+
+		const varying int32 i1 = Constraint.V[0];
+		const varying int32 i2 = Constraint.V[1];
+		const varying int32 i3 = Constraint.V[2];
+		const varying int32 i4 = Constraint.V[3];
+		const varying int8 StiffnessIndex = select(StiffnessHasMap, StiffnessIndices[i], 0);
+		const varying int8 StiffnessWeftIndex = select(StiffnessWeftHasMap, StiffnessWeftIndices[i], 0);
+		const varying int8 StiffnessBiasIndex = select(StiffnessBiasHasMap, StiffnessBiasIndices[i], 0);
+		const varying int8 BucklingStiffnessIndex = select(BucklingStiffnessHasMap, BucklingStiffnessIndices[i], 0);
+		const varying int8 BucklingStiffnessWeftIndex = select(BucklingStiffnessWeftHasMap, BucklingStiffnessWeftIndices[i], 0);
+		const varying int8 BucklingStiffnessBiasIndex = select(BucklingStiffnessBiasHasMap, BucklingStiffnessBiasIndices[i], 0);
+
+		const varying FVector4f PandInvM2 = VectorGather(&PandInvM[i2]);
+		const varying FVector4f PandInvM3 = VectorGather(&PandInvM[i3]);
+		const varying FVector4f PandInvM1 = VectorGather(&PandInvM[i1]);
+		const varying FVector4f PandInvM4 = VectorGather(&PandInvM[i4]);
+
+		varying FVector3f P1, P2, P3, P4;
+		varying float M1, M2, M3, M4;
+
+		UnzipPandInvM(PandInvM1, P1, M1);
+		UnzipPandInvM(PandInvM2, P2, M2);
+		UnzipPandInvM(PandInvM3, P3, M3);
+		UnzipPandInvM(PandInvM4, P4, M4);
+
+
+		varying FVector3f Grad1, Grad2, Grad3, Grad4;
+		const varying float Angle = CalcGradientsAndAngle(P1, P2, P3, P4, Grad1, Grad2, Grad3, Grad4);
+		
+		#pragma ignore warning(perf)
+		const varying float Stiffness = StiffnessTable[StiffnessIndex];
+		#pragma ignore warning(perf)
+		const varying float StiffnessWeft = StiffnessWeftTable[StiffnessIndex];
+		#pragma ignore warning(perf)
+		const varying float StiffnessBias = StiffnessBiasTable[StiffnessIndex];
+
+		const varying FVector3f StiffnessVector = {{Stiffness, StiffnessWeft, StiffnessBias}};
+
+		#pragma ignore warning(perf)
+		const varying float BucklingStiffness = BucklingStiffnessTable[BucklingStiffnessIndex];
+		#pragma ignore warning(perf)
+		const varying float BucklingStiffnessWeft = BucklingStiffnessWeftTable[BucklingStiffnessIndex];
+		#pragma ignore warning(perf)
+		const varying float BucklingStiffnessBias = BucklingStiffnessBiasTable[BucklingStiffnessIndex];
+		
+		const varying FVector3f BucklingStiffnessVector = {{BucklingStiffness, BucklingStiffnessWeft, BucklingStiffnessBias}};
+		
+		const varying float RestAngle = RestAngles[i];
+		const varying FVector3f WarpWeftBiasBaseMultiplier = VectorGather(&WarpWeftBiasBaseMultiplers[i]);
+		const varying float StiffnessValue = select(IsBuckled[i], VectorDot(WarpWeftBiasBaseMultiplier, BucklingStiffnessVector), VectorDot(WarpWeftBiasBaseMultiplier, StiffnessVector));
+		
+		const varying float Lambda = Lambdas[i];
+		const varying float Alpha  = 1.f / (StiffnessValue * Dt * Dt);
+		const varying float Denom = M1 * VectorSizeSquared(Grad1) + M2 * VectorSizeSquared(Grad2) + M3 * VectorSizeSquared(Grad3) + M4 * VectorSizeSquared(Grad4) + Alpha;
+
+		const varying float ComputedDLambda = (Angle - RestAngle - Alpha * Lambda) / Denom;
+		const varying float DLambda = select( StiffnessValue < MinStiffness, (uniform float)0, ComputedDLambda);
+		Lambdas[i] = Lambda + DLambda;
+
+		if (M1 > 0)
+		{
+			VectorScatter(&PandInvM[i1], SetVector4(P1 - DLambda * M1 * Grad1, M1));
+		}
+		if (M2 > 0)
+		{
+			VectorScatter(&PandInvM[i2], SetVector4(P2 - DLambda * M2 * Grad2, M2));
+		}
+		if (M3 > 0)
+		{
+			VectorScatter(&PandInvM[i3], SetVector4(P3 - DLambda * M3 * Grad3, M3));
+		}
+		if (M4 > 0)
+		{
+			VectorScatter(&PandInvM[i4], SetVector4(P4 - DLambda * M4 * Grad4, M4));
+		}
+	}
+}
+
+export void ApplyXPBDAnisotropicBendingConstraintsWithDampingAndMaps(uniform FVector4f PandInvM[],
+											const uniform FVector3f X[],
+											const uniform FIntVector4 Constraints[],
+											const uniform float RestAngles[],
+											const uniform bool IsBuckled[],
+											const uniform FVector3f WarpWeftBiasBaseMultiplers[],
+											uniform float Lambdas[],
+											const uniform float Dt,
+											const uniform float MinStiffness,
+											const uniform bool StiffnessHasMap,
+											const uniform uint8 StiffnessIndices[],
+											const uniform float StiffnessTable[],
+											const uniform bool StiffnessWeftHasMap,
+											const uniform uint8 StiffnessWeftIndices[],
+											const uniform float StiffnessWeftTable[],
+											const uniform bool StiffnessBiasHasMap,
+											const uniform uint8 StiffnessBiasIndices[],
+											const uniform float StiffnessBiasTable[],
+											const uniform bool BucklingStiffnessHasMap,
+											const uniform uint8 BucklingStiffnessIndices[],
+											const uniform float BucklingStiffnessTable[],
+											const uniform bool BucklingStiffnessWeftHasMap,
+											const uniform uint8 BucklingStiffnessWeftIndices[],
+											const uniform float BucklingStiffnessWeftTable[],
+											const uniform bool BucklingStiffnessBiasHasMap,
+											const uniform uint8 BucklingStiffnessBiasIndices[],
+											const uniform float BucklingStiffnessBiasTable[],
+											const uniform bool DampingHasMap,
+											const uniform uint8 DampingIndices[],
+											const uniform float DampingTable[],
+											const uniform int32 NumConstraints)
+{
+
+	foreach(i = 0 ... NumConstraints)
+	{
+		const varying FIntVector4 Constraint = VectorLoad(&Constraints[extract(i,0)]);
+
+		const varying int32 i1 = Constraint.V[0];
+		const varying int32 i2 = Constraint.V[1];
+		const varying int32 i3 = Constraint.V[2];
+		const varying int32 i4 = Constraint.V[3];
+		const varying int8 StiffnessIndex = select(StiffnessHasMap, StiffnessIndices[i], 0);
+		const varying int8 StiffnessWeftIndex = select(StiffnessWeftHasMap, StiffnessWeftIndices[i], 0);
+		const varying int8 StiffnessBiasIndex = select(StiffnessBiasHasMap, StiffnessBiasIndices[i], 0);
+		const varying int8 BucklingStiffnessIndex = select(BucklingStiffnessHasMap, BucklingStiffnessIndices[i], 0);
+		const varying int8 BucklingStiffnessWeftIndex = select(BucklingStiffnessWeftHasMap, BucklingStiffnessWeftIndices[i], 0);
+		const varying int8 BucklingStiffnessBiasIndex = select(BucklingStiffnessBiasHasMap, BucklingStiffnessBiasIndices[i], 0);
+		const varying int8 DampingIndex = select(DampingHasMap, DampingIndices[i], 0);
+
+		const varying FVector4f PandInvM2 = VectorGather(&PandInvM[i2]);
+		const varying FVector4f PandInvM3 = VectorGather(&PandInvM[i3]);
+		const varying FVector4f PandInvM1 = VectorGather(&PandInvM[i1]);
+		const varying FVector4f PandInvM4 = VectorGather(&PandInvM[i4]);
+
+		varying FVector3f P1, P2, P3, P4;
+		varying float M1, M2, M3, M4;
+
+		UnzipPandInvM(PandInvM1, P1, M1);
+		UnzipPandInvM(PandInvM2, P2, M2);
+		UnzipPandInvM(PandInvM3, P3, M3);
+		UnzipPandInvM(PandInvM4, P4, M4);
+
+
+		varying FVector3f Grad1, Grad2, Grad3, Grad4;
+		const varying float Angle = CalcGradientsAndAngle(P1, P2, P3, P4, Grad1, Grad2, Grad3, Grad4);
+		
+		
+		#pragma ignore warning(perf)
+		const varying float Stiffness = StiffnessTable[StiffnessIndex];
+		#pragma ignore warning(perf)
+		const varying float StiffnessWeft = StiffnessWeftTable[StiffnessIndex];
+		#pragma ignore warning(perf)
+		const varying float StiffnessBias = StiffnessBiasTable[StiffnessIndex];
+
+		const varying FVector3f StiffnessVector = {{Stiffness, StiffnessWeft, StiffnessBias}};
+
+		#pragma ignore warning(perf)
+		const varying float BucklingStiffness = BucklingStiffnessTable[BucklingStiffnessIndex];
+		#pragma ignore warning(perf)
+		const varying float BucklingStiffnessWeft = BucklingStiffnessWeftTable[BucklingStiffnessIndex];
+		#pragma ignore warning(perf)
+		const varying float BucklingStiffnessBias = BucklingStiffnessBiasTable[BucklingStiffnessIndex];
+		
+		const varying FVector3f BucklingStiffnessVector = {{BucklingStiffness, BucklingStiffnessWeft, BucklingStiffnessBias}};
+		
+		#pragma ignore warning(perf)
+		const varying float DampingRatio = DampingTable[DampingIndex];
+		
+		const varying float RestAngle = RestAngles[i];
+		const varying FVector3f WarpWeftBiasBaseMultiplier = VectorGather(&WarpWeftBiasBaseMultiplers[i]);
+		const varying float StiffnessValue = select(IsBuckled[i], VectorDot(WarpWeftBiasBaseMultiplier, BucklingStiffnessVector), VectorDot(WarpWeftBiasBaseMultiplier, StiffnessVector));
+
+		const varying FVector3f X1 = VectorGather(&X[i1]);
+		const varying FVector3f X2 = VectorGather(&X[i2]);
+		const varying FVector3f X3 = VectorGather(&X[i3]);
+		const varying FVector3f X4 = VectorGather(&X[i4]);
+
+		const varying float CombinedInvMass = M1 + M2 + M3 + M4;
+		const varying float Damping = 2.f * DampingRatio * sqrt(StiffnessValue / CombinedInvMass);
+		
+		const varying float Lambda = Lambdas[i];
+		const varying float Alpha  = 1.f / (StiffnessValue * Dt * Dt);
+		const varying float Gamma = Alpha * Damping * Dt;
+
+		const varying FVector3f V1TimesDt = P1 - X1;
+		const varying FVector3f V2TimesDt = P2 - X2;
+		const varying FVector3f V3TimesDt = P3 - X3;
+		const varying FVector3f V4TimesDt = P4 - X4;
+		const varying float DampingTerm = Gamma * (VectorDot(V1TimesDt, Grad1) + VectorDot(V2TimesDt, Grad2) + VectorDot(V3TimesDt, Grad3) + VectorDot(V4TimesDt, Grad4));
+
+		const varying float Denom = (1.f + Gamma)*(M1 * VectorSizeSquared(Grad1) + M2 * VectorSizeSquared(Grad2) + M3 * VectorSizeSquared(Grad3) + M4 * VectorSizeSquared(Grad4)) + Alpha;
+
+		const varying float ComputedDLambda = (Angle - RestAngle - Alpha * Lambda + DampingTerm) / Denom;
+		const varying float DLambda = select( StiffnessValue < MinStiffness, (uniform float)0, ComputedDLambda);
+		Lambdas[i] = Lambda + DLambda;
+
+		if (M1 > 0)
+		{
+			VectorScatter(&PandInvM[i1], SetVector4(P1 - DLambda * M1 * Grad1, M1));
+		}
+		if (M2 > 0)
+		{
+			VectorScatter(&PandInvM[i2], SetVector4(P2 - DLambda * M2 * Grad2, M2));
+		}
+		if (M3 > 0)
+		{
+			VectorScatter(&PandInvM[i3], SetVector4(P3 - DLambda * M3 * Grad3, M3));
+		}
+		if (M4 > 0)
+		{
+			VectorScatter(&PandInvM[i4], SetVector4(P4 - DLambda * M4 * Grad4, M4));
+		}
+	}
+}
\ No newline at end of file
diff --git a/Engine/Source/Runtime/Experimental/Chaos/Public/Chaos/PBDSoftsEvolutionFwd.h b/Engine/Source/Runtime/Experimental/Chaos/Public/Chaos/PBDSoftsEvolutionFwd.h
--- a/Engine/Source/Runtime/Experimental/Chaos/Public/Chaos/PBDSoftsEvolutionFwd.h
+++ b/Engine/Source/Runtime/Experimental/Chaos/Public/Chaos/PBDSoftsEvolutionFwd.h
@@ -52,6 +52,7 @@ namespace Chaos
 		class FXPBDBendingSpringConstraints;
 		class FPBDBendingConstraints;
 		class FXPBDBendingConstraints;
+		class FXPBDAnisotropicBendingConstraints;
 		class FPBDAxialSpringConstraints;
 		class FXPBDAxialSpringConstraints;
 		class FPBDAreaSpringConstraints;
diff --git a/Engine/Source/Runtime/Experimental/Chaos/Public/Chaos/XPBDAnisotropicBendingConstraints.h b/Engine/Source/Runtime/Experimental/Chaos/Public/Chaos/XPBDAnisotropicBendingConstraints.h
new file mode 100644
--- /dev/null
+++ b/Engine/Source/Runtime/Experimental/Chaos/Public/Chaos/XPBDAnisotropicBendingConstraints.h
@@ -0,0 +1,161 @@
+// Copyright Epic Games, Inc. All Rights Reserved.
+#pragma once
+
+#include "Chaos/PBDBendingConstraintsBase.h"
+#include "Chaos/CollectionPropertyFacade.h"
+#include "ChaosStats.h"
+
+namespace Chaos
+{
+class FTriangleMesh;
+}
+
+namespace Chaos::Softs
+{
+
+class CHAOS_API FXPBDAnisotropicBendingConstraints final : public FPBDBendingConstraintsBase
+{
+	typedef FPBDBendingConstraintsBase Base;
+
+public:
+	// Stiffness is in kg cm^2 / rad^2 s^2
+	static constexpr FSolverReal MinStiffness = (FSolverReal)1e-4; // Stiffness below this will be considered 0 since all of our calculations are actually based on 1 / stiffness.
+	static constexpr FSolverReal MaxStiffness = (FSolverReal)1e7;
+	static constexpr FSolverReal MinDamping = (FSolverReal)0.;
+	static constexpr FSolverReal MaxDamping = (FSolverReal)1000.;
+
+	static bool IsEnabled(const FCollectionPropertyConstFacade& PropertyCollection)
+	{
+		return IsXPBDBendingElementStiffnessWarpEnabled(PropertyCollection, false);
+	}
+
+	FXPBDAnisotropicBendingConstraints(const FSolverParticles& InParticles,
+		int32 ParticleOffset,
+		int32 ParticleCount,
+		const FTriangleMesh& TriangleMesh,
+		const TArray<TVec3<FVec2f>>& FaceVertexPatternPositions,
+		const TConstArrayView<FRealSingle>& StiffnessWarpMultipliers,
+		const TConstArrayView<FRealSingle>& StiffnessWeftMultipliers,
+		const TConstArrayView<FRealSingle>& StiffnessBiasMultipliers,
+		const TConstArrayView<FRealSingle>& BucklingStiffnessWarpMultipliers,
+		const TConstArrayView<FRealSingle>& BucklingStiffnessWeftMultipliers,
+		const TConstArrayView<FRealSingle>& BucklingStiffnessBiasMultipliers,
+		const TConstArrayView<FRealSingle>& DampingMultipliers,
+		const FCollectionPropertyConstFacade& PropertyCollection,
+		bool bTrimKinematicConstraints = false);
+
+	FXPBDAnisotropicBendingConstraints(const FSolverParticles& InParticles,
+		int32 ParticleOffset,
+		int32 ParticleCount,
+		const FTriangleMesh& TriangleMesh,
+		const TArray<TVec3<FVec2f>>& FaceVertexPatternPositions,
+		const TConstArrayView<FRealSingle>& StiffnessWarpMultipliers,
+		const TConstArrayView<FRealSingle>& StiffnessWeftMultipliers,
+		const TConstArrayView<FRealSingle>& StiffnessBiasMultipliers,
+		const TConstArrayView<FRealSingle>& BucklingStiffnessWarpMultipliers,
+		const TConstArrayView<FRealSingle>& BucklingStiffnessWeftMultipliers,
+		const TConstArrayView<FRealSingle>& BucklingStiffnessBiasMultipliers,
+		const TConstArrayView<FRealSingle>& DampingMultipliers,
+		const FSolverVec2& InStiffnessWarp,
+		const FSolverVec2& InStiffnessWeft,
+		const FSolverVec2& InStiffnessBias,
+		const FSolverReal InBucklingRatio,
+		const FSolverVec2& InBucklingStiffnessWarp,
+		const FSolverVec2& InBucklingStiffnessWeft,
+		const FSolverVec2& InBucklingStiffnessBias,
+		const FSolverVec2& InDampingRatio,
+		bool bTrimKinematicConstraints = false);
+
+	virtual ~FXPBDAnisotropicBendingConstraints() override {}
+
+	void Init(const FSolverParticles& InParticles)
+	{ 
+		Lambdas.Reset();
+		Lambdas.AddZeroed(Constraints.Num());
+		FPBDBendingConstraintsBase::Init(InParticles);
+	}
+
+	void SetProperties(const FCollectionPropertyConstFacade& PropertyCollection)
+	{
+		if (IsXPBDBendingElementStiffnessWarpMutable(PropertyCollection))
+		{
+			Stiffness.SetWeightedValue(FSolverVec2(GetWeightedFloatXPBDBendingElementStiffnessWarp(PropertyCollection)), MaxStiffness);
+		}
+		if (IsXPBDBendingElementStiffnessWeftMutable(PropertyCollection))
+		{
+			Stiffness.SetWeightedValue(FSolverVec2(GetWeightedFloatXPBDBendingElementStiffnessWeft(PropertyCollection)), MaxStiffness);
+		}
+		if (IsXPBDBendingElementStiffnessBiasMutable(PropertyCollection))
+		{
+			Stiffness.SetWeightedValue(FSolverVec2(GetWeightedFloatXPBDBendingElementStiffnessBias(PropertyCollection)), MaxStiffness);
+		}
+		if (IsXPBDBucklingRatioMutable(PropertyCollection))
+		{
+			BucklingRatio = FMath::Clamp(GetXPBDBucklingRatio(PropertyCollection), (FSolverReal)0., (FSolverReal)1.);
+		}
+		if (IsXPBDBucklingStiffnessWarpMutable(PropertyCollection))
+		{
+			BucklingStiffness.SetWeightedValue(FSolverVec2(GetWeightedFloatXPBDBucklingStiffnessWarp(PropertyCollection)), MaxStiffness);
+		}
+		if (IsXPBDBucklingStiffnessWeftMutable(PropertyCollection))
+		{
+			BucklingStiffness.SetWeightedValue(FSolverVec2(GetWeightedFloatXPBDBucklingStiffnessWeft(PropertyCollection)), MaxStiffness);
+		}
+		if (IsXPBDBucklingStiffnessBiasMutable(PropertyCollection))
+		{
+			BucklingStiffness.SetWeightedValue(FSolverVec2(GetWeightedFloatXPBDBucklingStiffnessBias(PropertyCollection)), MaxStiffness);
+		}
+		if (IsXPBDBendingElementDampingMutable(PropertyCollection))
+		{
+			DampingRatio.SetWeightedValue(FSolverVec2(GetWeightedFloatXPBDBendingElementDamping(PropertyCollection)).ClampAxes(MinDamping, MaxDamping));
+		}
+	}
+
+	// Update stiffness table, as well as the simulation stiffness exponent
+	void ApplyProperties(const FSolverReal /*Dt*/, const int32 /*NumIterations*/)
+	{
+		Stiffness.ApplyXPBDValues(MaxStiffness);
+		StiffnessWeft.ApplyXPBDValues(MaxStiffness);
+		StiffnessBias.ApplyXPBDValues(MaxStiffness);
+		BucklingStiffness.ApplyXPBDValues(MaxStiffness);
+		BucklingStiffnessWeft.ApplyXPBDValues(MaxStiffness);
+		BucklingStiffnessBias.ApplyXPBDValues(MaxStiffness);
+		DampingRatio.ApplyValues();
+	}
+
+	void Apply(FSolverParticles& Particles, const FSolverReal Dt) const;
+
+private:
+	void InitColor(const FSolverParticles& InParticles, const int32 ParticleOffset, const int32 ParticleCount);
+	void ApplyHelper(FSolverParticles& Particles, const FSolverReal Dt, const int32 ConstraintIndex, const FSolverVec3& ExpStiffnessValues, 
+		const FSolverVec3& ExpBucklingStiffnessValues, const FSolverReal DampingRatioValue) const;
+
+	TArray<FSolverVec3> GenerateWarpWeftBiasBaseMultipliers(const TArray<TVec3<FVec2f>>& FaceVertexPatternPositions, const FTriangleMesh& TriangleMesh) const;
+
+	using Base::Constraints;
+	using Base::RestAngles;
+	using Base::Stiffness; // Warp
+	using Base::BucklingStiffness; // Warp
+	
+	FPBDStiffness StiffnessWeft;
+	FPBDStiffness StiffnessBias;
+	FPBDStiffness BucklingStiffnessWeft;
+	FPBDStiffness BucklingStiffnessBias;
+
+	FPBDWeightMap DampingRatio;
+	mutable TArray<FSolverReal> Lambdas;
+	TArray<int32> ConstraintsPerColorStartIndex; // Constraints are ordered so each batch is contiguous. This is ColorNum + 1 length so it can be used as start and end.
+
+	TArray<FSolverVec3> WarpWeftBiasBaseMultipliers;
+
+	UE_CHAOS_DECLARE_PROPERTYCOLLECTION_NAME(XPBDBendingElementStiffnessWarp, float);
+	UE_CHAOS_DECLARE_PROPERTYCOLLECTION_NAME(XPBDBendingElementStiffnessWeft, float);
+	UE_CHAOS_DECLARE_PROPERTYCOLLECTION_NAME(XPBDBendingElementStiffnessBias, float);
+	UE_CHAOS_DECLARE_PROPERTYCOLLECTION_NAME(XPBDBendingElementDamping, float);
+	UE_CHAOS_DECLARE_PROPERTYCOLLECTION_NAME(XPBDBucklingRatio, float);
+	UE_CHAOS_DECLARE_PROPERTYCOLLECTION_NAME(XPBDBucklingStiffnessWarp, float);
+	UE_CHAOS_DECLARE_PROPERTYCOLLECTION_NAME(XPBDBucklingStiffnessWeft, float);
+	UE_CHAOS_DECLARE_PROPERTYCOLLECTION_NAME(XPBDBucklingStiffnessBias, float);
+};
+
+}  // End namespace Chaos::Softs
