Add an option to GetNumBreakingStrains to count strains N times if they are N times the internal strain

diff --git a/Engine/Source/Runtime/Experimental/Chaos/Private/Chaos/StrainModification.cpp b/Engine/Source/Runtime/Experimental/Chaos/Private/Chaos/StrainModification.cpp
--- a/Engine/Source/Runtime/Experimental/Chaos/Private/Chaos/StrainModification.cpp
+++ b/Engine/Source/Runtime/Experimental/Chaos/Private/Chaos/StrainModification.cpp
@@ -25,7 +25,7 @@ int32 Chaos::FStrainedProxyModifier::GetNumRestBreakables() const
 	return 0;
 }
 
-int32 Chaos::FStrainedProxyModifier::GetNumBreakingStrains(const uint8 StrainTypes) const
+int32 Chaos::FStrainedProxyModifier::GetNumBreakingStrains(const bool bDoubleCount, const uint8 StrainTypes) const
 {
 	// Make sure we have a proxy and rest-children
 	if (Proxy == nullptr) { return 0; }
@@ -51,9 +51,14 @@ int32 Chaos::FStrainedProxyModifier::GetNumBreakingStrains(const uint8 StrainTyp
 		const Chaos::FReal ExternalStrain
 			= (StrainTypes & Chaos::EStrainTypes::ExternalStrain)
 			? ChildHandle->GetExternalStrain() : FReal(0);
-		const Chaos::FReal MaxAppliedStrain = FMath::Max(CollisionImpulses, ExternalStrain);
 		const Chaos::FReal InternalStrain = ChildHandle->GetInternalStrains();
-		if (MaxAppliedStrain >= InternalStrain)
+		Chaos::FReal MaxAppliedStrain = FMath::Max(CollisionImpulses, ExternalStrain);
+
+		if (bDoubleCount && InternalStrain > SMALL_NUMBER)
+		{
+			NumBreakingStrains += (int32)(MaxAppliedStrain / InternalStrain);
+		}
+		else if (MaxAppliedStrain >= InternalStrain)
 		{
 			++NumBreakingStrains;
 		}
diff --git a/Engine/Source/Runtime/Experimental/Chaos/Public/Chaos/StrainModification.h b/Engine/Source/Runtime/Experimental/Chaos/Public/Chaos/StrainModification.h
--- a/Engine/Source/Runtime/Experimental/Chaos/Public/Chaos/StrainModification.h
+++ b/Engine/Source/Runtime/Experimental/Chaos/Public/Chaos/StrainModification.h
@@ -54,7 +54,9 @@ namespace Chaos
 		int32 GetNumRestBreakables() const;
 
 		// Get the number of breaking strains amongst the level-1 strainables
-		int32 GetNumBreakingStrains(const uint8 StrainTypes = EStrainTypes::ExternalStrain | EStrainTypes::CollisionStrain) const;
+		// If DoubleCount is true, then add N for each strain which is strong enough to
+		// break a connection N times.
+		int32 GetNumBreakingStrains(bool bDoubleCount = true, const uint8 StrainTypes = EStrainTypes::ExternalStrain | EStrainTypes::CollisionStrain) const;
 
 		// Clear strains for all strained cluster children
 		void ClearStrains();
