Iris - Prevent polling of function states.

#rnx

diff --git a/Engine/Source/Runtime/Experimental/Iris/Core/Private/Iris/ReplicationSystem/PropertyReplicationFragment.cpp b/Engine/Source/Runtime/Experimental/Iris/Core/Private/Iris/ReplicationSystem/PropertyReplicationFragment.cpp
--- a/Engine/Source/Runtime/Experimental/Iris/Core/Private/Iris/ReplicationSystem/PropertyReplicationFragment.cpp
+++ b/Engine/Source/Runtime/Experimental/Iris/Core/Private/Iris/ReplicationSystem/PropertyReplicationFragment.cpp
@@ -53,8 +53,12 @@ FPropertyReplicationFragment::FPropertyReplicationFragment(EReplicationFragmentT
 	
 	if (EnumHasAnyFlags(InTraits, EReplicationFragmentTraits::CanReplicate))
 	{
-		// For PropertyReplicationStates we need to poll properties from our owner in order to detect state changes.
-		Traits |= EReplicationFragmentTraits::NeedsPoll;
+		// For PropertyReplicationStates, except for pure function states, we need to poll properties from our owner in order to detect state changes.
+		if (InDescriptor->FunctionCount == 0)
+		{
+			// In theory we wouldn't need CanReplicate/CanReceive but there's a lot of logic that would then have to check whether the fragment has the appropriate buffers.
+			Traits |= EReplicationFragmentTraits::NeedsPoll;
+		}
 	}
 
 	// Propagate push based dirtiness.
