Iris - Bump fast array changemask bits

#rnx

diff --git a/Engine/Source/Runtime/Experimental/Iris/Core/Private/Iris/ReplicationState/IrisFastArraySerializer.cpp b/Engine/Source/Runtime/Experimental/Iris/Core/Private/Iris/ReplicationState/IrisFastArraySerializer.cpp
--- a/Engine/Source/Runtime/Experimental/Iris/Core/Private/Iris/ReplicationState/IrisFastArraySerializer.cpp
+++ b/Engine/Source/Runtime/Experimental/Iris/Core/Private/Iris/ReplicationState/IrisFastArraySerializer.cpp
@@ -7,6 +7,7 @@
 FIrisFastArraySerializer::FIrisFastArraySerializer()
 : FFastArraySerializer()
 {
+	static_assert(offsetof(FIrisFastArraySerializer, ChangeMaskStorage) - sizeof(UE::Net::FReplicationStateHeader) == offsetof(FIrisFastArraySerializer, ReplicationStateHeader), "GetReplicationStateHeader is no longer compatible with FIrisFastArraySerializer");
 	InitChangeMask();
 }
 
@@ -54,8 +55,18 @@ FIrisFastArraySerializer& FIrisFastArraySerializer::operator=(FIrisFastArraySeri
 void FIrisFastArraySerializer::InitChangeMask()
 {
 	// Init changemask and conditional changemask
-	ChangeMaskStorage[0] = 0U;
-	ChangeMaskStorage[1] = 0xffffffffU;
+	for (uint32& Storage : ChangeMaskStorage)
+	{
+		const SIZE_T Index = &Storage - &ChangeMaskStorage[0];
+		if (Index < ConditionalChangeMaskStorageIndex)
+		{
+			Storage = 0U;
+		}
+		else
+		{
+			Storage = 0xFFFFFFFFU;
+		}
+	}
 }
 
 void FIrisFastArraySerializer::InternalMarkAllItemsChanged()
@@ -81,12 +92,12 @@ namespace UE::Net::Private
 
 FNetBitArrayView FIrisFastArraySerializerPrivateAccessor::GetChangeMask(FIrisFastArraySerializer& Array)
 {
-	return MakeNetBitArrayView(&Array.ChangeMaskStorage[0], FIrisFastArraySerializer::IrisFastArrayChangeMaskBits + 1U);
+	return MakeNetBitArrayView(&Array.ChangeMaskStorage[FIrisFastArraySerializer::MemberChangeMaskStorageIndex], FIrisFastArraySerializer::IrisFastArrayChangeMaskBits + 1U);
 }
 
 FNetBitArrayView FIrisFastArraySerializerPrivateAccessor::GetConditionalChangeMask(FIrisFastArraySerializer& Array)
 {
-	return MakeNetBitArrayView(&Array.ChangeMaskStorage[1], FIrisFastArraySerializer::IrisFastArrayChangeMaskBits + 1U);
+	return MakeNetBitArrayView(&Array.ChangeMaskStorage[FIrisFastArraySerializer::ConditionalChangeMaskStorageIndex], FIrisFastArraySerializer::IrisFastArrayChangeMaskBits + 1U);
 }
 
 void FIrisFastArraySerializerPrivateAccessor::MarkAllArrayItemsDirty(FIrisFastArraySerializer& Array, uint32 StartingIndex)
diff --git a/Engine/Source/Runtime/Experimental/Iris/Core/Public/Iris/ReplicationState/IrisFastArraySerializer.h b/Engine/Source/Runtime/Experimental/Iris/Core/Public/Iris/ReplicationState/IrisFastArraySerializer.h
--- a/Engine/Source/Runtime/Experimental/Iris/Core/Public/Iris/ReplicationState/IrisFastArraySerializer.h
+++ b/Engine/Source/Runtime/Experimental/Iris/Core/Public/Iris/ReplicationState/IrisFastArraySerializer.h
@@ -22,7 +22,7 @@ USTRUCT()
 struct FIrisFastArraySerializer : public FFastArraySerializer
 {
 	// At the moment as we have no way to specify this we reserve 31 bits that can be used, the first bit is used for the array itself
-	enum { IrisFastArrayChangeMaskBits = 31U };
+	enum { IrisFastArrayChangeMaskBits = 63U };
 	enum { IrisFastArrayChangeMaskBitOffset = 1U };
 	enum { IrisFastArrayPropertyBitIndex = 0U };	
 
@@ -68,17 +68,23 @@ struct FIrisFastArraySerializer : public FFastArraySerializer
 	friend UE::Net::Private::FIrisFastArraySerializerPrivateAccessor;
 
 private:
+	enum : unsigned
+	{
+		MemberChangeMaskStorageIndex = 0U,
+		ConditionalChangeMaskStorageIndex = 2U,
+	};
+
 	// Mark item at index as changed, and set the array as dirty
 	void InitChangeMask();
 	void InternalMarkItemChanged(int32 ItemIdx);
 	void InternalMarkAllItemsChanged();
 	void InternalMarkArrayAsDirty();
 
+	// Header for dirty state tracking needs to be just before ChangeMaskStorage. See GetReplicationStateHeader for more info.
+	UE::Net::FReplicationStateHeader ReplicationStateHeader;
+
 	// Storage for changemask, this is currently hardcoded
 	UPROPERTY(Transient, NotReplicated)
-	uint32 ChangeMaskStorage[2];
-
-	// Header for dirty state tracking
-	UE::Net::FReplicationStateHeader ReplicationStateHeader;
+	uint32 ChangeMaskStorage[4];
 };
 
diff --git a/Engine/Source/Runtime/Experimental/Iris/Core/Public/Iris/ReplicationState/ReplicationStateUtil.h b/Engine/Source/Runtime/Experimental/Iris/Core/Public/Iris/ReplicationState/ReplicationStateUtil.h
--- a/Engine/Source/Runtime/Experimental/Iris/Core/Public/Iris/ReplicationState/ReplicationStateUtil.h
+++ b/Engine/Source/Runtime/Experimental/Iris/Core/Public/Iris/ReplicationState/ReplicationStateUtil.h
@@ -34,15 +34,8 @@ namespace UE::Net::Private
 inline UE::Net::FReplicationStateHeader& GetReplicationStateHeader(void* StateBuffer, const FReplicationStateDescriptor* Descriptor)
 {
 	// $IRIS: $TODO: Currently we do not store the offset to the internal replication state
-	if (!EnumHasAnyFlags(Descriptor->Traits, EReplicationStateTraits::IsNativeFastArrayReplicationState))
-	{
-		return *reinterpret_cast<UE::Net::FReplicationStateHeader*>(StateBuffer);
-	}
-	else
-	{
-		const SIZE_T ReplicationStateHeaderOffset = Descriptor->ChangeMasksExternalOffset + 8U;
-		return *reinterpret_cast<UE::Net::FReplicationStateHeader*>(reinterpret_cast<uint8*>(StateBuffer) + ReplicationStateHeaderOffset);
-	}
+	const SIZE_T ReplicationStateHeaderOffset = (EnumHasAnyFlags(Descriptor->Traits, EReplicationStateTraits::IsNativeFastArrayReplicationState) ? SIZE_T(Descriptor->ChangeMasksExternalOffset - sizeof(FReplicationStateHeader)) : SIZE_T(0));
+	return *reinterpret_cast<UE::Net::FReplicationStateHeader*>(reinterpret_cast<uint8*>(StateBuffer) + ReplicationStateHeaderOffset);
 }
 
 inline bool IsReplicationStateBound(void* StateBuffer, const FReplicationStateDescriptor* Descriptor)
