Iris
* Added ability to retrieve the name of an Iris Filter via it's handle
* Added IRIS_PROFILER_SCOPE_VERBOSE for low-level details of Iris.  Enabled via UE_IRIS_PROFILER_ENABLE_VERBOSE
* Formalized dynamic name captures via IRIS_PROFILER_PROTOCOL_NAME. Enabled via UE_IRIS_PROFILER_ENABLE_PROTOCOL_NAMES

diff --git a/Engine/Source/Runtime/Experimental/Iris/Core/Private/Iris/ReplicationState/PropertyReplicationState.cpp b/Engine/Source/Runtime/Experimental/Iris/Core/Private/Iris/ReplicationState/PropertyReplicationState.cpp
--- a/Engine/Source/Runtime/Experimental/Iris/Core/Private/Iris/ReplicationState/PropertyReplicationState.cpp
+++ b/Engine/Source/Runtime/Experimental/Iris/Core/Private/Iris/ReplicationState/PropertyReplicationState.cpp
@@ -225,9 +225,7 @@ bool FPropertyReplicationState::PollPropertyReplicationState(const void* RESTRIC
 		const FReplicationStateDescriptor* Descriptor = ReplicationStateDescriptor;
 		const uint8* SrcBuffer = reinterpret_cast<const uint8*>(SrcStateData);
 
-#if UE_IRIS_PROFILE_PROTOCOL_NAMES
-		IRIS_PROFILER_SCOPE_TEXT(ReplicationStateDescriptor->DebugName->Name);
-#endif
+		IRIS_PROFILER_PROTOCOL_NAME(ReplicationStateDescriptor->DebugName->Name);
 
 		const FReplicationStateMemberDescriptor* MemberDescriptors = Descriptor->MemberDescriptors;
 		const FProperty** MemberProperties = Descriptor->MemberProperties;
@@ -257,9 +255,7 @@ void FPropertyReplicationState::PushPropertyReplicationState(void* RESTRICT DstD
 		const FReplicationStateDescriptor* Descriptor = ReplicationStateDescriptor;
 		uint8* DstBuffer = reinterpret_cast<uint8*>(DstData);
 
-#if UE_IRIS_PROFILE_PROTOCOL_NAMES
-		IRIS_PROFILER_SCOPE_TEXT(ReplicationStateDescriptor->DebugName->Name);
-#endif
+		IRIS_PROFILER_PROTOCOL_NAME(ReplicationStateDescriptor->DebugName->Name);
 
 		const FReplicationStateMemberDescriptor* MemberDescriptors = Descriptor->MemberDescriptors;
 		const FProperty** MemberProperties = Descriptor->MemberProperties;
diff --git a/Engine/Source/Runtime/Experimental/Iris/Core/Private/Iris/ReplicationSystem/Filtering/ReplicationFiltering.cpp b/Engine/Source/Runtime/Experimental/Iris/Core/Private/Iris/ReplicationSystem/Filtering/ReplicationFiltering.cpp
--- a/Engine/Source/Runtime/Experimental/Iris/Core/Private/Iris/ReplicationSystem/Filtering/ReplicationFiltering.cpp
+++ b/Engine/Source/Runtime/Experimental/Iris/Core/Private/Iris/ReplicationSystem/Filtering/ReplicationFiltering.cpp
@@ -19,6 +19,36 @@
 namespace UE::Net::Private
 {
 
+FName GetStaticFilterName(FNetObjectFilterHandle Filter)
+{
+	switch (Filter)
+	{
+		case InvalidNetObjectFilterHandle:
+		{
+			return NAME_None;
+		} break;
+
+		case ToOwnerFilterHandle:
+		{
+			static const FName ToOwnerFilterName = TEXT("ToOwnerFilter");
+			return ToOwnerFilterName;
+		} break;
+
+		case ConnectionFilterHandle:
+		{
+			static const FName ConnectionFilterName = TEXT("ConnectionFilter");
+			return ConnectionFilterName;
+		} break;
+
+		default:
+		{
+			ensureMsgf(false, TEXT("ReplicationFiltering GetStaticFilterName() received undefined static filter handle %d."), Filter);
+		} break;
+	}
+
+	return NAME_None;
+}
+
 inline static ENetFilterStatus GetDependentObjectFilterStatus(const FNetRefHandleManager* NetRefHandleManager, const FNetBitArray& ObjectsInScope, FInternalNetRefIndex ObjectIndex)
 {
 	for (const FInternalNetRefIndex ParentObjectIndex : NetRefHandleManager->GetDependentObjectParents(ObjectIndex))
@@ -447,6 +477,17 @@ UNetObjectFilter* FReplicationFiltering::GetFilter(const FName FilterName) const
 	return nullptr;
 }
 
+FName FReplicationFiltering::GetFilterName(FNetObjectFilterHandle Filter) const
+{
+	if (FNetObjectFilterHandleUtil::IsDynamicFilter(Filter))
+	{
+		const uint32 DynamicFilterIndex = FNetObjectFilterHandleUtil::GetDynamicFilterIndex(Filter);
+		return DynamicFilterInfos[DynamicFilterIndex].Name;
+	}
+
+	return UE::Net::Private::GetStaticFilterName(Filter);
+}
+
 // Connection handling
 void FReplicationFiltering::AddConnection(uint32 ConnectionId)
 {
@@ -1937,4 +1978,4 @@ inline uint32 FNetObjectFilterHandleUtil::GetDynamicFilterIndex(FNetObjectFilter
 	return (Handle & DynamicNetObjectFilterHandleFlag) ? (Handle & ~DynamicNetObjectFilterHandleFlag) : ~0U;
 }
 
-}
+} // end namespace UE::Net::Private
diff --git a/Engine/Source/Runtime/Experimental/Iris/Core/Private/Iris/ReplicationSystem/Filtering/ReplicationFiltering.h b/Engine/Source/Runtime/Experimental/Iris/Core/Private/Iris/ReplicationSystem/Filtering/ReplicationFiltering.h
--- a/Engine/Source/Runtime/Experimental/Iris/Core/Private/Iris/ReplicationSystem/Filtering/ReplicationFiltering.h
+++ b/Engine/Source/Runtime/Experimental/Iris/Core/Private/Iris/ReplicationSystem/Filtering/ReplicationFiltering.h
@@ -67,6 +67,9 @@ public:
 	FNetObjectFilterHandle GetFilterHandle(const FName FilterName) const;
 	UNetObjectFilter* GetFilter(const FName FilterName) const;
 
+	/** Returns the name of the Filter represented by the handle. */
+	FName GetFilterName(FNetObjectFilterHandle Filter) const;
+
 	// Connection handling
 	void AddConnection(uint32 ConnectionId);
 	void RemoveConnection(uint32 ConnectionId);
diff --git a/Engine/Source/Runtime/Experimental/Iris/Core/Private/Iris/ReplicationSystem/ObjectReplicationBridge.cpp b/Engine/Source/Runtime/Experimental/Iris/Core/Private/Iris/ReplicationSystem/ObjectReplicationBridge.cpp
--- a/Engine/Source/Runtime/Experimental/Iris/Core/Private/Iris/ReplicationSystem/ObjectReplicationBridge.cpp
+++ b/Engine/Source/Runtime/Experimental/Iris/Core/Private/Iris/ReplicationSystem/ObjectReplicationBridge.cpp
@@ -329,6 +329,7 @@ UE::Net::FNetRefHandle UObjectReplicationBridge::BeginReplication(UObject* Insta
 				const FNetObjectFilterHandle FilterHandle = GetDynamicFilter(Instance->GetClass());
 				if (FilterHandle != InvalidNetObjectFilterHandle)
 				{
+					UE_LOG_OBJECTREPLICATIONBRIDGE(Verbose, TEXT("BeginReplication Filter: %s will be used for Object: %s "), *(ReplicationSystem->GetFilterName(FilterHandle).ToString()), *Instance->GetName());
 					ReplicationSystem->SetFilter(RefHandle, FilterHandle);
 				}
 			}
@@ -716,15 +717,12 @@ void UObjectReplicationBridge::PreUpdateAndPollImpl(FNetRefHandle Handle)
 		const FNetRefHandleManager::FReplicatedObjectData& ObjectData = LocalNetRefHandleManager.GetReplicatedObjectDataNoCheck(InternalObjectIndex);
 		if (ObjectData.InstanceProtocol)
 		{
-#if UE_IRIS_PROFILE_PROTOCOL_NAMES
-			IRIS_PROFILER_SCOPE_TEXT(ObjectData.Protocol->DebugName->Name);
-#endif
+			IRIS_PROFILER_PROTOCOL_NAME(ObjectData.Protocol->DebugName->Name);
+
 			// Call per-instance PreUpdate function
 			if (PreUpdateInstanceFunction && EnumHasAnyFlags(ObjectData.InstanceProtocol->InstanceTraits, EReplicationInstanceProtocolTraits::NeedsPreSendUpdate))
 			{
-#if UE_IRIS_PROFILE_PROTOCOL_NAMES
-				IRIS_PROFILER_SCOPE(PreReplicationUpdate);
-#endif
+				IRIS_PROFILER_SCOPE_VERBOSE(PreReplicationUpdate);
 				(*PreUpdateInstanceFunction)(ObjectData.RefHandle, ReplicatedInstances[InternalObjectIndex], this);
 				++Stats.PreUpdatedObjectCount;
 			}
@@ -732,9 +730,7 @@ void UObjectReplicationBridge::PreUpdateAndPollImpl(FNetRefHandle Handle)
 			// Poll properties if the instance protocol requires it
 			if (EnumHasAnyFlags(ObjectData.InstanceProtocol->InstanceTraits, EReplicationInstanceProtocolTraits::NeedsPoll))
 			{
-#if UE_IRIS_PROFILE_PROTOCOL_NAMES
-				IRIS_PROFILER_SCOPE(PreReplicationUpdate);
-#endif
+				IRIS_PROFILER_SCOPE_VERBOSE(Poll);
 
 				const bool bIsGCAffectedObject = GarbageCollectionAffectedObjects.GetBit(InternalObjectIndex);
 				GarbageCollectionAffectedObjects.ClearBit(InternalObjectIndex);
@@ -760,9 +756,8 @@ void UObjectReplicationBridge::PreUpdateAndPollImpl(FNetRefHandle Handle)
 		const FNetRefHandleManager::FReplicatedObjectData& ObjectData = LocalNetRefHandleManager.GetReplicatedObjectDataNoCheck(InternalObjectIndex);
 		if (const FReplicationInstanceProtocol* InstanceProtocol = ObjectData.InstanceProtocol)
 		{
-#if UE_IRIS_PROFILE_PROTOCOL_NAMES
-			IRIS_PROFILER_SCOPE_TEXT(ObjectData.Protocol->DebugName->Name);
-#endif
+			IRIS_PROFILER_PROTOCOL_NAME(ObjectData.Protocol->DebugName->Name);
+
 			const EReplicationInstanceProtocolTraits InstanceTraits = InstanceProtocol->InstanceTraits;
 			const bool bNeedsPoll = EnumHasAnyFlags(InstanceTraits, EReplicationInstanceProtocolTraits::NeedsPoll);
 			bool bIsDirtyObject = DirtyObjects.GetBit(InternalObjectIndex);
@@ -770,9 +765,8 @@ void UObjectReplicationBridge::PreUpdateAndPollImpl(FNetRefHandle Handle)
 			// Call per-instance PreUpdate function
 			if (PreUpdateInstanceFunction && EnumHasAnyFlags(InstanceTraits, EReplicationInstanceProtocolTraits::NeedsPreSendUpdate))
 			{
-#if UE_IRIS_PROFILE_PROTOCOL_NAMES
-				IRIS_PROFILER_SCOPE(PreReplicationUpdate);
-#endif
+				IRIS_PROFILER_SCOPE_VERBOSE(PreReplicationUpdate);
+
 				(*PreUpdateInstanceFunction)(ObjectData.RefHandle, ReplicatedInstances[InternalObjectIndex], this);
 				++Stats.PreUpdatedObjectCount;
 
@@ -797,9 +791,8 @@ void UObjectReplicationBridge::PreUpdateAndPollImpl(FNetRefHandle Handle)
 				return;
 			}
 
-#if UE_IRIS_PROFILE_PROTOCOL_NAMES
-			IRIS_PROFILER_SCOPE(PollPushBased);
-#endif
+			IRIS_PROFILER_SCOPE_VERBOSE(PollPushBased);
+
 			// If the object is fully push model we only need to poll it if it's dirty, unless it's a new object or was garbage collected.
 			bool bWasMarkedDirty = false;
 			if (EnumHasAnyFlags(InstanceTraits, EReplicationInstanceProtocolTraits::HasFullPushBasedDirtiness))
@@ -851,117 +844,121 @@ void UObjectReplicationBridge::PreUpdateAndPollImpl(FNetRefHandle Handle)
 		}
 	};
 
-	// Update all objects
-	if (!Handle.IsValid())
+	// If we are asked to poll a single object
+	if (Handle.IsValid())
 	{
-		const FNetBitArray& ScopableObjects = LocalNetRefHandleManager.GetScopableInternalIndices();
-		const FNetBitArray& WantToBeDormantObjects = LocalNetRefHandleManager.GetWantToBeDormantInternalIndices();
-
-		// Filter the set of objects considered for pre-update and polling
-		// We always want to consider objects marked as dirty and their subobjects
-		FNetBitArray ObjectsConsideredForPolling(LocalNetRefHandleManager.GetMaxActiveObjectCount());
-		if (bUseFrequencyBasedPolling)
-		{
-			FNetBitArrayView FrequencyBasedObjectsToPollView = MakeNetBitArrayView(ObjectsConsideredForPolling);
-			PollFrequencyLimiter->Update(MakeNetBitArrayView(ScopableObjects), DirtyObjects, FrequencyBasedObjectsToPollView);
-		}
-		else
+		if (uint32 InternalObjectIndex = LocalNetRefHandleManager.GetInternalIndex(Handle))
 		{
-			ObjectsConsideredForPolling.Copy(ScopableObjects);
+			UpdateAndPollFunction(InternalObjectIndex);
 		}
 
-		// Mask off objects pending dormancy as we do not want to poll/pre-update them unless they are marked for flush or are dirty
-		if (bUseDormancyToFilterPolling)
-		{
-			IRIS_PROFILER_SCOPE(PreUpdateAndPollImpl_Dormancy);
+		return;
+	}
+	
+	// If we poll all scopebable objects
 
-			// Mask off objects pending dormancy, masked off dirty objects are restored later
-			ObjectsConsideredForPolling.Combine(WantToBeDormantObjects, FNetBitArrayView::AndNotOp);
+	const FNetBitArray& ScopableObjects = LocalNetRefHandleManager.GetScopableInternalIndices();
+	const FNetBitArray& WantToBeDormantObjects = LocalNetRefHandleManager.GetWantToBeDormantInternalIndices();
 
-			// Add objects that have requested to flush dormancy
-			for (FNetRefHandle HandlePendingFlush : MakeArrayView(DormantHandlesPendingFlush))
+	// Filter the set of objects considered for pre-update and polling
+	// We always want to consider objects marked as dirty and their subobjects
+	FNetBitArray ObjectsConsideredForPolling(LocalNetRefHandleManager.GetMaxActiveObjectCount());
+	if (bUseFrequencyBasedPolling)
+	{
+		FNetBitArrayView FrequencyBasedObjectsToPollView = MakeNetBitArrayView(ObjectsConsideredForPolling);
+		PollFrequencyLimiter->Update(MakeNetBitArrayView(ScopableObjects), DirtyObjects, FrequencyBasedObjectsToPollView);
+	}
+	else
+	{
+		ObjectsConsideredForPolling.Copy(ScopableObjects);
+	}
+
+	// Mask off objects pending dormancy as we do not want to poll/pre-update them unless they are marked for flush or are dirty
+	if (bUseDormancyToFilterPolling)
+	{
+		IRIS_PROFILER_SCOPE(PreUpdateAndPollImpl_Dormancy);
+
+		// Mask off objects pending dormancy, masked off dirty objects are restored later
+		ObjectsConsideredForPolling.Combine(WantToBeDormantObjects, FNetBitArrayView::AndNotOp);
+
+		// Add objects that have requested to flush dormancy
+		for (FNetRefHandle HandlePendingFlush : MakeArrayView(DormantHandlesPendingFlush))
+		{
+			if (const uint32 InternalObjectIndex = LocalNetRefHandleManager.GetInternalIndex(HandlePendingFlush))
 			{
-				if (const uint32 InternalObjectIndex = LocalNetRefHandleManager.GetInternalIndex(HandlePendingFlush))
+				if (ObjectsConsideredForPolling.GetBit(InternalObjectIndex))
 				{
-					if (ObjectsConsideredForPolling.GetBit(InternalObjectIndex))
-					{
-						continue;
-					}
+					continue;
+				}
 
-					ObjectsConsideredForPolling.SetBit(InternalObjectIndex);
-					for (const FInternalNetRefIndex SubObjectIndex : LocalNetRefHandleManager.GetSubObjects(InternalObjectIndex))
-					{
-						ObjectsConsideredForPolling.SetBit(SubObjectIndex);
-					}
+				ObjectsConsideredForPolling.SetBit(InternalObjectIndex);
+				for (const FInternalNetRefIndex SubObjectIndex : LocalNetRefHandleManager.GetSubObjects(InternalObjectIndex))
+				{
+					ObjectsConsideredForPolling.SetBit(SubObjectIndex);
 				}
 			}
-
-			//TODO: Add size of DormantHandlesPendingFlush array to NetStats ?
 		}
-		DormantHandlesPendingFlush.Reset();
 
-		/**
-		 * Make sure to propagate polling for owners to subobjects and vice versa. If an actor is not due to update due to
-		 * polling frequency it can still be dirty or a dormant object marked for flush and polled for that reason. In order to make sure all recent state updates
-		 * are replicated atomically this polling propagation is required.
-		 */
-		{
-			IRIS_PROFILER_SCOPE(PreUpdateAndPollImpl_PropagateDirtyness);
+		//TODO: Add size of DormantHandlesPendingFlush array to NetStats ?
+	}
+	DormantHandlesPendingFlush.Reset();
 
-			auto PropagateSubObjectDirtinessToOwner = [&LocalNetRefHandleManager, &ObjectsConsideredForPolling](uint32 InternalObjectIndex)
-			{
-				const FNetRefHandleManager::FReplicatedObjectData& ObjectData = LocalNetRefHandleManager.GetReplicatedObjectDataNoCheck(InternalObjectIndex);
-				ObjectsConsideredForPolling.SetBit(ObjectData.SubObjectRootIndex);
-			};
+	/**
+	* Make sure to propagate polling for owners to subobjects and vice versa. If an actor is not due to update due to
+	* polling frequency it can still be dirty or a dormant object marked for flush and polled for that reason. In order to make sure all recent state updates
+	* are replicated atomically this polling propagation is required.
+	*/
+	{
+		IRIS_PROFILER_SCOPE(PreUpdateAndPollImpl_PropagateDirtyness);
 
-			auto PropagateOwnerDirtinessToSubObjectsAndDependentObjects = [&LocalNetRefHandleManager, &ObjectsConsideredForPolling](uint32 InternalObjectIndex)
-			{
-				const FNetRefHandleManager::FReplicatedObjectData& ObjectData = LocalNetRefHandleManager.GetReplicatedObjectDataNoCheck(InternalObjectIndex);
+		auto PropagateSubObjectDirtinessToOwner = [&LocalNetRefHandleManager, &ObjectsConsideredForPolling](uint32 InternalObjectIndex)
+		{
+			const FNetRefHandleManager::FReplicatedObjectData& ObjectData = LocalNetRefHandleManager.GetReplicatedObjectDataNoCheck(InternalObjectIndex);
+			ObjectsConsideredForPolling.SetBit(ObjectData.SubObjectRootIndex);
+		};
 
-				// Mark owner as well as we might have masked it out in the dormancy pass
-				ObjectsConsideredForPolling.SetBit(InternalObjectIndex);
+		auto PropagateOwnerDirtinessToSubObjectsAndDependentObjects = [&LocalNetRefHandleManager, &ObjectsConsideredForPolling](uint32 InternalObjectIndex)
+		{
+			const FNetRefHandleManager::FReplicatedObjectData& ObjectData = LocalNetRefHandleManager.GetReplicatedObjectDataNoCheck(InternalObjectIndex);
 
-				for (const FInternalNetRefIndex SubObjectInternalIndex : LocalNetRefHandleManager.GetSubObjects(InternalObjectIndex))
-				{
-					ObjectsConsideredForPolling.SetBit(SubObjectInternalIndex);
-				}
-			};
+			// Mark owner as well as we might have masked it out in the dormancy pass
+			ObjectsConsideredForPolling.SetBit(InternalObjectIndex);
 
-			// Update subobjects' owner first and owners' subobjects second. It's the only way to properly mark all groups of objects in two passes.
-			const FNetBitArrayView SubObjects = MakeNetBitArrayView(LocalNetRefHandleManager.GetSubObjectInternalIndices());
-			FNetBitArrayView::ForAllSetBits(DirtyObjects, SubObjects, FNetBitArray::AndOp, PropagateSubObjectDirtinessToOwner);
-			FNetBitArrayView::ForAllSetBits(DirtyObjects, SubObjects, FNetBitArray::AndNotOp, PropagateOwnerDirtinessToSubObjectsAndDependentObjects);
-			
-			// Patch dependent objects, not subtle as this point
+			for (const FInternalNetRefIndex SubObjectInternalIndex : LocalNetRefHandleManager.GetSubObjects(InternalObjectIndex))
 			{
-				IRIS_PROFILER_SCOPE(PreUpdateAndPollImpl_PatchIndepedentObjects);
-
-				FNetBitArray TempObjectsConsideredForPolling(ObjectsConsideredForPolling);
-				FNetBitArray::ForAllSetBits(TempObjectsConsideredForPolling, NetRefHandleManager->GetObjectsWithDependentObjectsInternalIndices(), FNetBitArray::AndOp,
-					[&LocalNetRefHandleManager, &ObjectsConsideredForPolling](FInternalNetRefIndex ObjectIndex) 
-					{
-						LocalNetRefHandleManager.ForAllDependentObjectsRecursive(ObjectIndex, [&ObjectsConsideredForPolling, ObjectIndex](FInternalNetRefIndex DependentObjectIndex) { ObjectsConsideredForPolling.SetBit(DependentObjectIndex); });
-					}
-				);
+				ObjectsConsideredForPolling.SetBit(SubObjectInternalIndex);
 			}
-		}
+		};
 
+		// Update subobjects' owner first and owners' subobjects second. It's the only way to properly mark all groups of objects in two passes.
+		const FNetBitArrayView SubObjects = MakeNetBitArrayView(LocalNetRefHandleManager.GetSubObjectInternalIndices());
+		FNetBitArrayView::ForAllSetBits(DirtyObjects, SubObjects, FNetBitArray::AndOp, PropagateSubObjectDirtinessToOwner);
+		FNetBitArrayView::ForAllSetBits(DirtyObjects, SubObjects, FNetBitArray::AndNotOp, PropagateOwnerDirtinessToSubObjectsAndDependentObjects);
+			
+		// Patch dependent objects, not subtle as this point
 		{
-			IRIS_PROFILER_SCOPE(PreUpdateAndPollImpl_Poll);
-			if (IsIrisPushModelEnabled())
-			{
-				ObjectsConsideredForPolling.ForAllSetBits(PushModelUpdateAndPollFunction);
-			}
-			else
-			{
-				ObjectsConsideredForPolling.ForAllSetBits(UpdateAndPollFunction);
-			}
+			IRIS_PROFILER_SCOPE(PreUpdateAndPollImpl_PatchDependentObjects);
+
+			FNetBitArray TempObjectsConsideredForPolling(ObjectsConsideredForPolling);
+			FNetBitArray::ForAllSetBits(TempObjectsConsideredForPolling, NetRefHandleManager->GetObjectsWithDependentObjectsInternalIndices(), FNetBitArray::AndOp,
+				[&LocalNetRefHandleManager, &ObjectsConsideredForPolling](FInternalNetRefIndex ObjectIndex) 
+				{
+					LocalNetRefHandleManager.ForAllDependentObjectsRecursive(ObjectIndex, [&ObjectsConsideredForPolling, ObjectIndex](FInternalNetRefIndex DependentObjectIndex) { ObjectsConsideredForPolling.SetBit(DependentObjectIndex); });
+				}
+			);
 		}
-		
 	}
-	else if (uint32 InternalObjectIndex = LocalNetRefHandleManager.GetInternalIndex(Handle))
+
 	{
-		UpdateAndPollFunction(InternalObjectIndex);
+		IRIS_PROFILER_SCOPE(PreUpdateAndPollImpl_Poll);
+		if (IsIrisPushModelEnabled())
+		{
+			ObjectsConsideredForPolling.ForAllSetBits(PushModelUpdateAndPollFunction);
+		}
+		else
+		{
+			ObjectsConsideredForPolling.ForAllSetBits(UpdateAndPollFunction);
+		}
 	}
 
 	// Report stats
diff --git a/Engine/Source/Runtime/Experimental/Iris/Core/Private/Iris/ReplicationSystem/ReplicationOperations.cpp b/Engine/Source/Runtime/Experimental/Iris/Core/Private/Iris/ReplicationSystem/ReplicationOperations.cpp
--- a/Engine/Source/Runtime/Experimental/Iris/Core/Private/Iris/ReplicationSystem/ReplicationOperations.cpp
+++ b/Engine/Source/Runtime/Experimental/Iris/Core/Private/Iris/ReplicationSystem/ReplicationOperations.cpp
@@ -510,6 +510,8 @@ bool FReplicationInstanceOperations::PollAndRefreshCachedObjectReferences(const
 
 void FReplicationInstanceOperations::CopyAndQuantize(FNetSerializationContext& Context, uint8* DstObjectStateBuffer, FNetBitStreamWriter* OutChangeMaskWriter, const FReplicationInstanceProtocol* InstanceProtocol, const FReplicationProtocol* Protocol)
 {
+	IRIS_PROFILER_SCOPE_VERBOSE(CopyAndQuantize);
+
 	check(InstanceProtocol && InstanceProtocol->FragmentCount == Protocol->ReplicationStateCount);
 
 	FReplicationInstanceProtocol::FFragmentData* FragmentData = InstanceProtocol->FragmentData;
@@ -548,6 +550,8 @@ void FReplicationInstanceOperations::CopyAndQuantize(FNetSerializationContext& C
 
 void FReplicationInstanceOperations::CopyAndQuantizeIfDirty(FNetSerializationContext& Context, uint8* DstObjectStateBuffer, FNetBitStreamWriter* OutChangeMaskWriter, const FReplicationInstanceProtocol* InstanceProtocol, const FReplicationProtocol* Protocol)
 {
+	IRIS_PROFILER_SCOPE_VERBOSE(CopyAndQuantizeIfDirty);
+
 	check(InstanceProtocol && InstanceProtocol->FragmentCount == Protocol->ReplicationStateCount);
 
 	FReplicationInstanceProtocol::FFragmentData* FragmentData = InstanceProtocol->FragmentData;
@@ -572,9 +576,7 @@ void FReplicationInstanceOperations::CopyAndQuantizeIfDirty(FNetSerializationCon
 		// Append changemask data to bitstreams and quantize state if mask is dirty
 		if (Private::AppendMemberChangeMasks(OutChangeMaskWriter, OutConditionalChangeMaskWriter, FragmentData[StateIt].ExternalSrcBuffer, CurrentDescriptor))
 		{
-#if UE_IRIS_PROFILE_PROTOCOL_NAMES
-			IRIS_PROFILER_SCOPE_TEXT(CurrentDescriptor->DebugName->Name);
-#endif
+			IRIS_PROFILER_PROTOCOL_NAME(CurrentDescriptor->DebugName->Name);
 			FReplicationStateOperations::Quantize(Context, CurrentInternalStateBuffer, FragmentData[StateIt].ExternalSrcBuffer, CurrentDescriptor);
 		}
 		
diff --git a/Engine/Source/Runtime/Experimental/Iris/Core/Private/Iris/ReplicationSystem/ReplicationOperationsInternal.cpp b/Engine/Source/Runtime/Experimental/Iris/Core/Private/Iris/ReplicationSystem/ReplicationOperationsInternal.cpp
--- a/Engine/Source/Runtime/Experimental/Iris/Core/Private/Iris/ReplicationSystem/ReplicationOperationsInternal.cpp
+++ b/Engine/Source/Runtime/Experimental/Iris/Core/Private/Iris/ReplicationSystem/ReplicationOperationsInternal.cpp
@@ -1,7 +1,13 @@
 // Copyright Epic Games, Inc. All Rights Reserved.
 
 #include "ReplicationOperationsInternal.h"
+
 #include "Iris/Core/IrisLog.h"
+#include "Iris/Core/IrisProfiler.h"
+
+#include "Net/Core/Trace/NetTrace.h"
+#include "Net/Core/Trace/NetDebugName.h"
+
 #include "Iris/ReplicationSystem/ChangeMaskCache.h"
 #include "Iris/ReplicationSystem/ChangeMaskUtil.h"
 #include "Iris/ReplicationSystem/NetRefHandleManager.h"
@@ -71,6 +77,8 @@ uint32 FReplicationInstanceOperationsInternal::CopyObjectStateData(FNetBitStream
 		// We cannot copy state data for zero sized objects or objects that no longer has an instance protocol.
 		if (Object.InstanceProtocol && Object.Protocol->InternalTotalSize > 0U)
 		{
+			IRIS_PROFILER_PROTOCOL_NAME(Object.Protocol->DebugName->Name);
+
 			// if the object was scopable prev frame we can do partial copy
 			bool bShouldPropagateChangedStates = Object.bShouldPropagateChangedStates;
 
@@ -85,10 +93,6 @@ uint32 FReplicationInstanceOperationsInternal::CopyObjectStateData(FNetBitStream
 				ChangeMaskStorageType* ChangeMaskData = Cache.GetChangeMaskStorage(Info);
 				ChangeMaskWriter.InitBytes(ChangeMaskData, ChangeMaskByteCount);
 
-#if UE_IRIS_PROFILE_PROTOCOL_NAMES
-				IRIS_PROFILER_SCOPE_TEXT(Object.Protocol->DebugName->Name);
-#endif
-
 				//UE_LOG(LogIris, Log, TEXT("Copying state data for ( InternalIndex: %u ) with NetRefHandle (Id=%u)"), InternalIndex, Object.RefHandle.GetId());
 				if (bUseFullCopyAndQuantize)
 				{
diff --git a/Engine/Source/Runtime/Experimental/Iris/Core/Private/Iris/ReplicationSystem/ReplicationSystem.cpp b/Engine/Source/Runtime/Experimental/Iris/Core/Private/Iris/ReplicationSystem/ReplicationSystem.cpp
--- a/Engine/Source/Runtime/Experimental/Iris/Core/Private/Iris/ReplicationSystem/ReplicationSystem.cpp
+++ b/Engine/Source/Runtime/Experimental/Iris/Core/Private/Iris/ReplicationSystem/ReplicationSystem.cpp
@@ -885,6 +885,11 @@ UNetObjectFilter* UReplicationSystem::GetFilter(const FName FilterName) const
 	return Impl->ReplicationSystemInternal.GetFiltering().GetFilter(FilterName);
 }
 
+FName UReplicationSystem::GetFilterName(UE::Net::FNetObjectFilterHandle Filter) const
+{
+	return Impl->ReplicationSystemInternal.GetFiltering().GetFilterName(Filter);
+}
+
 bool UReplicationSystem::SetConnectionFilter(FNetRefHandle Handle, const TBitArray<>& Connections, UE::Net::ENetFilterStatus ReplicationStatus)
 {
 	using namespace UE::Net;
diff --git a/Engine/Source/Runtime/Experimental/Iris/Core/Private/Iris/ReplicationSystem/ReplicationWriter.cpp b/Engine/Source/Runtime/Experimental/Iris/Core/Private/Iris/ReplicationSystem/ReplicationWriter.cpp
--- a/Engine/Source/Runtime/Experimental/Iris/Core/Private/Iris/ReplicationSystem/ReplicationWriter.cpp
+++ b/Engine/Source/Runtime/Experimental/Iris/Core/Private/Iris/ReplicationSystem/ReplicationWriter.cpp
@@ -1799,6 +1799,8 @@ FReplicationWriter::EWriteObjectStatus FReplicationWriter::WriteObjectInBatch(FN
 	const FNetRefHandleManager::FReplicatedObjectData& ObjectData = NetRefHandleManager->GetReplicatedObjectDataNoCheck(InternalIndex);
 	uint8* ReplicatedObjectStateBuffer = NetRefHandleManager->GetReplicatedObjectStateBufferNoCheck(InternalIndex);
 
+	//IRIS_PROFILER_PROTOCOL_NAME(ObjectData.Protocol?ObjectData.Protocol->DebugName->Name:TEXT("NoProtocol"));
+
 	// Filter out changemasks that are not supposed to be replicated to this connection
 	const bool bNeedToFilterChangeMask = (bIsInitialState || Info.HasDirtyChangeMask) && Info.HasChangemaskFilter;
 	if (bNeedToFilterChangeMask)
diff --git a/Engine/Source/Runtime/Experimental/Iris/Core/Public/Iris/Core/IrisProfiler.h b/Engine/Source/Runtime/Experimental/Iris/Core/Public/Iris/Core/IrisProfiler.h
--- a/Engine/Source/Runtime/Experimental/Iris/Core/Public/Iris/Core/IrisProfiler.h
+++ b/Engine/Source/Runtime/Experimental/Iris/Core/Public/Iris/Core/IrisProfiler.h
@@ -10,11 +10,21 @@
 #	endif
 #endif
 
-#ifndef UE_IRIS_PROFILE_PROTOCOL_NAMES
-	// Set this to true to see protocol names in profiling captures. The downside is the events will add an overhead on the cpu when cpu captures are occuring
-	#define UE_IRIS_PROFILE_PROTOCOL_NAMES !UE_BUILD_SHIPPING
+// When true this adds dynamic protocol names in profile captures. The downside is a noticeable cpu cost overhead but only while cpu trace recording is occurring.
+#ifndef UE_IRIS_PROFILER_ENABLE_PROTOCOL_NAMES
+#	define UE_IRIS_PROFILER_ENABLE_PROTOCOL_NAMES !UE_BUILD_SHIPPING
 #endif
 
+// When true this adds low-level cpu trace captures of operations in Iris. Adds a little cpu overhead but only while cpu trace recording is occurring.
+#ifndef UE_IRIS_PROFILER_ENABLE_VERBOSE
+#	if (UE_BUILD_SHIPPING || UE_BUILD_TEST)
+#		define UE_IRIS_PROFILER_ENABLE_VERBOSE 0
+#	else
+#		define UE_IRIS_PROFILER_ENABLE_VERBOSE 1
+#	endif
+#endif
+
+
 //#define IRIS_USE_SUPERLUMINAL
 
 #if IRIS_PROFILER_ENABLE
@@ -34,3 +44,15 @@
 #	define IRIS_PROFILER_SCOPE(x)
 #	define IRIS_PROFILER_SCOPE_TEXT(X)
 #endif
+
+#if UE_IRIS_PROFILER_ENABLE_PROTOCOL_NAMES
+#	define IRIS_PROFILER_PROTOCOL_NAME(x) IRIS_PROFILER_SCOPE_TEXT(x)
+#else
+#	define IRIS_PROFILER_PROTOCOL_NAME(x)
+#endif
+
+#if UE_IRIS_PROFILER_ENABLE_VERBOSE
+#	define IRIS_PROFILER_SCOPE_VERBOSE(x) IRIS_PROFILER_SCOPE(x);
+#else
+#	define IRIS_PROFILER_SCOPE_VERBOSE(x)
+#endif
diff --git a/Engine/Source/Runtime/Experimental/Iris/Core/Public/Iris/ReplicationSystem/FastArrayReplicationFragment.h b/Engine/Source/Runtime/Experimental/Iris/Core/Public/Iris/ReplicationSystem/FastArrayReplicationFragment.h
--- a/Engine/Source/Runtime/Experimental/Iris/Core/Public/Iris/ReplicationSystem/FastArrayReplicationFragment.h
+++ b/Engine/Source/Runtime/Experimental/Iris/Core/Public/Iris/ReplicationSystem/FastArrayReplicationFragment.h
@@ -148,9 +148,7 @@ bool TFastArrayReplicationFragment<FastArrayItemType, FastArrayType>::PollAllSta
 	FastArrayType* SrcArraySerializer = GetFastArraySerializerFromOwner();
 	ItemArrayType* SrcWrappedArray = reinterpret_cast<ItemArrayType*>(reinterpret_cast<uint8*>(SrcArraySerializer) + WrappedArrayOffsetRelativeFastArraySerializerProperty);
 	
-#if UE_IRIS_PROFILE_PROTOCOL_NAMES
-	IRIS_PROFILER_SCOPE_TEXT(ReplicationStateDescriptor->DebugName->Name);
-#endif
+	IRIS_PROFILER_PROTOCOL_NAME(ReplicationStateDescriptor->DebugName->Name);
 
 	// Lookup destination data
 	FastArrayType* DstArraySerializer = GetFastArraySerializerFromReplicationState();
diff --git a/Engine/Source/Runtime/Experimental/Iris/Core/Public/Iris/ReplicationSystem/ReplicationSystem.h b/Engine/Source/Runtime/Experimental/Iris/Core/Public/Iris/ReplicationSystem/ReplicationSystem.h
--- a/Engine/Source/Runtime/Experimental/Iris/Core/Public/Iris/ReplicationSystem/ReplicationSystem.h
+++ b/Engine/Source/Runtime/Experimental/Iris/Core/Public/Iris/ReplicationSystem/ReplicationSystem.h
@@ -355,6 +355,11 @@ public:
 	 */
 	IRISCORE_API UNetObjectFilter* GetFilter(const FName FilterName) const;
 
+	/**
+	 * Returns the name of the filter associated to this handle.
+	 */
+	IRISCORE_API FName GetFilterName(UE::Net::FNetObjectFilterHandle Filter) const;
+
 	// Group Filtering
 
 	/**
