collision conversion needs to check against the inner shape type instead to handle the case of a transformed union

- add GetNestedType to support this check

diff --git a/Engine/Source/Runtime/Engine/Private/Collision/CollisionConversions.cpp b/Engine/Source/Runtime/Engine/Private/Collision/CollisionConversions.cpp
--- a/Engine/Source/Runtime/Engine/Private/Collision/CollisionConversions.cpp
+++ b/Engine/Source/Runtime/Engine/Private/Collision/CollisionConversions.cpp
@@ -103,7 +103,7 @@ static FVector FindGeomOpposingNormal(ECollisionShapeType QueryGeomType, const T
 	{
 		if (const FPhysicsShape* Shape = GetShape(Hit))
 		{
-			const Chaos::EImplicitObjectType ShapeType = Shape->GetGeometry()->GetType();
+			const Chaos::EImplicitObjectType ShapeType = Shape->GetGeometry()->GetNestedType();
 			if (ShapeType == Chaos::ImplicitObjectType::Union && !Hit.FaceNormal.IsNearlyZero())
 			{
 				return Hit.FaceNormal;
diff --git a/Engine/Source/Runtime/Experimental/Chaos/Public/Chaos/ImplicitObject.h b/Engine/Source/Runtime/Experimental/Chaos/Public/Chaos/ImplicitObject.h
--- a/Engine/Source/Runtime/Experimental/Chaos/Public/Chaos/ImplicitObject.h
+++ b/Engine/Source/Runtime/Experimental/Chaos/Public/Chaos/ImplicitObject.h
@@ -121,6 +121,7 @@ public:
 	//Not all implicit objects can be duplicated, up to user code to use this in cases that make sense
 	virtual FImplicitObject* Duplicate() const { check(false); return nullptr; }
 
+	virtual EImplicitObjectType GetNestedType() const { return GetType(); }
 	EImplicitObjectType GetType() const;
 	static int32 GetOffsetOfType() { return offsetof(FImplicitObject, Type); }
 
diff --git a/Engine/Source/Runtime/Experimental/Chaos/Public/Chaos/ImplicitObjectScaled.h b/Engine/Source/Runtime/Experimental/Chaos/Public/Chaos/ImplicitObjectScaled.h
--- a/Engine/Source/Runtime/Experimental/Chaos/Public/Chaos/ImplicitObjectScaled.h
+++ b/Engine/Source/Runtime/Experimental/Chaos/Public/Chaos/ImplicitObjectScaled.h
@@ -198,6 +198,11 @@ public:
 		return MObject->GetTypeHash();
 	}
 
+	virtual EImplicitObjectType GetNestedType() const override
+	{
+		return MObject->GetNestedType();
+	}
+
 	virtual TUniquePtr<FImplicitObject> Copy() const override
 	{
 		return TUniquePtr<FImplicitObject>(CopyHelper(this));
@@ -581,6 +586,11 @@ public:
 		return MObject.Get();
 	}
 
+	virtual EImplicitObjectType GetNestedType() const override
+	{
+		return MObject->GetNestedType();
+	}
+
 	FReal GetRadius() const
 	{
 		return (MObject->GetRadius() > 0.0f) ? Margin : 0.0f;
diff --git a/Engine/Source/Runtime/Experimental/Chaos/Public/Chaos/ImplicitObjectTransformed.h b/Engine/Source/Runtime/Experimental/Chaos/Public/Chaos/ImplicitObjectTransformed.h
--- a/Engine/Source/Runtime/Experimental/Chaos/Public/Chaos/ImplicitObjectTransformed.h
+++ b/Engine/Source/Runtime/Experimental/Chaos/Public/Chaos/ImplicitObjectTransformed.h
@@ -100,6 +100,11 @@ public:
 		}
 	}
 
+	virtual EImplicitObjectType GetNestedType() const override
+	{
+		return MObject->GetNestedType();
+	}
+
 	virtual TUniquePtr<FImplicitObject> Copy() const
 	{
 		if (MObject)
