move the logic of setting child to parent to the cluster union manager to let the GC proxy leverage this when it detects it scene component transform has been changed when it's a part of a cluster union

diff --git a/Engine/Source/Runtime/Experimental/Chaos/Private/Chaos/ClusterUnionManager.cpp b/Engine/Source/Runtime/Experimental/Chaos/Private/Chaos/ClusterUnionManager.cpp
--- a/Engine/Source/Runtime/Experimental/Chaos/Private/Chaos/ClusterUnionManager.cpp
+++ b/Engine/Source/Runtime/Experimental/Chaos/Private/Chaos/ClusterUnionManager.cpp
@@ -520,4 +520,35 @@ namespace Chaos
 		}
 		return INDEX_NONE;
 	}
+
+	DECLARE_CYCLE_STAT(TEXT("FClusterUnionManager::UpdateClusterUnionParticlesChildToparent"), STAT_UpdateClusterUnionParticlesChildToparent, STATGROUP_Chaos);
+	void FClusterUnionManager::UpdateClusterUnionParticlesChildToparent(FClusterUnionIndex Index, const TArray<FPBDRigidParticleHandle*>& Particles, const TArray<FTransform>& ChildToParent)
+	{
+		SCOPE_CYCLE_COUNTER(STAT_UpdateClusterUnionParticlesChildToparent);
+
+		if (FClusterUnion* ClusterUnion = FindClusterUnion(Index))
+		{
+			for (int32 ParticleIndex = 0; ParticleIndex < Particles.Num() && ParticleIndex < ChildToParent.Num(); ++ParticleIndex)
+			{
+				FPBDRigidParticleHandle* Particle = Particles[ParticleIndex];
+				if (!ensure(Particle))
+				{
+					return;
+				}
+
+				const int32 ChildIndex = ClusterUnion->ChildParticles.Find(Particle->CastToRigidParticle());
+				if (ChildIndex != INDEX_NONE && ClusterUnion->InternalCluster)
+				{
+					if (FPBDRigidClusteredParticleHandle* ChildHandle = ClusterUnion->ChildParticles[ChildIndex]->CastToClustered())
+					{
+						ChildHandle->SetChildToParent(ChildToParent[ParticleIndex]);
+					}
+				}
+			}
+
+			UpdateAllClusterUnionProperties(*ClusterUnion, false);
+			MEvolution.GetParticles().MarkTransientDirtyParticle(ClusterUnion->InternalCluster);
+			MEvolution.DirtyParticle(*ClusterUnion->InternalCluster);
+		}
+	}
 }
\ No newline at end of file
diff --git a/Engine/Source/Runtime/Experimental/Chaos/Private/PhysicsProxy/ClusterUnionPhysicsProxy.cpp b/Engine/Source/Runtime/Experimental/Chaos/Private/PhysicsProxy/ClusterUnionPhysicsProxy.cpp
--- a/Engine/Source/Runtime/Experimental/Chaos/Private/PhysicsProxy/ClusterUnionPhysicsProxy.cpp
+++ b/Engine/Source/Runtime/Experimental/Chaos/Private/PhysicsProxy/ClusterUnionPhysicsProxy.cpp
@@ -539,33 +539,13 @@ namespace Chaos
 		Solver->EnqueueCommandImmediate(
 			[this, Objects, Transforms]() mutable
 			{
-				FPBDRigidsEvolutionGBF& Evolution = *GetEvolution(this);
-				FClusterUnionManager& ClusterUnionManager = Evolution.GetRigidClustering().GetClusterUnionManager();
-
-				if (FClusterUnion* ClusterUnion = ClusterUnionManager.FindClusterUnion(ClusterUnionIndex))
+				FReadPhysicsObjectInterface_Internal Interface = FPhysicsObjectInternalInterface::GetRead();
+				TArray<FPBDRigidParticleHandle*> Particles = Interface.GetAllRigidParticles(Objects);
+				if (ensure(Particles.Num() == Objects.Num()))
 				{
-					for (int32 Index = 0; Index < Objects.Num() && Index < Transforms.Num(); ++Index)
-					{
-						FPhysicsObjectHandle Child = Objects[Index];
-						FGeometryParticleHandle* Particle = Child->GetParticle<EThreadContext::Internal>();
-						if (!ensure(Particle))
-						{
-							return;
-						}
-
-						const int32 ParticleIndex = ClusterUnion->ChildParticles.Find(Particle->CastToRigidParticle());
-						if (ParticleIndex != INDEX_NONE && ClusterUnion->InternalCluster)
-						{
-							if (FPBDRigidClusteredParticleHandle* ChildHandle = ClusterUnion->ChildParticles[ParticleIndex]->CastToClustered())
-							{
-								ChildHandle->SetChildToParent(Transforms[Index]);
-							}
-						}
-					}
-
-					ClusterUnionManager.UpdateAllClusterUnionProperties(*ClusterUnion, false);
-					Evolution.GetParticles().MarkTransientDirtyParticle(ClusterUnion->InternalCluster);
-					Evolution.DirtyParticle(*ClusterUnion->InternalCluster);
+					FPBDRigidsEvolutionGBF& Evolution = *static_cast<FPBDRigidsSolver*>(Solver)->GetEvolution();
+					FClusterUnionManager& ClusterUnionManager = Evolution.GetRigidClustering().GetClusterUnionManager();
+					ClusterUnionManager.UpdateClusterUnionParticlesChildToparent(ClusterUnionIndex, Particles, Transforms);
 				}
 			}
 		);
diff --git a/Engine/Source/Runtime/Experimental/Chaos/Private/PhysicsProxy/GeometryCollectionPhysicsProxy.cpp b/Engine/Source/Runtime/Experimental/Chaos/Private/PhysicsProxy/GeometryCollectionPhysicsProxy.cpp
--- a/Engine/Source/Runtime/Experimental/Chaos/Private/PhysicsProxy/GeometryCollectionPhysicsProxy.cpp
+++ b/Engine/Source/Runtime/Experimental/Chaos/Private/PhysicsProxy/GeometryCollectionPhysicsProxy.cpp
@@ -2645,37 +2645,71 @@ void FGeometryCollectionPhysicsProxy::PushToPhysicsState()
 		const FTransform& ActorToWorld = Parameters.WorldTransform;
 
 		const int32 NumTransforms = PhysicsThreadCollection.NumElements(FGeometryCollection::TransformGroup);
-		for (int32 TransformGroupIndex = 0; TransformGroupIndex < NumTransforms; ++TransformGroupIndex)
+		if (Chaos::FPhysicsSolver* RigidSolver = GetSolver<Chaos::FPhysicsSolver>())
 		{
-			if (Chaos::FPBDRigidClusteredParticleHandle* Handle = SolverParticleHandles[TransformGroupIndex])
+			FClusterUnionManager& ClusterUnionManager = RigidSolver->GetEvolution()->GetRigidClustering().GetClusterUnionManager();
+
+			// Assume all particles are in the same cluster union.
+			FClusterUnionIndex ClusterUnionIndex = INDEX_NONE;
+			TArray<FPBDRigidParticleHandle*> DeferredClusterUnionParticleUpdates;
+			TArray<FTransform> DeferredClusterUnionChildToParentUpdates;
+
+			for (int32 TransformGroupIndex = 0; TransformGroupIndex < NumTransforms; ++TransformGroupIndex)
 			{
-				FClusterHandle* KinematicRootHandle = nullptr;
-				if (!Handle->Disabled() && Handle->ObjectState() == Chaos::EObjectStateType::Kinematic)
-				{
-					KinematicRootHandle = Handle;
-				}
-				else 
+				if (Chaos::FPBDRigidClusteredParticleHandle* Handle = SolverParticleHandles[TransformGroupIndex])
 				{
-					// is there a internal parent as a kinematic root?
+					FClusterHandle* KinematicRootHandle = nullptr;
 					FClusterHandle* ParentHandle = Handle->Parent();
-					if (ParentHandle && ParentHandle->InternalCluster() && !ParentHandle->Disabled() && ParentHandle->ObjectState() == Chaos::EObjectStateType::Kinematic)
+
+					if (!Handle->Disabled() && Handle->ObjectState() == Chaos::EObjectStateType::Kinematic)
 					{
-						if (!ProcessedInternalClusters.Contains(ParentHandle))
+						KinematicRootHandle = Handle;
+					}
+					else
+					{
+						// is there a internal parent as a kinematic root?
+						if (ParentHandle && ParentHandle->InternalCluster() && !ParentHandle->Disabled() && ParentHandle->ObjectState() == Chaos::EObjectStateType::Kinematic && ParentHandle->PhysicsProxy() == this)
 						{
-							ProcessedInternalClusters.Add(ParentHandle);
-							KinematicRootHandle = ParentHandle;
+							if (!ProcessedInternalClusters.Contains(ParentHandle))
+							{
+								ProcessedInternalClusters.Add(ParentHandle);
+								KinematicRootHandle = ParentHandle;
+							}
 						}
 					}
-				}
-				if (KinematicRootHandle)
-				{
-					const FTransform RootWorldTransform(KinematicRootHandle->R(), KinematicRootHandle->X());
-					const FTransform RootRelativeTransform = RootWorldTransform.GetRelativeTransform(Parameters.PrevWorldTransform);
-					const FTransform WorldTransform = RootRelativeTransform * ActorToWorld;
 
-					SetClusteredParticleKinematicTarget_Internal(KinematicRootHandle, WorldTransform);
+					if (KinematicRootHandle)
+					{
+						const FTransform RootWorldTransform(KinematicRootHandle->R(), KinematicRootHandle->X());
+						const FTransform RootRelativeTransform = RootWorldTransform.GetRelativeTransform(Parameters.PrevWorldTransform);
+						const FTransform WorldTransform = RootRelativeTransform * ActorToWorld;
+
+						SetClusteredParticleKinematicTarget_Internal(KinematicRootHandle, WorldTransform);
+					}
+					else if (ParentHandle)
+					{
+						if (ClusterUnionIndex == INDEX_NONE)
+						{
+							ClusterUnionIndex = ClusterUnionManager.FindClusterUnionIndexFromParticle(Handle);
+						}
+
+						if (ClusterUnionIndex != INDEX_NONE)
+						{
+							const FTransform ParentWorldTransform{ ParentHandle->R(), ParentHandle->X() };
+							const FTransform NewWorldTransform = PhysicsThreadCollection.MassToLocal[TransformGroupIndex] * PhysicsThreadCollection.Transform[TransformGroupIndex] * Parameters.WorldTransform;
+							const FTransform RelativeTransform = NewWorldTransform.GetRelativeTransform(ParentWorldTransform);
+
+							DeferredClusterUnionParticleUpdates.Add(Handle);
+							DeferredClusterUnionChildToParentUpdates.Add(RelativeTransform);
+						}
+					}
 				}
 			}
+
+			if (ClusterUnionIndex != INDEX_NONE && !DeferredClusterUnionParticleUpdates.IsEmpty() && !DeferredClusterUnionChildToParentUpdates.IsEmpty())
+			{
+				ClusterUnionManager.UpdateClusterUnionParticlesChildToparent(ClusterUnionIndex, DeferredClusterUnionParticleUpdates, DeferredClusterUnionChildToParentUpdates);
+			}
 		}
 	}
 }
diff --git a/Engine/Source/Runtime/Experimental/Chaos/Public/Chaos/ClusterUnionManager.h b/Engine/Source/Runtime/Experimental/Chaos/Public/Chaos/ClusterUnionManager.h
--- a/Engine/Source/Runtime/Experimental/Chaos/Public/Chaos/ClusterUnionManager.h
+++ b/Engine/Source/Runtime/Experimental/Chaos/Public/Chaos/ClusterUnionManager.h
@@ -96,6 +96,9 @@ namespace Chaos
 
 		FClusterUnionIndex FindClusterUnionIndexFromParticle(FPBDRigidParticleHandle* Particle);
 
+		// Changes the ChildToParent of a number of particles in a cluster union.
+		void UpdateClusterUnionParticlesChildToparent(FClusterUnionIndex Index, const TArray<FPBDRigidParticleHandle*>& Particles, const TArray<FTransform>& ChildToParent);
+
 		// Update the cluster union's properties after its set of particle changes.
 		void UpdateAllClusterUnionProperties(FClusterUnion& ClusterUnion, bool bRecomputeMassOrientation);
 
