fix capitalization of FClusterUnionManager::UpdateClusterUnionParticlesChildToParent

diff --git a/Engine/Source/Runtime/Experimental/Chaos/Private/Chaos/ClusterUnionManager.cpp b/Engine/Source/Runtime/Experimental/Chaos/Private/Chaos/ClusterUnionManager.cpp
--- a/Engine/Source/Runtime/Experimental/Chaos/Private/Chaos/ClusterUnionManager.cpp
+++ b/Engine/Source/Runtime/Experimental/Chaos/Private/Chaos/ClusterUnionManager.cpp
@@ -521,10 +521,10 @@ namespace Chaos
 		return INDEX_NONE;
 	}
 
-	DECLARE_CYCLE_STAT(TEXT("FClusterUnionManager::UpdateClusterUnionParticlesChildToparent"), STAT_UpdateClusterUnionParticlesChildToparent, STATGROUP_Chaos);
-	void FClusterUnionManager::UpdateClusterUnionParticlesChildToparent(FClusterUnionIndex Index, const TArray<FPBDRigidParticleHandle*>& Particles, const TArray<FTransform>& ChildToParent)
+	DECLARE_CYCLE_STAT(TEXT("FClusterUnionManager::UpdateClusterUnionParticlesChildToParent"), STAT_UpdateClusterUnionParticlesChildToParent, STATGROUP_Chaos);
+	void FClusterUnionManager::UpdateClusterUnionParticlesChildToParent(FClusterUnionIndex Index, const TArray<FPBDRigidParticleHandle*>& Particles, const TArray<FTransform>& ChildToParent)
 	{
-		SCOPE_CYCLE_COUNTER(STAT_UpdateClusterUnionParticlesChildToparent);
+		SCOPE_CYCLE_COUNTER(STAT_UpdateClusterUnionParticlesChildToParent);
 
 		if (FClusterUnion* ClusterUnion = FindClusterUnion(Index))
 		{
diff --git a/Engine/Source/Runtime/Experimental/Chaos/Private/PhysicsProxy/ClusterUnionPhysicsProxy.cpp b/Engine/Source/Runtime/Experimental/Chaos/Private/PhysicsProxy/ClusterUnionPhysicsProxy.cpp
--- a/Engine/Source/Runtime/Experimental/Chaos/Private/PhysicsProxy/ClusterUnionPhysicsProxy.cpp
+++ b/Engine/Source/Runtime/Experimental/Chaos/Private/PhysicsProxy/ClusterUnionPhysicsProxy.cpp
@@ -545,7 +545,7 @@ namespace Chaos
 				{
 					FPBDRigidsEvolutionGBF& Evolution = *static_cast<FPBDRigidsSolver*>(Solver)->GetEvolution();
 					FClusterUnionManager& ClusterUnionManager = Evolution.GetRigidClustering().GetClusterUnionManager();
-					ClusterUnionManager.UpdateClusterUnionParticlesChildToparent(ClusterUnionIndex, Particles, Transforms);
+					ClusterUnionManager.UpdateClusterUnionParticlesChildToParent(ClusterUnionIndex, Particles, Transforms);
 				}
 			}
 		);
diff --git a/Engine/Source/Runtime/Experimental/Chaos/Private/PhysicsProxy/GeometryCollectionPhysicsProxy.cpp b/Engine/Source/Runtime/Experimental/Chaos/Private/PhysicsProxy/GeometryCollectionPhysicsProxy.cpp
--- a/Engine/Source/Runtime/Experimental/Chaos/Private/PhysicsProxy/GeometryCollectionPhysicsProxy.cpp
+++ b/Engine/Source/Runtime/Experimental/Chaos/Private/PhysicsProxy/GeometryCollectionPhysicsProxy.cpp
@@ -2708,7 +2708,7 @@ void FGeometryCollectionPhysicsProxy::PushToPhysicsState()
 
 			if (ClusterUnionIndex != INDEX_NONE && !DeferredClusterUnionParticleUpdates.IsEmpty() && !DeferredClusterUnionChildToParentUpdates.IsEmpty())
 			{
-				ClusterUnionManager.UpdateClusterUnionParticlesChildToparent(ClusterUnionIndex, DeferredClusterUnionParticleUpdates, DeferredClusterUnionChildToParentUpdates);
+				ClusterUnionManager.UpdateClusterUnionParticlesChildToParent(ClusterUnionIndex, DeferredClusterUnionParticleUpdates, DeferredClusterUnionChildToParentUpdates);
 			}
 		}
 	}
diff --git a/Engine/Source/Runtime/Experimental/Chaos/Public/Chaos/ClusterUnionManager.h b/Engine/Source/Runtime/Experimental/Chaos/Public/Chaos/ClusterUnionManager.h
--- a/Engine/Source/Runtime/Experimental/Chaos/Public/Chaos/ClusterUnionManager.h
+++ b/Engine/Source/Runtime/Experimental/Chaos/Public/Chaos/ClusterUnionManager.h
@@ -97,7 +97,7 @@ namespace Chaos
 		FClusterUnionIndex FindClusterUnionIndexFromParticle(FPBDRigidParticleHandle* Particle);
 
 		// Changes the ChildToParent of a number of particles in a cluster union.
-		void UpdateClusterUnionParticlesChildToparent(FClusterUnionIndex Index, const TArray<FPBDRigidParticleHandle*>& Particles, const TArray<FTransform>& ChildToParent);
+		void UpdateClusterUnionParticlesChildToParent(FClusterUnionIndex Index, const TArray<FPBDRigidParticleHandle*>& Particles, const TArray<FTransform>& ChildToParent);
 
 		// Update the cluster union's properties after its set of particle changes.
 		void UpdateAllClusterUnionProperties(FClusterUnion& ClusterUnion, bool bRecomputeMassOrientation);
