Restore some of the momentum of objects which have caused clusters to break

diff --git a/Engine/Source/Runtime/Experimental/Chaos/Private/Chaos/PBDRigidClustering.cpp b/Engine/Source/Runtime/Experimental/Chaos/Private/Chaos/PBDRigidClustering.cpp
--- a/Engine/Source/Runtime/Experimental/Chaos/Private/Chaos/PBDRigidClustering.cpp
+++ b/Engine/Source/Runtime/Experimental/Chaos/Private/Chaos/PBDRigidClustering.cpp
@@ -53,6 +53,9 @@ namespace Chaos
 	float GraphPropagationBasedCollisionFactor = 1;
 	FAutoConsoleVariableRef CVarGraphPropagationBasedCollisionFactor(TEXT("p.GraphPropagationBasedCollisionFactor"), GraphPropagationBasedCollisionFactor, TEXT("when p.GraphPropagationBasedCollisionImpulseProcessing is on, the percentage [0-1] of remaining damage that is distributed to the connected pieces"));
 
+	bool bRestoreBreakingMomentum = 1;
+	FAutoConsoleVariableRef CVarRestoreBreakingMomentum(TEXT("p.RestoreBreakingMomentum"), bRestoreBreakingMomentum, TEXT("When a rigid cluster is broken, objects that its in contact with will receive an impulse to restore their momentum prior to the break."));
+
 	static FGeometryCollectionPhysicsProxy* GetConcreteProxy(FPBDRigidClusteredParticleHandle* ClusteredParticle)
 	{
 		if (ClusteredParticle)
@@ -429,6 +432,81 @@ namespace Chaos
 		CrumbledSinceLastUpdate.Reset();
 	}
 
+	void FRigidClustering::TrackBreakingCollision(FPBDRigidClusteredParticleHandle* ClusteredParticle)
+	{
+		if (auto Rigid = ClusteredParticle->CastToRigidParticle())
+		{
+			Rigid->ParticleCollisions().VisitCollisions([this, Rigid](FPBDCollisionConstraint& Collision)
+			{
+				// Get a generic handle for the "other" particle
+				uint8 OtherIdx = Collision.GetParticle0() == Rigid ? 1 : 0;
+
+				// Make sure this collision actually includes the clustered particle
+				if (!ensure(Collision.GetParticle(1 - OtherIdx) == Rigid))
+				{
+					return ECollisionVisitorResult::Continue;
+				}
+
+				FGeometryParticleHandle* OtherGeometry = Collision.GetParticle(OtherIdx);
+				if (OtherGeometry == nullptr)
+				{
+					return ECollisionVisitorResult::Continue;
+				}
+
+				FPBDRigidParticleHandle* OtherRigid = OtherGeometry->CastToRigidParticle();
+				if (OtherRigid == nullptr)
+				{
+					return ECollisionVisitorResult::Continue;
+				}
+
+				if (Collision.AccumulatedImpulse.SizeSquared() <= SMALL_NUMBER)
+				{
+					return ECollisionVisitorResult::Continue;
+				}
+
+				// Track this collision
+				BreakingCollisions.Add(TPair<FPBDCollisionConstraint*, FPBDRigidParticleHandle*>(&Collision, OtherRigid));
+
+				return ECollisionVisitorResult::Continue;
+			});
+		}
+	}
+
+	void FRigidClustering::RestoreBreakingMomentum()
+	{
+		for (TPair<FPBDCollisionConstraint*, FPBDRigidParticleHandle*>& Pair : BreakingCollisions)
+		{
+			FPBDCollisionConstraint& Collision = *Pair.Key;
+			FPBDRigidParticleHandle& Rigid = *Pair.Value;
+			FConstGenericParticleHandle Generic(&Rigid);
+
+			// Flip the impulse if we're restoring particle 0's momentum.
+			// This is because by convention constraint impulses point from 1 to 0.
+			uint8 OtherIdx = Collision.GetParticle0() == &Rigid ? 1 : 0;
+			const FVec3 Impulse
+				= OtherIdx == 0
+				? -Collision.AccumulatedImpulse
+				: Collision.AccumulatedImpulse;
+
+			// Compute the angular impulse based on distance from the contact point to the CoM
+			const FVec3 Location = Collision.CalculateWorldContactLocation();
+			const Chaos::FVec3 AngularImpulse = Chaos::FVec3::CrossProduct(Location - Generic->PCom(), Impulse);
+
+			// Compute impulse velocities
+			const FVec3 ImpulseVelocity = Generic->InvM() * Impulse;
+
+			const FMatrix33 OtherInvI = Utilities::ComputeWorldSpaceInertia(Generic->QCom(), Generic->ConditionedInvI());
+			const FVec3 AngularImpulseVelocity = OtherInvI * AngularImpulse;
+
+			// Temporary: Always restore only 50% of the impulse
+			const float RestorationPercent = .5f;
+
+			// Update linear and angular impulses for the body, to be integrated next solve
+			Rigid.V() += ImpulseVelocity * RestorationPercent;
+			Rigid.W() += AngularImpulseVelocity * RestorationPercent;
+		}
+	}
+
 	void FRigidClustering::SendBreakingEvent(FPBDRigidClusteredParticleHandle* ClusteredParticle, bool bFromCrumble)
 	{
 		// only emit break event if the proxy needs it 
@@ -776,6 +854,12 @@ namespace Chaos
 				
 				ActivatedChildren.Add(Child);
 				SendBreakingEvent(Child, bParentCrumbled);
+
+				// Restore some of the momentum of whatever collided with the parent
+				if (bRestoreBreakingMomentum)
+				{
+					TrackBreakingCollision(ClusteredParticle);
+				}
 			}
 			if (bUseDamagePropagation)
 			{
@@ -1017,6 +1101,9 @@ namespace Chaos
 	{
 		SCOPE_CYCLE_COUNTER(STAT_BreakingModel);
 
+		// Clear the set tracking breaking collisions
+		BreakingCollisions.Empty();
+
 		//make copy because release cluster modifies active indices. We want to iterate over original active indices
 		TArray<FPBDRigidClusteredParticleHandle*> ClusteredParticlesToProcess;
 		for (auto& Particle : MEvolution.GetNonDisabledClusteredView())
@@ -1031,6 +1118,12 @@ namespace Chaos
 				ReleaseClusterParticles(ClusteredParticle);
 			}
 		}
+
+		// Restore some of the momentum of objects that were touching rigid clusters that broke
+		if (bRestoreBreakingMomentum)
+		{
+			RestoreBreakingMomentum();
+		}
 	}
 
 	DECLARE_CYCLE_STAT(TEXT("FRigidClustering::Visitor"), STAT_ClusterVisitor, STATGROUP_Chaos);
diff --git a/Engine/Source/Runtime/Experimental/Chaos/Public/Chaos/PBDRigidClustering.h b/Engine/Source/Runtime/Experimental/Chaos/Public/Chaos/PBDRigidClustering.h
--- a/Engine/Source/Runtime/Experimental/Chaos/Public/Chaos/PBDRigidClustering.h
+++ b/Engine/Source/Runtime/Experimental/Chaos/Public/Chaos/PBDRigidClustering.h
@@ -404,6 +404,14 @@ public:
 
 	void RemoveChildFromParent(FPBDRigidParticleHandle* Child, const FPBDRigidClusteredParticleHandle* ClusteredParent);
 
+	// When a body has broken due to contact resolution, record an entry in a set
+	// for the collision and the particle who's momentum should be restored.
+	void TrackBreakingCollision(FPBDRigidClusteredParticleHandle* ClusteredParticle);
+
+	// Restore some percentage of momenta for objects which were involved in collisions
+	// with destroyed GCs
+	void RestoreBreakingMomentum();
+
 	void SendBreakingEvent(FPBDRigidClusteredParticleHandle* ClusteredParticle, bool bFromCrumble);
 	void SendCrumblingEvent(FPBDRigidClusteredParticleHandle* ClusteredParticle);
 
@@ -455,6 +463,11 @@ private:
 
 	TSet<FPBDRigidClusteredParticleHandle*> CrumbledSinceLastUpdate;
 
+	// Pairs of collision constraints and rigid particle handles of particles which collided with
+	// rigid clusters which broken. Some portion of the momentum change due to the constraint
+	// will be restored to each of the corresponding particles.
+	TSet<TPair<FPBDCollisionConstraint*, FPBDRigidParticleHandle*>> BreakingCollisions;
+
 	FReal MClusterConnectionFactor;
 	FClusterCreationParameters::EConnectionMethod MClusterUnionConnectionType;
 
