Geometry collection : fix issue where exploded vectors would remain in the collection uninitialized after executing a dataflow causing global matrices to have NaNs translations and potentially causing bounds computation to generate bogus results
- Implemented the inherited Reset method overrides for FTransformCOllection and FGeometryCollection
- Fixed a bug when resizing a birtarray based managed array
- Geometry collection asset ResetFrom method now calls Reset on the FGeometryCollection object
- Added an optional check for Nans for global matrices

diff --git a/Engine/Source/Runtime/Experimental/Chaos/Private/GeometryCollection/GeometryCollection.cpp b/Engine/Source/Runtime/Experimental/Chaos/Private/GeometryCollection/GeometryCollection.cpp
--- a/Engine/Source/Runtime/Experimental/Chaos/Private/GeometryCollection/GeometryCollection.cpp
+++ b/Engine/Source/Runtime/Experimental/Chaos/Private/GeometryCollection/GeometryCollection.cpp
@@ -704,6 +704,12 @@ void FGeometryCollection::Empty()
 	SetNumUVLayers(1);
 }
 
+void FGeometryCollection::Reset()
+{
+	Super::Reset();
+	Construct();
+}
+
 void FGeometryCollection::InitializeInterfaces()
 {
 	FGeometryCollectionConvexPropertiesInterface::InitializeInterface();
diff --git a/Engine/Source/Runtime/Experimental/Chaos/Private/GeometryCollection/TransformCollection.cpp b/Engine/Source/Runtime/Experimental/Chaos/Private/GeometryCollection/TransformCollection.cpp
--- a/Engine/Source/Runtime/Experimental/Chaos/Private/GeometryCollection/TransformCollection.cpp
+++ b/Engine/Source/Runtime/Experimental/Chaos/Private/GeometryCollection/TransformCollection.cpp
@@ -40,6 +40,12 @@ void FTransformCollection::Construct()
 	AddExternalAttribute<TSet<int32>>(FTransformCollection::ChildrenAttribute, FTransformCollection::TransformGroup, Children);
 }
 
+void FTransformCollection::Reset()
+{
+	Super::Reset();
+	Construct();
+}
+
 void FTransformCollection::Serialize(Chaos::FChaosArchive& Ar)
 {
 	Super::Serialize(Ar);
diff --git a/Engine/Source/Runtime/Experimental/Chaos/Public/GeometryCollection/GeometryCollection.h b/Engine/Source/Runtime/Experimental/Chaos/Public/GeometryCollection/GeometryCollection.h
--- a/Engine/Source/Runtime/Experimental/Chaos/Public/GeometryCollection/GeometryCollection.h
+++ b/Engine/Source/Runtime/Experimental/Chaos/Public/GeometryCollection/GeometryCollection.h
@@ -174,6 +174,11 @@ public:
 	*/
 	void Empty();
 
+	/**
+	* reset internal state
+	*/
+	virtual void Reset() override;
+
 	/**
 	* Reorders elements in a group. NewOrder must be the same length as the group.
 	*/
diff --git a/Engine/Source/Runtime/Experimental/Chaos/Public/GeometryCollection/ManagedArray.h b/Engine/Source/Runtime/Experimental/Chaos/Public/GeometryCollection/ManagedArray.h
--- a/Engine/Source/Runtime/Experimental/Chaos/Public/GeometryCollection/ManagedArray.h
+++ b/Engine/Source/Runtime/Experimental/Chaos/Public/GeometryCollection/ManagedArray.h
@@ -953,7 +953,7 @@ private:
 		else if (Size < Array.Num())
 		{
 			const int32 NumToRemove = (Array.Num() - Size);
-			Array.RemoveAt((Array.Num() - NumToRemove - 1), NumToRemove);
+			Array.RemoveAt((Array.Num() - NumToRemove), NumToRemove);
 		}
 	}
 
diff --git a/Engine/Source/Runtime/Experimental/Chaos/Public/GeometryCollection/TransformCollection.h b/Engine/Source/Runtime/Experimental/Chaos/Public/GeometryCollection/TransformCollection.h
--- a/Engine/Source/Runtime/Experimental/Chaos/Public/GeometryCollection/TransformCollection.h
+++ b/Engine/Source/Runtime/Experimental/Chaos/Public/GeometryCollection/TransformCollection.h
@@ -102,6 +102,11 @@ public:
 	*/
 	virtual void RemoveElements(const FName & Group, const TArray<int32> & SortedDeletionList, FProcessingParameters Params = FProcessingParameters()) override;
 
+	/**
+	* reset internal state
+	*/
+	virtual void Reset() override;
+
 	// Transform Group
 	TManagedArray<FTransform>   Transform;
 	TManagedArray<FString>      BoneName;
diff --git a/Engine/Source/Runtime/Experimental/GeometryCollectionEngine/Private/GeometryCollection/GeometryCollectionComponent.cpp b/Engine/Source/Runtime/Experimental/GeometryCollectionEngine/Private/GeometryCollection/GeometryCollectionComponent.cpp
--- a/Engine/Source/Runtime/Experimental/GeometryCollectionEngine/Private/GeometryCollection/GeometryCollectionComponent.cpp
+++ b/Engine/Source/Runtime/Experimental/GeometryCollectionEngine/Private/GeometryCollection/GeometryCollectionComponent.cpp
@@ -4057,6 +4057,18 @@ void UGeometryCollectionComponent::CalculateLocalBounds()
 	LocalBounds = ComputeBounds(FMatrix::Identity);
 }
 
+#define GEOMETRY_COLLECTION_CHECK_FOR_NANS_IN_MATRICES 0
+
+#if GEOMETRY_COLLECTION_CHECK_FOR_NANS_IN_MATRICES
+static void CheckForNaNs(const TArray<FMatrix>& Matrices)
+{
+	for (const FMatrix& Matrix : Matrices)
+	{
+		ensureAlways(false == Matrix.ContainsNaN());
+	}
+}
+#endif
+
 void UGeometryCollectionComponent::CalculateGlobalMatrices()
 {
 	SCOPE_CYCLE_COUNTER(STAT_GCCUGlobalMatrices);
@@ -4087,10 +4099,14 @@ void UGeometryCollectionComponent::CalculateGlobalMatrices()
 			GeometryCollectionAlgo::GlobalMatrices(GetTransformArray(), GetParentArray(), GlobalMatrices);
 		}
 	}
-	
+
 #if WITH_EDITOR
 	UpdateGlobalMatricesWithExplodedVectors(GlobalMatrices, *(RestCollection->GetGeometryCollection()));
 #endif
+
+#if GEOMETRY_COLLECTION_CHECK_FOR_NANS_IN_MATRICES
+	CheckForNaNs(GlobalMatrices);
+#endif
 }
 
 int32 UGeometryCollectionComponent::GetNumMaterials() const
diff --git a/Engine/Source/Runtime/Experimental/GeometryCollectionEngine/Private/GeometryCollection/GeometryCollectionObject.cpp b/Engine/Source/Runtime/Experimental/GeometryCollectionEngine/Private/GeometryCollection/GeometryCollectionObject.cpp
--- a/Engine/Source/Runtime/Experimental/GeometryCollectionEngine/Private/GeometryCollection/GeometryCollectionObject.cpp
+++ b/Engine/Source/Runtime/Experimental/GeometryCollectionEngine/Private/GeometryCollection/GeometryCollectionObject.cpp
@@ -414,7 +414,7 @@ void UGeometryCollection::Reset()
 	if (GeometryCollection.IsValid())
 	{
 		Modify(); 
-		GeometryCollection->Empty();
+		GeometryCollection->Reset();
 		Materials.Empty();
 		EmbeddedGeometryExemplar.Empty();
 		AutoInstanceMeshes.Empty();
