When geometry collection is empty, don't build the render mesh data.
This fixes crash when trying to render with empty mesh data.

diff --git a/Engine/Source/Runtime/Experimental/GeometryCollectionEngine/Private/GeometryCollection/DerivedDataGeometryCollectionCooker.cpp b/Engine/Source/Runtime/Experimental/GeometryCollectionEngine/Private/GeometryCollection/DerivedDataGeometryCollectionCooker.cpp
--- a/Engine/Source/Runtime/Experimental/GeometryCollectionEngine/Private/GeometryCollection/DerivedDataGeometryCollectionCooker.cpp
+++ b/Engine/Source/Runtime/Experimental/GeometryCollectionEngine/Private/GeometryCollection/DerivedDataGeometryCollectionCooker.cpp
@@ -121,7 +121,7 @@ bool FDerivedDataGeometryCollectionRenderDataCooker::Build(TArray<uint8>& OutDat
 
 const TCHAR* FDerivedDataGeometryCollectionRenderDataCooker::GetVersionString() const
 {
-	const TCHAR* VersionString = TEXT("043C044F08614AAF8A56BA8070BD0D5A");
+	const TCHAR* VersionString = TEXT("EADF0CE05C96495E8621D418077E9C02");
 
 	static FString CachedNaniteVersionString;
 	if (CachedNaniteVersionString.IsEmpty())
diff --git a/Engine/Source/Runtime/Experimental/GeometryCollectionEngine/Private/GeometryCollection/GeometryCollectionRenderData.cpp b/Engine/Source/Runtime/Experimental/GeometryCollectionEngine/Private/GeometryCollection/GeometryCollectionRenderData.cpp
--- a/Engine/Source/Runtime/Experimental/GeometryCollectionEngine/Private/GeometryCollection/GeometryCollectionRenderData.cpp
+++ b/Engine/Source/Runtime/Experimental/GeometryCollectionEngine/Private/GeometryCollection/GeometryCollectionRenderData.cpp
@@ -155,7 +155,7 @@ struct FGeometryCollectionBuiltMeshData
 };
 
 /** Get the data needed to build render data from a Geometry Collection. */
-void BuildMeshDataFromGeometryCollection(FGeometryCollection& InCollection, FGeometryCollectionBuiltMeshData& OutMeshData)
+bool BuildMeshDataFromGeometryCollection(FGeometryCollection& InCollection, FGeometryCollectionBuiltMeshData& OutMeshData)
 {
 	// Vertices Group
 	const TManagedArray<FVector3f>& VertexArray = InCollection.Vertex;
@@ -329,6 +329,9 @@ void BuildMeshDataFromGeometryCollection(FGeometryCollection& InCollection, FGeo
 	OutMeshData.MaterialsPerTriangle = MoveTemp(MaterialsPerFace);
 	OutMeshData.InternalPerTriangle = MoveTemp(InternalPerFace);
 	OutMeshData.MeshTriangleCounts = MoveTemp(MeshTriangleCounts);
+
+	// Return false if we found no mesh data.
+	return OutMeshData.Indices.Num() > 0 && OutMeshData.Vertices.Position.Num() > 0;
 }
 
 struct FGeometryCollectionBuiltMeshSectionData
@@ -675,16 +678,17 @@ TUniquePtr<FGeometryCollectionRenderData> FGeometryCollectionRenderData::Create(
 	TUniquePtr<FGeometryCollectionRenderData> RenderData = MakeUnique<FGeometryCollectionRenderData>();
 
 	FGeometryCollectionBuiltMeshData MeshBuildData;
-	BuildMeshDataFromGeometryCollection(InCollection, MeshBuildData);
-
-	if (bInEnableNanite)
-	{
-		CreateNaniteData(MeshBuildData, *RenderData.Get());
-	}
-	else
+	if (BuildMeshDataFromGeometryCollection(InCollection, MeshBuildData))
 	{
-		// Could always create mesh data if we want to be able to enable/disable nanite at runtime in cooked build.
-		CreateMeshData(InCollection, MeshBuildData, bInUseFullPrecisionUVs, *RenderData.Get());
+		if (bInEnableNanite)
+		{
+			CreateNaniteData(MeshBuildData, *RenderData.Get());
+		}
+		else
+		{
+			// Could always create mesh data if we want to be able to enable/disable nanite at runtime in cooked build.
+			CreateMeshData(InCollection, MeshBuildData, bInUseFullPrecisionUVs, *RenderData.Get());
+		}
 	}
 
 	return RenderData;
