World Partition - HLOD: Added UGeometryCollectionComponent HLOD support
* Added CreateProxyComponents() method which creates static mesh component(s) from the collection proxy meshes
* Override IsHLODRelevant()

diff --git a/Engine/Source/Runtime/Experimental/GeometryCollectionEngine/Private/GeometryCollection/GeometryCollectionComponent.cpp b/Engine/Source/Runtime/Experimental/GeometryCollectionEngine/Private/GeometryCollection/GeometryCollectionComponent.cpp
--- a/Engine/Source/Runtime/Experimental/GeometryCollectionEngine/Private/GeometryCollection/GeometryCollectionComponent.cpp
+++ b/Engine/Source/Runtime/Experimental/GeometryCollectionEngine/Private/GeometryCollection/GeometryCollectionComponent.cpp
@@ -4552,6 +4552,28 @@ UInstancedStaticMeshComponent* UGeometryCollectionComponent::GetRootProxyISM() c
 	return nullptr;
 }
 
+TArray<UStaticMeshComponent*> UGeometryCollectionComponent::CreateProxyComponents() const
+{
+	TArray<UStaticMeshComponent*> Components;
+
+	if (RestCollection)
+	{
+		for (int32 MeshIndex = 0; MeshIndex < RestCollection->RootProxyData.ProxyMeshes.Num(); MeshIndex++)
+		{
+			const TObjectPtr<UStaticMesh>& Mesh = RestCollection->RootProxyData.ProxyMeshes[MeshIndex];
+			if (Mesh != nullptr)
+			{
+				UStaticMeshComponent* NewComponent = NewObject<UStaticMeshComponent>(GetOwner());
+				NewComponent->SetStaticMesh(Mesh);
+				NewComponent->SetComponentToWorld(GetComponentToWorld());
+				Components.Add(NewComponent);
+			}
+		}
+	}
+
+	return Components;
+}
+
 bool UGeometryCollectionComponent::IsRootBroken() const
 {
 	if (DynamicCollection && DynamicCollection->Active.Num() > 0)
@@ -5118,3 +5140,40 @@ void UGeometryCollectionComponent::SetEnableDamageFromCollision(bool bValue)
 		PhysicsProxy->SetEnableDamageFromCollision_External(bValue);
 	}
 }
+
+
+#if WITH_EDITOR
+
+bool UGeometryCollectionComponent::IsHLODRelevant() const
+{
+	if (!RestCollection)
+	{
+		return false;
+	}
+
+	if (RestCollection->RootProxyData.ProxyMeshes.IsEmpty())
+	{
+		return false;
+	}
+
+	if (!IsVisible())
+	{
+		return false;
+	}
+
+#if WITH_EDITORONLY_DATA
+	if (IsVisualizationComponent())
+	{
+		return false;
+	}
+
+	if (!bEnableAutoLODGeneration)
+	{
+		return false;
+	}
+#endif
+
+	return true;
+}
+
+#endif // #if WITH_EDITOR
diff --git a/Engine/Source/Runtime/Experimental/GeometryCollectionEngine/Public/GeometryCollection/GeometryCollectionComponent.h b/Engine/Source/Runtime/Experimental/GeometryCollectionEngine/Public/GeometryCollection/GeometryCollectionComponent.h
--- a/Engine/Source/Runtime/Experimental/GeometryCollectionEngine/Public/GeometryCollection/GeometryCollectionComponent.h
+++ b/Engine/Source/Runtime/Experimental/GeometryCollectionEngine/Public/GeometryCollection/GeometryCollectionComponent.h
@@ -1013,6 +1013,7 @@ public:
 	FGeometryDynamicCollection* GetDynamicCollection();  // TEMP HACK?
 
 	UInstancedStaticMeshComponent* GetRootProxyISM() const;
+	TArray<UStaticMeshComponent*> CreateProxyComponents() const;
 
 public:
 	UPROPERTY(BlueprintAssignable, Category = "Collision")
@@ -1306,4 +1307,10 @@ public:
 	virtual TArray<Chaos::FPhysicsObject*> GetAllPhysicsObjects() const override;
 	virtual Chaos::FPhysicsObjectId GetIdFromGTParticle(Chaos::FGeometryParticle* Particle) const override;
 	//~ End IPhysicsComponent Interface.
+
+#if WITH_EDITOR
+	//~ Begin UActorComponent interface.
+	virtual bool IsHLODRelevant() const override;
+	//~ End UActorComponent interface.
+#endif
 };
