Chaos: Dynamic trees dirty tree can now be enabled/disabled during runtime to avoid a crash

diff --git a/Engine/Source/Runtime/Experimental/Chaos/Public/Chaos/AABBTree.h b/Engine/Source/Runtime/Experimental/Chaos/Public/Chaos/AABBTree.h
--- a/Engine/Source/Runtime/Experimental/Chaos/Public/Chaos/AABBTree.h
+++ b/Engine/Source/Runtime/Experimental/Chaos/Public/Chaos/AABBTree.h
@@ -2032,8 +2032,14 @@ public:
 
 		if (DirtyElementTree)
 		{
-			check(From.DirtyElementTree);
-			DirtyElementTree->PrepareCopyTimeSliced(*(From.DirtyElementTree));
+			if (From.DirtyElementTree)
+			{
+				DirtyElementTree->PrepareCopyTimeSliced(*(From.DirtyElementTree));
+			}
+			else
+			{
+				DirtyElementTree = nullptr;
+			}			
 		}
 	}
 	
@@ -3584,7 +3590,10 @@ private:
 			bBuildOverlapCache = Rhs.bBuildOverlapCache;			
 			if (Rhs.DirtyElementTree)
 			{
-				check(DirtyElementTree); // We should have allocated this already
+				if (!DirtyElementTree)
+				{
+					DirtyElementTree = TUniquePtr<TAABBTree<TPayloadType, TLeafType, bMutable, T>>(new TAABBTree<TPayloadType, TLeafType, bMutable, T>());
+				}				
 				*DirtyElementTree = *Rhs.DirtyElementTree;
 			}
 		}
