Chaos - Reduce FImplicitObjectUnion BVH memorty footprint
- Fix partitioning in TBoundingVolumeHierarchy
- Use BVH in Collision midphase when available
- Improve midphase perf for many-shape objects (even without BVH)
- Add cvars to disable BVH

diff --git a/Engine/Source/Runtime/Core/Public/UObject/FortniteMainBranchObjectVersion.h b/Engine/Source/Runtime/Core/Public/UObject/FortniteMainBranchObjectVersion.h
--- a/Engine/Source/Runtime/Core/Public/UObject/FortniteMainBranchObjectVersion.h
+++ b/Engine/Source/Runtime/Core/Public/UObject/FortniteMainBranchObjectVersion.h
@@ -286,6 +286,9 @@ struct CORE_API FFortniteMainBranchObjectVersion
 		// FDataLayerInstanceDesc support for private data layers
 		WorldPartitionPrivateDataLayers,
 
+		// Reduce size and improve behaviour of Chaos::FImplicitObjectUnion
+		ChaosImplicitObjectUnionBVHRefactor,
+
 		// -----<new versions can be added above this line>-------------------------------------------------
 		VersionPlusOne,
 		LatestVersion = VersionPlusOne - 1
diff --git a/Engine/Source/Runtime/Engine/Private/PhysicsEngine/Experimental/PhysInterface_Chaos.cpp b/Engine/Source/Runtime/Engine/Private/PhysicsEngine/Experimental/PhysInterface_Chaos.cpp
--- a/Engine/Source/Runtime/Engine/Private/PhysicsEngine/Experimental/PhysInterface_Chaos.cpp
+++ b/Engine/Source/Runtime/Engine/Private/PhysicsEngine/Experimental/PhysInterface_Chaos.cpp
@@ -716,15 +716,20 @@ void FPhysInterface_Chaos::AddGeometry(FPhysicsActorHandle& InActor, const FGeom
 		// todo: we should not be creating unique geometry per actor
 		bool bMergeShapesArray = false;
 		{
-			if (InActor->GetGameThreadAPI().Geometry()) // geometry already exists - combine new geometry with the existing
+			if (InActor->GetGameThreadAPI().Geometry())
 			{
+				// Geometry already exists - combine new geometry with the existing
+				// NOTE: We do not need to set the AllowBVH flag because it will be cloned (see below)
 				InActor->GetGameThreadAPI().MergeGeometry(MoveTemp(Geoms));
 				bMergeShapesArray = true;
 			}
 			else
 			{
-				// we always have a union so we can support any future welding operations. (Non-trivial converting the SharedPtr to UniquePtr)
-				InActor->GetGameThreadAPI().SetGeometry(MakeUnique<Chaos::FImplicitObjectUnion>(MoveTemp(Geoms)));
+				// We always have a union so we can support any future welding operations. (Non-trivial converting the SharedPtr to UniquePtr).
+				// The root union always supports BVH and is the only Union in the hierarchy that is allowed to do so.
+				TUniquePtr<Chaos::FImplicitObjectUnion> Union = MakeUnique<Chaos::FImplicitObjectUnion>(MoveTemp(Geoms));
+				Union->SetAllowBVH(true);
+				InActor->GetGameThreadAPI().SetGeometry(MoveTemp(Union));
 			}
 		}
 
diff --git a/Engine/Source/Runtime/Experimental/Chaos/Private/Chaos/BoundingVolumeHierarchy.cpp b/Engine/Source/Runtime/Experimental/Chaos/Private/Chaos/BoundingVolumeHierarchy.cpp
--- a/Engine/Source/Runtime/Experimental/Chaos/Private/Chaos/BoundingVolumeHierarchy.cpp
+++ b/Engine/Source/Runtime/Experimental/Chaos/Private/Chaos/BoundingVolumeHierarchy.cpp
@@ -3,6 +3,7 @@
 #include "Chaos/BoundingVolumeHierarchy.h"
 #include "Chaos/BoundingVolume.h"
 #include "Chaos/BoundingVolumeUtilities.h"
+#include "Chaos/ImplicitObjectBVH.h"
 #include "Chaos/Tetrahedron.h"
 #include "ChaosStats.h"
 
@@ -423,8 +424,16 @@ namespace Chaos
 		LocalLeafs.SetNum(2);
 		TAABB<T, d> GlobalBox(GlobalMin, GlobalMax);
 		const TVector<T, d> WorldCenter = GlobalBox.Center();
-		const TVector<T, d> MinCenterSearch = TAABB<T, d>(GlobalMin, WorldCenter).Center();
-		const TVector<T, d> MaxCenterSearch = TAABB<T, d>(WorldCenter, GlobalMax).Center();
+
+		// @todo(chaos): I'm not sure what MinCenterSearch and MaxCenterSearch are supposed to be
+		// doing here, but using two different points and counting objects either side along each
+		// axis (see AccumulateNextLevelCount) will double-count objects. In the simplest case of
+		// a regular grid of objects, this results in always splitting along one axis rather than
+		// alternating between the three axes and that's probably not what was intended. Consider
+		// switching to KD Tree logic (sort along axis and pick median point for split position).
+		// OK, how super amazing was it that this comment block perfectly right-justified itself?
+		const TVector<T, d> MinCenterSearch = WorldCenter;//TAABB<T, d>(GlobalMin, WorldCenter).Center();
+		const TVector<T, d> MaxCenterSearch = WorldCenter;//TAABB<T, d>(WorldCenter, GlobalMax).Center();
 
 		for(int32 i = 0; i < Objects.Num(); ++i)
 		{
@@ -692,3 +701,5 @@ template class Chaos::TBoundingVolumeHierarchy<Chaos::TGeometryParticles<Chaos::
 template class Chaos::TBoundingVolumeHierarchy<Chaos::TPBDRigidParticles<Chaos::FReal, 3>, Chaos::TBoundingVolume<Chaos::TPBDRigidParticleHandle<Chaos::FReal, 3>*, Chaos::FReal, 3>, Chaos::FReal, 3>;
 template class Chaos::TBoundingVolumeHierarchy<Chaos::TGeometryParticles<Chaos::FReal, 3>, Chaos::TBoundingVolume<Chaos::TGeometryParticleHandle<Chaos::FReal, 3>*, Chaos::FReal, 3>, Chaos::FReal, 3>;
 template class Chaos::TBoundingVolumeHierarchy<TArray<TUniquePtr<Chaos::FImplicitObject>>, TArray<int32>, Chaos::FReal, 3>;
+template class Chaos::TBoundingVolumeHierarchy<TArray<Chaos::Private::FImplicitBVHObject>, TArray<int32>, Chaos::FReal, 3>;
+
diff --git a/Engine/Source/Runtime/Experimental/Chaos/Private/Chaos/Collision/ParticlePairMidPhase.cpp b/Engine/Source/Runtime/Experimental/Chaos/Private/Chaos/Collision/ParticlePairMidPhase.cpp
--- a/Engine/Source/Runtime/Experimental/Chaos/Private/Chaos/Collision/ParticlePairMidPhase.cpp
+++ b/Engine/Source/Runtime/Experimental/Chaos/Private/Chaos/Collision/ParticlePairMidPhase.cpp
@@ -6,6 +6,7 @@
 #include "Chaos/Collision/PBDCollisionConstraint.h"
 #include "Chaos/CollisionResolution.h"
 #include "Chaos/ImplicitObject.h"
+#include "Chaos/ImplicitObjectBVH.h"
 #include "Chaos/ParticleHandle.h"
 #include "Chaos/Particle/ParticleUtilities.h"
 #include "Chaos/PBDCollisionConstraints.h"
@@ -34,6 +35,8 @@ namespace Chaos
 		FAutoConsoleVariableRef CVarChaos_Collision_CullDistanceReferenceSize(TEXT("p.Chaos.Collision.CullDistanceReferenceSize"), Chaos_Collision_CullDistanceScaleInverseSize, TEXT(""));
 		FAutoConsoleVariableRef CVarChaos_Collision_MinCullDistanecScale(TEXT("p.Chaos.Collision.MinCullDistanceScale"), Chaos_Collision_MinCullDistanceScale, TEXT(""));
 
+		int32 Chaos_Collision_MidPhase_MaxShapePairs = 100;
+		FAutoConsoleVariableRef CVarChaos_Collision_MaxShapePairs(TEXT("p.Chaos.Collision.MaxShapePairs"), Chaos_Collision_MidPhase_MaxShapePairs, TEXT(""));
 	}
 	using namespace CVars;
 
@@ -108,6 +111,36 @@ namespace Chaos
 		return Context.GetAllocator()->CreateConstraint(Particle0, Implicit0, Shape0, BVHParticles0, ShapeRelativeTransform0, Particle1, Implicit1, Shape1, BVHParticles1, ShapeRelativeTransform1, CullDistance, bUseManifold, ShapePairType);
 	}
 
+	struct FLeafImplicitObject
+	{
+		const FImplicitObject* Geometry;
+		const FShapeInstance* ShapeInstance;
+		FRigidTransform3 RelativeTransform;
+		FRigidTransform3 WorldTransform;
+		FAABB3 WorldBounds;
+		int32 RootObjectIndex;
+		int32 LeafObjectIndex;
+	};
+
+	void CollectLeafImplicitObjects(const FAABB3& LocalBounds, const FConstGenericParticleHandle& P, const FImplicitObjectUnion* Union, const FRigidTransform3 ParticleWorldTransform, const FReal CullDistance, TArray<FLeafImplicitObject>& Objects)
+	{
+		Union->VisitOverlappingLeafObjects(LocalBounds,
+			[&Objects, &P, &ParticleWorldTransform, CullDistance](const FImplicitObject* Implicit, const FRigidTransform3& RelativeTransform, const int32 RootObjectIndex, const int32 ObjectIndex, const int32 LeafObjectIndex)
+			{
+				const FRigidTransform3 WorldTransform = RelativeTransform * ParticleWorldTransform;
+				const FLeafImplicitObject LeafObject =
+				{
+					Implicit,
+					P->ShapeInstances()[RootObjectIndex].Get(),
+					RelativeTransform,
+					WorldTransform,
+					Implicit->CalculateTransformedBounds(WorldTransform).ThickenSymmetrically(FVec3(CullDistance)),
+					RootObjectIndex,
+					LeafObjectIndex
+				};
+				Objects.Emplace(LeafObject);
+			});
+	};
 
 	////////////////////////////////////////////////////////////////////////////////////////
 	////////////////////////////////////////////////////////////////////////////////////////
@@ -941,6 +974,440 @@ namespace Chaos
 	////////////////////////////////////////////////////////////////////////////////////////
 
 
+	FParticlePairCollisionDetector::FParticlePairCollisionDetector(
+		FGeometryParticleHandle* InParticle0,
+		FGeometryParticleHandle* InParticle1,
+		FParticlePairMidPhase& InMidPhase)
+		: MidPhase(InMidPhase)
+		, Constraints()
+		, Particle0(InParticle0)
+		, Particle1(InParticle1)
+	{
+	}
+
+	FParticlePairCollisionDetector::FParticlePairCollisionDetector(FParticlePairCollisionDetector&& R)
+		: MidPhase(R.MidPhase)
+		, Constraints(MoveTemp(R.Constraints))
+		, Particle0(R.Particle0)
+		, Particle1(R.Particle1)
+	{
+		check(R.NewConstraints.IsEmpty());
+	}
+
+	FParticlePairCollisionDetector::~FParticlePairCollisionDetector()
+	{
+	}
+
+	int32 FParticlePairCollisionDetector::GenerateCollisions(
+		const FReal CullDistance,
+		const FReal Dt,
+		const FCollisionContext& Context)
+	{
+		return GenerateCollisionsImpl(false, CullDistance, Dt, Context);
+	}
+
+
+	int32 FParticlePairCollisionDetector::GenerateCollisionsCCD(
+		const bool bEnableCCDSweep,
+		const FReal CullDistance,
+		const FReal Dt,
+		const FCollisionContext& Context)
+	{
+		return GenerateCollisionsImpl(bEnableCCDSweep, CullDistance, Dt, Context);
+	}
+
+	int32 FParticlePairCollisionDetector::GenerateCollisionsImpl(
+		const bool bEnableCCDSweep,
+		const FReal CullDistance,
+		const FReal Dt,
+		const FCollisionContext& Context)
+	{
+		const FConstGenericParticleHandle P0 = Particle0;
+		const FConstGenericParticleHandle P1 = Particle1;
+
+		// We should have a Union at the root of each geometry hierarchy
+		const FImplicitObjectUnion* Union0 = P0->Geometry()->template GetObject<FImplicitObjectUnion>();
+		const FImplicitObjectUnion* Union1 = P1->Geometry()->template GetObject<FImplicitObjectUnion>();
+		if ((Union0 == nullptr) || (Union1 == nullptr))
+		{
+			return 0;
+		}
+
+		// Levelset collision if present
+		const FBVHParticles* BVHParticles0 = P0->CollisionParticles().Get();
+		const FBVHParticles* BVHParticles1 = P1->CollisionParticles().Get();
+
+		// Particle transforms
+		const FRigidTransform3 ParticleWorldTransform0 = P0->GetTransformPQ();
+		const FRigidTransform3 ParticleWorldTransform1 = P1->GetTransformPQ();
+
+		// Calculate the overlapping volume in each particle's space
+		FAABB3 ParticleLocalIntersect0 = FAABB3::FullAABB();
+		FAABB3 ParticleLocalIntersect1 = FAABB3::FullAABB();
+		if (P0->HasBoundingBox() && P1->HasBoundingBox())
+		{
+			const FRigidTransform3 ParticleTransform1To0 = ParticleWorldTransform1.GetRelativeTransform(ParticleWorldTransform0);
+			const FRigidTransform3 ParticleTransform0To1 = ParticleWorldTransform0.GetRelativeTransform(ParticleWorldTransform1);
+			ParticleLocalIntersect0 = P1->Geometry()->CalculateTransformedBounds(ParticleTransform1To0).ThickenSymmetrically(FVec3(CullDistance));
+			ParticleLocalIntersect1 = P0->Geometry()->CalculateTransformedBounds(ParticleTransform0To1).ThickenSymmetrically(FVec3(CullDistance));
+		}
+
+		// Get the set of implicits from each particle that are in the overlap zone
+		TArray<FLeafImplicitObject> Objects0;
+		TArray<FLeafImplicitObject> Objects1;
+		CollectLeafImplicitObjects(ParticleLocalIntersect0, P0, Union0, ParticleWorldTransform0, CullDistance, Objects0);
+		CollectLeafImplicitObjects(ParticleLocalIntersect1, P1, Union1, ParticleWorldTransform1, CullDistance, Objects1);
+
+		// Modify the context so that Collision allocations come to us
+		// @todo(chaos): maybe pass a functor into ConstructConstraints instead
+		FCollisionContext LocalContext = Context;
+		LocalContext.ParticlePairCollisionDetector = this;
+
+		// Check all shape pairs for collisions
+		for (const FLeafImplicitObject& Object0 : Objects0)
+		{
+			for (const FLeafImplicitObject& Object1 : Objects1)
+			{
+				// NOTE: We have 2 sets of bodies from the total overlapping volume of the 2 hierarchies. 
+				// We still need to perform a shape pair bounds check
+				if (!Object0.WorldBounds.Intersects(Object1.WorldBounds))
+				{
+					continue;
+				}
+
+				// Create the constraint for the shape pair if the bounds overlap
+				Collisions::ConstructConstraints(
+					Particle0, Particle1,
+					Object0.Geometry, Object0.ShapeInstance, BVHParticles0,
+					Object1.Geometry, Object1.ShapeInstance, BVHParticles1,
+					ParticleWorldTransform0, Object0.RelativeTransform,
+					ParticleWorldTransform1, Object1.RelativeTransform,
+					CullDistance, Dt, bEnableCCDSweep, LocalContext);
+			}
+		}
+
+		// Generate manifolds for each constraint we created/recovered and (re)activate if necessary
+		int32 NumActiveConstraints = ProcessNewConstraints(CullDistance, Dt, Context);
+
+		// @todo(chaos): we could clean up unused collisions between this pair, but probably not worth it
+		// and we would have to prevent cleanup for sleeping particles because the collisions will still
+		// be referenced in the IslandManager's constraint graph for the sleeping island.
+		//PruneConstraints();
+
+		return NumActiveConstraints;
+	}
+
+	FPBDCollisionConstraint* FParticlePairCollisionDetector::FindOrCreateConstraint(
+		FGeometryParticleHandle* InParticle0,
+		const FImplicitObject* InImplicit0,
+		const int32 InImplicitId0,
+		const FShapeInstance* InShape0,
+		const FBVHParticles* InBVHParticles0,
+		const FRigidTransform3& InShapeRelativeTransform0,
+		FGeometryParticleHandle* InParticle1,
+		const FImplicitObject* InImplicit1,
+		const int32 InImplicitId1,
+		const FShapeInstance* InShape1,
+		const FBVHParticles* InBVHParticles1,
+		const FRigidTransform3& InShapeRelativeTransform1,
+		const FReal CullDistance,
+		const EContactShapesType ShapePairType,
+		const bool bUseManifold,
+		const bool bEnableSweep,
+		const FCollisionContext& Context)
+	{
+		// This is a callback from the low-level collision function. It should always be the same two particles, though the
+		// shapes may be some children in the implicit hierarchy. The particles could be in the opposite order though, and
+		// this will depend on the shape types involved. E.g., with two particles each with a sphere and a box in a union
+		// would require up to two Sphere-Box contacts, with the particles in opposite orders.
+		if (!ensure(((InParticle0 == Particle0) && (InParticle1 == Particle1)) || ((InParticle0 == Particle1) && (InParticle1 == Particle0))))
+		{
+			// We somehow received a callback for the wrong particle pair...this should not happen
+			return nullptr;
+		}
+
+		const FCollisionParticlePairConstraintKey Key = FCollisionParticlePairConstraintKey(InShape0, InImplicit0, InImplicitId0, InBVHParticles0, InShape1, InImplicit1, InImplicitId1, InBVHParticles1);
+		FPBDCollisionConstraint* Constraint = FindConstraint(Key);
+
+		// @todo(chaos): fix key uniqueness guarantee.  We need a truly unique key gen function
+		const bool bIsKeyCollision = (Constraint != nullptr) && ((Constraint->GetImplicit0() != InImplicit0) || (Constraint->GetImplicit1() != InImplicit1) || (Constraint->GetCollisionParticles0() != InBVHParticles0) || (Constraint->GetCollisionParticles1() != InBVHParticles1));
+		if (bIsKeyCollision)
+		{
+			// If we get here, we have a key collision. The key uses a hash of pointers which is very likely to be unique for different implicit pairs, 
+			// especially since it only needs to be unique for this particle pair, but it is not guaranteed.
+			// Creating a new constraint with the same key could cause fatal problems (the original constraint will be deleted when we add the new one 
+			// to the map, but if it is asleep it will be referenced in the contact graph) so we just abort and accept we will miss collisions. 
+			// It is extremely unlikely to happen but we should fix it at some point.
+			ensure(false);
+			return nullptr;
+		}
+
+		if (Constraint == nullptr)
+		{
+			// NOTE: Using InParticle0 and InParticle1 here because the order may be different to what we have stored
+			Constraint = CreateConstraint(InParticle0, InImplicit0, InShape0, InBVHParticles0, InShapeRelativeTransform0, InParticle1, InImplicit1, InShape1, InBVHParticles1, InShapeRelativeTransform1, CullDistance, ShapePairType, bUseManifold, Key, Context);
+
+			// Is this a CCD constraint?
+			Constraint->SetCCDEnabled(MidPhase.IsCCD());
+		}
+
+		// Do we want to sweep on this tick? 
+		// CCD may be temporarily disabled by the user or because we are moving slowly.
+		Constraint->SetCCDSweepEnabled(MidPhase.IsCCD() && bEnableSweep);
+
+		NewConstraints.Add(Constraint);
+		return Constraint;
+	}
+
+	FPBDCollisionConstraint* FParticlePairCollisionDetector::FindConstraint(const FCollisionParticlePairConstraintKey& Key)
+	{
+		FPBDCollisionConstraintPtr* PConstraint = Constraints.Find(Key.GetKey());
+		if (PConstraint != nullptr)
+		{
+			return (*PConstraint).Get();
+		}
+		return nullptr;
+	}
+
+	FPBDCollisionConstraint* FParticlePairCollisionDetector::CreateConstraint(
+		FGeometryParticleHandle* InParticle0,
+		const FImplicitObject* Implicit0,
+		const FPerShapeData* InShape0,
+		const FBVHParticles* BVHParticles0,
+		const FRigidTransform3& ShapeRelativeTransform0,
+		FGeometryParticleHandle* InParticle1,
+		const FImplicitObject* Implicit1,
+		const FPerShapeData* InShape1,
+		const FBVHParticles* BVHParticles1,
+		const FRigidTransform3& ShapeRelativeTransform1,
+		const FReal CullDistance,
+		const EContactShapesType ShapePairType,
+		const bool bUseManifold,
+		const FCollisionParticlePairConstraintKey& Key,
+		const FCollisionContext& Context)
+	{
+		PHYSICS_CSV_SCOPED_EXPENSIVE(PhysicsVerbose, NarrowPhase_CreateConstraint);
+		FPBDCollisionConstraintPtr Constraint = CreateImplicitPairConstraint(InParticle0, Implicit0, InShape0, BVHParticles0, ShapeRelativeTransform0, InParticle1, Implicit1, InShape1, BVHParticles1, ShapeRelativeTransform1, CullDistance, ShapePairType, bUseManifold, Context);
+
+		const int32 CurrentEpoch = Context.GetAllocator()->GetCurrentEpoch();
+		Constraint->GetContainerCookie().MidPhase = &MidPhase;
+		Constraint->GetContainerCookie().bIsMultiShapePair = true;
+		Constraint->GetContainerCookie().CreationEpoch = CurrentEpoch;
+
+		return Constraints.Add(Key.GetKey(), MoveTemp(Constraint)).Get();
+	}
+
+	void FParticlePairCollisionDetector::WakeCollisions(const int32 SleepEpoch, const int32 CurrentEpoch)
+	{
+		for (auto& KVP : Constraints)
+		{
+			FPBDCollisionConstraintPtr& Constraint = KVP.Value;
+			if (Constraint->GetContainerCookie().LastUsedEpoch >= SleepEpoch)
+			{
+				Constraint->GetContainerCookie().LastUsedEpoch = CurrentEpoch;
+			}
+		}
+	}
+
+	int32 FParticlePairCollisionDetector::ProcessNewConstraints(
+		const FReal CullDistance,
+		const FReal Dt,
+		const FCollisionContext& Context)
+	{
+		int32 NumActiveConstraints = 0;
+		for (FPBDCollisionConstraint* Constraint : NewConstraints)
+		{
+			if (!Constraint->GetCCDSweepEnabled())
+			{
+				UpdateCollision(Constraint, CullDistance, Dt, Context);
+			}
+			else
+			{
+				UpdateCollisionCCD(Constraint, CullDistance, Dt, Context);
+			}
+
+			if (!Constraint->GetDisabled())
+			{
+				++NumActiveConstraints;
+			}
+		}
+		NewConstraints.Reset();
+		return NumActiveConstraints;
+	}
+
+	void FParticlePairCollisionDetector::UpdateCollision(
+		FPBDCollisionConstraint* Constraint,
+		const FReal CullDistance,
+		const FReal Dt,
+		const FCollisionContext& Context)
+	{
+		// @todo(chaos): share this code with FSingleShapePairCollisionDetector
+
+		const FRigidTransform3 ShapeWorldTransform0 = Constraint->GetShapeRelativeTransform0() * FConstGenericParticleHandle(Constraint->GetParticle0())->GetTransformPQ();
+		const FRigidTransform3 ShapeWorldTransform1 = Constraint->GetShapeRelativeTransform1() * FConstGenericParticleHandle(Constraint->GetParticle1())->GetTransformPQ();
+		const int32 CurrentEpoch = Context.GetAllocator()->GetCurrentEpoch();
+		const int32 LastEpoch = CurrentEpoch - 1;
+		const bool bWasUpdatedLastTick = (Constraint->GetContainerCookie().LastUsedEpoch >= LastEpoch);
+
+		// Update the world shape transforms on the constraint (we cannot just give it the PerShapeData 
+		// pointer because of Unions - see FMultiShapePairCollisionDetector)
+		Constraint->SetShapeWorldTransforms(ShapeWorldTransform0, ShapeWorldTransform1);
+
+		Constraint->SetCullDistance(CullDistance);
+
+		// If the constraint was not used last frame, it needs to be reset, otherwise we will try to reuse
+		if (!bWasUpdatedLastTick || (Constraint->GetManifoldPoints().Num() == 0))
+		{
+			// Clear all manifold data including saved contact data
+			Constraint->ResetManifold();
+		}
+
+		bool bWasManifoldRestored = false;
+		//if (Context.GetSettings().bAllowManifoldReuse && Flags.bEnableManifoldUpdate && bWasUpdatedLastTick)
+		//{
+		//	// Update the existing manifold. We can re-use as-is if none of the points have moved much and the bodies have not moved much
+		//	// NOTE: this can succeed in "restoring" even if we have no manifold points
+		//	// NOTE: this uses the transforms from SetLastShapeWorldTransforms, so we can only do this if we were updated last tick
+		//	bWasManifoldRestored = Constraint->UpdateAndTryRestoreManifold();
+		//}
+		//else
+		{
+			// We are not trying to reuse manifold points, so reset them but leave stored data intact (for friction)
+			Constraint->ResetActiveManifoldContacts();
+		}
+
+		if (!bWasManifoldRestored)
+		{
+			// We will be updating the manifold so update transforms used to check for movement in UpdateAndTryRestoreManifold on future ticks
+			Constraint->SetLastShapeWorldTransforms(ShapeWorldTransform0, ShapeWorldTransform1);
+
+			if (!Context.GetSettings().bDeferNarrowPhase)
+			{
+				Collisions::UpdateConstraint(*Constraint, Constraint->GetShapeWorldTransform0(), Constraint->GetShapeWorldTransform1(), Dt);
+			}
+		}
+
+		// If we have a valid contact, add it to the active list
+		// We also add it to the active list if collision detection is deferred because the data will be filled in later and we
+		// don't know in advance whether we will pass the Phi check (deferred narrow phase is used with RBAN)
+		if (Constraint->GetPhi() <= CullDistance || Context.GetSettings().bDeferNarrowPhase)
+		{
+			if (Context.GetAllocator()->ActivateConstraint(Constraint))
+			{
+				return;
+			}
+		}
+
+		// If we get here, we did not activate hte constraint and it should be disabled for this tick
+		Constraint->SetDisabled(true);
+	}
+
+
+	void FParticlePairCollisionDetector::UpdateCollisionCCD(
+		FPBDCollisionConstraint* Constraint,
+		const FReal CullDistance,
+		const FReal Dt,
+		const FCollisionContext& Context)
+	{
+		// @todo(chaos): share this code with FSingleShapePairCollisionDetector
+
+		if (!Constraint->GetCCDSweepEnabled())
+		{
+			return UpdateCollision(Constraint, CullDistance, Dt, Context);
+		}
+
+		Constraint->ResetManifold();
+		Constraint->ResetActiveManifoldContacts();
+
+		const FRigidTransform3 ShapeWorldTransform0 = Constraint->GetShapeRelativeTransform0() * FConstGenericParticleHandle(Constraint->GetParticle0())->GetTransformPQ();
+		const FRigidTransform3 ShapeWorldTransform1 = Constraint->GetShapeRelativeTransform1() * FConstGenericParticleHandle(Constraint->GetParticle1())->GetTransformPQ();
+
+		// Update the world shape transforms on the constraint (we cannot just give it the PerShapeData 
+		// pointer because of Unions - see FMultiShapePairCollisionDetector)
+		// NOTE: these are not used by CCD which continuously moves the particles
+		Constraint->SetShapeWorldTransforms(ShapeWorldTransform0, ShapeWorldTransform1);
+
+		Constraint->SetCullDistance(CullDistance);
+		Constraint->ResetManifold();
+		Constraint->ResetActiveManifoldContacts();
+
+		FConstGenericParticleHandle P0 = Constraint->GetParticle0();
+		FConstGenericParticleHandle P1 = Constraint->GetParticle1();
+
+		// We need the previous transform for the swept collision detector. It assumes that the current
+		// transform has been set on the constraint. 
+		// We assume that the particle's center of mass moved in a stright line and that it's rotation has 
+		// not changed so we calculate the previous transform from the current one and the velocity.
+		// NOTE: These are actor transforms, not CoM transforms
+		// @todo(chaos): Pass both start and end transforms to the collision detector
+		const FRigidTransform3 CCDParticleWorldTransform0 = FRigidTransform3(P0->P() - P0->V() * Dt, P0->Q());
+		const FRigidTransform3 CCDParticleWorldTransform1 = FRigidTransform3(P1->P() - P1->V() * Dt, P1->Q());
+		const FRigidTransform3 CCDShapeWorldTransform0 = Constraint->GetShapeRelativeTransform0() * CCDParticleWorldTransform0;
+		const FRigidTransform3 CCDShapeWorldTransform1 = Constraint->GetShapeRelativeTransform1() * CCDParticleWorldTransform1;
+		const bool bDidSweep = Collisions::UpdateConstraintSwept(*Constraint, CCDShapeWorldTransform0, CCDShapeWorldTransform1, Dt);
+
+#if CHAOS_DEBUG_DRAW
+		if (CVars::ChaosSolverDrawCCDInteractions)
+		{
+			if (FConstGenericParticleHandle(Constraint->GetParticle0())->CCDEnabled())
+			{
+				DebugDraw::DrawShape(CCDShapeWorldTransform0, Constraint->GetImplicit0(), Constraint->GetShape0(), FColor::Black, &CVars::ChaosSolverDebugDebugDrawSettings);
+				DebugDraw::DrawShape(ShapeWorldTransform0, Constraint->GetImplicit0(), Constraint->GetShape0(), FColor::White, &CVars::ChaosSolverDebugDebugDrawSettings);
+			}
+			if (FConstGenericParticleHandle(Constraint->GetParticle1())->CCDEnabled())
+			{
+				DebugDraw::DrawShape(CCDShapeWorldTransform1, Constraint->GetImplicit1(), Constraint->GetShape1(), FColor::Black, &CVars::ChaosSolverDebugDebugDrawSettings);
+				DebugDraw::DrawShape(ShapeWorldTransform1, Constraint->GetImplicit1(), Constraint->GetShape1(), FColor::White, &CVars::ChaosSolverDebugDebugDrawSettings);
+			}
+		}
+#endif
+
+		// If we did get a hit but it's at TOI = 1, treat this constraint as a regular non-swept constraint (skip the rewind)
+		if ((!bDidSweep) || Constraint->GetCCDTimeOfImpact() == FReal(1))
+		{
+			Collisions::UpdateConstraint(*Constraint, Constraint->GetShapeWorldTransform0(), Constraint->GetShapeWorldTransform1(), Dt);
+		}
+
+		// If the sweep did not find a hit, or we are treating it like a regular contact, we will skip the CCD rewind step for this contact
+		if (Constraint->GetCCDTimeOfImpact() >= FReal(1))
+		{
+			Constraint->SetCCDSweepEnabled(false);
+		}
+
+		Context.GetAllocator()->ActivateConstraint(Constraint);
+	}
+
+	void FParticlePairCollisionDetector::PruneConstraints(const int32 CurrentEpoch)
+	{
+		// We don't prune from NewCollisions - must call ProcessNewCollisions before Prune
+		check(NewConstraints.Num() == 0);
+
+		// Find all the expired collisions
+		TArray<uint32> Pruned;
+		for (auto& KVP : Constraints)
+		{
+			const uint32 Key = KVP.Key;
+			FPBDCollisionConstraintPtr& Constraint = KVP.Value;
+			if (Constraint->GetContainerCookie().LastUsedEpoch < CurrentEpoch)
+			{
+				Pruned.Add(Key);
+			}
+		}
+
+		// Destroy expired collisions
+		for (uint32 Key : Pruned)
+		{
+			Constraints.Remove(Key);
+		}
+	}
+
+
+	////////////////////////////////////////////////////////////////////////////////////////
+	////////////////////////////////////////////////////////////////////////////////////////
+	////////////////////////////////////////////////////////////////////////////////////////
+	////////////////////////////////////////////////////////////////////////////////////////
+
+
 	FParticlePairMidPhase::FParticlePairMidPhase()
 		: Particle0(nullptr)
 		, Particle1(nullptr)
@@ -953,6 +1420,7 @@ namespace Chaos
 		, ParticleCollisionsIndex1(INDEX_NONE)
 		, MultiShapePairDetectors()
 		, ShapePairDetectors()
+		, ParticlePairDetector()
 	{
 	}
 
@@ -1031,17 +1499,42 @@ namespace Chaos
 	{
 		if (IsValid() && (Particle0 != Particle1))
 		{
-			const FShapesArray& Shapes0 = Particle0->ShapesArray();
-			const FShapesArray& Shapes1 = Particle1->ShapesArray();
-			for (int32 ShapeIndex0 = 0; ShapeIndex0 < Shapes0.Num(); ++ShapeIndex0)
+			// For particles with a small number of shapes, we prebuild the set of potentially colliding
+			// shape pairs and cache various data with them. However we don't do this if there are a
+			// large number of shapes, or if the geometry hierarchy contains Unions of Unions (the
+			// FParticlePairCollisionDetector assumes that it has a leaf shape and not a Union).
+			bool bPrebuildShapePairs = true;
+			if (CVars::Chaos_Collision_MidPhase_MaxShapePairs > 0)
+			{
+				const FImplicitObjectUnion* Union0 = (Particle0->Geometry().Get() != nullptr) ? Particle0->Geometry()->template GetObject<FImplicitObjectUnion>() : nullptr;
+				const FImplicitObjectUnion* Union1 = (Particle1->Geometry().Get() != nullptr) ? Particle1->Geometry()->template GetObject<FImplicitObjectUnion>() : nullptr;
+				if ((Union0 != nullptr) && (Union1 != nullptr))
+				{
+					const int32 NumShapePairs = Union0->GetNumLeafObjects() * Union1->GetNumLeafObjects();
+					const bool bIsTree0 = (Union0->GetNumLeafObjects() > Union0->GetObjects().Num());	// Do we have Unions of Unions?
+					const bool bIsTree1 = (Union1->GetNumLeafObjects() > Union1->GetObjects().Num());	// Do we have Unions of Unions?
+					bPrebuildShapePairs = (NumShapePairs < CVars::Chaos_Collision_MidPhase_MaxShapePairs) && !bIsTree0 && !bIsTree1;
+				}
+			}
+
+			if (bPrebuildShapePairs)
 			{
-				const FPerShapeData* Shape0 = Shapes0[ShapeIndex0].Get();
-				for (int32 ShapeIndex1 = 0; ShapeIndex1 < Shapes1.Num(); ++ShapeIndex1)
+				const FShapesArray& Shapes0 = Particle0->ShapesArray();
+				const FShapesArray& Shapes1 = Particle1->ShapesArray();
+				for (int32 ShapeIndex0 = 0; ShapeIndex0 < Shapes0.Num(); ++ShapeIndex0)
 				{
-					const FPerShapeData* Shape1 = Shapes1[ShapeIndex1].Get();
-					TryAddShapePair(Shape0, Shape1);
+					const FPerShapeData* Shape0 = Shapes0[ShapeIndex0].Get();
+					for (int32 ShapeIndex1 = 0; ShapeIndex1 < Shapes1.Num(); ++ShapeIndex1)
+					{
+						const FPerShapeData* Shape1 = Shapes1[ShapeIndex1].Get();
+						TryAddShapePair(Shape0, Shape1);
+					}
 				}
 			}
+			else
+			{
+				ParticlePairDetector = MakeUnique<FParticlePairCollisionDetector>(Particle0, Particle1, *this);
+			}
 		}
 	}
 
@@ -1177,6 +1670,10 @@ namespace Chaos
 				{
 					NumActiveConstraints += MultiShapePair.GenerateCollisionsCCD(bUseSweep, CullDistance, Dt, Context);
 				}
+				if (ParticlePairDetector.IsValid())
+				{
+					ParticlePairDetector->GenerateCollisionsCCD(bUseSweep, CullDistance, Dt, Context);
+				}
 			}
 			else
 			{
@@ -1188,6 +1685,10 @@ namespace Chaos
 				{
 					NumActiveConstraints += MultiShapePair.GenerateCollisions(CullDistance, Dt, Context);
 				}
+				if (ParticlePairDetector.IsValid())
+				{
+					ParticlePairDetector->GenerateCollisions(CullDistance, Dt, Context);
+				}
 			}
 		}
 
@@ -1218,7 +1719,7 @@ namespace Chaos
 		}
 		else
 		{
-			// @todo(chaos): implement cluster Resim restore
+			// @todo(chaos): implement cluster/large hierarchy Resim restore
 			ensure(false);
 		}
 
@@ -1257,6 +1758,10 @@ namespace Chaos
 					{
 						MultiShapePair.WakeCollisions(LastUsedEpoch, CurrentEpoch);
 					}
+					if (ParticlePairDetector.IsValid())
+					{
+						ParticlePairDetector->WakeCollisions(LastUsedEpoch, CurrentEpoch);
+					}
 					LastUsedEpoch = CurrentEpoch;
 				}
 			}
diff --git a/Engine/Source/Runtime/Experimental/Chaos/Private/Chaos/CollisionResolution.cpp b/Engine/Source/Runtime/Experimental/Chaos/Private/Chaos/CollisionResolution.cpp
--- a/Engine/Source/Runtime/Experimental/Chaos/Private/Chaos/CollisionResolution.cpp
+++ b/Engine/Source/Runtime/Experimental/Chaos/Private/Chaos/CollisionResolution.cpp
@@ -1274,10 +1274,20 @@ namespace Chaos
 
 			if (bDoOverlap)
 			{
-				Context.MultiShapeCollisionDetector->FindOrCreateConstraint(
-					Particle0, Implicit0, ImplicitID0, Shape0, Simplicial0, LocalTransform0,
-					Particle1, Implicit1, ImplicitID1, Shape1, Simplicial1, LocalTransform1,
-					CullDistance, ShapePairType, Context.GetSettings().bAllowManifolds, bEnableSweep, Context);
+				if (Context.MultiShapeCollisionDetector != nullptr)
+				{
+					Context.MultiShapeCollisionDetector->FindOrCreateConstraint(
+						Particle0, Implicit0, ImplicitID0, Shape0, Simplicial0, LocalTransform0,
+						Particle1, Implicit1, ImplicitID1, Shape1, Simplicial1, LocalTransform1,
+						CullDistance, ShapePairType, Context.GetSettings().bAllowManifolds, bEnableSweep, Context);
+				}
+				else if (Context.ParticlePairCollisionDetector != nullptr)
+				{
+					Context.ParticlePairCollisionDetector->FindOrCreateConstraint(
+						Particle0, Implicit0, ImplicitID0, Shape0, Simplicial0, LocalTransform0,
+						Particle1, Implicit1, ImplicitID1, Shape1, Simplicial1, LocalTransform1,
+						CullDistance, ShapePairType, Context.GetSettings().bAllowManifolds, bEnableSweep, Context);
+				}
 			}
 		}
 
@@ -1392,11 +1402,11 @@ namespace Chaos
 
 					// @todo(chaos): we shouldn't need to treat Union and ClusterUnion differentlty here - change how we get the simplicial to be more generic
 					const FImplicitObjectUnionClustered* ClusterUnion0 = Implicit0->template GetObject<FImplicitObjectUnionClustered>();
-					Implicit0->VisitOverlappingObjects(QueryBounds0,
-						[&](const FImplicitObject* ChildImplicit0, const FRigidTransform3& ChildTransform0, const int32 ChildId0)
+					Implicit0->VisitOverlappingLeafObjects(QueryBounds0,
+						[&](const FImplicitObject* ChildImplicit0, const FRigidTransform3& ChildTransform0, const int32 RootObjectIndex0, const int32 ObjectIndex, const int32 LeafObjectIndex0)
 						{
-							const FBVHParticles* ChildSimplicial0 = (ClusterUnion0 != nullptr) ? ClusterUnion0->GetChildSimplicial(ChildId0) : Simplicial0;
-							ConstructConstraintsRecursive(Particle0, Particle1, ChildImplicit0, Shape0, ChildSimplicial0, ChildId0, Implicit1, Shape1, Simplicial1, ImplicitID1, ParticleWorldTransform0, ChildTransform0 * LocalTransform0, ParticleWorldTransform1, LocalTransform1, CullDistance, Dt, bEnableSweep, Context);
+							const FBVHParticles* ChildSimplicial0 = (ClusterUnion0 != nullptr) ? ClusterUnion0->GetChildSimplicial(LeafObjectIndex0) : Simplicial0;
+							ConstructConstraintsRecursive(Particle0, Particle1, ChildImplicit0, Shape0, ChildSimplicial0, LeafObjectIndex0, Implicit1, Shape1, Simplicial1, ImplicitID1, ParticleWorldTransform0, ChildTransform0 * LocalTransform0, ParticleWorldTransform1, LocalTransform1, CullDistance, Dt, bEnableSweep, Context);
 						});
 					return;
 				}
@@ -1408,11 +1418,11 @@ namespace Chaos
 					const FAABB3 QueryBounds1 = Implicit0->HasBoundingBox() ? Implicit0->BoundingBox().TransformedAABB(TM0ToTM1) : FAABB3::FullAABB();
 
 					const FImplicitObjectUnionClustered* ClusterUnion1 = Implicit1->template GetObject<FImplicitObjectUnionClustered>();
-					Implicit1->VisitOverlappingObjects(QueryBounds1,
-						[&](const FImplicitObject* ChildImplicit1, const FRigidTransform3& ChildTransform1, const int32 ChildId1)
+					Implicit1->VisitOverlappingLeafObjects(QueryBounds1,
+						[&](const FImplicitObject* ChildImplicit1, const FRigidTransform3& ChildTransform1, const int32 RootObjectIndex1, const int32 ObjectIndex, const int32 LeafObjectIndex1)
 						{
-							const FBVHParticles* ChildSimplicial1 = (ClusterUnion1 != nullptr) ? ClusterUnion1->GetChildSimplicial(ChildId1) : Simplicial1;
-							ConstructConstraintsRecursive(Particle0, Particle1, Implicit0, Shape0, Simplicial0, ImplicitID0, ChildImplicit1, Shape1, Simplicial1, ChildId1, ParticleWorldTransform0, LocalTransform0, ParticleWorldTransform1, ChildTransform1 * LocalTransform1, CullDistance, Dt, bEnableSweep, Context);
+							const FBVHParticles* ChildSimplicial1 = (ClusterUnion1 != nullptr) ? ClusterUnion1->GetChildSimplicial(LeafObjectIndex1) : Simplicial1;
+							ConstructConstraintsRecursive(Particle0, Particle1, Implicit0, Shape0, Simplicial0, ImplicitID0, ChildImplicit1, Shape1, Simplicial1, LeafObjectIndex1, ParticleWorldTransform0, LocalTransform0, ParticleWorldTransform1, ChildTransform1 * LocalTransform1, CullDistance, Dt, bEnableSweep, Context);
 						});
 					return;
 				}
diff --git a/Engine/Source/Runtime/Experimental/Chaos/Private/Chaos/ImplicitObjectBVH.cpp b/Engine/Source/Runtime/Experimental/Chaos/Private/Chaos/ImplicitObjectBVH.cpp
new file mode 100644
--- /dev/null
+++ b/Engine/Source/Runtime/Experimental/Chaos/Private/Chaos/ImplicitObjectBVH.cpp
@@ -0,0 +1,159 @@
+// Copyright Epic Games, Inc. All Rights Reserved.
+#include "Chaos/ImplicitObjectBVH.h"
+#include "Chaos/ChaosArchive.h"
+
+//UE_DISABLE_OPTIMIZATION
+
+namespace Chaos
+{
+	namespace CVars
+	{
+		extern int32 ChaosUnionBVHMaxDepth;
+	}
+
+	namespace Private
+	{
+
+		FImplicitBVHObject::FImplicitBVHObject()
+		{
+		}
+
+		FImplicitBVHObject::FImplicitBVHObject(
+			const TSerializablePtr<FImplicitObject>& InGeometry, 
+			const FVec3& InX, 
+			const FRotation3& InR, 
+			const FAABB3& InBounds,
+			const int32 InRootObjectIndex,
+			const int32 InObjectIndex)
+			: R(FRotation3f(InR))
+			, X(FVec3f(InX))
+			, Bounds(FAABB3f(InBounds))
+			, Geometry(InGeometry)
+			, RootObjectIndex(InRootObjectIndex)
+			, ObjectIndex(InObjectIndex)
+		{
+		}
+
+		///////////////////////////////////////////////////////////////////////////////////////////
+		///////////////////////////////////////////////////////////////////////////////////////////
+
+		int32 FImplicitBVH::CountLeafObjects(const TArrayView<const TUniquePtr<FImplicitObject>>& InRootObjects)
+		{
+			// Count the objects in the hierarchy
+			// @todo(chaos): provide a visitor that does not check bounds
+			int32 NumObjects = 0;
+			for (const TUniquePtr<FImplicitObject>& RootObject : InRootObjects)
+			{
+				RootObject->VisitLeafObjects(
+					[&NumObjects](const FImplicitObject* Object, const FRigidTransform3& ParentTransform, const int32 RootObjectIndex, const int32 ObjectIndex, const int32 LeafObjectIndex)
+					{ 
+						++NumObjects;
+					});
+			}
+			return NumObjects;
+		}
+
+		FImplicitBVH::FObjects FImplicitBVH::CollectLeafObjects(const TArrayView<const TUniquePtr<FImplicitObject>>& InRootObjects)
+		{
+			// We visit the hierarchy once to ensure we can create a tight-fitting array of leaf objects (the array growth
+			// policy will over-allocate if we don't size exactly)
+			TArray<FImplicitBVHObject> Objects;
+			Objects.Reserve(CountLeafObjects(InRootObjects));
+
+			for (int32 RootObjectIndex = 0; RootObjectIndex < InRootObjects.Num(); ++RootObjectIndex)
+			{
+				const TUniquePtr<FImplicitObject>& RootObject = InRootObjects[RootObjectIndex];
+
+				RootObject->VisitLeafObjects(
+					[RootObjectIndex, &Objects](const FImplicitObject* Object, const FRigidTransform3& ParentTransform, const int32 UnusedRootObjectIndex, const int32 ObjectIndex, const int32 UnusedLeafObjectIndex)
+					{
+						// @todo(chaos): clean this up (SetFromRawLowLevel). We know all the objects we visit are children of a UniquePtr because we own it
+						TSerializablePtr<FImplicitObject> SerializableObject;
+						SerializableObject.SetFromRawLowLevel(Object);
+
+						Objects.Emplace(
+							SerializableObject,
+							ParentTransform.GetTranslation(),
+							ParentTransform.GetRotation(),
+							Object->CalculateTransformedBounds(ParentTransform),
+							RootObjectIndex,
+							ObjectIndex);
+					});
+			}
+
+			return Objects;
+		}
+
+		TUniquePtr<FImplicitBVH> FImplicitBVH::MakeEmpty()
+		{
+			return TUniquePtr<FImplicitBVH>(new FImplicitBVH());
+		}
+
+		TUniquePtr<FImplicitBVH> FImplicitBVH::TryMake(const TArrayView<const TUniquePtr<FImplicitObject>>& InRootObjects, const int32 InMinObjects, const int32 InMaxBVHDepth)
+		{
+			TArray<FImplicitBVHObject> Objects = CollectLeafObjects(InRootObjects);
+			if (Objects.Num() > InMinObjects)
+			{
+				return TUniquePtr<FImplicitBVH>(new FImplicitBVH(MoveTemp(Objects), InMaxBVHDepth));
+			}
+			return TUniquePtr<FImplicitBVH>();
+		}
+
+		///////////////////////////////////////////////////////////////////////////////////////////
+		///////////////////////////////////////////////////////////////////////////////////////////
+
+		FImplicitBVH::FImplicitBVH()
+		{
+		}
+
+		FImplicitBVH::FImplicitBVH(FObjects&& InObjects, const int32 InMaxBVHDepth)
+		{
+			Init(MoveTemp(InObjects), InMaxBVHDepth);
+		}
+
+		FImplicitBVH::~FImplicitBVH()
+		{
+		}
+
+		void FImplicitBVH::Init(FObjects&& InObjects, const int32 InMaxBVHDepth)
+		{
+			Objects = MoveTemp(InObjects);
+			BVH = FBVH(Objects, FMath::Max(1, InMaxBVHDepth));
+		}
+
+		///////////////////////////////////////////////////////////////////////////////////////////
+		///////////////////////////////////////////////////////////////////////////////////////////
+		// 
+		// SERIALIZATION
+		// 
+		///////////////////////////////////////////////////////////////////////////////////////////
+		///////////////////////////////////////////////////////////////////////////////////////////
+
+		FChaosArchive& operator<<(FChaosArchive& Ar, FImplicitBVH& BVH)
+		{ 
+			return BVH.Serialize(Ar);
+		}
+
+		FChaosArchive& operator<<(FChaosArchive& Ar, FImplicitBVHObject& BVHObject)
+		{ 
+			return BVHObject.Serialize(Ar);
+		}
+
+		FChaosArchive& FImplicitBVH::Serialize(FChaosArchive& Ar)
+		{
+			Ar << Objects;
+			Ar << BVH;
+			return Ar;
+		}
+
+		FChaosArchive& FImplicitBVHObject::Serialize(FChaosArchive& Ar)
+		{
+			Ar << Geometry;
+			Ar << X;
+			Ar << R;
+			Ar << RootObjectIndex;
+			return Ar;
+		}
+
+	}
+}
\ No newline at end of file
diff --git a/Engine/Source/Runtime/Experimental/Chaos/Private/Chaos/ImplicitObjectUnion.cpp b/Engine/Source/Runtime/Experimental/Chaos/Private/Chaos/ImplicitObjectUnion.cpp
--- a/Engine/Source/Runtime/Experimental/Chaos/Private/Chaos/ImplicitObjectUnion.cpp
+++ b/Engine/Source/Runtime/Experimental/Chaos/Private/Chaos/ImplicitObjectUnion.cpp
@@ -1,53 +1,38 @@
 // Copyright Epic Games, Inc. All Rights Reserved.
 #include "Chaos/ImplicitObjectUnion.h"
+#include "Chaos/ImplicitObjectBVH.h"
 #include "Chaos/BoundingVolumeHierarchy.h"
 
-namespace Chaos
-{
-struct FLargeImplicitObjectUnionData
-{
-	// @todo(chaos): FGeometryParticles is way too much data for this. Use a custom struct
-	FGeometryParticles GeomParticles;
-	TBoundingVolumeHierarchy<FGeometryParticles, TArray<int32>> Hierarchy;
-
-	FLargeImplicitObjectUnionData(const TArray<Pair<TSerializablePtr<FImplicitObject>,FRigidTransform3>>& SubObjects)
-	{
-		const int32 NumObjects = SubObjects.Num();
-		GeomParticles.Resize(NumObjects);
-		for (int32 i = 0; i < NumObjects; ++i)
-		{
-			GeomParticles.X(i) = SubObjects[i].Second.GetLocation();
-			GeomParticles.R(i) = SubObjects[i].Second.GetRotation();
-			GeomParticles.SetGeometry(i, SubObjects[i].First);
-			//check(!SubObjects[i].First->IsUnderlyingUnion());	//we don't support union of unions
-		}
+#include "UObject/FortniteMainBranchObjectVersion.h"
 
-		Hierarchy = TBoundingVolumeHierarchy<FGeometryParticles,TArray<int32>> (GeomParticles,1);
-	}
+//UE_DISABLE_OPTIMIZATION
 
-	void Serialize(FChaosArchive& Ar)
+namespace Chaos
+{
+	namespace CVars
 	{
-		Ar << GeomParticles << Hierarchy;
+		int32 ChaosUnionBVHMinShapes = 32;
+		int32 ChaosUnionBVHMaxDepth = 8;
+		bool bChaosUnionBVHEnabled = true;
+		FAutoConsoleVariableRef CVarChaosUnionBVHMinShapes(TEXT("p.Chaos.Collision.UnionBVH.NumShapes"), ChaosUnionBVHMinShapes, TEXT("If a geometry hierarchy has this many shapes, wrap it in a BVH for collision detection (negative to disable BVH)"));
+		FAutoConsoleVariableRef CVarChaosUnionBVHMaxDepth(TEXT("p.Chaos.Collision.UnionBVH.MaxDepth"), ChaosUnionBVHMaxDepth, TEXT("The allowed depth of the BVH when used to wrap a shape hiererchy"));
+		FAutoConsoleVariableRef CVarChaosUnionBVHEnabled(TEXT("p.Chaos.Collision.UnionBVH.Enabled"), bChaosUnionBVHEnabled, TEXT("Set to false to disable use of BVH during collision detection (without affecting creations and serialization)"));
 	}
 
-	FLargeImplicitObjectUnionData(){}
-
-	FLargeImplicitObjectUnionData(const FLargeImplicitObjectUnionData& Other) = delete;
-	FLargeImplicitObjectUnionData& operator=(const FLargeImplicitObjectUnionData& Other) = delete;
-};
-
-FChaosArchive& operator<<(FChaosArchive& Ar, FLargeImplicitObjectUnionData& LargeUnionData)
+FImplicitObjectUnion::FImplicitObjectUnion() 
+	: FImplicitObject(EImplicitObject::HasBoundingBox, ImplicitObjectType::Union)
+	, MLocalBoundingBox()
+	, NumLeafObjects(0)
+	, Flags()
 {
-	LargeUnionData.Serialize(Ar);
-	return Ar;
 }
 
-FImplicitObjectUnion::FImplicitObjectUnion() : FImplicitObject(EImplicitObject::HasBoundingBox, ImplicitObjectType::Union){}
-
 FImplicitObjectUnion::FImplicitObjectUnion(TArray<TUniquePtr<FImplicitObject>>&& Objects)
 	: FImplicitObject(EImplicitObject::HasBoundingBox, ImplicitObjectType::Union)
 	, MObjects(MoveTemp(Objects))
 	, MLocalBoundingBox()
+	, NumLeafObjects(0)
+	, Flags()
 {
 	ensure(MObjects.Num());
 	for (int32 i = 0; i < MObjects.Num(); ++i)
@@ -62,9 +47,21 @@ FImplicitObjectUnion::FImplicitObjectUnion(TArray<TUniquePtr<FImplicitObject>>&&
 		}
 	}
 
-	CacheAllImplicitObjects();
+	SetNumLeafObjects(Private::FImplicitBVH::CountLeafObjects(MakeArrayView(MObjects)));
+}
+
+FImplicitObjectUnion::FImplicitObjectUnion(FImplicitObjectUnion&& Other)
+	: FImplicitObject(EImplicitObject::HasBoundingBox, ImplicitObjectType::Union)
+	, MObjects(MoveTemp(Other.MObjects))
+	, MLocalBoundingBox(MoveTemp(Other.MLocalBoundingBox))
+	, BVH(MoveTemp(Other.BVH))
+	, NumLeafObjects(Other.NumLeafObjects)
+{
+	Flags.Bits = Other.Flags.Bits;
 }
 
+FImplicitObjectUnion::~FImplicitObjectUnion() = default;
+
 void FImplicitObjectUnion::Combine(TArray<TUniquePtr<FImplicitObject>>& OtherObjects)
 {
 	ensure(MObjects.Num());
@@ -75,18 +72,22 @@ void FImplicitObjectUnion::Combine(TArray<TUniquePtr<FImplicitObject>>& OtherObj
 	}
 
 	MObjects.Reserve(MObjects.Num() + OtherObjects.Num());
-	for (TUniquePtr<FImplicitObject>& ChildObject: OtherObjects)
+	for (TUniquePtr<FImplicitObject>& ChildObject : OtherObjects)
 	{
+		SetNumLeafObjects(GetNumLeafObjects() + Private::FImplicitBVH::CountLeafObjects(MakeArrayView(&ChildObject, 1)));
+
 		MObjects.Add(MoveTemp(ChildObject));
 	}
 
-	CacheAllImplicitObjects();
+	RebuildBVH();
 }
 
 void FImplicitObjectUnion::RemoveAt(int32 RemoveIndex)
 {
 	if (RemoveIndex < MObjects.Num())
 	{
+		SetNumLeafObjects(GetNumLeafObjects() - Private::FImplicitBVH::CountLeafObjects(MakeArrayView(&MObjects[RemoveIndex], 1)));
+
 		MObjects[RemoveIndex].Reset(nullptr);
 		MObjects.RemoveAt(RemoveIndex);
 	}
@@ -104,31 +105,54 @@ void FImplicitObjectUnion::RemoveAt(int32 RemoveIndex)
 		}
 	}
 
-	CacheAllImplicitObjects();
+	RebuildBVH();
+}
 
+void FImplicitObjectUnion::SetNumLeafObjects(int32 InNumLeafObjects)
+{
+	constexpr int32 MaxNumLeafObjects = int32(TNumericLimits<decltype(NumLeafObjects)>::Max());
+	ensure(InNumLeafObjects <= MaxNumLeafObjects);
+	check(InNumLeafObjects >= 0);
+
+	NumLeafObjects = uint16(FMath::Min(InNumLeafObjects, MaxNumLeafObjects));
 }
 
+void FImplicitObjectUnion::CreateBVH()
+{
+	if (Flags.bAllowBVH)
+	{
+		const int32 MinBVHShapes = CVars::ChaosUnionBVHMinShapes;
+		const int32 MaxBVHDepth = CVars::ChaosUnionBVHMaxDepth;
+		BVH = Private::FImplicitBVH::TryMake(MakeArrayView(MObjects), MinBVHShapes, MaxBVHDepth);
+		Flags.bHasBVH = BVH.IsValid();
+	}
+}
 
-FImplicitObjectUnion::FImplicitObjectUnion(FImplicitObjectUnion&& Other)
-	: FImplicitObject(EImplicitObject::HasBoundingBox, ImplicitObjectType::Union)
-	, MObjects(MoveTemp(Other.MObjects))
-	, MLocalBoundingBox(MoveTemp(Other.MLocalBoundingBox))
-	, LargeUnionData(MoveTemp(Other.LargeUnionData))
+void FImplicitObjectUnion::DestroyBVH()
 {
+	if (BVH.IsValid())
+	{
+		BVH.Reset();
+		Flags.bHasBVH = false;
+	}
 }
 
-FImplicitObjectUnion::~FImplicitObjectUnion() = default;
+void FImplicitObjectUnion::RebuildBVH()
+{
+	DestroyBVH();
+	CreateBVH();
+}
 
-void FImplicitObjectUnion::FindAllIntersectingObjects(TArray < Pair<const FImplicitObject*,FRigidTransform3>>& Out, const FAABB3& LocalBounds) const
+void FImplicitObjectUnion::FindAllIntersectingObjects(TArray<Pair<const FImplicitObject*,FRigidTransform3>>& Out, const FAABB3& LocalBounds) const
 {
-	if (LargeUnionData)
+	if (BVH.IsValid() && CVars::bChaosUnionBVHEnabled)
 	{
-		TArray<int32> Overlaps = LargeUnionData->Hierarchy.FindAllIntersections(LocalBounds);
+		TArray<int32> Overlaps = BVH->GetBVH().FindAllIntersections(LocalBounds);
 		Out.Reserve(Out.Num() + Overlaps.Num());
 		for (int32 Idx : Overlaps)
 		{
-			const FImplicitObject* Obj = LargeUnionData->GeomParticles.Geometry(Idx).Get();
-			Out.Add(MakePair(Obj,FRigidTransform3(LargeUnionData->GeomParticles.X(Idx), LargeUnionData->GeomParticles.R(Idx))));
+			const FImplicitObject* Obj = BVH->GetGeometry(Idx);
+			Out.Add(MakePair(Obj, BVH->GetTransform(Idx)));
 		}
 	}
 	else
@@ -140,55 +164,75 @@ void FImplicitObjectUnion::FindAllIntersectingObjects(TArray < Pair<const FImpli
 	}
 }
 
-void FImplicitObjectUnion::CacheAllImplicitObjects()
+void FImplicitObjectUnion::VisitOverlappingLeafObjectsImpl(
+	const FAABB3& LocalBounds,
+	const FRigidTransform3& ObjectTransform,
+	const int32 InRootObjectIndex,
+	int32& ObjectIndex,
+	int32& LeafObjectIndex,
+	const FImplicitHierarchyVisitor& VisitorFunc) const
 {
-	TArray < Pair<TSerializablePtr<FImplicitObject>,FRigidTransform3>> SubObjects;
-	AccumulateAllSerializableImplicitObjectsHelper(SubObjects,FRigidTransform3::Identity);
-	//build hierarchy
+	if (BVH.IsValid() && CVars::bChaosUnionBVHEnabled)
+	{
+		// Visit children
+		BVH->VisitAllIntersections(LocalBounds,
+			[this, &ObjectIndex, &LeafObjectIndex, &ObjectTransform, &VisitorFunc](const int32 BVHObjectIndex)
+			{
+				VisitorFunc(BVH->GetGeometry(BVHObjectIndex), BVH->GetTransform(BVHObjectIndex) * ObjectTransform, BVH->GetRootObjectIndex(BVHObjectIndex), BVH->GetObjectIndex(BVHObjectIndex), BVHObjectIndex);
+			});
+	}
+	else
 	{
-		const int32 NumObjects = SubObjects.Num();
-		constexpr int32 MinSubObjectsToCache = 32;	//todo(make this tunable?)
-		if (NumObjects > MinSubObjectsToCache)
+		// Skip self
+		++ObjectIndex;
+
+		for (int32 BVHObjectIndex = 0; BVHObjectIndex < MObjects.Num(); ++BVHObjectIndex)
 		{
-			LargeUnionData = MakeUnique<FLargeImplicitObjectUnionData>(SubObjects);
+			// If we are the root our object index is the root index, otherwise just pass on the value we were given (from the actual root)
+			const int32 RootObjectIndex = (InRootObjectIndex != INDEX_NONE) ? InRootObjectIndex : BVHObjectIndex;
+
+			MObjects[BVHObjectIndex]->VisitOverlappingLeafObjectsImpl(LocalBounds, ObjectTransform, RootObjectIndex, ObjectIndex, LeafObjectIndex, VisitorFunc);
 		}
 	}
 }
 
-void FImplicitObjectUnion::VisitOverlappingObjectsImpl(const FAABB3& LocalBounds, const FRigidTransform3& ParentTM, int32& VisitId, const TFunctionRef<void(const FImplicitObject*, const FRigidTransform3&, const int32)>& VisitorFunc) const
+void FImplicitObjectUnion::VisitLeafObjectsImpl(
+	const FRigidTransform3& ObjectTransform,
+	const int32 InRootObjectIndex,
+	int32& ObjectIndex,
+	int32& LeafObjectIndex,
+	const FImplicitHierarchyVisitor& VisitorFunc) const
 {
-	if (LargeUnionData.IsValid())
-	{
-		// @todo(chaos): can a child be in multiple leafs? If not we can ditch this visited array, but
-		// see TBoundingVolumeHierarchy::FindAllIntersectionsHelper which implies that dupes are possible
-		TArray<bool> LeafVisited;
-		LeafVisited.AddZeroed(LargeUnionData->GeomParticles.Size());
+	// Skip self
+	++ObjectIndex;
 
-		// A visitor for the bounding volume.
-		// NOTE: The leaf data is an array of indices into the Hierarchy's flattened list of implicit objects
-		const auto& LeafVisitor = [this, &LeafVisited, &VisitorFunc](const TArray<int32>& Indices)
-		{
-			for (const int32 Index : Indices)
-			{
-				if (!LeafVisited[Index])
-				{
-					LeafVisited[Index] = true;
-					const FImplicitObject* Implicit = LargeUnionData->GeomParticles.Geometry(Index).Get();
-					const FRigidTransform3 Transform = FRigidTransform3(LargeUnionData->GeomParticles.X(Index), LargeUnionData->GeomParticles.R(Index));
-					VisitorFunc(Implicit, Transform, Index);
-				}
-			}
-		};
+	for (int32 BVHObjectIndex = 0; BVHObjectIndex < MObjects.Num(); ++BVHObjectIndex)
+	{
+		// If we are the root our object index is the root index, otherwise just pass on the value we were given (from the actual root)
+		const int32 RootObjectIndex = (InRootObjectIndex != INDEX_NONE) ? InRootObjectIndex : BVHObjectIndex;
 
-		LargeUnionData->Hierarchy.VisitAllIntersections(LocalBounds, LeafVisitor);
+		MObjects[BVHObjectIndex]->VisitLeafObjectsImpl(ObjectTransform, RootObjectIndex, ObjectIndex, LeafObjectIndex, VisitorFunc);
 	}
-	else
+}
+
+void FImplicitObjectUnion::VisitObjectsImpl(
+	const FRigidTransform3& ObjectTransform,
+	const int32 InRootObjectIndex,
+	int32& ObjectIndex,
+	int32& LeafObjectIndex,
+	const FImplicitHierarchyVisitor& VisitorFunc) const
+{
+	// Visit self
+	VisitorFunc(this, ObjectTransform, InRootObjectIndex, ObjectIndex, INDEX_NONE);
+	++ObjectIndex;
+
+	// Visit Children
+	for (int32 BVHObjectIndex = 0; BVHObjectIndex < MObjects.Num(); ++BVHObjectIndex)
 	{
-		for (int32 ObjectIndex = 0; ObjectIndex < MObjects.Num(); ++ObjectIndex)
-		{
-			MObjects[ObjectIndex]->VisitOverlappingObjectsImpl(LocalBounds, ParentTM, VisitId, VisitorFunc);
-			++VisitId;
-		}
+		// If we are the root our object index is the root index, otherwise just pass on the value we were given (from the actual root)
+		const int32 RootObjectIndex = (InRootObjectIndex != INDEX_NONE) ? InRootObjectIndex : BVHObjectIndex;
+
+		MObjects[BVHObjectIndex]->VisitObjectsImpl(ObjectTransform, RootObjectIndex, ObjectIndex, LeafObjectIndex, VisitorFunc);
 	}
 }
 
@@ -196,7 +240,7 @@ TUniquePtr<FImplicitObject> FImplicitObjectUnion::Copy() const
 {
 	TArray<TUniquePtr<FImplicitObject>> CopyOfObjects;
 	CopyOfObjects.Reserve(MObjects.Num());
-	for (const TUniquePtr<FImplicitObject>& Object: MObjects)
+	for (const TUniquePtr<FImplicitObject>& Object : MObjects)
 	{
 		CopyOfObjects.Emplace(Object->Copy());
 	}
@@ -207,7 +251,7 @@ TUniquePtr<FImplicitObject> FImplicitObjectUnion::CopyWithScale(const FVec3& Sca
 {
 	TArray<TUniquePtr<FImplicitObject>> CopyOfObjects;
 	CopyOfObjects.Reserve(MObjects.Num());
-	for (const TUniquePtr<FImplicitObject>& Object: MObjects)
+	for (const TUniquePtr<FImplicitObject>& Object : MObjects)
 	{
 		CopyOfObjects.Emplace(Object->CopyWithScale(Scale));
 	}
@@ -218,7 +262,7 @@ TUniquePtr<FImplicitObject> FImplicitObjectUnion::DeepCopy() const
 {
 	TArray<TUniquePtr<FImplicitObject>> CopyOfObjects;
 	CopyOfObjects.Reserve(MObjects.Num());
-	for (const TUniquePtr<FImplicitObject>& Object: MObjects)
+	for (const TUniquePtr<FImplicitObject>& Object : MObjects)
 	{
 		CopyOfObjects.Emplace(Object->DeepCopy());
 	}
@@ -229,7 +273,7 @@ TUniquePtr<FImplicitObject> FImplicitObjectUnion::DeepCopyWithScale(const FVec3&
 {
 	TArray<TUniquePtr<FImplicitObject>> CopyOfObjects;
 	CopyOfObjects.Reserve(MObjects.Num());
-	for (const TUniquePtr<FImplicitObject>& Object: MObjects)
+	for (const TUniquePtr<FImplicitObject>& Object : MObjects)
 	{
 		CopyOfObjects.Emplace(Object->DeepCopyWithScale(Scale));
 	}
@@ -238,14 +282,15 @@ TUniquePtr<FImplicitObject> FImplicitObjectUnion::DeepCopyWithScale(const FVec3&
 
 void FImplicitObjectUnion::ForEachObject(TFunctionRef<bool(const FImplicitObject&, const FRigidTransform3&)> Lambda) const
 {
-	if (LargeUnionData)
+	// @todo(chaos): this implementation is strange. If we have as BVH we will visit all children in the hierarchy, but if not
+	// we only visit our immediate children, and not their children. It should probably just ignore the BVH?
+	if (BVH.IsValid())
 	{
-		for (int32 Index = 0; Index < static_cast<int32>(LargeUnionData->GeomParticles.Size()); ++Index)
+		for (int32 Index = 0; Index < BVH->NumObjects(); ++Index)
 		{
-			if (TSerializablePtr<FImplicitObject> SubObject = LargeUnionData->GeomParticles.Geometry(Index))
+			if (const FImplicitObject* SubObject = BVH->GetGeometry(Index))
 			{
-				const FRigidTransform3 SubTransform{ LargeUnionData->GeomParticles.X(Index), LargeUnionData->GeomParticles.R(Index) };
-				if (Lambda(*SubObject, SubTransform))
+				if (Lambda(*SubObject, BVH->GetTransform(Index)))
 				{
 					break;
 				}
@@ -266,36 +311,72 @@ void FImplicitObjectUnion::ForEachObject(TFunctionRef<bool(const FImplicitObject
 		}
 	}
 }
-	
+
 void FImplicitObjectUnion::Serialize(FChaosArchive& Ar)
 {
 	Ar.UsingCustomVersion(FExternalPhysicsCustomObjectVersion::GUID);
+	Ar.UsingCustomVersion(FFortniteMainBranchObjectVersion::GUID);
+
 	FChaosArchiveScopedMemory ScopedMemory(Ar, GetTypeName(), false);
 	FImplicitObject::SerializeImp(Ar);
 	Ar << MObjects;
-	TBox<FReal,3>::SerializeAsAABB(Ar, MLocalBoundingBox);
+	TBox<FReal, 3>::SerializeAsAABB(Ar, MLocalBoundingBox);
 
-	bool bHierarchyBuilt = LargeUnionData.Get() != nullptr;
-	if(Ar.CustomVer(FExternalPhysicsCustomObjectVersion::GUID) < FExternalPhysicsCustomObjectVersion::UnionObjectsCanAvoidHierarchy)
+	bool bHierarchyBuilt = BVH.IsValid();
+	if (Ar.CustomVer(FExternalPhysicsCustomObjectVersion::GUID) < FExternalPhysicsCustomObjectVersion::UnionObjectsCanAvoidHierarchy)
 	{
-		LargeUnionData = MakeUnique<FLargeImplicitObjectUnionData>();
-		Ar << *LargeUnionData;
+		LegacySerializeBVH(Ar);
 		Ar << bHierarchyBuilt;
 	}
-	else
+	else if (Ar.CustomVer(FFortniteMainBranchObjectVersion::GUID) < FFortniteMainBranchObjectVersion::ChaosImplicitObjectUnionBVHRefactor)
 	{
 		Ar << bHierarchyBuilt;
-		if(bHierarchyBuilt)
+		if (bHierarchyBuilt)
+		{
+			LegacySerializeBVH(Ar);
+		}
+	}
+	else
+	{
+		Ar << Flags.Bits;
+		Ar << NumLeafObjects;
+		if (Flags.bHasBVH)
 		{
-			if(Ar.IsLoading())
+			if (Ar.IsLoading())
 			{
-				LargeUnionData = MakeUnique<FLargeImplicitObjectUnionData>();
+				BVH = Private::FImplicitBVH::MakeEmpty();
 			}
-			Ar << *LargeUnionData;
+			Ar << *BVH;
 		}
 	}
 }
 
+void FImplicitObjectUnion::LegacySerializeBVH(FChaosArchive& Ar)
+{
+	// We should only ever be loading old data. never saving it
+	check(Ar.IsLoading());
+
+	// The old data structure used FGeometryParticles which contains a lot of data we don't need
+	struct FLargeImplicitObjectUnionData
+	{
+		FGeometryParticles GeomParticles;
+		TBoundingVolumeHierarchy<FGeometryParticles, TArray<int32>> Hierarchy;
+	};
+
+	// Load the old data structure and chuck it away
+	{
+		FLargeImplicitObjectUnionData LegacyData;
+		Ar << LegacyData.GeomParticles << LegacyData.Hierarchy;
+	}
+
+	// Count the objects in the hierarchy
+	SetNumLeafObjects(Private::FImplicitBVH::CountLeafObjects(MakeArrayView(MObjects)));
+
+	// Only the root Union should allow BVH, but we don't know which that is at this stage
+	// so just revert to the original behaviour of every Union potentially having a BVH
+	Flags.bAllowBVH = true;
+	RebuildBVH();
+}
 
 FImplicitObjectUnionClustered::FImplicitObjectUnionClustered()
 	: FImplicitObjectUnion()
@@ -341,15 +422,15 @@ FImplicitObjectUnionClustered::FImplicitObjectUnionClustered(FImplicitObjectUnio
 
 void FImplicitObjectUnionClustered::FindAllIntersectingClusteredObjects(TArray<FLargeUnionClusteredImplicitInfo>& Out, const FAABB3& LocalBounds) const
 {
-	if (LargeUnionData)
+	if (BVH.IsValid())
 	{
-		TArray<int32> Overlaps = LargeUnionData->Hierarchy.FindAllIntersections(LocalBounds);
+		TArray<int32> Overlaps = BVH->GetBVH().FindAllIntersections(LocalBounds);
 		Out.Reserve(Out.Num() + Overlaps.Num());
 		for (int32 Idx : Overlaps)
 		{
-			const FImplicitObject* Obj = LargeUnionData->GeomParticles.Geometry(Idx).Get();
+			const FImplicitObject* Obj = BVH->GetGeometry(Idx);
 			const FBVHParticles* Simplicial = MOriginalParticleLookupHack.IsValidIndex(Idx) ? MOriginalParticleLookupHack[Idx]->CollisionParticles().Get() : nullptr;
-			Out.Add(FLargeUnionClusteredImplicitInfo(Obj,FRigidTransform3(LargeUnionData->GeomParticles.X(Idx), LargeUnionData->GeomParticles.R(Idx)), Simplicial));
+			Out.Add(FLargeUnionClusteredImplicitInfo(Obj, BVH->GetTransform(Idx), Simplicial));
 		}
 	}
 	else
@@ -380,9 +461,9 @@ TArray<FPBDRigidParticleHandle*>
 FImplicitObjectUnionClustered::FindAllIntersectingChildren(const FAABB3& LocalBounds) const
 {
 	TArray<FPBDRigidParticleHandle*> IntersectingChildren;
-	if (LargeUnionData) //todo: make this work when hierarchy is not built
+	if (BVH.IsValid()) //todo: make this work when hierarchy is not built
 	{
-		TArray<int32> IntersectingIndices = LargeUnionData->Hierarchy.FindAllIntersections(LocalBounds);
+		TArray<int32> IntersectingIndices = BVH->GetBVH().FindAllIntersections(LocalBounds);
 		IntersectingChildren.Reserve(IntersectingIndices.Num());
 		for (const int32 Idx : IntersectingIndices)
 		{
diff --git a/Engine/Source/Runtime/Experimental/Chaos/Public/Chaos/Collision/CollisionConstraintAllocator.h b/Engine/Source/Runtime/Experimental/Chaos/Public/Chaos/Collision/CollisionConstraintAllocator.h
--- a/Engine/Source/Runtime/Experimental/Chaos/Public/Chaos/Collision/CollisionConstraintAllocator.h
+++ b/Engine/Source/Runtime/Experimental/Chaos/Public/Chaos/Collision/CollisionConstraintAllocator.h
@@ -664,6 +664,42 @@ namespace Chaos
 		return ECollisionVisitorResult::Continue;
 	}
 
+	template<typename TLambda>
+	inline ECollisionVisitorResult FParticlePairCollisionDetector::VisitCollisions(const TLambda& Visitor, const bool bOnlyActive)
+	{
+		for (auto& KVP : Constraints)
+		{
+			const FPBDCollisionConstraintPtr& Constraint = KVP.Value;
+
+			if (Constraint.IsValid() && (!bOnlyActive || !Constraint->GetDisabled()))
+			{
+				if (Visitor(*Constraint) == ECollisionVisitorResult::Stop)
+				{
+					return ECollisionVisitorResult::Stop;
+				}
+			}
+		}
+		return ECollisionVisitorResult::Continue;
+	}
+
+	template<typename TLambda>
+	inline ECollisionVisitorResult FParticlePairCollisionDetector::VisitConstCollisions(const TLambda& Visitor, const bool bOnlyActive) const
+	{
+		for (auto& KVP : Constraints)
+		{
+			const FPBDCollisionConstraintPtr& Constraint = KVP.Value;
+
+			if (Constraint.IsValid() && (!bOnlyActive || !Constraint->GetDisabled()))
+			{
+				if (!bOnlyActive || (Visitor(*Constraint) == ECollisionVisitorResult::Stop))
+				{
+					return ECollisionVisitorResult::Stop;
+				}
+			}
+		}
+		return ECollisionVisitorResult::Continue;
+	}
+
 	template<typename TLambda>
 	inline ECollisionVisitorResult FParticlePairMidPhase::VisitCollisions(const TLambda& Visitor, const bool bOnlyActive)
 	{
@@ -686,6 +722,14 @@ namespace Chaos
 			}
 		}
 
+		if (ParticlePairDetector.IsValid())
+		{
+			if (ParticlePairDetector->VisitCollisions(Visitor, bOnlyActive) == ECollisionVisitorResult::Stop)
+			{
+				return ECollisionVisitorResult::Stop;
+			}
+		}
+
 		return ECollisionVisitorResult::Continue;
 	}
 
@@ -712,6 +756,14 @@ namespace Chaos
 			}
 		}
 
+		if (ParticlePairDetector.IsValid())
+		{
+			if (ParticlePairDetector->VisitConstCollisions(Visitor, bOnlyActive) == ECollisionVisitorResult::Stop)
+			{
+				return ECollisionVisitorResult::Stop;
+			}
+		}
+
 		return ECollisionVisitorResult::Continue;
 	}
 
diff --git a/Engine/Source/Runtime/Experimental/Chaos/Public/Chaos/Collision/CollisionContext.h b/Engine/Source/Runtime/Experimental/Chaos/Public/Chaos/Collision/CollisionContext.h
--- a/Engine/Source/Runtime/Experimental/Chaos/Public/Chaos/Collision/CollisionContext.h
+++ b/Engine/Source/Runtime/Experimental/Chaos/Public/Chaos/Collision/CollisionContext.h
@@ -10,6 +10,7 @@ namespace Chaos
 		class FCollisionContextAllocator;
 	}
 	class FMultiShapePairCollisionDetector;
+	class FParticlePairCollisionDetector;
 
 	class FCollisionDetectorSettings
 	{
@@ -64,6 +65,7 @@ namespace Chaos
 			: Settings(nullptr)
 			, Allocator(nullptr)
 			, MultiShapeCollisionDetector(nullptr)
+			, ParticlePairCollisionDetector(nullptr)
 		{
 		}
 
@@ -72,6 +74,7 @@ namespace Chaos
 			Settings = nullptr;
 			Allocator = nullptr;
 			MultiShapeCollisionDetector = nullptr;
+			ParticlePairCollisionDetector = nullptr;
 		}
 
 		const FCollisionDetectorSettings& GetSettings() const { return *Settings; }
@@ -80,10 +83,6 @@ namespace Chaos
 		Private::FCollisionContextAllocator* GetAllocator() const { return Allocator; }
 		void SetAllocator(Private::FCollisionContextAllocator* InAllocator) { Allocator = InAllocator; }
 
-		FMultiShapePairCollisionDetector* GetMultiShapePairCollisionDetector() const { return MultiShapeCollisionDetector; }
-		void SetMultiShapePairCollisionDetector(FMultiShapePairCollisionDetector* InDetector) { MultiShapeCollisionDetector = InDetector; }
-
-
 		const FCollisionDetectorSettings* Settings;
 
 		Private::FCollisionContextAllocator* Allocator;
@@ -92,5 +91,6 @@ namespace Chaos
 		// into the Particle's ShapesArray. Currently this is only Clusters.
 		// @todo(chaos): remove thsi from here and make it a parameter on ConstructCollisions and all inner functions.
 		FMultiShapePairCollisionDetector* MultiShapeCollisionDetector;
+		FParticlePairCollisionDetector* ParticlePairCollisionDetector;
 	};
 }
\ No newline at end of file
diff --git a/Engine/Source/Runtime/Experimental/Chaos/Public/Chaos/Collision/PBDCollisionConstraint.h b/Engine/Source/Runtime/Experimental/Chaos/Public/Chaos/Collision/PBDCollisionConstraint.h
--- a/Engine/Source/Runtime/Experimental/Chaos/Public/Chaos/Collision/PBDCollisionConstraint.h
+++ b/Engine/Source/Runtime/Experimental/Chaos/Public/Chaos/Collision/PBDCollisionConstraint.h
@@ -168,6 +168,7 @@ namespace Chaos
 		friend class Private::FCollisionConstraintAllocator;
 		friend class Private::FCollisionContextAllocator;
 		friend class FMultiShapePairCollisionDetector;
+		friend class FParticlePairCollisionDetector;
 		friend class FParticlePairMidPhase;
 		friend class FPBDCollisionConstraints;
 		friend class FSingleShapePairCollisionDetector;
diff --git a/Engine/Source/Runtime/Experimental/Chaos/Public/Chaos/Collision/ParticlePairMidPhase.h b/Engine/Source/Runtime/Experimental/Chaos/Public/Chaos/Collision/ParticlePairMidPhase.h
--- a/Engine/Source/Runtime/Experimental/Chaos/Public/Chaos/Collision/ParticlePairMidPhase.h
+++ b/Engine/Source/Runtime/Experimental/Chaos/Public/Chaos/Collision/ParticlePairMidPhase.h
@@ -343,6 +343,133 @@ namespace Chaos
 		const FPerShapeData* Shape1;
 	};
 
+	/**
+	 * 
+	 */
+	class FParticlePairCollisionDetector
+	{
+	public:
+		FParticlePairCollisionDetector(
+			FGeometryParticleHandle* InParticle0,
+			FGeometryParticleHandle* InParticle1,
+			FParticlePairMidPhase& MidPhase);
+		FParticlePairCollisionDetector(FParticlePairCollisionDetector&& R);
+		FParticlePairCollisionDetector(const FParticlePairCollisionDetector& R) = delete;
+		FParticlePairCollisionDetector& operator=(const FParticlePairCollisionDetector& R) = delete;
+		~FParticlePairCollisionDetector();
+
+		/**
+		 * @brief Perform a bounds check and run the narrow phase if necessary
+		 * @return The number of collisions constraints that were activated
+		*/
+		int32 GenerateCollisions(
+			const FReal CullDistance,
+			const FReal Dt,
+			const FCollisionContext& Context);
+
+		int32 GenerateCollisionsCCD(
+			const bool bEnableCCDSweep,
+			const FReal CullDistance,
+			const FReal Dt,
+			const FCollisionContext& Context);
+
+		/**
+		 * @brief Callback from the narrow phase to create a collision constraint for this particle pair.
+		 * We should never be asked for a collision for a different particle pair, but the
+		 * implicit objects may be children of the root shape.
+		*/
+		FPBDCollisionConstraint* FindOrCreateConstraint(
+			FGeometryParticleHandle* InParticle0,
+			const FImplicitObject* InImplicit0,
+			const int32 InImplicitId0,
+			const FShapeInstance* InShape0,
+			const FBVHParticles* InBVHParticles0,
+			const FRigidTransform3& InShapeRelativeTransform0,
+			FGeometryParticleHandle* InParticle1,
+			const FImplicitObject* Implicit1,
+			const int32 InImplicitId1,
+			const FShapeInstance* InShape1,
+			const FBVHParticles* InBVHParticles1,
+			const FRigidTransform3& InhapeRelativeTransform1,
+			const FReal CullDistance,
+			const EContactShapesType ShapePairType,
+			const bool bUseManifold,
+			const bool bEnableSweep,
+			const FCollisionContext& Context);
+
+		/**
+		 * @brief Reactivate the constraint
+		 * @parame SleepEpoch The tick on which the particle went to sleep.
+		 * Only constraints that were active when the particle went to sleep should be reactivated.
+		*/
+		void WakeCollisions(const int32 SleepEpoch, const int32 CurrentEpoch);
+
+		/**
+		 * @brief Visit all the collisions
+		 * @param Visitor functor with signature ECollisionVisitorResult(FPBDCollisionConstaint& Constraint)
+		*/
+		template<typename TLambda>
+		ECollisionVisitorResult VisitCollisions(const TLambda& Visitor, const bool bOnlyActive = true);
+
+		/**
+		 * @brief Visit all the collisions
+		 * @param Visitor functor with signature ECollisionVisitorResult(const FPBDCollisionConstaint& Constraint)
+		*/
+		template<typename TLambda>
+		ECollisionVisitorResult VisitConstCollisions(const TLambda& Visitor, const bool bOnlyActive = true) const;
+
+	private:
+		int32 GenerateCollisionsImpl(
+			const bool bEnableCCDSweep,
+			const FReal CullDistance,
+			const FReal Dt,
+			const FCollisionContext& Context);
+
+		FPBDCollisionConstraint* FindConstraint(const FCollisionParticlePairConstraintKey& Key);
+
+		FPBDCollisionConstraint* CreateConstraint(
+			FGeometryParticleHandle* Particle0,
+			const FImplicitObject* Implicit0,
+			const FPerShapeData* Shape0,
+			const FBVHParticles* BVHParticles0,
+			const FRigidTransform3& ShapeRelativeTransform0,
+			FGeometryParticleHandle* Particle1,
+			const FImplicitObject* Implicit1,
+			const FPerShapeData* Shape1,
+			const FBVHParticles* BVHParticles1,
+			const FRigidTransform3& ShapeRelativeTransform1,
+			const FReal CullDistance,
+			const EContactShapesType ShapePairType,
+			const bool bInUseManifold,
+			const FCollisionParticlePairConstraintKey& Key,
+			const FCollisionContext& Context);
+
+		int32 ProcessNewConstraints(
+			const FReal CullDistance,
+			const FReal Dt,
+			const FCollisionContext& Context);
+
+		void PruneConstraints(const int32 CurrentEpoch);
+
+		void UpdateCollision(
+			FPBDCollisionConstraint* Constraint,
+			const FReal CullDistance,
+			const FReal Dt,
+			const FCollisionContext& Context);
+
+		void UpdateCollisionCCD(
+			FPBDCollisionConstraint* Constraint,
+			const FReal CullDistance,
+			const FReal Dt,
+			const FCollisionContext& Context);
+
+		FParticlePairMidPhase& MidPhase;
+		TMap<uint32, FPBDCollisionConstraintPtr> Constraints;
+		TArray<FPBDCollisionConstraint*> NewConstraints;
+		FGeometryParticleHandle* Particle0;
+		FGeometryParticleHandle* Particle1;
+	};
+
 	/**
 	 * @brief Produce collisions for a particle pair
 	 * A FParticlePairMidPhase object is created for every particle pair whose bounds overlap. It is 
@@ -587,5 +714,6 @@ namespace Chaos
 
 		TArray<FMultiShapePairCollisionDetector> MultiShapePairDetectors; // 16 bytes
 		TArray<FSingleShapePairCollisionDetector, TInlineAllocator<1>> ShapePairDetectors; //88 bytes
+		TUniquePtr<FParticlePairCollisionDetector> ParticlePairDetector;
 	};
 }
diff --git a/Engine/Source/Runtime/Experimental/Chaos/Public/Chaos/ImplicitObject.h b/Engine/Source/Runtime/Experimental/Chaos/Public/Chaos/ImplicitObject.h
--- a/Engine/Source/Runtime/Experimental/Chaos/Public/Chaos/ImplicitObject.h
+++ b/Engine/Source/Runtime/Experimental/Chaos/Public/Chaos/ImplicitObject.h
@@ -59,6 +59,17 @@ struct TImplicitObjectPtrStorage<T, d, true>
 	}
 };
 
+/**
+* A visitor for use in FImplicitObject hierarchy visiting functions.
+* @param Implicit The geometry we are currently visiting
+* @param Transform The net transform relative to the hierarchy root (the originating visit call)
+* @param RootObjectIndex The index of our ancestor in the root union. Will be INDEX_NONE if no Union at the root. Used to index into a Particle's ShapeInstances or get our ancenstor.
+* @param ObjectIndex A counter tracking the current implicit object index in the flattened hierarchy (pre-order, depth first)
+* @param LeafObjectIndex A counter tracking the current leaf index in the flattened hierarchy (pre-order, depth first). Used to differentiate between geometry when duplicated in the hierarchy. INDEX_NONE if not visiting a leaf.
+*/
+using FImplicitHierarchyVisitor = TFunctionRef<void(const FImplicitObject* Implicit, const FRigidTransform3& Transform, const int32 RootObjectIndex, const int32 ObjectIndex, const int32 LeafObjectIndex)>;
+
+
 /*
  * Base class for implicit collision geometry such as spheres, capsules, boxes, etc.
  * 
@@ -315,22 +326,85 @@ public:
 
 	virtual uint16 GetMaterialIndex(uint32 HintIndex) const { return 0; }
 
-	// Visit all the objects in the hierarchy that overlap the specified local-space bounds.
-	// The Visitor ChildId is an unique index into the hierarchy that can be used to distinguish ImplicitObjects when
-	// they are reused in multiple locations in the hierarchy.
-	void VisitOverlappingObjects(const FAABB3& LocalBounds, const TFunctionRef<void(const FImplicitObject* Implicit, const FRigidTransform3& Transform, const int32 ChildId)>& Visitor) const
+	/**
+	* Visit all the leaf objects in the hierarchy that overlap the specified local-space bounds.
+	* NOTE: Templated decorators like Instanced and Scaled cound as leafs, but object decorators like Transformed do not.
+	* @see FImplicitHierarchyVisitor for visitor notes
+	*/
+	void VisitOverlappingLeafObjects(const FAABB3& LocalBounds, const FImplicitHierarchyVisitor& Visitor) const
 	{
-		int32 VisitId = 0;
-		VisitOverlappingObjectsImpl(LocalBounds, FRigidTransform3::Identity, VisitId, Visitor);
+		int32 LeafObjectIndex = 0;
+		int32 ObjectIndex = 0;
+		const int32 RootObjectIndex = INDEX_NONE;
+		VisitOverlappingLeafObjectsImpl(LocalBounds, FRigidTransform3::Identity, RootObjectIndex, ObjectIndex, LeafObjectIndex, Visitor);
 	}
 
-	// This should not be public, but it needs to be callable by derived classes on another instance (e.g., see FImplicitObjectUnion::VisitOverlappingObjectsImpl)
-	virtual void VisitOverlappingObjectsImpl(const FAABB3& LocalBounds, const FRigidTransform3& ParentTM, int32& VisitId, const TFunctionRef<void(const FImplicitObject*, const FRigidTransform3&, const int32)>& VisitorFunc) const
+	/**
+	* Visit all the leaf objects in the hierarchy
+	* @see FImplicitHierarchyVisitor for visitor notes
+	*/
+	void VisitLeafObjects(const FImplicitHierarchyVisitor& Visitor) const
+	{
+		int32 LeafObjectIndex = 0;
+		int32 ObjectIndex = 0;
+		const int32 RootObjectIndex = INDEX_NONE;
+		VisitLeafObjectsImpl(FRigidTransform3::Identity, RootObjectIndex, ObjectIndex, LeafObjectIndex, Visitor);
+	}
+
+	/**
+	* Visit all the objects in the hierarchy, including inner nodes like Union and Transform
+	* @see FImplicitHierarchyVisitor for visitor notes
+	*/
+	void VisitObjects(const FImplicitHierarchyVisitor& Visitor) const
+	{
+		int32 LeafObjectIndex = 0;
+		int32 ObjectIndex = 0;
+		const int32 RootObjectIndex = INDEX_NONE;
+		VisitObjectsImpl(FRigidTransform3::Identity, RootObjectIndex, ObjectIndex, LeafObjectIndex, Visitor);
+	}
+
+//protected:
+	// This should not be public, but it needs to be callable by derived classes on another instance
+	virtual void VisitOverlappingLeafObjectsImpl(
+		const FAABB3& LocalBounds,
+		const FRigidTransform3& ObjectTransform,
+		const int32 RootObjectIndex,
+		int32& ObjectIndex,
+		int32& LeafObjectIndex,
+		const FImplicitHierarchyVisitor& VisitorFunc) const
 	{
 		if (!HasBoundingBox() || LocalBounds.Intersects(BoundingBox()))
 		{
-			VisitorFunc(this, ParentTM, VisitId);
+			VisitorFunc(this, ObjectTransform, RootObjectIndex, ObjectIndex, LeafObjectIndex);
 		}
+		++ObjectIndex;
+		++LeafObjectIndex;
+	}
+
+	// This should not be public, but it needs to be callable by derived classes on another instance
+	virtual void VisitLeafObjectsImpl(
+		const FRigidTransform3& ObjectTransform,
+		const int32 RootObjectIndex,
+		int32& ObjectIndex,
+		int32& LeafObjectIndex,
+		const FImplicitHierarchyVisitor& VisitorFunc) const
+	{
+		VisitorFunc(this, ObjectTransform, RootObjectIndex, ObjectIndex, LeafObjectIndex);
+		++ObjectIndex;
+		++LeafObjectIndex;
+	}
+
+	// This should not be public, but it needs to be callable by derived classes on another instance
+	virtual void VisitObjectsImpl(
+		const FRigidTransform3& ObjectTransform,
+		const int32 RootObjectIndex,
+		int32& ObjectIndex,
+		int32& LeafObjectIndex,
+		const FImplicitHierarchyVisitor& VisitorFunc) const
+	{
+		VisitorFunc(this, ObjectTransform, RootObjectIndex, ObjectIndex, LeafObjectIndex);
+		++ObjectIndex;
+		++LeafObjectIndex;
 	}
 
 protected:
diff --git a/Engine/Source/Runtime/Experimental/Chaos/Public/Chaos/ImplicitObjectBVH.h b/Engine/Source/Runtime/Experimental/Chaos/Public/Chaos/ImplicitObjectBVH.h
new file mode 100644
--- /dev/null
+++ b/Engine/Source/Runtime/Experimental/Chaos/Public/Chaos/ImplicitObjectBVH.h
@@ -0,0 +1,185 @@
+// Copyright Epic Games, Inc. All Rights Reserved.
+#pragma once
+
+#include "Chaos/Core.h"
+#include "Chaos/BoundingVolumeHierarchy.h"
+#include "Chaos/ImplicitFwd.h"
+#include "Chaos/Serializable.h"
+
+namespace Chaos
+{
+	class FChaosArchive;
+
+	namespace Private
+	{
+		class FImplicitBVH;
+		class FImplicitObjectUnion;
+
+		// An entry in the ImplicitObject BVH that holds the leaf geometry and transform
+		class FImplicitBVHObject
+		{
+		public:
+			friend class FImplicitBVH;
+
+			friend FChaosArchive& operator<<(FChaosArchive& Ar, FImplicitBVHObject& BVHObjecty);
+
+			FImplicitBVHObject();
+			FImplicitBVHObject(const TSerializablePtr<FImplicitObject>& InGeometry, const FVec3& InX, const FRotation3& InR, const FAABB3& InBounds, const int32 InRootObjectIndex, const int32 InObjectIndex);
+
+			const FImplicitObject* GetGeometry() const { return Geometry.Get(); }
+
+			const FVec3f& GetX() const { return X; }
+
+			const FRotation3f& GetR() const { return R; }
+
+			const FAABB3f& GetBounds() const { return Bounds; }
+
+			FRigidTransform3f GetTransformf() const { return FRigidTransform3f(GetX(), GetR()); }
+
+			FRigidTransform3 GetTransform() const { return FRigidTransform3(FVec3(GetX()), FRotation3(GetR())); }
+
+			// A unique index for this object in the hierarchy. E.g., if the same FImplicirtObject is referenced 
+			// multiple times in the hierarchy in a union of transformed objects, each will have a different ObjectIndex 
+			// (it is the index into the pre-order depth first traversal).
+			int32 GetObjectIndex() const { return ObjectIndex; }
+
+			// The index of our most distant ancestor. I.e., the index in the root Union. This is used to map
+			// each object to a ShapeInstance.
+			int32 GetRootObjectIndex() const { return RootObjectIndex; }
+
+		private:
+			FChaosArchive& Serialize(FChaosArchive& Ar);
+
+			// Transform and bounds in the space of the BVH owner (Union Implicit)
+			FRotation3f R;
+			FVec3f X;
+			FAABB3f Bounds;
+
+			// The leaf geometry stripped of decorators (but not Instanced or Scaled)
+			TSerializablePtr<FImplicitObject> Geometry;
+
+			// The index of our ancestor in the array of RootObjects that was provided when creating the BVH
+			int32 RootObjectIndex;
+
+			// Our index in the hierarchy. This could be used to uniquely identity copies of the same implicit in the hierarchy
+			int32 ObjectIndex;
+		};
+
+		// A Bounding Volume Hierarchy of a set of Implicit Objects
+		class FImplicitBVH
+		{
+		public:
+			using FObjects = TArray<FImplicitBVHObject>;
+			using FLeaf = TArray<int32>;
+			using FBVH = TBoundingVolumeHierarchy<FObjects, FLeaf>;
+
+			friend FChaosArchive& operator<<(FChaosArchive& Ar, FImplicitBVH& BVH);
+
+			// Amke an empty BVH (for serialization only)
+			static TUniquePtr<FImplicitBVH> MakeEmpty();
+
+			// Utility for processing the hierarchy
+			static int32 CountLeafObjects(const TArrayView<const TUniquePtr<FImplicitObject>>& InRootObjects);
+			static FObjects CollectLeafObjects(const TArrayView<const TUniquePtr<FImplicitObject>>& InRootObjects);
+
+			// Create a BVH around a set of ImplicitObjects. Usually these are the immediate child elements of an FImplcitObjectUnion
+			// TryMake will then recurse into the geometry hierachy and add all descendents to the BVH. Will return null if the 
+			// number of descendents is less that MinObjscts.
+			static TUniquePtr<FImplicitBVH> TryMake(const TArrayView<const TUniquePtr<FImplicitObject>>& InRootObjects, const int32 MinObjects, const int32 InMaxBVHDepth);
+
+			~FImplicitBVH();
+
+			int32 NumObjects() const { return Objects.Num(); }
+
+			const FImplicitBVHObject& GetObject(const int32 ObjectIndex) const { return Objects[ObjectIndex]; }
+
+			const FImplicitObject* GetGeometry(const int32 ObjectIndex) const { return Objects[ObjectIndex].GetGeometry(); }
+
+			const FVec3f& GetX(const int32 ObjectIndex) const { return Objects[ObjectIndex].GetX(); }
+
+			const FRotation3f& GetR(const int32 ObjectIndex) const { return Objects[ObjectIndex].GetR(); }
+
+			const FAABB3f& GetBounds(const int32 ObjectIndex) const { return Objects[ObjectIndex].GetBounds(); }
+
+			FRigidTransform3f GetTransformf(const int32 ObjectIndex) const { return Objects[ObjectIndex].GetTransformf(); }
+
+			FRigidTransform3 GetTransform(const int32 ObjectIndex) const { return Objects[ObjectIndex].GetTransform(); }
+
+			int32 GetRootObjectIndex(const int32 ObjectIndex) const { return Objects[ObjectIndex].GetRootObjectIndex(); }
+			int32 GetObjectIndex(const int32 ObjectIndex) const { return Objects[ObjectIndex].GetObjectIndex(); }
+
+			const FBVH& GetBVH() const { return BVH; }
+
+			// @param Visitor void<int ObjectIndex> where ObjectIndex is an index into the array of objects in the hierarchy
+			// for use with GetObject(ObjectIndex), GetGeometry(ObjectIndex), etc.
+			template<typename TVisitor>
+			void VisitAllIntersections(const FAABB3& LocalBounds, const TVisitor& Visitor) const
+			{
+				// Each objects can be in multiple leafs, so we keep track of whichones we have visited
+				TArray<bool> ObjectVisited;
+				ObjectVisited.AddZeroed(NumObjects());
+			
+				const auto& LeafVisitor = [this, &ObjectVisited, &Visitor](const TArray<int32>& Indices)
+				{
+					for (const int32 Index : Indices)
+					{
+						if (!ObjectVisited[Index])
+						{
+							ObjectVisited[Index] = true;
+							Visitor(Index);
+						}
+					}
+				};
+
+				BVH.VisitAllIntersections(LocalBounds, LeafVisitor);
+			}
+
+		private:
+			FImplicitBVH();
+			FImplicitBVH(FObjects&& InObjects, const int32 InMaxBVHDepth);
+			FChaosArchive& Serialize(FChaosArchive& Ar);
+
+			// Initialize the BVH from the specified set of children.
+			// NOTE: InChildren should be immediate children of the BVH owner, not further-removed descendents
+			void Init(FObjects&& InObjects, const int32 InMaxBVHDepth);
+
+			// A BVH leaf holds an array of indices into the Objects array
+			FObjects Objects;
+			FBVH BVH;
+		};
+
+		class FImplicitCache
+		{
+		public:
+
+		private:
+			TUniquePtr<FImplicitBVH> BVH;
+		};
+	}
+
+	///////////////////////////////////////////////////////////////////////////////////////////////
+	///////////////////////////////////////////////////////////////////////////////////////////////
+	// 
+	// TBoundingVolumeHierarchy Template API
+	// 
+	///////////////////////////////////////////////////////////////////////////////////////////////
+	///////////////////////////////////////////////////////////////////////////////////////////////
+
+	template<>
+	inline bool HasBoundingBox(const TArray<Private::FImplicitBVHObject>& Objects, const int32 ObjectIndex)
+	{
+		return Objects[ObjectIndex].GetGeometry()->HasBoundingBox();
+	}
+
+	template<class T, int d>
+	inline const TAABB<T, d> GetWorldSpaceBoundingBox(const TArray<Private::FImplicitBVHObject>& Objects, const int32 ObjectIndex, const TMap<int32, TAABB<T, d>>& WorldSpaceBoxes)
+	{
+		return TAABB<T, d>(Objects[ObjectIndex].GetBounds());
+	}
+
+	template<class T, int d>
+	void ComputeAllWorldSpaceBoundingBoxes(const TArray<Private::FImplicitBVHObject>& Objects, const TArray<int32>& AllObjects, const bool bUseVelocity, const T Dt, TMap<int32, TAABB<T, d>>& WorldSpaceBoxes)
+	{
+	}
+
+}
\ No newline at end of file
diff --git a/Engine/Source/Runtime/Experimental/Chaos/Public/Chaos/ImplicitObjectTransformed.h b/Engine/Source/Runtime/Experimental/Chaos/Public/Chaos/ImplicitObjectTransformed.h
--- a/Engine/Source/Runtime/Experimental/Chaos/Public/Chaos/ImplicitObjectTransformed.h
+++ b/Engine/Source/Runtime/Experimental/Chaos/Public/Chaos/ImplicitObjectTransformed.h
@@ -359,10 +359,49 @@ public:
 	}
 
 protected:
-	virtual void VisitOverlappingObjectsImpl(const FAABB3& InLocalBounds, const FRigidTransform3& ParentTM, int32& VisitId, const TFunctionRef<void(const FImplicitObject*, const FRigidTransform3&, const int32)>& VisitorFunc) const
-	{
+	virtual void VisitOverlappingLeafObjectsImpl(
+		const FAABB3& InLocalBounds,
+		const FRigidTransform3& ObjectTransform,
+		const int32 RootObjectIndex,
+		int32& ObjectIndex,
+		int32& LeafObjectIndex,
+		const FImplicitHierarchyVisitor& VisitorFunc) const override
+	{
+		// Skip self
+		++ObjectIndex;
+
+		// Visit child
 		const FAABB3 LocalBounds = InLocalBounds.InverseTransformedAABB(MTransform);
-		MObject->VisitOverlappingObjectsImpl(LocalBounds, MTransform * ParentTM, VisitId, VisitorFunc);
+		MObject->VisitOverlappingLeafObjectsImpl(LocalBounds, MTransform * ObjectTransform, RootObjectIndex, ObjectIndex, LeafObjectIndex, VisitorFunc);
+	}
+
+	virtual void VisitLeafObjectsImpl(
+		const FRigidTransform3& ObjectTransform,
+		const int32 RootObjectIndex,
+		int32& ObjectIndex,
+		int32& LeafObjectIndex,
+		const FImplicitHierarchyVisitor& VisitorFunc) const override
+	{
+		// Skip self
+		++ObjectIndex;
+
+		// Visit child
+		MObject->VisitLeafObjectsImpl(MTransform * ObjectTransform, RootObjectIndex, ObjectIndex, LeafObjectIndex, VisitorFunc);
+	}
+
+	virtual void VisitObjectsImpl(
+		const FRigidTransform3& ObjectTransform,
+		const int32 RootObjectIndex,
+		int32& ObjectIndex,
+		int32& LeafObjectIndex,
+		const FImplicitHierarchyVisitor& VisitorFunc) const override
+	{
+		// Visit self
+		VisitorFunc(this, ObjectTransform, RootObjectIndex, ObjectIndex, INDEX_NONE);
+		++ObjectIndex;
+
+		// Visit child
+		MObject->VisitObjectsImpl(MTransform * ObjectTransform, RootObjectIndex, ObjectIndex, LeafObjectIndex, VisitorFunc);
 	}
 
 private:
@@ -407,4 +446,6 @@ namespace Utilities
 template <typename T, int d>
 using TImplicitObjectTransformedNonSerializable = TImplicitObjectTransformed<T, d, false>;
 
+using FImplicitObjectTransformed = TImplicitObjectTransformed<FReal, 3>;
+
 }
diff --git a/Engine/Source/Runtime/Experimental/Chaos/Public/Chaos/ImplicitObjectUnion.h b/Engine/Source/Runtime/Experimental/Chaos/Public/Chaos/ImplicitObjectUnion.h
--- a/Engine/Source/Runtime/Experimental/Chaos/Public/Chaos/ImplicitObjectUnion.h
+++ b/Engine/Source/Runtime/Experimental/Chaos/Public/Chaos/ImplicitObjectUnion.h
@@ -12,7 +12,15 @@
 namespace Chaos
 {
 
-struct FLargeImplicitObjectUnionData;
+namespace Private
+{
+	class FImplicitBVH;
+}
+
+namespace CVars
+{
+	extern bool bChaosUnionBVHEnabled;
+}
 
 class CHAOS_API FImplicitObjectUnion : public FImplicitObject
 {
@@ -21,9 +29,6 @@ class CHAOS_API FImplicitObjectUnion : public FImplicitObject
 	using FImplicitObject::GetTypeName;
 
 	FImplicitObjectUnion(TArray<TUniquePtr<FImplicitObject>>&& Objects);
-	void Combine(TArray<TUniquePtr<FImplicitObject>>& Objects);
-	void RemoveAt(int32 RemoveIndex);
-
 	FImplicitObjectUnion(const FImplicitObjectUnion& Other) = delete;
 	FImplicitObjectUnion(FImplicitObjectUnion&& Other);
 	virtual ~FImplicitObjectUnion();
@@ -33,6 +38,25 @@ class CHAOS_API FImplicitObjectUnion : public FImplicitObject
 		return ImplicitObjectType::Union;
 	}
 
+	void Combine(TArray<TUniquePtr<FImplicitObject>>& Objects);
+	void RemoveAt(int32 RemoveIndex);
+
+	// The total number of leaf objects in the hierarchy
+	int32 GetNumLeafObjects() const
+	{
+		return int32(NumLeafObjects);
+	}
+
+	// Enable BVH suport for this Union. This should only be done for the root Union in a hierarchy
+	void SetAllowBVH(const bool bInAllowBVH)
+	{
+		if (bInAllowBVH != Flags.bAllowBVH)
+		{
+			Flags.bAllowBVH = bInAllowBVH;
+			RebuildBVH();
+		}
+	}
+
 	virtual TUniquePtr<FImplicitObject> Copy() const;
 	virtual TUniquePtr<FImplicitObject> CopyWithScale(const FVec3& Scale) const override;
 	virtual TUniquePtr<FImplicitObject> DeepCopy() const;
@@ -93,7 +117,8 @@ class CHAOS_API FImplicitObjectUnion : public FImplicitObject
 	}
 
 	virtual void FindAllIntersectingObjects(TArray < Pair<const FImplicitObject*, FRigidTransform3>>& Out, const FAABB3& LocalBounds) const;
-	virtual void CacheAllImplicitObjects();
+	
+	virtual void CacheAllImplicitObjects() { RebuildBVH(); }
 
 	virtual bool Raycast(const FVec3& StartPoint, const FVec3& Dir, const FReal Length, const FReal Thickness, FReal& OutTime, FVec3& OutPosition, FVec3& OutNormal, int32& OutFaceIndex) const override
 	{
@@ -178,7 +203,9 @@ class CHAOS_API FImplicitObjectUnion : public FImplicitObject
 			}
 		}
 
-		return new FImplicitObjectUnion(MoveTemp(NewObjects));
+		FImplicitObjectUnion* Union = new FImplicitObjectUnion(MoveTemp(NewObjects));
+		Union->SetAllowBVH(Flags.bAllowBVH);
+		return Union;
 	}
 
 #if INTEL_ISPC && !UE_BUILD_SHIPPING
@@ -214,16 +241,61 @@ protected:
 		return ClosestIntersection;
 	}
 
-	virtual void VisitOverlappingObjectsImpl(const FAABB3& LocalBounds, const FRigidTransform3& ParentTM, int32& VisitId, const TFunctionRef<void(const FImplicitObject*, const FRigidTransform3&, const int32)>& VisitorFunc) const override final;
+	virtual void VisitOverlappingLeafObjectsImpl(
+		const FAABB3& LocalBounds,
+		const FRigidTransform3& ObjectTransform,
+		const int32 RootObjectIndex,
+		int32& ObjectIndex,
+		int32& LeafObjectIndex,
+		const FImplicitHierarchyVisitor& VisitorFunc) const override final;
+
+	virtual void VisitLeafObjectsImpl(
+		const FRigidTransform3& ObjectTransform,
+		const int32 RootObjectIndex,
+		int32& ObjectIndex,
+		int32& LeafObjectIndex,
+		const FImplicitHierarchyVisitor& VisitorFunc) const override final;
+
+	virtual void VisitObjectsImpl(
+		const FRigidTransform3& ObjectTransform,
+		const int32 RootObjectIndex,
+		int32& ObjectIndex,
+		int32& LeafObjectIndex,
+		const FImplicitHierarchyVisitor& VisitorFunc) const override final;
 
   protected:
 	//needed for serialization
 	FImplicitObjectUnion();
 	friend FImplicitObject;	//needed for serialization
 
+	void SetNumLeafObjects(const int32 InNumLeafObjects);
+	void CreateBVH();
+	void DestroyBVH();
+	void RebuildBVH();
+
+	void LegacySerializeBVH(FChaosArchive& Ar);
+
+	union FFLags
+	{
+	public:
+		FFLags() : Bits(0) {}
+		struct
+		{
+			// NOTE: Flags are serialized. Ordering must be retained
+			uint8 bAllowBVH : 1;				// Are we allowed to use a BVH for this Union? (We generally only want a BVH in the root union)
+			uint8 bHasBVH : 1;					// Do we currently have a BVH? Equivalent to checking BVH pointer, but used by serialization to know whether to load a BVH
+		};
+		uint8 Bits;
+	};
+
 	TArray<TUniquePtr<FImplicitObject>> MObjects;
 	FAABB3 MLocalBoundingBox;
-	TUniquePtr<FLargeImplicitObjectUnionData> LargeUnionData;	//only needed when there are many objects
+
+	// BVH is only created when there are many objects.
+	// @todo(chaos): consider registering particles that may need BVH updated in evolution instead
+	TUniquePtr<Private::FImplicitBVH> BVH;
+	uint16 NumLeafObjects;
+	FFLags Flags;
 };
 
 struct FLargeUnionClusteredImplicitInfo
diff --git a/Engine/Source/Runtime/Experimental/ChaosCore/Public/Chaos/Core.h b/Engine/Source/Runtime/Experimental/ChaosCore/Public/Chaos/Core.h
--- a/Engine/Source/Runtime/Experimental/ChaosCore/Public/Chaos/Core.h
+++ b/Engine/Source/Runtime/Experimental/ChaosCore/Public/Chaos/Core.h
@@ -51,7 +51,7 @@ namespace Chaos
 	// NOTE: if you get a merge conflict on the GUID, you must replace it with a new GUID - do not accept the source or target
 	// or you will likely get DDC version conflicts resulting in crashes during load.
 	// Core version string for Chaos data. Any DDC builder dependent on Chaos for serialization should depend on this version
-	inline const TCHAR* const ChaosVersionGUID = TEXT("E2C28A68-EE02-4BED-9343-8659AF804349");
+	inline const TCHAR* const ChaosVersionGUID = TEXT("F960A384-9C87-4377-AE3A-1E9D861E34AB");
 
 	inline FString GetChaosVersionStringInner()
 	{
