Chaos - add option to solve Joint Drives in the velocity phase (cvar controlled, disable by default)

This resolves an issue when using velocity drives with PBD. PBD velocity drives will completely override the pushout from other constraints. E.g., a velocity drive would push an objects right through a wall because it moves the object so that the implicit velocity is the velocity target, completely undoing any pushout from collisions.

diff --git a/Engine/Source/Programs/HeadlessChaos/Private/HeadlessChaosTestJointForces.cpp b/Engine/Source/Programs/HeadlessChaos/Private/HeadlessChaosTestJointForces.cpp
--- a/Engine/Source/Programs/HeadlessChaos/Private/HeadlessChaosTestJointForces.cpp
+++ b/Engine/Source/Programs/HeadlessChaos/Private/HeadlessChaosTestJointForces.cpp
@@ -244,6 +244,11 @@ namespace ChaosTest {
 		FJointChainTest<FPBDRigidsEvolutionGBF> Test(NumSolverIterations, Gravity);
 		Test.InitChain(NumBodies, FVec3(0, 0, -1));
 
+		if (!Test.Evolution.GetJointConstraints().GetSettings().bUsePositionBasedDrives)
+		{
+			Test.Evolution.SetNumVelocityIterations(NumSolverIterations);
+		}
+
 		// Disable all limits
 		Test.JointSettings[0].LinearMotionTypes = { EJointMotionType::Free, EJointMotionType::Free, EJointMotionType::Free };
 		Test.JointSettings[0].AngularMotionTypes = { EJointMotionType::Free, EJointMotionType::Free, EJointMotionType::Free };
@@ -270,7 +275,7 @@ namespace ChaosTest {
 		const FReal M = Test.ParticleMasses[1];
 		FReal ExpectedForceZ = 0;
 		FReal DP = 0;
-		for (int32 It = 0; It < 10; ++It)
+		for (int32 It = 0; It < NumSolverIterations; ++It)
 		{
 			const FReal V = -Velocity + DP / Dt;
 			const FReal F = -Damping * V;
@@ -281,10 +286,68 @@ namespace ChaosTest {
 
 		// Check the joint forces agree
 		const FReal ForceZ = -Test.Evolution.GetJointConstraints().GetConstraintLinearImpulse(0).Z / Dt;
-		EXPECT_NEAR(ForceZ, ExpectedForceZ, 0.1);
+		EXPECT_NEAR(ForceZ, ExpectedForceZ, 1);
 	}
 
-	// Check that a handing mass on a joint drive reaches the correct extension with the correct spring force when using Force mode.
+	GTEST_TEST(JointForceTests, TestAngularDriveForceMode_Damping)
+	{
+		const int32 NumSolverIterations = 20;
+		const FReal Gravity = 0;
+		const FReal Dt = 0.01;
+		const int32 NumBodies = 2;
+
+		const FReal AngularVelocity = 3;
+		const FReal Stiffness = 0;
+		const FReal Damping = 200;
+
+		FJointChainTest<FPBDRigidsEvolutionGBF> Test(NumSolverIterations, Gravity);
+		Test.InitChain(NumBodies, FVec3(0, 0, -1));
+
+		if (!Test.Evolution.GetJointConstraints().GetSettings().bUsePositionBasedDrives)
+		{
+			Test.Evolution.SetNumVelocityIterations(NumSolverIterations);
+		}
+
+		// Disable all limits
+		Test.JointSettings[0].LinearMotionTypes = { EJointMotionType::Free, EJointMotionType::Free, EJointMotionType::Free };
+		Test.JointSettings[0].AngularMotionTypes = { EJointMotionType::Free, EJointMotionType::Free, EJointMotionType::Free };
+
+		// Set up the drive in force mode
+		Test.JointSettings[0].bAngularSLerpVelocityDriveEnabled = true;
+		Test.JointSettings[0].AngularDriveForceMode = EJointForceMode::Force;
+		Test.JointSettings[0].AngularDriveStiffness = FVec3(Stiffness, Stiffness, Stiffness);
+		Test.JointSettings[0].AngularDriveDamping = FVec3(Damping, Damping, Damping);
+
+		Test.Create();
+
+		FGenericParticleHandle P1 = Test.GetParticle(1);
+
+		// Give the particle angular velocity so that joint damping has some work to do
+		P1->SetW(FVec3(0, 0, -AngularVelocity));
+
+		// Run the sim
+		Test.Evolution.AdvanceOneTimeStep(Dt);
+		Test.Evolution.EndFrame(Dt);
+
+		// Calculate expected Trque from T = -D.W with implicit integration
+		const FReal I = FConstGenericParticleHandle(Test.GetParticle(1))->I().Z;
+		FReal ExpectedTorqueZ = 0;
+		FReal DQ = 0;
+		for (int32 It = 0; It < NumSolverIterations; ++It)
+		{
+			const FReal W = -AngularVelocity + DQ / Dt;
+			const FReal T = -Damping * W;
+			const FReal DW = ((T - ExpectedTorqueZ) / I) * Dt;
+			DQ += DW * Dt;
+			ExpectedTorqueZ = T;
+		}
+
+		// Check the joint forces agree
+		const FReal TorqueZ = -Test.Evolution.GetJointConstraints().GetConstraintAngularImpulse(0).Z / Dt;
+		EXPECT_NEAR(TorqueZ, ExpectedTorqueZ, 0.1);
+	}
+
+	// Check that a hanging mass on a joint drive reaches the correct extension with the correct spring force when using Force mode.
 	// This is just a pre-test for TestLinearDriveForceMode_MaxForce to verify that everything works in the absense of a max force
 	//
 	GTEST_TEST(JointForceTests, TestLinearDriveForceMode_MaxForcePreTest)
@@ -299,6 +362,11 @@ namespace ChaosTest {
 		FJointChainTest<FPBDRigidsEvolutionGBF> Test(NumSolverIterations, Gravity);
 		Test.InitChain(NumBodies, FVec3(0, 0, -1));
 
+		if (!Test.Evolution.GetJointConstraints().GetSettings().bUsePositionBasedDrives)
+		{
+			Test.Evolution.SetNumVelocityIterations(NumSolverIterations);
+		}
+
 		// Disable all limits
 		Test.JointSettings[0].LinearMotionTypes = { EJointMotionType::Free, EJointMotionType::Free, EJointMotionType::Free };
 		Test.JointSettings[0].AngularMotionTypes = { EJointMotionType::Free, EJointMotionType::Free, EJointMotionType::Free };
@@ -326,8 +394,8 @@ namespace ChaosTest {
 		// We should be stationary at the desired extension
 		const FReal ExpectedVZ = 0;
 		const FReal ExpectedZ = Test.ParticlePositions[1].Z - Extension;
-		EXPECT_NEAR(P1->V().Z, ExpectedVZ, 0.1);
-		EXPECT_NEAR(P1->X().Z, ExpectedZ, 0.1);
+		EXPECT_NEAR(P1->V().Z, ExpectedVZ, 1);
+		EXPECT_NEAR(P1->X().Z, ExpectedZ, 1);
 	}
 
 	// Check that the maximum drive force setting honored for linear drives.
@@ -345,6 +413,11 @@ namespace ChaosTest {
 		FJointChainTest<FPBDRigidsEvolutionGBF> Test(NumSolverIterations, Gravity);
 		Test.InitChain(NumBodies, FVec3(0, 0, -1));
 
+		if (!Test.Evolution.GetJointConstraints().GetSettings().bUsePositionBasedDrives)
+		{
+			Test.Evolution.SetNumVelocityIterations(NumSolverIterations);
+		}
+
 		// Disable all limits
 		Test.JointSettings[0].LinearMotionTypes = { EJointMotionType::Free, EJointMotionType::Free, EJointMotionType::Free };
 		Test.JointSettings[0].AngularMotionTypes = { EJointMotionType::Free, EJointMotionType::Free, EJointMotionType::Free };
@@ -391,6 +464,11 @@ namespace ChaosTest {
 		FJointChainTest<FPBDRigidsEvolutionGBF> Test(NumSolverIterations, Gravity);
 		Test.InitChain(NumBodies, FVec3(0, 0, -1));
 
+		if (!Test.Evolution.GetJointConstraints().GetSettings().bUsePositionBasedDrives)
+		{
+			Test.Evolution.SetNumVelocityIterations(NumSolverIterations);
+		}
+
 		// Disable all limits
 		Test.JointSettings[0].LinearMotionTypes = { EJointMotionType::Free, EJointMotionType::Free, EJointMotionType::Free };
 		Test.JointSettings[0].AngularMotionTypes = { EJointMotionType::Free, EJointMotionType::Free, EJointMotionType::Free };
@@ -419,8 +497,8 @@ namespace ChaosTest {
 		// We should be stationary at the desired extension
 		const FReal ExpectedVZ = 0;
 		const FReal ExpectedZ = Test.ParticlePositions[1].Z - Extension;
-		EXPECT_NEAR(P1->V().Z, ExpectedVZ, 0.1);
-		EXPECT_NEAR(P1->X().Z, ExpectedZ, 0.1);
+		EXPECT_NEAR(P1->V().Z, ExpectedVZ, 1);
+		EXPECT_NEAR(P1->X().Z, ExpectedZ, 1);
 	}
 
 
@@ -440,6 +518,11 @@ namespace ChaosTest {
 		FJointChainTest<FPBDRigidsEvolutionGBF> Test(NumSolverIterations, Gravity);
 		Test.InitChain(NumBodies, FVec3(0, 0, -1));
 
+		if (!Test.Evolution.GetJointConstraints().GetSettings().bUsePositionBasedDrives)
+		{
+			Test.Evolution.SetNumVelocityIterations(NumSolverIterations);
+		}
+
 		// Disable all limits
 		Test.JointSettings[0].LinearMotionTypes = { EJointMotionType::Free, EJointMotionType::Free, EJointMotionType::Free };
 		Test.JointSettings[0].AngularMotionTypes = { EJointMotionType::Free, EJointMotionType::Free, EJointMotionType::Free };
diff --git a/Engine/Source/Runtime/Engine/Private/PhysicsEngine/ImmediatePhysics/ImmediatePhysicsChaos/ImmediatePhysicsSimulation_Chaos.cpp b/Engine/Source/Runtime/Engine/Private/PhysicsEngine/ImmediatePhysics/ImmediatePhysicsChaos/ImmediatePhysicsSimulation_Chaos.cpp
--- a/Engine/Source/Runtime/Engine/Private/PhysicsEngine/ImmediatePhysics/ImmediatePhysicsChaos/ImmediatePhysicsSimulation_Chaos.cpp
+++ b/Engine/Source/Runtime/Engine/Private/PhysicsEngine/ImmediatePhysics/ImmediatePhysicsChaos/ImmediatePhysicsSimulation_Chaos.cpp
@@ -93,6 +93,7 @@ Chaos::FRealSingle ChaosImmediate_Joint_PositionTolerance = 0.025f;
 Chaos::FRealSingle ChaosImmediate_Joint_AngleTolerance = 0.001f;
 int32 ChaosImmediate_Joint_NumShockPropagationIterations = -1;
 int32 ChaosImmediate_Joint_SolvePositionLast = 1;
+int32 ChaosImmediate_Joint_UsePositionBasedDrives = 1;
 int32 ChaosImmediate_Joint_EnableTwistLimits = 1;
 int32 ChaosImmediate_Joint_EnableSwingLimits = 1;
 int32 ChaosImmediate_Joint_EnableDrives = 1;
@@ -116,6 +117,7 @@ FAutoConsoleVariableRef CVarChaosImmPhysJointPositionTolerance(TEXT("p.Chaos.Imm
 FAutoConsoleVariableRef CVarChaosImmPhysJointAngleTolerance(TEXT("p.Chaos.ImmPhys.Joint.AngleTolerance"), ChaosImmediate_Joint_AngleTolerance, TEXT("AngleTolerance."));
 FAutoConsoleVariableRef CVarChaosImmPhysJointNumShockPropagationIterations(TEXT("p.Chaos.ImmPhys.Joint.NumShockPropagationIterations"), ChaosImmediate_Joint_NumShockPropagationIterations, TEXT("How many iterations to run shock propagation for"));
 FAutoConsoleVariableRef CVarChaosImmPhysJointSolvePositionLast(TEXT("p.Chaos.ImmPhys.Joint.SolvePositionLast"), ChaosImmediate_Joint_SolvePositionLast, TEXT("Should we solve joints in position-then-rotation order (false) rotation-then-position order (true, default)"));
+FAutoConsoleVariableRef CVarChaosImmPhysJointUsePBDVelocityDrives(TEXT("p.Chaos.ImmPhys.Joint.UsePBDDrives"), ChaosImmediate_Joint_UsePositionBasedDrives, TEXT("Whether to solve drives in the position or velocity phase of the solver (default true)"));
 FAutoConsoleVariableRef CVarChaosImmPhysJointEnableTwistLimits(TEXT("p.Chaos.ImmPhys.Joint.EnableTwistLimits"), ChaosImmediate_Joint_EnableTwistLimits, TEXT("EnableTwistLimits."));
 FAutoConsoleVariableRef CVarChaosImmPhysJointEnableSwingLimits(TEXT("p.Chaos.ImmPhys.Joint.EnableSwingLimits"), ChaosImmediate_Joint_EnableSwingLimits, TEXT("EnableSwingLimits."));
 FAutoConsoleVariableRef CVarChaosImmPhysJointEnableDrives(TEXT("p.Chaos.ImmPhys.Joint.EnableDrives"), ChaosImmediate_Joint_EnableDrives, TEXT("EnableDrives."));
@@ -821,6 +823,7 @@ namespace ImmediatePhysics_Chaos
 			JointsSettings.MinParentMassRatio = ChaosImmediate_Joint_MinParentMassRatio;
 			JointsSettings.MaxInertiaRatio = ChaosImmediate_Joint_MaxInertiaRatio;
 			JointsSettings.bSolvePositionLast = ChaosImmediate_Joint_SolvePositionLast != 0;
+			JointsSettings.bUsePositionBasedDrives = ChaosImmediate_Joint_UsePositionBasedDrives != 0;
 			JointsSettings.bEnableTwistLimits = ChaosImmediate_Joint_EnableTwistLimits != 0;
 			JointsSettings.bEnableSwingLimits = ChaosImmediate_Joint_EnableSwingLimits != 0;
 			JointsSettings.bEnableDrives = ChaosImmediate_Joint_EnableDrives != 0;
diff --git a/Engine/Source/Runtime/Experimental/Chaos/Private/Chaos/Joint/PBDJointCachedSolverGaussSeidel.cpp b/Engine/Source/Runtime/Experimental/Chaos/Private/Chaos/Joint/PBDJointCachedSolverGaussSeidel.cpp
--- a/Engine/Source/Runtime/Experimental/Chaos/Private/Chaos/Joint/PBDJointCachedSolverGaussSeidel.cpp
+++ b/Engine/Source/Runtime/Experimental/Chaos/Private/Chaos/Joint/PBDJointCachedSolverGaussSeidel.cpp
@@ -7,6 +7,9 @@
 #include "Chaos/PBDJointConstraintUtilities.h"
 #include "Chaos/Utilities.h"
 
+
+//UE_DISABLE_OPTIMIZATION
+
 namespace Chaos
 {
 	
@@ -178,9 +181,6 @@ void FPBDJointCachedSolver::Init(
 	SolverBodies[0].SetShockPropagationScale(FReal(1));
 	SolverBodies[1].SetShockPropagationScale(FReal(1));
 
-	NetLinearImpulse = FVec3(0);
-	NetAngularImpulse = FVec3(0);
-
 	// Tolerances are positional errors below visible detection. But in PBD the errors
 	// we leave behind get converted to velocity, so we need to ensure that the resultant
 	// movement from that erroneous velocity is less than the desired position tolerance.
@@ -195,6 +195,7 @@ void FPBDJointCachedSolver::Init(
 	NumActiveConstraints = -1;
 	bIsActive = true;
 	bIsBroken = false;
+	bUsePositionBasedDrives = SolverSettings.bUsePositionBasedDrives;
 
 	SolverStiffness = 1.0f;
 
@@ -300,8 +301,6 @@ void FPBDJointCachedSolver::ApplyConstraints(
 		ApplyPositionDrives(Dt);
 		ApplyRotationDrives(Dt);
 	}
-
-	//UpdateIsActive();
 }
 
 void FPBDJointCachedSolver::ApplyVelocityConstraints(
@@ -322,14 +321,18 @@ void FPBDJointCachedSolver::ApplyVelocityConstraints(
 	{
 		ApplyAngularVelocityConstraints();
 		ApplyLinearVelocityConstraints();
+
+		ApplyRotationVelocityDrives(Dt);
+		ApplyPositionVelocityDrives(Dt);
 	}
 	else
 	{
 		ApplyLinearVelocityConstraints();
 		ApplyAngularVelocityConstraints();
-	}
 
-	// @todo(chaos): We can also apply velocity drives here rather than in the Pbd pass
+		ApplyPositionVelocityDrives(Dt);
+		ApplyRotationVelocityDrives(Dt);
+	}
 }
 
 /** UTILS FOR POSITION CONSTRAINT **************************************************************************************/
@@ -448,9 +451,20 @@ void FPBDJointCachedSolver::InitPositionDatasMass(
 	const FReal II0 = FVec3::DotProduct(AngularAxis0, IA0);
 	const FReal II1 = FVec3::DotProduct(AngularAxis1, IA1);
 
-	PositionDatas.UpdateMass(ConstraintIndex, IA0, IA1, InvM(0) + II0 + InvM(1) + II1, Dt);
+	PositionDatas.UpdateMass(ConstraintIndex, IA0, IA1, InvM(0) + II0 + InvM(1) + II1, Dt, bUsePositionBasedDrives);
 }
 	
+void FPBDJointCachedSolver::SetInitConstraintVelocity(
+	const FVec3& ConstraintArm0,
+	const FVec3& ConstraintArm1)
+
+{
+	const FVec3 CV0 = V(0) + FVec3::CrossProduct(W(0), ConstraintArm0);
+	const FVec3 CV1 = V(1) + FVec3::CrossProduct(W(1), ConstraintArm1);
+	const FVec3 CV = CV1 - CV0;
+	InitConstraintVelocity = CV;
+}
+
 void FPBDJointCachedSolver::InitPositionConstraintDatas(
 	const int32 ConstraintIndex,
 	const FVec3& ConstraintAxis,
@@ -480,11 +494,7 @@ void FPBDJointCachedSolver::InitPositionConstraintDatas(
 			ConstraintRestitution, true, ConstraintArm0, ConstraintArm1);
 	}
 
-	const FVec3 CV0 = V(0) + FVec3::CrossProduct(W(0), ConstraintArm0);
-	const FVec3 CV1 = V(1) + FVec3::CrossProduct(W(1), ConstraintArm1);
-	const FVec3 CV = CV1 - CV0;
-
-	InitConstraintAxisLinearVelocities[ConstraintIndex] = FVec3::DotProduct(CV, LocalAxis);
+	InitConstraintAxisLinearVelocities[ConstraintIndex] = FVec3::DotProduct(InitConstraintVelocity, LocalAxis);
 
 	InitPositionDatasMass(PositionConstraints, ConstraintIndex, Dt);
 }
@@ -515,6 +525,9 @@ void FPBDJointCachedSolver::InitLockedPositionConstraint(
 			
 		}
 	}
+
+	SetInitConstraintVelocity(ConstraintArm0, ConstraintArm1);
+
 	for(int32 ConstraintIndex = 0; ConstraintIndex < 3; ++ConstraintIndex)
 	{
 		if(LinearMotion[ConstraintIndex] == EJointMotionType::Locked)
@@ -543,6 +556,8 @@ void FPBDJointCachedSolver::InitSphericalPositionConstraint(
 	const FVec3 ConstraintArm0 = ConnectorXs[1] - CurrentPs[0];
 	const FVec3 ConstraintArm1 = ConnectorXs[1] - CurrentPs[1];
 	
+	SetInitConstraintVelocity(ConstraintArm0, ConstraintArm1);
+
 	InitPositionConstraintDatas(0, SphereAxis0,  SphereDelta0, JointSettings.LinearRestitution, Dt,
 		JointSettings.LinearLimit, EJointMotionType::Limited, ConstraintArm0, ConstraintArm1);
 
@@ -573,6 +588,8 @@ void FPBDJointCachedSolver::InitCylindricalPositionConstraint(
 	const FVec3 ConstraintArm0 = ConnectorXs[1] - CurrentPs[0];
 	const FVec3 ConstraintArm1 = ConnectorXs[1] - CurrentPs[1];
 	
+	SetInitConstraintVelocity(ConstraintArm0, ConstraintArm1);
+
 	InitPositionConstraintDatas((AxisIndex+1)%3, RadialAxis0, RadialDelta0, JointSettings.LinearRestitution, Dt,
 		JointSettings.LinearLimit, EJointMotionType::Limited, ConstraintArm0, ConstraintArm1);
 
@@ -592,6 +609,8 @@ void FPBDJointCachedSolver::InitPlanarPositionConstraint(
 	const FVec3 ConstraintArm0 = ConnectorXs[1] - CurrentPs[0];
 	const FVec3 ConstraintArm1 = ConnectorXs[1] - CurrentPs[1];
 	
+	SetInitConstraintVelocity(ConstraintArm0, ConstraintArm1);
+
 	InitPositionConstraintDatas(AxisIndex, PlaneAxis, PlaneDelta, JointSettings.LinearRestitution, Dt,
 		JointSettings.LinearLimit, EJointMotionType::Limited, ConstraintArm0, ConstraintArm1);
 }
@@ -632,7 +651,6 @@ void FPBDJointCachedSolver::SolvePositionConstraintDelta(
 		ApplyRotationDelta(1,DR1);
 	}
 
-	NetLinearImpulse += DX;
 	++NumActiveConstraints;
 }
 
@@ -727,9 +745,12 @@ void FPBDJointCachedSolver::SolveLinearVelocityConstraint(
 
 	const FReal DeltaLambda = SolverStiffness * PositionConstraints.ConstraintHardStiffness[ConstraintIndex] *
 	 (FVec3::DotProduct(CV, PositionConstraints.ConstraintAxis[ConstraintIndex]) - TargetVel) / PositionConstraints.ConstraintHardIM[ConstraintIndex];
+
+	// @todo(chaos): We should be adding to the net positional impulse here
+	//PositionConstraints.ConstraintLambda[ConstraintIndex] += DeltaLambda * Dt;
 	
 	const FVec3 MDV = DeltaLambda * PositionConstraints.ConstraintAxis[ConstraintIndex];
-
+	
 	if(Body(0).IsDynamic())
 	{
 		const FVec3 DV0 = InvM(0) * MDV;
@@ -746,16 +767,21 @@ void FPBDJointCachedSolver::SolveLinearVelocityConstraint(
 	}
 }
 
-void FPBDJointCachedSolver::ApplyAxisVelocityConstraint(const int32 ConstraintIndex)
+void FPBDJointCachedSolver::ApplyAxisVelocityConstraint(
+	const int32 ConstraintIndex)
 {
-	if(!NetLinearImpulse.IsNearlyZero() && (FMath::Abs(PositionConstraints.ConstraintLambda[ConstraintIndex]) > UE_SMALL_NUMBER) )
+	// Apply restitution for limited joints when we have exceeded the limits
+	// We also drive the velocity to zero for locked constraints (ignoring restitution)
+	if (FMath::Abs(PositionConstraints.ConstraintLambda[ConstraintIndex]) > UE_SMALL_NUMBER)
 	{
 		FReal TargetVel = 0.0f;
-		if (PositionConstraints.MotionType[ConstraintIndex] == EJointMotionType::Limited && PositionConstraints.ConstraintRestitution[ConstraintIndex] != 0.0f)
+		const FReal Restitution = PositionConstraints.ConstraintRestitution[ConstraintIndex];
+		const bool bIsLimited = (PositionConstraints.MotionType[ConstraintIndex] == EJointMotionType::Limited);
+		if (bIsLimited && (Restitution != 0.0f))
 		{
 			const FReal InitVel = InitConstraintAxisLinearVelocities[ConstraintIndex];
-			TargetVel = InitVel > Chaos_Joint_LinearVelocityThresholdToApplyRestitution ?
-				-PositionConstraints.ConstraintRestitution[ConstraintIndex] * InitVel : 0.0f; 
+			const FReal Threshold = Chaos_Joint_LinearVelocityThresholdToApplyRestitution;
+			TargetVel = (InitVel > Threshold) ? -Restitution * InitVel : 0.0f;
 		}
 		SolveLinearVelocityConstraint(ConstraintIndex, TargetVel);
 	}
@@ -909,7 +935,7 @@ void FPBDJointCachedSolver::InitRotationDatasMass(
 	const FReal II0 = FVec3::DotProduct(RotationDatas.ConstraintAxis[ConstraintIndex], IA0);
 	const FReal II1 = FVec3::DotProduct(RotationDatas.ConstraintAxis[ConstraintIndex], IA1);
 
-	RotationDatas.UpdateMass(ConstraintIndex,  IA0, IA1, II0 + II1, Dt);
+	RotationDatas.UpdateMass(ConstraintIndex,  IA0, IA1, II0 + II1, Dt, bUsePositionBasedDrives);
 }
 
 void FPBDJointCachedSolver::InitRotationConstraintDatas(
@@ -1118,7 +1144,6 @@ void FPBDJointCachedSolver::SolveRotationConstraintDelta(
 			DeltaImpulse * ( FVec3::DotProduct(ConstraintDatas.ConstraintAxis[ConstraintIndex], ConstraintDatas.ConstraintDRAxis[ConstraintIndex][1]));
 		ApplyRotationDelta(1, DR1);
 	}
-	NetAngularImpulse += DeltaImpulse;
 	++NumActiveConstraints;
 }
 
@@ -1214,6 +1239,9 @@ void FPBDJointCachedSolver::SolveAngularVelocityConstraint(
 	const FReal DeltaLambda = SolverStiffness * RotationConstraints.ConstraintHardStiffness[ConstraintIndex] *
 	 (FVec3::DotProduct(CW, RotationConstraints.ConstraintAxis[ConstraintIndex]) - TargetVel) / RotationConstraints.ConstraintHardIM[ConstraintIndex];
 
+	 // @todo(chaos): we should be adding to the net positional impulse here
+	 // RotationConstraints.ConstraintLambda[ConsraintIndex] += DeltaLambda * Dt;
+
 	if(Body(0).IsDynamic())
 	{
 		const FVec3 DW0 = RotationConstraints.ConstraintDRAxis[ConstraintIndex][0] * DeltaLambda;
@@ -1228,12 +1256,15 @@ void FPBDJointCachedSolver::SolveAngularVelocityConstraint(
 	}
 }
 
-void FPBDJointCachedSolver::ApplyAngularVelocityConstraint(const int32 ConstraintIndex)
+void FPBDJointCachedSolver::ApplyAngularVelocityConstraint(
+	const int32 ConstraintIndex)
 {
-	if(!NetAngularImpulse.IsNearlyZero() && (FMath::Abs(RotationConstraints.ConstraintLambda[ConstraintIndex]) > UE_SMALL_NUMBER))
+	// Apply restitution for limited joints when we have exceeded the limits
+	// We also drive the velocity to zero for locked constraints (ignoring restitution)
+	if (FMath::Abs(RotationConstraints.ConstraintLambda[ConstraintIndex]) > UE_SMALL_NUMBER)
 	{
 		FReal TargetVel = 0.0f;
-		if (RotationConstraints.ConstraintRestitution[ConstraintIndex] != 0.0f)
+		if ((RotationConstraints.MotionType[ConstraintIndex] == EJointMotionType::Limited) && (RotationConstraints.ConstraintRestitution[ConstraintIndex] != 0.0f))
 		{ 
 			const FReal InitVel = InitConstraintAxisAngularVelocities[ConstraintIndex];
 			TargetVel = InitVel > Chaos_Joint_AngularVelocityThresholdToApplyRestitution ?
@@ -1273,6 +1304,10 @@ void FPBDJointCachedSolver::InitPositionDrives(
 			const FVec3 VTarget = ConnectorRs[0] * JointSettings.LinearDriveVelocityTarget;
 			const FVec3 CX = ConnectorXs[1] - XTarget;
 
+			const FVec3 ConstraintArm0 = ConnectorXs[1] - CurrentPs[0];
+			const FVec3 ConstraintArm1 = ConnectorXs[1] - CurrentPs[1];
+			SetInitConstraintVelocity(ConstraintArm0, ConstraintArm1);
+
 			for (int32 AxisIndex = 0; AxisIndex < 3; ++AxisIndex)
 			{
 				if (bDriven[AxisIndex])
@@ -1531,7 +1566,7 @@ void FPBDJointCachedSolver::ApplyRotationTeleport(
 /** APPLY POSITION  DRIVES *********************************************************************************/
 
 void FPBDJointCachedSolver::ApplyPositionDrives(
-		const FReal Dt)
+	const FReal Dt)
 {
 	for(int32 ConstraintIndex = 0; ConstraintIndex < 3; ++ConstraintIndex)
 	{
@@ -1544,34 +1579,106 @@ void FPBDJointCachedSolver::ApplyPositionDrives(
 	
 
 void FPBDJointCachedSolver::ApplyAxisPositionDrive(
-		const int32 ConstraintIndex,
-		const FReal Dt)
+	const int32 ConstraintIndex,
+	const FReal Dt)
 {
+	const FReal Stiffness = PositionDrives.ConstraintSoftStiffness[ConstraintIndex];
+	const FReal Damping = PositionDrives.ConstraintSoftDamping[ConstraintIndex];
+	const FReal IM = PositionDrives.ConstraintSoftIM[ConstraintIndex];
+
 	const FVec3 Delta0 = Body(0).DP() + FVec3::CrossProduct(Body(0).DQ(), PositionDrives.ConstraintArms[ConstraintIndex][0]);
 	const FVec3 Delta1 = Body(1).DP() + FVec3::CrossProduct(Body(1).DQ(), PositionDrives.ConstraintArms[ConstraintIndex][1]);
+	const FReal CX = PositionDrives.ConstraintCX[ConstraintIndex] + FVec3::DotProduct(Delta1 - Delta0, PositionDrives.ConstraintAxis[ConstraintIndex]);
 
-	const FReal DeltaPos = PositionDrives.ConstraintCX[ConstraintIndex] + FVec3::DotProduct(Delta1 - Delta0, PositionDrives.ConstraintAxis[ConstraintIndex]);
-
-	FReal VelDt = 0;
-	if (PositionDrives.ConstraintSoftDamping[ConstraintIndex] > UE_KINDA_SMALL_NUMBER)
+	FReal CVDt = 0;
+	if (Damping > UE_KINDA_SMALL_NUMBER)
 	{
 		const FVec3 V0Dt = FVec3::CalculateVelocity(InitConnectorXs[0], ConnectorXs[0] + Delta0, 1.0f);
 		const FVec3 V1Dt = FVec3::CalculateVelocity(InitConnectorXs[1], ConnectorXs[1] + Delta1, 1.0f);
-		VelDt = PositionDrives.ConstraintVX[ConstraintIndex] * Dt + FVec3::DotProduct(V0Dt - V1Dt, PositionDrives.ConstraintAxis[ConstraintIndex] );
+		const FReal TargetVDt = PositionDrives.ConstraintVX[ConstraintIndex] * Dt;
+		CVDt = TargetVDt + FVec3::DotProduct(V0Dt - V1Dt, PositionDrives.ConstraintAxis[ConstraintIndex]);
 	}
 
-	FReal DeltaLambda = SolverStiffness * (PositionDrives.ConstraintSoftStiffness[ConstraintIndex] * DeltaPos -
-		PositionDrives.ConstraintSoftDamping[ConstraintIndex] * VelDt - PositionDrives.ConstraintLambda[ConstraintIndex]) /
-		PositionDrives.ConstraintSoftIM[ConstraintIndex];
-	FReal Lambda = PositionDrives.ConstraintLambda[ConstraintIndex] + DeltaLambda;
+	FReal Lambda = PositionDrives.ConstraintLambda[ConstraintIndex];
+	FReal DeltaLambda = SolverStiffness * (Stiffness * CX - Damping * CVDt - Lambda) / IM;
+	Lambda += DeltaLambda;
 
 	PositionDrives.ApplyMaxLambda(ConstraintIndex, DeltaLambda, Lambda);
-
 	PositionDrives.ConstraintLambda[ConstraintIndex] = Lambda;
 
 	SolvePositionConstraintDelta(ConstraintIndex, DeltaLambda, PositionDrives);
 }
 
+void FPBDJointCachedSolver::ApplyPositionVelocityDrives(
+	const FReal Dt)
+{
+	if (bUsePositionBasedDrives)
+	{
+		return;
+	}
+
+	for (int32 ConstraintIndex = 0; ConstraintIndex < 3; ++ConstraintIndex)
+	{
+		if (PositionDrives.bValidDatas[ConstraintIndex])
+		{
+			ApplyAxisPositionVelocityDrive(ConstraintIndex, Dt);
+		}
+	}
+}
+
+void FPBDJointCachedSolver::ApplyAxisPositionVelocityDrive(
+	const int32 ConstraintIndex,
+	const FReal Dt)
+{
+	// NOTE: Using the actual damping, not the PBD modified value
+	const FReal Damping = PositionDrives.SettingsSoftDamping[ConstraintIndex] * Dt;
+	if (Damping < UE_SMALL_NUMBER)
+	{
+		return;
+	}
+
+	const FReal MassScale = PositionDrives.bAccelerationMode ? (FReal(1) / PositionDrives.ConstraintHardIM[ConstraintIndex]) : FReal(1);
+	const FReal IM = MassScale * Damping * PositionDrives.ConstraintHardIM[ConstraintIndex] + (FReal)1;
+
+	// Velocity error to correct
+	const FVec3 V0 = V(0) + FVec3::CrossProduct(W(0), PositionDrives.ConstraintArms[ConstraintIndex][0]);
+	const FVec3 V1 = V(1) + FVec3::CrossProduct(W(1), PositionDrives.ConstraintArms[ConstraintIndex][1]);
+	const FReal VRel = FVec3::DotProduct(V1 - V0, PositionDrives.ConstraintAxis[ConstraintIndex]);
+	const FReal TargetV = PositionDrives.ConstraintVX[ConstraintIndex];
+	const FReal CV = (VRel - TargetV);
+
+	// Implicit scheme: F(t) = -D x V(t+dt)
+	FReal& LambdaVel = PositionDrives.ConstraintLambdaVelocity[ConstraintIndex];
+	FReal DeltaLambdaVel = SolverStiffness * (MassScale * Damping * CV - LambdaVel) / IM;
+
+	// Apply limits and accumulate total impulse 
+	// (NOTE: Limits and net impulses are position based, not velocity based)
+	FReal DeltaLambda = DeltaLambdaVel * Dt;
+	FReal Lambda = PositionDrives.ConstraintLambda[ConstraintIndex] + DeltaLambda;
+	PositionDrives.ApplyMaxLambda(ConstraintIndex, DeltaLambda, Lambda);
+	PositionDrives.ConstraintLambda[ConstraintIndex] = Lambda;	
+	DeltaLambdaVel = DeltaLambda / Dt;
+
+	LambdaVel += DeltaLambdaVel;
+	const FVec3 Impulse = DeltaLambdaVel * PositionDrives.ConstraintAxis[ConstraintIndex];
+
+	if (Body(0).IsDynamic())
+	{
+		const FVec3 DV0 = InvM(0) * Impulse;
+		const FVec3 DW0 = PositionDrives.ConstraintDRAxis[ConstraintIndex][0] * DeltaLambdaVel;
+
+		Body(0).ApplyVelocityDelta(DV0, DW0);
+	}
+	if (Body(1).IsDynamic())
+	{
+		const FVec3 DV1 = -InvM(1) * Impulse;
+		const FVec3 DW1 = PositionDrives.ConstraintDRAxis[ConstraintIndex][1] * DeltaLambdaVel;
+
+		Body(1).ApplyVelocityDelta(DV1, DW1);
+	}
+}
+
+
 /** INIT ROTATION DRIVES *********************************************************************************/
 
 void FPBDJointCachedSolver::InitRotationDrives(
@@ -1752,33 +1859,105 @@ void FPBDJointCachedSolver::ApplyRotationDrives(
 }
 
 void FPBDJointCachedSolver::ApplyAxisRotationDrive(
-		const int32 ConstraintIndex,
-		const FReal Dt)
+	const int32 ConstraintIndex,
+	const FReal Dt)
 {
-	const FReal DeltaConstraint = RotationDrives.ConstraintCX[ConstraintIndex]+
-		FVec3::DotProduct(Body(1).DQ()-Body(0).DQ(), RotationDrives.ConstraintAxis[ConstraintIndex]);
+	const FReal Stiffness = RotationDrives.ConstraintSoftStiffness[ConstraintIndex];
+	const FReal Damping = RotationDrives.ConstraintSoftDamping[ConstraintIndex];
+	const FReal IM = RotationDrives.ConstraintSoftIM[ConstraintIndex];
 
-	// Damping angular velocity
-	FReal AngVelDt = 0;
-	if (RotationDrives.ConstraintSoftDamping[ConstraintIndex] > UE_KINDA_SMALL_NUMBER)
+	// Stiffness position delta
+	FReal CX = 0;
+	if (Stiffness > UE_KINDA_SMALL_NUMBER)
+	{
+		const FReal DX = FVec3::DotProduct(Body(1).DQ() - Body(0).DQ(), RotationDrives.ConstraintAxis[ConstraintIndex]);
+		const FReal TargetX = RotationDrives.ConstraintCX[ConstraintIndex];
+		CX = TargetX + DX;
+	}
+
+	// Damping angular velocity delta
+	FReal CVDt = 0;
+	if (Damping > UE_KINDA_SMALL_NUMBER)
 	{
 		const FVec3 W0Dt = FVec3(Body(0).DQ()) + ConnectorWDts[0];
 		const FVec3 W1Dt = FVec3(Body(1).DQ()) + ConnectorWDts[1];
-		AngVelDt = RotationDrives.ConstraintVX[ConstraintIndex] * Dt + FVec3::DotProduct(RotationDrives.ConstraintAxis[ConstraintIndex] , W0Dt - W1Dt);
+		const FReal TargetW = RotationDrives.ConstraintVX[ConstraintIndex];
+		CVDt = (TargetW * Dt) + FVec3::DotProduct(RotationDrives.ConstraintAxis[ConstraintIndex] , W0Dt - W1Dt);
 	}
 
-	FReal DeltaLambda = SolverStiffness * (RotationDrives.ConstraintSoftStiffness[ConstraintIndex] * DeltaConstraint -
-		RotationDrives.ConstraintSoftDamping[ConstraintIndex] * AngVelDt - RotationDrives.ConstraintLambda[ConstraintIndex]) /
-		RotationDrives.ConstraintSoftIM[ConstraintIndex];
-	FReal Lambda = RotationDrives.ConstraintLambda[ConstraintIndex] + DeltaLambda;
+	FReal Lambda = RotationDrives.ConstraintLambda[ConstraintIndex];
+	FReal DeltaLambda = SolverStiffness * (Stiffness * CX - Damping * CVDt - Lambda) / IM;
+	Lambda += DeltaLambda;
 
 	RotationDrives.ApplyMaxLambda(ConstraintIndex, DeltaLambda, Lambda);
-
 	RotationDrives.ConstraintLambda[ConstraintIndex] = Lambda;
 
 	SolveRotationConstraintDelta(ConstraintIndex, DeltaLambda,true, RotationDrives);
 }
 
+void FPBDJointCachedSolver::ApplyRotationVelocityDrives(
+	const FReal Dt)
+{
+	if (bUsePositionBasedDrives)
+	{
+		return;
+	}
+
+	for (int32 ConstraintIndex = 0; ConstraintIndex < 3; ++ConstraintIndex)
+	{
+		if (RotationDrives.bValidDatas[ConstraintIndex])
+		{
+			ApplyAxisRotationVelocityDrive(ConstraintIndex, Dt);
+		}
+	}
+}
+
+void FPBDJointCachedSolver::ApplyAxisRotationVelocityDrive(
+	const int32 ConstraintIndex,
+	const FReal Dt)
+{
+	// NOTE: Using the actual damping, not the PBD modified value
+	const FReal Damping = RotationDrives.SettingsSoftDamping[ConstraintIndex] * Dt;
+	if (Damping < UE_SMALL_NUMBER)
+	{
+		return;
+	}
+
+	const FReal MassScale = RotationDrives.bAccelerationMode ? (FReal(1) / RotationDrives.ConstraintHardIM[ConstraintIndex]) : FReal(1);
+	const FReal IM = MassScale * Damping * RotationDrives.ConstraintHardIM[ConstraintIndex] + (FReal)1;
+
+	// Angular velocity error to correct
+	const FReal WRel = FVec3::DotProduct(W(1) - W(0), RotationDrives.ConstraintAxis[ConstraintIndex]);
+	const FReal TargetW = RotationDrives.ConstraintVX[ConstraintIndex];
+	const FReal CV = (WRel - TargetW);
+
+	// Implicit scheme: F(t) = -D x W(t+dt)
+	FReal& LambdaVel = RotationDrives.ConstraintLambdaVelocity[ConstraintIndex];
+	FReal DeltaLambdaVel = SolverStiffness * (MassScale * Damping * CV - LambdaVel) / IM;
+
+	// Apply limits and accumulate total impulse 
+	// (NOTE: Limits and net impulses are position based, not velocity based)
+	FReal DeltaLambda = DeltaLambdaVel * Dt;
+	FReal Lambda = RotationDrives.ConstraintLambda[ConstraintIndex] + DeltaLambda;
+	RotationDrives.ApplyMaxLambda(ConstraintIndex, DeltaLambda, Lambda);
+	RotationDrives.ConstraintLambda[ConstraintIndex] = Lambda;
+	DeltaLambdaVel = DeltaLambda / Dt;
+
+	LambdaVel += DeltaLambdaVel;
+	const FVec3 Impulse = DeltaLambdaVel * RotationDrives.ConstraintAxis[ConstraintIndex];
+
+	if (Body(0).IsDynamic())
+	{
+		const FVec3 DW0 = RotationDrives.ConstraintDRAxis[ConstraintIndex][0] * DeltaLambdaVel;
+		Body(0).ApplyAngularVelocityDelta(DW0);
+	}
+	if (Body(1).IsDynamic())
+	{
+		const FVec3 DW1 = RotationDrives.ConstraintDRAxis[ConstraintIndex][1] * DeltaLambdaVel;
+		Body(1).ApplyAngularVelocityDelta(DW1);
+	}
+}
+
  // Joint utilities
 
 void FPBDJointCachedSolver::ApplyPositionDelta(
@@ -1812,6 +1991,7 @@ void FAxisConstraintDatas::InitDatas(
 	bValidDatas[ConstraintIndex] = false;
 	bLimitsCheck[ConstraintIndex] = true;
 	ConstraintLambda = FVec3::Zero();
+	ConstraintLambdaVelocity = FVec3::Zero();
 	ConstraintLimits = FVec3::Zero();
 	MotionType[ConstraintIndex] = EJointMotionType::Free;
 }
@@ -1842,18 +2022,21 @@ void FAxisConstraintDatas::UpdateMass(
 		const FVec3& DatasIA0,
 		const FVec3& DatasIA1,
 		const FReal DatasIM,
-		const FReal Dt)
+		const FReal Dt,
+		const bool bUsePositionBasedDrives)
 {
 	ConstraintDRAxis[ConstraintIndex][0] = DatasIA0;
 	ConstraintDRAxis[ConstraintIndex][1] = -DatasIA1;
 	ConstraintHardIM[ConstraintIndex] = DatasIM;
 	ConstraintLambda = FVec3::Zero();
+	ConstraintLambdaVelocity = FVec3::Zero();
 
 	if(bSoftLimit[ConstraintIndex])
 	{
+		// If bUsePositionBasedDrives is false, we apply the velocity drive in the velocity solver phase so we don't include in the PBD settings
 		const FReal SpringMassScale = (bAccelerationMode) ? (FReal)1 / (ConstraintHardIM[ConstraintIndex]) : (FReal)1;
 		ConstraintSoftStiffness[ConstraintIndex] = SpringMassScale * SettingsSoftStiffness[ConstraintIndex] * Dt * Dt;
-		ConstraintSoftDamping[ConstraintIndex] = SpringMassScale * SettingsSoftDamping[ConstraintIndex] * Dt;
+		ConstraintSoftDamping[ConstraintIndex] = (bUsePositionBasedDrives) ? SpringMassScale * SettingsSoftDamping[ConstraintIndex] * Dt : 0;
 		ConstraintSoftIM[ConstraintIndex] = (ConstraintSoftStiffness[ConstraintIndex] + ConstraintSoftDamping[ConstraintIndex]) * ConstraintHardIM[ConstraintIndex] + (FReal)1;
 	}
 }
diff --git a/Engine/Source/Runtime/Experimental/Chaos/Private/Chaos/Joint/PBDJointContainerSolver.cpp b/Engine/Source/Runtime/Experimental/Chaos/Private/Chaos/Joint/PBDJointContainerSolver.cpp
--- a/Engine/Source/Runtime/Experimental/Chaos/Private/Chaos/Joint/PBDJointContainerSolver.cpp
+++ b/Engine/Source/Runtime/Experimental/Chaos/Private/Chaos/Joint/PBDJointContainerSolver.cpp
@@ -224,6 +224,7 @@ namespace Chaos
 					SolverType& Solver = Solvers[SolverConstraintIndex];
 
 					// NOTE: Particle order was revered in the solver...
+					// NOTE: Solver impulses are positional impulses
 					const FVec3 LinearImpulse = -Solver.GetNetLinearImpulse() / Dt;
 					const FVec3 AngularImpulse = -Solver.GetNetAngularImpulse() / Dt;
 					const FSolverBody* SolverBody0 = &Solver.Body0().SolverBody();
diff --git a/Engine/Source/Runtime/Experimental/Chaos/Private/Chaos/PBDJointConstraints.cpp b/Engine/Source/Runtime/Experimental/Chaos/Private/Chaos/PBDJointConstraints.cpp
--- a/Engine/Source/Runtime/Experimental/Chaos/Private/Chaos/PBDJointConstraints.cpp
+++ b/Engine/Source/Runtime/Experimental/Chaos/Private/Chaos/PBDJointConstraints.cpp
@@ -364,6 +364,7 @@ namespace Chaos
 		, bUseLinearSolver(true)
 		, bSortEnabled(false)
 		, bSolvePositionLast(true)
+		, bUsePositionBasedDrives(true)
 		, bEnableTwistLimits(true)
 		, bEnableSwingLimits(true)
 		, bEnableDrives(true)
diff --git a/Engine/Source/Runtime/Experimental/Chaos/Private/PBDRigidsSolver.cpp b/Engine/Source/Runtime/Experimental/Chaos/Private/PBDRigidsSolver.cpp
--- a/Engine/Source/Runtime/Experimental/Chaos/Private/PBDRigidsSolver.cpp
+++ b/Engine/Source/Runtime/Experimental/Chaos/Private/PBDRigidsSolver.cpp
@@ -299,12 +299,14 @@ namespace Chaos
 		float ChaosSolverJointMaxSolverStiffness = 1.0f;
 		int32 ChaosSolverJointNumIterationsAtMaxSolverStiffness = 1;
 		bool bChaosSolverJointSolvePositionLast = true;
+		bool bChaosSolverJointUsePositionBasedDrives = true;
 		int32 ChaosSolverJointNumShockProagationIterations = 0;
 		FRealSingle ChaosSolverJointShockPropagation = -1.0f;
 		FAutoConsoleVariableRef CVarChaosSolverJointMinSolverStiffness(TEXT("p.Chaos.Solver.Joint.MinSolverStiffness"), ChaosSolverJointMinSolverStiffness, TEXT("Solver stiffness on first iteration, increases each iteration toward MaxSolverStiffness."));
 		FAutoConsoleVariableRef CVarChaosSolverJointMaxSolverStiffness(TEXT("p.Chaos.Solver.Joint.MaxSolverStiffness"), ChaosSolverJointMaxSolverStiffness, TEXT("Solver stiffness on last iteration, increases each iteration from MinSolverStiffness."));
 		FAutoConsoleVariableRef CVarChaosSolverJointNumIterationsAtMaxSolverStiffness(TEXT("p.Chaos.Solver.Joint.NumIterationsAtMaxSolverStiffness"), ChaosSolverJointNumIterationsAtMaxSolverStiffness, TEXT("How many iterations we want at MaxSolverStiffness."));
 		FAutoConsoleVariableRef CVarChaosSolverJointSolvePositionFirst(TEXT("p.Chaos.Solver.Joint.SolvePositionLast"), bChaosSolverJointSolvePositionLast, TEXT("Should we solve joints in position-then-rotation order (false) or rotation-then-position order (true, default)"));
+		FAutoConsoleVariableRef CVarChaosSolverJointUsePBDVelocityDrives(TEXT("p.Chaos.Solver.Joint.UsePBDDrives"), bChaosSolverJointUsePositionBasedDrives, TEXT("Whether to solve drives in the position or velocity phase of the solver (default true"));
 		FAutoConsoleVariableRef CVarChaosSolverJointNumShockPropagationIterations(TEXT("p.Chaos.Solver.Joint.NumShockPropagationIterations"), ChaosSolverJointNumShockProagationIterations, TEXT("How many iterations to enable SHockProagation for."));
 		FAutoConsoleVariableRef CVarChaosSolverJointShockPropagation(TEXT("p.Chaos.Solver.Joint.ShockPropagation"), ChaosSolverJointShockPropagation, TEXT("6Dof joint shock propagation override (if >= 0)."));
 
@@ -1049,6 +1051,7 @@ namespace Chaos
 		JointsSettings.MinParentMassRatio = ChaosSolverJointMinParentMassRatio;
 		JointsSettings.MaxInertiaRatio = ChaosSolverJointMaxInertiaRatio;
 		JointsSettings.bSolvePositionLast = bChaosSolverJointSolvePositionLast;
+		JointsSettings.bUsePositionBasedDrives = bChaosSolverJointUsePositionBasedDrives;
 		JointsSettings.NumShockPropagationIterations = ChaosSolverJointNumShockProagationIterations;
 		JointsSettings.ShockPropagationOverride = ChaosSolverJointShockPropagation;
 		JointsSettings.bUseLinearSolver = bChaosSolverJointUseLinearSolver;
diff --git a/Engine/Source/Runtime/Experimental/Chaos/Public/Chaos/Joint/PBDJointCachedSolverGaussSeidel.h b/Engine/Source/Runtime/Experimental/Chaos/Public/Chaos/Joint/PBDJointCachedSolverGaussSeidel.h
--- a/Engine/Source/Runtime/Experimental/Chaos/Public/Chaos/Joint/PBDJointCachedSolverGaussSeidel.h
+++ b/Engine/Source/Runtime/Experimental/Chaos/Public/Chaos/Joint/PBDJointCachedSolverGaussSeidel.h
@@ -47,7 +47,8 @@ struct FAxisConstraintDatas
 		const FVec3& DatasIA0,
 		const FVec3& DatasIA1,
 		const FReal DatasIM,
-		const FReal Dt);
+		const FReal Dt,
+		const bool bUsePositionBasedDrives);
 
 	/** Apply the impulse limits to the impulse delta and net impulse */
 	void ApplyMaxLambda(
@@ -74,6 +75,7 @@ struct FAxisConstraintDatas
 	FVec3 ConstraintDRAxis[3][2];
 	FVec3 ConstraintRestitution;
 	FVec3 ConstraintLambda;
+	FVec3 ConstraintLambdaVelocity;
 	
 	bool bValidDatas[3];
 	bool bLimitsCheck[3];
@@ -183,14 +185,40 @@ struct FAxisConstraintDatas
 			return (InvM(BodyIndex) > 0);
 		}
 
-		inline const FVec3& GetNetLinearImpulse() const
+		// NOTE: This is a positional impulse
+		inline FVec3 GetNetLinearImpulse() const
 		{
-			return NetLinearImpulse;
+			FVec3 Impulse = FVec3(0);
+			for (int32 Axis = 0; Axis < 3; ++Axis)
+			{
+				if (PositionConstraints.bValidDatas[Axis])
+				{
+					Impulse += PositionConstraints.ConstraintLambda[Axis] * PositionConstraints.ConstraintAxis[Axis];
+				}
+				if (PositionDrives.bValidDatas[Axis])
+				{
+					Impulse += PositionDrives.ConstraintLambda[Axis] * PositionDrives.ConstraintAxis[Axis];
+				}
+			}
+			return Impulse;
 		}
 
-		inline const FVec3& GetNetAngularImpulse() const
+		// NOTE: This is a positional impulse
+		inline FVec3 GetNetAngularImpulse() const
 		{
-			return NetAngularImpulse;
+			FVec3 Impulse = FVec3(0);
+			for (int32 Axis = 0; Axis < 3; ++Axis)
+			{
+				if (RotationConstraints.bValidDatas[Axis])
+				{
+					Impulse += RotationConstraints.ConstraintLambda[Axis] * RotationConstraints.ConstraintAxis[Axis];
+				}
+				if (RotationDrives.bValidDatas[Axis])
+				{
+					Impulse += RotationDrives.ConstraintLambda[Axis] * RotationDrives.ConstraintAxis[Axis];
+				}
+			}
+			return Impulse;
 		}
 
 		inline int32 GetNumActiveConstraints() const
@@ -506,7 +534,8 @@ struct FAxisConstraintDatas
 			const int32 ConstraintIndex,
 			const FReal TargetVel);
 
-		void ApplyAngularVelocityConstraint(const int32 ConstraintIndex);
+		void ApplyAngularVelocityConstraint(
+			const int32 ConstraintIndex);
 
 		/** Init Position Drives */
 		
@@ -531,6 +560,13 @@ struct FAxisConstraintDatas
 			const int32 ConstraintIndex,
 			const FReal Dt);
 
+		void ApplyPositionVelocityDrives(
+			const FReal Dt);
+
+		void ApplyAxisPositionVelocityDrive(
+			const int32 ConstraintIndex,
+			const FReal Dt);
+
 		/** Init Rotation Drives */
 
 		void InitRotationDrives(
@@ -563,9 +599,20 @@ struct FAxisConstraintDatas
 			const FReal Dt);
 		
 		void ApplyAxisRotationDrive(
-    		const int32 ConstraintIndex,
-    		const FReal Dt);
-		
+			const int32 ConstraintIndex,
+			const FReal Dt);
+
+		void ApplyRotationVelocityDrives(
+			const FReal Dt);
+
+		void ApplyAxisRotationVelocityDrive(
+			const int32 ConstraintIndex,
+			const FReal Dt);
+
+		void SetInitConstraintVelocity(
+			const FVec3& ConstraintArm0,
+			const FVec3& ConstraintArm1);
+
 		
 		// The cached body state on which the joint operates
 		FConstraintSolverBody SolverBodies[MaxConstrainedBodies];
@@ -581,32 +628,30 @@ struct FAxisConstraintDatas
 		// XPBD Initial iteration world-space body state
 		FVec3 InitConnectorXs[MaxConstrainedBodies];		// World-space joint connector positions
 		FRotation3 InitConnectorRs[MaxConstrainedBodies];	// World-space joint connector rotations
+		FVec3 InitConstraintVelocity;						// Initial relative velocity at the constrained position
 
 		// Inverse Mass and Inertia
 		FReal InvMs[MaxConstrainedBodies];
 		FMatrix33 InvIs[MaxConstrainedBodies];				// World-space
-
-		// Accumulated Impulse and AngularImpulse (Impulse * Dt since they are mass multiplied position corrections)
-		FVec3 NetLinearImpulse;
-		FVec3 NetAngularImpulse;
 		
 		// Solver stiffness - increased over iterations for stability
 		// \todo(chaos): remove Stiffness from SolverSettings (because it is not a solver constant)
 		FReal SolverStiffness;
 
 		// Tolerances below which we stop solving
-		FReal PositionTolerance;					// Distance error below which we consider a constraint or drive solved
-		FReal AngleTolerance;						// Angle error below which we consider a constraint or drive solved
+		FReal PositionTolerance;							// Distance error below which we consider a constraint or drive solved
+		FReal AngleTolerance;								// Angle error below which we consider a constraint or drive solved
 
 		// Tracking whether the solver is resolved
-		FVec3 LastDPs[MaxConstrainedBodies];			// Positions at the beginning of the iteration
-		FVec3 LastDQs[MaxConstrainedBodies];	// Rotations at the beginning of the iteration
-		FVec3 CurrentPs[MaxConstrainedBodies];			// Positions at the beginning of the iteration
-		FRotation3 CurrentQs[MaxConstrainedBodies];	// Rotations at the beginning of the iteration
-		FVec3 InitConstraintAxisLinearVelocities; // Linear velocities along the constraint axes at the begining of the frame, used by restitution
-		FVec3 InitConstraintAxisAngularVelocities; // Angular velocities along the constraint axes at the begining of the frame, used by restitution
-		int32 NumActiveConstraints;					// The number of active constraints and drives in the last iteration (-1 initial value)
-		bool bIsActive;								// Whether any constraints actually moved any bodies in last iteration
+		FVec3 LastDPs[MaxConstrainedBodies];				// Positions at the beginning of the iteration
+		FVec3 LastDQs[MaxConstrainedBodies];				// Rotations at the beginning of the iteration
+		FVec3 CurrentPs[MaxConstrainedBodies];				// Positions at the beginning of the iteration
+		FRotation3 CurrentQs[MaxConstrainedBodies];			// Rotations at the beginning of the iteration
+		FVec3 InitConstraintAxisLinearVelocities;			// Linear velocities along the constraint axes at the begining of the frame, used by restitution
+		FVec3 InitConstraintAxisAngularVelocities;			// Angular velocities along the constraint axes at the begining of the frame, used by restitution
+		int32 NumActiveConstraints;							// The number of active constraints and drives in the last iteration (-1 initial value)
+		bool bIsActive;										// Whether any constraints actually moved any bodies in last iteration
+		bool bUsePositionBasedDrives;						// Whether to apply velocity drive in the PBD step of the VBD step
 		
 		FAxisConstraintDatas PositionConstraints;
 		FAxisConstraintDatas RotationConstraints;
diff --git a/Engine/Source/Runtime/Experimental/Chaos/Public/Chaos/PBDJointConstraintTypes.h b/Engine/Source/Runtime/Experimental/Chaos/Public/Chaos/PBDJointConstraintTypes.h
--- a/Engine/Source/Runtime/Experimental/Chaos/Public/Chaos/PBDJointConstraintTypes.h
+++ b/Engine/Source/Runtime/Experimental/Chaos/Public/Chaos/PBDJointConstraintTypes.h
@@ -232,6 +232,9 @@ namespace Chaos
 		// Solving position last leads to less separation at the joints when limits are being forced
 		bool bSolvePositionLast;
 
+		// Whether joints are position-based or velocity-based in the solver
+		bool bUsePositionBasedDrives;
+
 		// @todo(chaos): remove these TEMP overrides for testing
 		bool bEnableTwistLimits;
 		bool bEnableSwingLimits;
