Chaos - disable particle array shrinking by default (cvar)
- add cvars to control when shrinking occurs

diff --git a/Engine/Source/Runtime/Experimental/Chaos/Private/PBDRigidsSolver.cpp b/Engine/Source/Runtime/Experimental/Chaos/Private/PBDRigidsSolver.cpp
--- a/Engine/Source/Runtime/Experimental/Chaos/Private/PBDRigidsSolver.cpp
+++ b/Engine/Source/Runtime/Experimental/Chaos/Private/PBDRigidsSolver.cpp
@@ -225,8 +225,12 @@ namespace Chaos
 		FAutoConsoleVariableRef CVarChaosSolverCollisionDisable(TEXT("p.Chaos.Solver.Collision.Enabled"), bChaosSolverCollisionEnabled, TEXT("Enable/Disable collisions in the main scene."));
 
 		// Shrink particle arrays every frame to recove rmemory when a scene changes significantly
-		bool bChaosSolverShrinkArrays = true;
+		bool bChaosSolverShrinkArrays = false;
+		float ChaosArrayCollectionMaxSlackFraction = 0.5f;
+		int32 ChaosArrayCollectionMinSlack = 100;
 		FAutoConsoleVariableRef CVarChaosSolverShrinkArrays(TEXT("p.Chaos.Solver.ShrinkArrays"), bChaosSolverShrinkArrays, TEXT("Enable/Disable particle array shrinking in the main scene"));
+		FAutoConsoleVariableRef CVarChaosSolverArrayCollectionMaxSlackFraction(TEXT("p.Chaos.ArrayCollection.MaxSlackFraction"), ChaosArrayCollectionMaxSlackFraction, TEXT("Shrink particle arrays if the number of slack elements exceeds the number of elements by this fraction"));
+		FAutoConsoleVariableRef CVarChaosSolverArrayCollectionMinSlack(TEXT("p.Chaos.ArrayCollection.MinSlack"), ChaosArrayCollectionMinSlack, TEXT("Do not reduce the size of particle arrays if it would leave less slack than this"));
 
 		// Iteration count cvars
 		// These override the engine config if >= 0
@@ -493,7 +497,7 @@ namespace Chaos
 			if (bChaosSolverShrinkArrays)
 			{
 				QUICK_SCOPE_CYCLE_COUNTER(ShrinkParticleArrays);
-				MSolver->GetParticles().ShrinkArrays();
+				MSolver->GetParticles().ShrinkArrays(CVars::ChaosArrayCollectionMaxSlackFraction, CVars::ChaosArrayCollectionMinSlack);
 			}
 
 			if (FRewindData* RewindData = MSolver->GetRewindData())
diff --git a/Engine/Source/Runtime/Experimental/Chaos/Public/Chaos/PBDRigidsSOAs.h b/Engine/Source/Runtime/Experimental/Chaos/Public/Chaos/PBDRigidsSOAs.h
--- a/Engine/Source/Runtime/Experimental/Chaos/Public/Chaos/PBDRigidsSOAs.h
+++ b/Engine/Source/Runtime/Experimental/Chaos/Public/Chaos/PBDRigidsSOAs.h
@@ -230,21 +230,21 @@ public:
 		check(0);
 	}
 
-	void ShrinkArrays()
+	void ShrinkArrays(const float MaxSlackFraction, const int32 MinSlack)
 	{
-		StaticParticles->ShrinkArrays();
-		StaticDisabledParticles->ShrinkArrays();
+		StaticParticles->ShrinkArrays(MaxSlackFraction, MinSlack);
+		StaticDisabledParticles->ShrinkArrays(MaxSlackFraction, MinSlack);
 
-		KinematicParticles->ShrinkArrays();;
-		KinematicDisabledParticles->ShrinkArrays();
+		KinematicParticles->ShrinkArrays(MaxSlackFraction, MinSlack);
+		KinematicDisabledParticles->ShrinkArrays(MaxSlackFraction, MinSlack);
 
-		DynamicDisabledParticles->ShrinkArrays();
-		DynamicParticles->ShrinkArrays();
-		DynamicKinematicParticles->ShrinkArrays();
+		DynamicDisabledParticles->ShrinkArrays(MaxSlackFraction, MinSlack);
+		DynamicParticles->ShrinkArrays(MaxSlackFraction, MinSlack);
+		DynamicKinematicParticles->ShrinkArrays(MaxSlackFraction, MinSlack);
 
-		ClusteredParticles->ShrinkArrays();
+		ClusteredParticles->ShrinkArrays(MaxSlackFraction, MinSlack);
 
-		GeometryCollectionParticles->ShrinkArrays();
+		GeometryCollectionParticles->ShrinkArrays(MaxSlackFraction, MinSlack);
 	}
 
 	void UpdateDirtyViews()
diff --git a/Engine/Source/Runtime/Experimental/ChaosCore/Public/Chaos/ArrayCollection.h b/Engine/Source/Runtime/Experimental/ChaosCore/Public/Chaos/ArrayCollection.h
--- a/Engine/Source/Runtime/Experimental/ChaosCore/Public/Chaos/ArrayCollection.h
+++ b/Engine/Source/Runtime/Experimental/ChaosCore/Public/Chaos/ArrayCollection.h
@@ -22,13 +22,13 @@ public:
 		}
 	}
 
-	void ShrinkArrays()
+	void ShrinkArrays(const float MaxSlackFraction, const int32 MinSlack)
 	{
 		for (int32 Index = 0; Index < MArrays.Num(); Index++)
 		{
 			if (MArrays[Index] != nullptr)
 			{
-				MArrays[Index]->Shrink();
+				MArrays[Index]->ApplyShrinkPolicy(MaxSlackFraction, MinSlack);
 			}
 		}
 	}
diff --git a/Engine/Source/Runtime/Experimental/ChaosCore/Public/Chaos/ArrayCollectionArray.h b/Engine/Source/Runtime/Experimental/ChaosCore/Public/Chaos/ArrayCollectionArray.h
--- a/Engine/Source/Runtime/Experimental/ChaosCore/Public/Chaos/ArrayCollectionArray.h
+++ b/Engine/Source/Runtime/Experimental/ChaosCore/Public/Chaos/ArrayCollectionArray.h
@@ -18,6 +18,7 @@ class TArrayCollectionArray : public TArrayCollectionArrayBase, public TArray<T>
 	using TArray<T>::RemoveAtSwap;
 	using TArray<T>::Emplace;
 	using TArray<T>::Shrink;
+	using TArray<T>::Max;
 
 public:
 	constexpr static bool bAllowShrinkOnRemove = false;
@@ -57,9 +58,22 @@ public:
 		return NewArray;
 	}
 
-	void Shrink() override
+	// If we have more slack space than MaxSlackFraction x Num(), run the default Shrink policy
+	void ApplyShrinkPolicy(const float MaxSlackFraction, const int32 MinSlack) override
 	{
-		TArray<T>::Shrink();
+		// Never shrink below this size
+		const int32 Slack = Max() - Num();
+		if (Slack <= MinSlack)
+		{
+			return;
+		}
+
+		// Shrink if we exceed the maximum allowed slack
+		const int32 MaxSlack = FMath::Max(MinSlack, FMath::FloorToInt(MaxSlackFraction * float(Num())));
+		if (Slack > MaxSlack)
+		{
+			Shrink();
+		}
 	}
 
 	void Resize(const int Num) override
diff --git a/Engine/Source/Runtime/Experimental/ChaosCore/Public/Chaos/ArrayCollectionArrayBase.h b/Engine/Source/Runtime/Experimental/ChaosCore/Public/Chaos/ArrayCollectionArrayBase.h
--- a/Engine/Source/Runtime/Experimental/ChaosCore/Public/Chaos/ArrayCollectionArrayBase.h
+++ b/Engine/Source/Runtime/Experimental/ChaosCore/Public/Chaos/ArrayCollectionArrayBase.h
@@ -8,7 +8,7 @@ namespace Chaos
 class TArrayCollectionArrayBase
 {
   public:
-	virtual void Shrink() = 0;
+	virtual void ApplyShrinkPolicy(const float MaxSlackFraction, const int32 MinSlack) = 0;
 	virtual void Resize(const int Num) = 0;
 	virtual void RemoveAt(const int Idx, const int Count) = 0;
 	virtual void RemoveAtSwap(const int Idx) = 0;
