Chaos Physics: Spatial Acceleration structure substructure Spatial Indices can now be changed at any time. The Elements will now migrate automatically without causing hard to find crashes

diff --git a/Engine/Source/Runtime/Experimental/Chaos/Public/Chaos/AABBTree.h b/Engine/Source/Runtime/Experimental/Chaos/Public/Chaos/AABBTree.h
--- a/Engine/Source/Runtime/Experimental/Chaos/Public/Chaos/AABBTree.h
+++ b/Engine/Source/Runtime/Experimental/Chaos/Public/Chaos/AABBTree.h
@@ -1711,7 +1711,8 @@ public:
 		return false;
 	}
 
-	virtual void UpdateElement(const TPayloadType& Payload, const TAABB<T, 3>& NewBounds, bool bInHasBounds) override
+	// Returns true if element was updated, or false when it was added instead
+	virtual bool UpdateElement(const TPayloadType& Payload, const TAABB<T, 3>& NewBounds, bool bInHasBounds) override
 	{
 		if (UNLIKELY(!ensure(bModifyingTreeMultiThreadingFastCheck == false)))
 		{
@@ -1733,6 +1734,7 @@ public:
 				*NewBounds.Min().ToString(), *NewBounds.Max().ToString());
 		}
 
+		bool bElementExisted = true;
 		if (ensure(bMutable))
 		{
 			FAABBTreePayloadInfo* PayloadInfo = PayloadToInfo.Find(Payload);
@@ -1753,7 +1755,7 @@ public:
 								Leaves[PayloadInfo->LeafIdx].UpdateElement(Payload, NewBounds, bHasBounds);
 								Leaves[PayloadInfo->LeafIdx].RecomputeBounds();
 								bModifyingTreeMultiThreadingFastCheck = false;
-								return;
+								return bElementExisted;
 							}
 						}
 						else
@@ -1764,7 +1766,7 @@ public:
 								// We still need to update the constituent bounds
 								Leaves[PayloadInfo->LeafIdx].UpdateElement(Payload, NewBounds, bHasBounds);
 								bModifyingTreeMultiThreadingFastCheck = false;
-								return;
+								return bElementExisted;
 							}
 						}
 					}
@@ -1786,6 +1788,7 @@ public:
 			}
 			else
 			{
+				bElementExisted = false;
 				PayloadInfo = &PayloadToInfo.Add(Payload);
 			}
 
@@ -1924,6 +1927,7 @@ public:
 			ReoptimizeTree();
 		}
 		bModifyingTreeMultiThreadingFastCheck = false;
+		return bElementExisted;
 	}
 
 	int32 NumDirtyElements() const
diff --git a/Engine/Source/Runtime/Experimental/Chaos/Public/Chaos/BoundingVolume.h b/Engine/Source/Runtime/Experimental/Chaos/Public/Chaos/BoundingVolume.h
--- a/Engine/Source/Runtime/Experimental/Chaos/Public/Chaos/BoundingVolume.h
+++ b/Engine/Source/Runtime/Experimental/Chaos/Public/Chaos/BoundingVolume.h
@@ -261,9 +261,10 @@ public:
 		return false;
 	}
 
-	virtual void UpdateElement(const TPayloadType& Payload, const TAABB<T,d>& NewBounds, bool bHasBounds) override
+	virtual bool UpdateElement(const TPayloadType& Payload, const TAABB<T,d>& NewBounds, bool bHasBounds) override
 	{
 		SCOPE_CYCLE_COUNTER(STAT_BoundingVolumeUpdateElement);
+		bool bElementExisted = true;
 		if (FPayloadInfo* PayloadInfo = MPayloadInfo.Find(Payload))
 		{
 			ensure(bHasBounds || PayloadInfo->GlobalPayloadIdx != INDEX_NONE);
@@ -280,9 +281,11 @@ public:
 		}
 		else
 		{
+			bElementExisted = false;
 			FPayloadInfo& NewPayloadInfo = MPayloadInfo.Add(Payload);
 			AddElementToExistingGrid(Payload, NewPayloadInfo, NewBounds, bHasBounds);
 		}
+		return bElementExisted;
 	}
 
 	inline void AddElement(const TPayloadBoundsElement<TPayloadType, T>& Payload)
diff --git a/Engine/Source/Runtime/Experimental/Chaos/Public/Chaos/ISpatialAcceleration.h b/Engine/Source/Runtime/Experimental/Chaos/Public/Chaos/ISpatialAcceleration.h
--- a/Engine/Source/Runtime/Experimental/Chaos/Public/Chaos/ISpatialAcceleration.h
+++ b/Engine/Source/Runtime/Experimental/Chaos/Public/Chaos/ISpatialAcceleration.h
@@ -295,9 +295,10 @@ public:
 		return true;
 	}
 
-	virtual void UpdateElement(const TPayloadType& Payload, const TAABB<T, d>& NewBounds, bool bHasBounds)
+	virtual bool UpdateElement(const TPayloadType& Payload, const TAABB<T, d>& NewBounds, bool bHasBounds)
 	{
 		check(false);
+		return true;
 	}
 
 	virtual void RemoveElementFrom(const TPayloadType& Payload, FSpatialAccelerationIdx Idx)
diff --git a/Engine/Source/Runtime/Experimental/Chaos/Public/Chaos/SpatialAccelerationCollection.h b/Engine/Source/Runtime/Experimental/Chaos/Public/Chaos/SpatialAccelerationCollection.h
--- a/Engine/Source/Runtime/Experimental/Chaos/Public/Chaos/SpatialAccelerationCollection.h
+++ b/Engine/Source/Runtime/Experimental/Chaos/Public/Chaos/SpatialAccelerationCollection.h
@@ -574,28 +574,36 @@ public:
 	{
 		const uint16 UseBucket = ((1 << SpatialIdx.Bucket) & this->ActiveBucketsMask) ? SpatialIdx.Bucket : 0;
 		const bool Success = Buckets[UseBucket].Objects[SpatialIdx.InnerIdx].Acceleration->RemoveElement(Payload);
-		//ensure(Success); // We are removing this element with the wrong spatial index
-		if (!Success)
+		//ensure(Success); // Debug check to see if SpatialIdx was what the caller expected it to be
+		// Make sure that we remove this Payload even if the SpatialIdx is wrong
+		for (FSpatialAccelerationIdx Idx : GetAllSpatialIndices())
+		{
+			if (!(Idx == SpatialIdx))
+			{
+				const uint16 Buckt = ((1 << Idx.Bucket) & this->ActiveBucketsMask) ? Idx.Bucket : 0;
+				const bool Removed = Buckets[Buckt].Objects[Idx.InnerIdx].Acceleration->RemoveElement(Payload);					
+			}
+		}		
+	}
+
+	virtual void UpdateElementIn(const TPayloadType& Payload, const TAABB<T, d>& NewBounds, bool bHasBounds, FSpatialAccelerationIdx SpatialIdx)
+	{
+		const uint16 UseBucket = ((1 << SpatialIdx.Bucket) & this->ActiveBucketsMask) ? SpatialIdx.Bucket : 0;
+		const bool bElementExisted = Buckets[UseBucket].Objects[SpatialIdx.InnerIdx].Acceleration->UpdateElement(Payload, NewBounds, bHasBounds);
+		// In case spatial index changed, remove this element in all other substructures
+		if (!bElementExisted)
 		{
-			// Make sure that we remove this Payload to prevent crashes
 			for (FSpatialAccelerationIdx Idx : GetAllSpatialIndices())
 			{
 				if (!(Idx == SpatialIdx))
 				{
 					const uint16 Buckt = ((1 << Idx.Bucket) & this->ActiveBucketsMask) ? Idx.Bucket : 0;
-					const bool Removed = Buckets[Buckt].Objects[Idx.InnerIdx].Acceleration->RemoveElement(Payload);
-					//ensure(!Removed); // If this ensure hits, the payload was found in another substructure and successfully removed
+					Buckets[Buckt].Objects[Idx.InnerIdx].Acceleration->RemoveElement(Payload);
 				}
 			}
 		}
 	}
 
-	virtual void UpdateElementIn(const TPayloadType& Payload, const TAABB<T, d>& NewBounds, bool bHasBounds, FSpatialAccelerationIdx SpatialIdx)
-	{
-		const uint16 UseBucket = ((1 << SpatialIdx.Bucket) & this->ActiveBucketsMask) ? SpatialIdx.Bucket : 0;
-		Buckets[UseBucket].Objects[SpatialIdx.InnerIdx].Acceleration->UpdateElement(Payload, NewBounds, bHasBounds);
-	}
-
 	virtual TUniquePtr<ISpatialAcceleration<TPayloadType, T, d>> Copy() const
 	{
 		return TUniquePtr<ISpatialAcceleration<TPayloadType, T, d>>(new TSpatialAccelerationCollection<TSpatialAccelerationTypes...>(*this));
