Track breaking collisions on children and parents, and fix impulse sign.

diff --git a/Engine/Source/Runtime/Experimental/Chaos/Private/Chaos/PBDRigidClustering.cpp b/Engine/Source/Runtime/Experimental/Chaos/Private/Chaos/PBDRigidClustering.cpp
--- a/Engine/Source/Runtime/Experimental/Chaos/Private/Chaos/PBDRigidClustering.cpp
+++ b/Engine/Source/Runtime/Experimental/Chaos/Private/Chaos/PBDRigidClustering.cpp
@@ -485,8 +485,8 @@ namespace Chaos
 			uint8 OtherIdx = Collision.GetParticle0() == &Rigid ? 1 : 0;
 			const FVec3 Impulse
 				= OtherIdx == 0
-				? -Collision.AccumulatedImpulse
-				: Collision.AccumulatedImpulse;
+				? Collision.AccumulatedImpulse
+				: -Collision.AccumulatedImpulse;
 
 			// Compute the angular impulse based on distance from the contact point to the CoM
 			const FVec3 Location = Collision.CalculateWorldContactLocation();
@@ -858,6 +858,7 @@ namespace Chaos
 				// Restore some of the momentum of whatever collided with the parent
 				if (bRestoreBreakingMomentum)
 				{
+					TrackBreakingCollision(Child);
 					TrackBreakingCollision(ClusteredParticle);
 				}
 			}
