Chaos : Fields
- Particle handles filtered with MetaDataFilter do not produce the same results as the default resolution method because the active cluster particle�s children are ignored.

diff --git a/Engine/Source/Runtime/Experimental/Chaos/Private/PhysicsProxy/PerSolverFieldSystem.cpp b/Engine/Source/Runtime/Experimental/Chaos/Private/PhysicsProxy/PerSolverFieldSystem.cpp
--- a/Engine/Source/Runtime/Experimental/Chaos/Private/PhysicsProxy/PerSolverFieldSystem.cpp
+++ b/Engine/Source/Runtime/Experimental/Chaos/Private/PhysicsProxy/PerSolverFieldSystem.cpp
@@ -392,6 +392,27 @@ void FPerSolverFieldSystem::RemovePersistentCommand(const FFieldSystemCommand& F
 	PersistentCommands.Remove(FieldCommand);
 }
 
+
+void AddClusterChildren(TArray<Chaos::FGeometryParticleHandle*>& Handles,
+						const Chaos::FPBDRigidsSolver* RigidSolver,
+						const Chaos::TPBDRigidClusteredParticleHandleImp<Chaos::FReal, 3, false>* Clustered)
+{
+	using FClusterMap = Chaos::FRigidClustering::FClusterMap;
+	using FParticleHandel = Chaos::TPBDRigidParticleHandle<Chaos::FReal, 3>;
+
+	if (Clustered && Clustered->ClusterIds().NumChildren)
+	{
+		const FClusterMap& ClusterMap = RigidSolver->GetEvolution()->GetRigidClustering().GetChildrenMap();
+		if (ClusterMap.Contains(Clustered->Handle()))
+		{
+			for (FParticleHandel* Child : ClusterMap[Clustered->Handle()])
+			{
+				Handles.Add(Child);
+			}
+		}
+	}
+}
+
 void FPerSolverFieldSystem::GetRelevantParticleHandles(
 	TArray<Chaos::FGeometryParticleHandle*>& Handles,
 	const Chaos::FPBDRigidsSolver* RigidSolver,
@@ -402,9 +423,6 @@ void FPerSolverFieldSystem::GetRelevantParticleHandles(
 
 	if (ResolutionType == EFieldResolutionType::Field_Resolution_Minimal)
 	{
-		const auto& Clustering = RigidSolver->GetEvolution()->GetRigidClustering();
-		const auto& ClusterMap = Clustering.GetChildrenMap();
-
 		const Chaos::TParticleView<Chaos::FGeometryParticles>& ParticleView =
 			SolverParticles.GetNonDisabledView();
 		Handles.Reserve(ParticleView.Num()); // ?? what about additional number of children added
@@ -413,19 +431,7 @@ void FPerSolverFieldSystem::GetRelevantParticleHandles(
 		{
 			const Chaos::TTransientGeometryParticleHandle<Chaos::FReal,3>* Handle = &(*It);
 			Handles.Add(GetHandleHelper(const_cast<Chaos::TTransientGeometryParticleHandle<Chaos::FReal,3>*>(Handle)));
-
-			const auto* Clustered = Handle->CastToClustered();
-			if (Clustered && Clustered->ClusterIds().NumChildren)
-			{
-				Chaos::TPBDRigidClusteredParticleHandle<Chaos::FReal, 3>* ClusterHandle = (*It).Handle()->CastToClustered();
-				if (ClusterMap.Contains(ClusterHandle))
-				{
-					for (Chaos::TPBDRigidParticleHandle<Chaos::FReal, 3> * Child : ClusterMap[ClusterHandle])
-					{
-						Handles.Add(Child);
-					}
-				}
-			}
+			AddClusterChildren(Handles, RigidSolver, It->CastToClustered());
 		}
 	}
 	else if (ResolutionType == EFieldResolutionType::Field_Resolution_DisabledParents)
@@ -487,6 +493,7 @@ void FPerSolverFieldSystem::GetFilteredParticleHandles(
 			{
 				const Chaos::TTransientGeometryParticleHandle<Chaos::FReal, 3>* Handle = &(*It);
 				Handles.Add(GetHandleHelper(const_cast<Chaos::TTransientGeometryParticleHandle<Chaos::FReal, 3>*>(Handle)));
+				AddClusterChildren(Handles, RigidSolver, It->CastToClustered());
 			}
 		}
 	}
@@ -568,6 +575,7 @@ void FPerSolverFieldSystem::GetFilteredParticleHandles(
 			{
 				const Chaos::TTransientGeometryParticleHandle<Chaos::FReal, 3>* Handle = &(*It);
 				Handles.Add(GetHandleHelper(const_cast<Chaos::TTransientGeometryParticleHandle<Chaos::FReal, 3>*>(Handle)));
+				AddClusterChildren(Handles, RigidSolver, It->CastToClustered());
 			}
 		}
 	}
