ChaosFlesh : Add code that allows turning XPBD into PBD

diff --git a/Engine/Source/Runtime/Experimental/Chaos/Public/Chaos/XPBDCorotatedConstraints.h b/Engine/Source/Runtime/Experimental/Chaos/Public/Chaos/XPBDCorotatedConstraints.h
--- a/Engine/Source/Runtime/Experimental/Chaos/Public/Chaos/XPBDCorotatedConstraints.h
+++ b/Engine/Source/Runtime/Experimental/Chaos/Public/Chaos/XPBDCorotatedConstraints.h
@@ -133,7 +133,13 @@ namespace Chaos::Softs
 				for (int32 j = 0; j < 4; j++)
 				{
 					ensure(MeshConstraints[e][j] > -1 && MeshConstraints[e][j] < int32(InParticles.Size()));
+					
+					if (InParticles.InvM(MeshConstraints[e][j]) == (T)0.)
+					{
+						AlphaJArray[e] = (T)1.;
+					}
 				}
+				
 			}
 
 			for (int e = 0; e < InMesh.Num(); e++)
@@ -457,6 +463,11 @@ namespace Chaos::Softs
 			//T AlphaTilde = (T)2. / (Dt * Dt * Lambda * Measure[ElementIndex]);
 			T AlphaTilde = (T)2. / (Dt * Dt * LambdaElementArray[ElementIndex] * Measure[ElementIndex]);
 
+			if (LambdaElementArray[ElementIndex] > (T)1. / (T)UE_SMALL_NUMBER)
+			{
+				AlphaTilde = (T)0.;
+			}
+
 			if (bRecordMetric)
 			{
 				HError += J - 1 + AlphaTilde * LambdaArray[2 * ElementIndex + 1];
@@ -531,6 +542,11 @@ namespace Chaos::Softs
 			//T AlphaTilde = (T)1. / (Dt * Dt * Mu * Measure[ElementIndex]);
 			T AlphaTilde = (T)1. / (Dt * Dt * MuElementArray[ElementIndex] * Measure[ElementIndex]);
 
+			if (MuElementArray[ElementIndex] > (T)1./ (T)UE_SMALL_NUMBER) 
+			{
+				AlphaTilde = (T)0.;
+			}
+
 			if (bRecordMetric)
 			{
 				HError += C1 + AlphaTilde * LambdaArray[2 * ElementIndex + 0];
