ChaosFlesh : Fix CIS

diff --git a/Engine/Source/Runtime/Experimental/Chaos/Private/Chaos/Deformable/ChaosDeformableSolver.cpp b/Engine/Source/Runtime/Experimental/Chaos/Private/Chaos/Deformable/ChaosDeformableSolver.cpp
--- a/Engine/Source/Runtime/Experimental/Chaos/Private/Chaos/Deformable/ChaosDeformableSolver.cpp
+++ b/Engine/Source/Runtime/Experimental/Chaos/Private/Chaos/Deformable/ChaosDeformableSolver.cpp
@@ -126,7 +126,8 @@ namespace Chaos::Softs
 			/*CoefficientOfFriction = */(FSolverReal)0.,
 			/*FSolverReal Damping = */Property.Damping,
 			/*FSolverReal LocalDamping = */(FSolverReal)0.,
-			Property.bDoQuasistatics));
+			Property.bDoQuasistatics, 
+			false));
 		Evolution->Particles().AddArray(&MObjects);
 		if (Property.bDoSelfCollision || Property.CacheToFile)
 		{
diff --git a/Engine/Source/Runtime/Experimental/Chaos/Private/Chaos/PBDEvolution.cpp b/Engine/Source/Runtime/Experimental/Chaos/Private/Chaos/PBDEvolution.cpp
--- a/Engine/Source/Runtime/Experimental/Chaos/Private/Chaos/PBDEvolution.cpp
+++ b/Engine/Source/Runtime/Experimental/Chaos/Private/Chaos/PBDEvolution.cpp
@@ -81,7 +81,8 @@ FPBDEvolution::FPBDEvolution(
 	FSolverReal CoefficientOfFriction,
 	FSolverReal Damping,
 	FSolverReal LocalDamping,
-	bool bDoQuasistatics)
+	bool bDoQuasistatics, 
+	bool InbUsePerParticleDamping)
 	: MParticles(MoveTemp(InParticles))
 	, MParticlesActiveView(MParticles)
 	, MCollisionParticles(MoveTemp(InGeometryParticles))
@@ -98,6 +99,7 @@ FPBDEvolution::FPBDEvolution(
 	, MLocalDamping(LocalDamping)
 	, MTime(0)
 	, bDoQuasistatics(bDoQuasistatics)
+	, bUsePerParticleDamping(InbUsePerParticleDamping)
 {
 	// Add group arrays
 	TArrayCollection::AddArray(&MGroupGravityAccelerations);
@@ -112,7 +114,10 @@ FPBDEvolution::FPBDEvolution(
 
 	// Add particle arrays
 	MParticles.AddArray(&MParticleGroupIds);
-	MParticles.AddArray(&MParticleDampings);
+	if (bUsePerParticleDamping)
+	{
+		MParticles.AddArray(&MParticleDampings);
+	}
 	MCollisionParticles.AddArray(&MCollisionTransforms);
 	MCollisionParticles.AddArray(&MCollided);
 	MCollisionParticles.AddArray(&MCollisionParticleGroupIds);
@@ -140,7 +145,14 @@ int32 FPBDEvolution::AddParticleRange(int32 NumParticles, uint32 GroupId, bool b
 		for (int32 i = Offset; i < (int32)MParticles.Size(); ++i)
 		{
 			MParticleGroupIds[i] = GroupId;
-			MParticleDampings[i] = (FSolverReal)0.;
+		}
+
+		if (bUsePerParticleDamping)
+		{
+			for (int32 i = Offset; i < (int32)MParticles.Size(); ++i)
+			{
+				MParticleDampings[i] = (FSolverReal)0.;
+			}
 		}
 
 		// Resize the group parameter arrays
@@ -339,11 +351,11 @@ void FPBDEvolution::PreIterationUpdate(
 					}
 				}, RangeSize < MinParallelBatchSize);
 		} 
-		else
+		else if (bUsePerParticleDamping)
 		{
 			const int32 RangeSize = Range - Offset;
 			PhysicsParallelFor(RangeSize,
-				[this, &Offset, &ForceRule, &Gravity, &VelocityAndPressureField, &DampVelocityRule, &DampingFrequency, &ParticleGroupId,DampingPowDt, DampingIntegrated, Dt](int32 i)
+				[this, &Offset, &ForceRule, &Gravity, &VelocityAndPressureField, &DampVelocityRule, &DampingFrequency, &ParticleGroupId, DampingPowDt, DampingIntegrated, Dt](int32 i)
 				{
 					const int32 Index = Offset + i;
 
@@ -366,8 +378,7 @@ void FPBDEvolution::PreIterationUpdate(
 						ParticleDampingPowDt = (FSolverReal)1.;
 						ParticleDampingIntegrated = Dt;
 					}
-
-
+					
 					if (MParticles.InvM(Index) != (FSolverReal)0.)  // Process dynamic particles
 					{
 						// Init forces with GravityForces
@@ -404,6 +415,52 @@ void FPBDEvolution::PreIterationUpdate(
 						MKinematicUpdate(MParticles, Dt, MTime, Index);
 					}
 				}, RangeSize < MinParallelBatchSize);
+			
+		}
+		else
+		{
+			const int32 RangeSize = Range - Offset;
+			PhysicsParallelFor(RangeSize,
+				[this, &Offset, &ForceRule, &Gravity, &VelocityAndPressureField, &DampVelocityRule, &DampingFrequency, &ParticleGroupId, DampingPowDt, DampingIntegrated, Dt](int32 i)
+				{
+					const int32 Index = Offset + i;
+
+					if (MParticles.InvM(Index) != (FSolverReal)0.)  // Process dynamic particles
+					{
+						// Init forces with GravityForces
+						MParticles.Acceleration(Index) = Gravity;
+
+						// Force Rule
+						if (bForceRule)
+						{
+							ForceRule(MParticles, Dt, Index); // F += M * A
+						}
+
+						// Velocity Field
+						if (bVelocityField)
+						{
+							VelocityAndPressureField.Apply(MParticles, Dt, Index);
+						}
+
+						// Euler Step Velocity
+						MParticles.V(Index) += MParticles.Acceleration(Index) * Dt;
+
+						// Damp Velocity Rule
+						if (bDampVelocityRule)
+						{
+							DampVelocityRule.ApplyFast(MParticles, Dt, Index);
+						}
+
+						// Euler Step with point damping integration
+						MParticles.P(Index) = MParticles.X(Index) + MParticles.V(Index) * DampingIntegrated;
+
+						MParticles.V(Index) *= DampingPowDt;
+					}
+					else  // Process kinematic particles
+					{
+						MKinematicUpdate(MParticles, Dt, MTime, Index);
+					}
+				}, RangeSize < MinParallelBatchSize);
 		}
 	}
 }
diff --git a/Engine/Source/Runtime/Experimental/Chaos/Public/Chaos/PBDEvolution.h b/Engine/Source/Runtime/Experimental/Chaos/Public/Chaos/PBDEvolution.h
--- a/Engine/Source/Runtime/Experimental/Chaos/Public/Chaos/PBDEvolution.h
+++ b/Engine/Source/Runtime/Experimental/Chaos/Public/Chaos/PBDEvolution.h
@@ -26,7 +26,8 @@ class CHAOS_API FPBDEvolution : public TArrayCollection
 		FSolverReal CoefficientOfFriction = (FSolverReal)0.,
 		FSolverReal Damping = (FSolverReal)0.04,
 		FSolverReal LocalDamping = (FSolverReal)0.,
-		bool bDoQuasistatics = false);
+		bool bDoQuasistatics = false, 
+		bool InbUsePerParticleDamping = false);
 	~FPBDEvolution() {}
 
 	// Advance one time step. Filter the input time step if specified.
@@ -170,8 +171,8 @@ class CHAOS_API FPBDEvolution : public TArrayCollection
 	bool GetUseCCD(const uint32 GroupId = 0) const { check(GroupId < TArrayCollection::Size()); return MGroupUseCCDs[GroupId]; }
 	void SetUseCCD(const bool bUseCCD, const uint32 GroupId = 0) { check(GroupId < TArrayCollection::Size()); MGroupUseCCDs[GroupId] = bUseCCD; }
 
-	FSolverReal GetParticleDamping(const uint32 ParticleIndex = 0) const { check(ParticleIndex < this->Particles().Size()); return MParticleDampings[ParticleIndex]; }
-	void SetParticleDamping(const FSolverReal Damping, const uint32 ParticleIndex = 0) { check(ParticleIndex < this->Particles().Size()); MParticleDampings[ParticleIndex] = Damping; }
+	FSolverReal GetParticleDamping(const uint32 ParticleIndex = 0) const { check(ParticleIndex < this->Particles().Size() && bUsePerParticleDamping); return MParticleDampings[ParticleIndex]; }
+	void SetParticleDamping(const FSolverReal Damping, const uint32 ParticleIndex = 0) { check(ParticleIndex < this->Particles().Size()); if (bUsePerParticleDamping) { MParticleDampings[ParticleIndex] = Damping; } }
 
 
 	UE_DEPRECATED(4.27, "Use GetCollisionStatus() instead")
@@ -242,6 +243,7 @@ private:
 	FSolverReal MLocalDamping;
 	FSolverReal MTime;
 	bool bDoQuasistatics = false;
+	bool bUsePerParticleDamping = false;
 };
 
 }  // End namespace Chaos::Softs
