disable the internal cluster creation code that runs for GCs when dealing with cluster unions

diff --git a/Engine/Source/Runtime/Experimental/Chaos/Private/Chaos/ClusterUnionManager.cpp b/Engine/Source/Runtime/Experimental/Chaos/Private/Chaos/ClusterUnionManager.cpp
--- a/Engine/Source/Runtime/Experimental/Chaos/Private/Chaos/ClusterUnionManager.cpp
+++ b/Engine/Source/Runtime/Experimental/Chaos/Private/Chaos/ClusterUnionManager.cpp
@@ -521,6 +521,16 @@ namespace Chaos
 		return INDEX_NONE;
 	}
 
+	bool FClusterUnionManager::IsClusterUnionParticle(FPBDRigidClusteredParticleHandle* Particle)
+	{
+		FClusterUnionIndex UnionIndex = FindClusterUnionIndexFromParticle(Particle);
+		if (FClusterUnion* Union = FindClusterUnion(UnionIndex))
+		{
+			return Union->InternalCluster == Particle;
+		}
+		return false;
+	}
+
 	DECLARE_CYCLE_STAT(TEXT("FClusterUnionManager::UpdateClusterUnionParticlesChildToParent"), STAT_UpdateClusterUnionParticlesChildToParent, STATGROUP_Chaos);
 	void FClusterUnionManager::UpdateClusterUnionParticlesChildToParent(FClusterUnionIndex Index, const TArray<FPBDRigidParticleHandle*>& Particles, const TArray<FTransform>& ChildToParent)
 	{
diff --git a/Engine/Source/Runtime/Experimental/Chaos/Private/Chaos/PBDRigidClustering.cpp b/Engine/Source/Runtime/Experimental/Chaos/Private/Chaos/PBDRigidClustering.cpp
--- a/Engine/Source/Runtime/Experimental/Chaos/Private/Chaos/PBDRigidClustering.cpp
+++ b/Engine/Source/Runtime/Experimental/Chaos/Private/Chaos/PBDRigidClustering.cpp
@@ -922,6 +922,7 @@ namespace Chaos
 
 		if (ActivatedChildren.Num() > 0)
 		{
+			const bool bIsClusterUnion = ClusterUnionManager.IsClusterUnionParticle(ClusteredParticle);
 			if (Children.Num() == 0)
 			{
 				// Free the memory if we can do so cheaply (no data copies).
@@ -938,7 +939,10 @@ namespace Chaos
 					RemoveNodeConnections(Child);
 				}
 
-				if (Children.Num())
+				// If we're breaking a geometry collection, we'll need to create internal clusters to parent the remaining particles.
+				// However, we do not need to do this if we're currently operating on a cluster union! Its remaining particles should stay
+				// attached to the cluster union because they can handle particles being dynamically added/removed.
+				if (Children.Num() && !bIsClusterUnion)
 				{
 					TArray<FParticleIsland> Islands = FindIslandsInChildren(ClusteredParticle);
 					for (const FParticleIsland& Island: Islands)
@@ -964,8 +968,11 @@ namespace Chaos
 				UpdateKinematicProperties(Child, MChildren, MEvolution);
 			}
 
-			//disable cluster
-			DisableCluster(ClusteredParticle);
+			// Disable the cluster only if we're not a cluster union. Cluster unions will handle themselves separately.
+			if (!bIsClusterUnion)
+			{
+				DisableCluster(ClusteredParticle);
+			}
 		}
 
 		return ActivatedChildren;
diff --git a/Engine/Source/Runtime/Experimental/Chaos/Public/Chaos/ClusterUnionManager.h b/Engine/Source/Runtime/Experimental/Chaos/Public/Chaos/ClusterUnionManager.h
--- a/Engine/Source/Runtime/Experimental/Chaos/Public/Chaos/ClusterUnionManager.h
+++ b/Engine/Source/Runtime/Experimental/Chaos/Public/Chaos/ClusterUnionManager.h
@@ -96,6 +96,9 @@ namespace Chaos
 
 		FClusterUnionIndex FindClusterUnionIndexFromParticle(FPBDRigidParticleHandle* Particle);
 
+		// An extension to FindClusterUnionIndexFromParticle to check whether or not the given particle is the cluster union particle itself.
+		bool IsClusterUnionParticle(FPBDRigidClusteredParticleHandle* Particle);
+
 		// Changes the ChildToParent of a number of particles in a cluster union.
 		void UpdateClusterUnionParticlesChildToParent(FClusterUnionIndex Index, const TArray<FPBDRigidParticleHandle*>& Particles, const TArray<FTransform>& ChildToParent);
 
