Move geometry collection ISM Pool management code into custom renderer object.
Allows us to have different rendering strategies for geometry collection.

diff --git a/Engine/Source/Runtime/Experimental/GeometryCollectionEngine/Private/GeometryCollection/GeometryCollectionComponent.cpp b/Engine/Source/Runtime/Experimental/GeometryCollectionEngine/Private/GeometryCollection/GeometryCollectionComponent.cpp
--- a/Engine/Source/Runtime/Experimental/GeometryCollectionEngine/Private/GeometryCollection/GeometryCollectionComponent.cpp
+++ b/Engine/Source/Runtime/Experimental/GeometryCollectionEngine/Private/GeometryCollection/GeometryCollectionComponent.cpp
@@ -29,7 +29,7 @@
 #include "GeometryCollection/GeometryCollectionUtility.h"
 #include "GeometryCollection/GeometryCollectionISMPoolActor.h"
 #include "GeometryCollection/GeometryCollectionISMPoolComponent.h"
-#include "GeometryCollection/GeometryCollectionISMPoolSubSystem.h"
+#include "GeometryCollection/GeometryCollectionISMPoolRenderer.h"
 #include "Math/Sphere.h"
 #include "Modules/ModuleManager.h"
 #include "Net/Core/PushModel/PushModel.h"
@@ -65,6 +65,7 @@
 #include "GeometryCollection/Facades/CollectionAnchoringFacade.h"
 #include "GeometryCollection/Facades/CollectionRemoveOnBreakFacade.h"
 #include "GeometryCollection/Facades/CollectionInstancedMeshFacade.h"
+#include "GeometryCollection/GeometryCollectionExternalRenderInterface.h"
 
 #include UE_INLINE_GENERATED_CPP_BY_NAME(GeometryCollectionComponent)
 
@@ -105,20 +106,8 @@ static FAutoConsoleVariableRef CVarChaosBoxCalcBoundsISPCEnabled(TEXT("p.Chaos.B
 bool bChaos_GC_CacheComponentSpaceBounds = true;
 FAutoConsoleVariableRef CVarChaosGCCacheComponentSpaceBounds(TEXT("p.Chaos.GC.CacheComponentSpaceBounds"), bChaos_GC_CacheComponentSpaceBounds, TEXT("Cache component space bounds for performance"));
 
-bool bChaos_GC_UseISMPool = true;
-FAutoConsoleVariableRef CVarChaosGCUseISMPool(TEXT("p.Chaos.GC.UseISMPool"), bChaos_GC_UseISMPool, TEXT("When enabled, use the ISM pool if specified"));
-
-bool bChaos_GC_UseISMPoolForNonFracturedParts = true;
-FAutoConsoleVariableRef CVarChaosGCUseISMPoolForNonFracturedParts(TEXT("p.Chaos.GC.UseISMPoolForNonFracturedParts"), bChaos_GC_UseISMPoolForNonFracturedParts, TEXT("When enabled, non fractured part will use the ISM pool if specified"));
-
-bool bChaos_GC_ForceAutoAssignISMPool = false;
-FAutoConsoleVariableRef CVarChaosGCForceAutoAssignISMPool(TEXT("p.Chaos.GC.ForceAutoAssignISMPool"), bChaos_GC_ForceAutoAssignISMPool, TEXT("When enabled, force assignement of ISMPool regardgless of the settings of the components"));
-
-bool bChaos_GC_UseHierarchicalISMForProxyMesh = true;
-FAutoConsoleVariableRef CVarChaosGCUseHierarchicalISMForProxyMesh(TEXT("p.Chaos.GC.UseHierarchicalISMForProxyMesh"), bChaos_GC_UseHierarchicalISMForProxyMesh, TEXT("When enabled along with ISM Pool, proxy mesh will prefer using HISM vs standard ISM"));
-
-bool bChaos_GC_UseHierarchicalISMForLeafMeshes = true;
-FAutoConsoleVariableRef CVarChaosGCUseHierarchicalISMForLeafMeshes(TEXT("p.Chaos.GC.UseHierarchicalISMForLeafMeshes"), bChaos_GC_UseHierarchicalISMForLeafMeshes, TEXT("When enabled along with ISM Pool, leaf meshes will prefer using HISM vs standard ISM"));
+bool bChaos_GC_UseCustomRenderer = true;
+FAutoConsoleVariableRef CVarChaosGCUseISMPool(TEXT("p.Chaos.GC.UseCustomRenderer"), bChaos_GC_UseCustomRenderer, TEXT("When enabled, use a custom renderer if specified"));
 
 bool bChaos_GC_InitConstantDataUseParallelFor = true;
 FAutoConsoleVariableRef CVarChaosGCInitConstantDataUseParallelFor(TEXT("p.Chaos.GC.InitConstantDataUseParallelFor"), bChaos_GC_InitConstantDataUseParallelFor, TEXT("When enabled, InitConstant data will use parallelFor for copying some of the data"));
@@ -129,9 +118,6 @@ FAutoConsoleVariableRef CVarChaosGCInitConstantDataParallelForBatchSize(TEXT("p.
 int32 MaxGeometryCollectionAsyncPhysicsTickIdleTimeMs = 30;
 FAutoConsoleVariableRef CVarMaxGeometryCollectionAsyncPhysicsTickIdleTimeMs(TEXT("p.Chaos.GC.MaxGeometryCollectionAsyncPhysicsTickIdleTimeMs"), MaxGeometryCollectionAsyncPhysicsTickIdleTimeMs, TEXT("Amount of time in milliseconds before the async tick turns off when it is otherwise not doing anything."));
 
-bool bChaos_GC_DeferAddingAutoInstancesToISMPool = true;
-FAutoConsoleVariableRef CVarDeferAddingAutoInstancesToISMPool(TEXT("p.Chaos.GC.DeferAddingAutoInstancesToISMPool"), bChaos_GC_DeferAddingAutoInstancesToISMPool, TEXT("When enabled, auto instances will be added to the ISM pool after the GC is broken"));
-
 DEFINE_LOG_CATEGORY_STATIC(UGCC_LOG, Error, All);
 
 extern FGeometryCollectionDynamicDataPool GDynamicDataPool;
@@ -247,7 +233,6 @@ bool bGeometryCollectionSingleThreadedBoundsCalculation = false;
 FAutoConsoleVariableRef CVarGeometryCollectionSingleThreadedBoundsCalculation(TEXT("p.GeometryCollectionSingleThreadedBoundsCalculation"), bGeometryCollectionSingleThreadedBoundsCalculation, TEXT("[Debug Only] Single threaded bounds calculation. [def:false]"));
 
 
-
 FGeomComponentCacheParameters::FGeomComponentCacheParameters()
 	: CacheMode(EGeometryCollectionCacheType::None)
 	, TargetCache(nullptr)
@@ -343,8 +328,6 @@ UGeometryCollectionComponent::UGeometryCollectionComponent(const FObjectInitiali
 	, bEnableRunTimeDataCollection(false)
 	, RunTimeDataCollectionGuid(FGuid::NewGuid())
 #endif
-	, ISMPool(nullptr)
-	, bAutoAssignISMPool(false)
 	, bEnableReplication(false)
 	, bEnableAbandonAfterLevel(true)
 	, AbandonedCollisionProfileName(UCollisionProfile::CustomCollisionProfileName)
@@ -435,8 +418,7 @@ void UGeometryCollectionComponent::BeginPlay()
 	// default current cache time
 	CurrentCacheTime = MAX_flt;	
 
-	// we only enable ISM if we are playing ( not in editing mode because of various side effect like selection )
-	RegisterToISMPool();
+	RefreshCustomRenderer();
 }
 
 
@@ -447,8 +429,6 @@ void UGeometryCollectionComponent::EndPlay(const EEndPlayReason::Type ReasonEnd)
 	EditorActor = EditorUtilities::GetEditorWorldCounterpartActor(GetTypedOuter<AActor>());
 #endif
 
-	UnregisterFromISMPool();
-
 	Super::EndPlay(ReasonEnd);
 
 	CurrentCacheTime = MAX_flt;
@@ -685,8 +665,7 @@ FPrimitiveSceneProxy* UGeometryCollectionComponent::CreateSceneProxy()
 
 	FPrimitiveSceneProxy* LocalSceneProxy = nullptr;
 
-	const bool bUsesISMPool = this->CanUseISMPool();
-	if (RestCollection && !bUsesISMPool)
+	if (RestCollection && !CanUseCustomRenderer())
 	{
 		if (UseNanite(GetScene()->GetShaderPlatform()) &&
 			RestCollection->EnableNanite &&
@@ -1268,6 +1247,7 @@ void UGeometryCollectionComponent::SetRestState(TArray<FTransform>&& InRestTrans
 	}
 
 	RefreshEmbeddedGeometry();
+	RefreshCustomRenderer();
 }
 
 void UGeometryCollectionComponent::InitializeComponent()
@@ -2339,7 +2319,7 @@ void UGeometryCollectionComponent::UpdateAttachedChildrenTransform() const
 
 bool UGeometryCollectionComponent::HasVisibleGeometry() const
 {
-	return AssignedISMPool || RestCollection->HasVisibleGeometry();
+	return CustomRenderer != nullptr || RestCollection->HasVisibleGeometry();
 }
 
 void UGeometryCollectionComponent::TickComponent(float DeltaTime, enum ELevelTick TickType, FActorComponentTickFunction *ThisTickFunction)
@@ -2401,23 +2381,6 @@ void UGeometryCollectionComponent::AsyncPhysicsTickComponent(float DeltaTime, fl
 
 void UGeometryCollectionComponent::OnRegister()
 {
-	// important : we shoudl assign this as soon as possible to avoid the scene proxy from being created
-	AssignedISMPool = nullptr;
-	if (bChaos_GC_UseISMPool && GetWorld()->IsGameWorld())
-	{
-		if (bAutoAssignISMPool || bChaos_GC_ForceAutoAssignISMPool)
-		{
-			if (UGeometryCollectionISMPoolSubSystem* ISMPoolSubSystem = UWorld::GetSubsystem<UGeometryCollectionISMPoolSubSystem>(GetWorld()))
-			{
-				AssignedISMPool = ISMPoolSubSystem->FindISMPoolActor(*this);
-			}
-		}
-		else
-		{
-			AssignedISMPool = ISMPool;
-		}
-	}
-
 	//UE_LOG(UGCC_LOG, Log, TEXT("GeometryCollectionComponent[%p]::OnRegister()[%p]"), this,RestCollection );
 	ResetDynamicCollection();
 
@@ -2434,13 +2397,28 @@ void UGeometryCollectionComponent::OnRegister()
 
 	InitializeEmbeddedGeometry();
 
+	// Look for any attached custom renderer.
+	CustomRenderer = GetOwner()->FindComponentByInterface<UGeometryCollectionExternalRenderInterface>();
+	// Use the ISMPool custom renderer if it was set.
+	if (!CustomRenderer && ISMPool)
+	{
+		UGeometryCollectionCustomRendererISMPool* ISMPoolRenderer = NewObject<UGeometryCollectionCustomRendererISMPool>(this);
+		ISMPoolRenderer->ISMPoolActor = ISMPool;
+		CustomRenderer = (UGeometryCollectionExternalRenderInterface*)ISMPoolRenderer;
+	}
+
 	Super::OnRegister();
 }
 
 void UGeometryCollectionComponent::OnUnregister()
 {
 	Super::OnUnregister();
-	AssignedISMPool = nullptr;
+
+	if (IGeometryCollectionExternalRenderInterface* RendererInterface = Cast<IGeometryCollectionExternalRenderInterface>(CustomRenderer))
+	{
+		RendererInterface->OnUnregisterGeometryCollection();
+	}
+	CustomRenderer = nullptr;
 }
 
 void UGeometryCollectionComponent::ResetDynamicCollection()
@@ -2840,9 +2818,7 @@ void UGeometryCollectionComponent::UpdateRenderSystemsIfNeeded(bool bDynamicColl
 		// #todo review: When we've made changes to ISMC, we need to move this function call to SetRenderDynamicData_Concurrent
 		RefreshEmbeddedGeometry();
 
-		// we may want to call this when the geometry collection updates ( notified by the proxy buffer updates ) 
-		// otherwise we are getting a frame delay 
-		RefreshISMPoolInstances();
+		RefreshCustomRenderer();
 
 		if (SceneProxy && SceneProxy->IsNaniteMesh())
 		{
@@ -4310,248 +4286,33 @@ void UGeometryCollectionComponent::InitializeEmbeddedGeometry()
 	}
 }
 
-bool UGeometryCollectionComponent::CanUseISMPool() const 
-{
-	return bChaos_GC_UseISMPool && AssignedISMPool && GetWorld()->IsGameWorld();
-}
-
-void UGeometryCollectionComponent::RegisterToISMPool()
-{
-	UnregisterFromISMPool();
-
-	if (CanUseISMPool())
-	{
-		if (UGeometryCollectionISMPoolComponent* ISMPoolComp = AssignedISMPool->GetISMPoolComp())
-		{
-			bool bCanRenderComponent = true;
-			if (RestCollection)
-			{
-				ISMPoolMeshGroupIndex = ISMPoolComp->CreateMeshGroup();
-
-				// root proxy if available 
-				// TODO : if ISM pool is not available : uses a standard static mesh component
-				for (int32 MeshIndex = 0; MeshIndex < RestCollection->RootProxyData.ProxyMeshes.Num(); MeshIndex++)
-				{
-					const TObjectPtr<UStaticMesh>& Mesh = RestCollection->RootProxyData.ProxyMeshes[MeshIndex];
-					if (Mesh != nullptr)
-					{
-						// if we use a mesh proxy hide the component for rendering 
-						bCanRenderComponent = false;
-
-						FGeometryCollectionStaticMeshInstance StaticMeshInstance;
-						StaticMeshInstance.StaticMesh = Mesh;
-						StaticMeshInstance.NumCustomDataFloats = ISMPoolMaterialCustomData.Num();
-
-						ISMPoolRootProxyMeshIds.Add(ISMPoolComp->AddMeshToGroup(ISMPoolMeshGroupIndex, StaticMeshInstance, 1, ISMPoolMaterialCustomData, bChaos_GC_UseHierarchicalISMForProxyMesh));
-					}
-				}
-
-				if (bChaos_GC_UseISMPoolForNonFracturedParts)
-				{
-					// if we use ISM pool for the hierarchy we must hide the component for rendering 
-					bCanRenderComponent = false;
-
-					const bool bHasRootProxyMesh = !ISMPoolRootProxyMeshIds.IsEmpty();
-					if (!(bChaos_GC_DeferAddingAutoInstancesToISMPool && bHasRootProxyMesh))
-					{
-						AddAutoInstancesToISMPool();
-					}
-				}
-			}
-
-			SetVisibility(bCanRenderComponent);
-
-			RefreshISMPoolInstances();
-		}
-	}
-}
-
-void UGeometryCollectionComponent::UnregisterFromISMPool()
+bool UGeometryCollectionComponent::CanUseCustomRenderer() const 
 {
-	if (AssignedISMPool)
-	{
-		if (UGeometryCollectionISMPoolComponent* ISMPoolComp = AssignedISMPool->GetISMPoolComp())
-		{
-			ISMPoolComp->DestroyMeshGroup(ISMPoolMeshGroupIndex);
-			ISMPoolMeshGroupIndex = INDEX_NONE;
-			ISMPoolRootProxyMeshIds.Empty();
-			ISMPoolAutoInstancesMeshIds.Empty();
-		}
-		SetVisibility(true);
-	}
+	return bChaos_GC_UseCustomRenderer && CustomRenderer != nullptr && GetWorld()->IsGameWorld();
 }
 
-void UGeometryCollectionComponent::AddAutoInstancesToISMPool()
+void UGeometryCollectionComponent::RefreshCustomRenderer()
 {
-	if (ISMPoolAutoInstancesMeshIds.Num() > 0 || ISMPoolMeshGroupIndex == -1)
-	{
-		return;
-	}
-
-	if (UGeometryCollectionISMPoolComponent* ISMPoolComp = AssignedISMPool->GetISMPoolComp())
+	if (CanUseCustomRenderer())
 	{
-		if (bChaos_GC_UseISMPoolForNonFracturedParts)
+		if (IGeometryCollectionExternalRenderInterface* RendererInterface = Cast<IGeometryCollectionExternalRenderInterface>(CustomRenderer))
 		{
-			if (RestCollection->GetGeometryCollection())
+			if (RestCollection != nullptr)
 			{
-				// first count the instance per mesh 
-				const int32 NumMeshes = RestCollection->AutoInstanceMeshes.Num();
-				const int32 NumCustomDataFloatsFromComponent = ISMPoolMaterialCustomData.Num();
-
-				// Apply custom instance data from geometry collcection object if custom data array is set up correctly so that it's size matches a multiple of the number of instances.
-				TArray<float, TInlineAllocator<64>> CustomFloatData;
-
-				// now register each mesh 
-				int32 AutoInstanceMeshIndex = 0;
-				for (const FGeometryCollectionAutoInstanceMesh& AutoInstanceMesh : RestCollection->AutoInstanceMeshes)
-				{
-					if (const UStaticMesh* StaticMesh = AutoInstanceMesh.Mesh)
-					{
-						bool bMaterialOverride = false;
-						for (int32 MatIndex = 0; MatIndex < AutoInstanceMesh.Materials.Num(); MatIndex++)
-						{
-							const UMaterialInterface* OriginalMaterial = StaticMesh->GetMaterial(MatIndex);
-							if (OriginalMaterial != AutoInstanceMesh.Materials[MatIndex])
-							{
-								bMaterialOverride = true;
-								break;
-							}
-						}
-						FGeometryCollectionStaticMeshInstance StaticMeshInstance;
-						StaticMeshInstance.StaticMesh = const_cast<UStaticMesh*>(StaticMesh);
-						if (bMaterialOverride)
-						{
-							StaticMeshInstance.MaterialsOverrides = AutoInstanceMesh.Materials;
-						}
-
-						const bool bHasValidAssetCustomInstanceData = AutoInstanceMesh.NumInstances && (AutoInstanceMesh.CustomData.Num() % AutoInstanceMesh.NumInstances == 0);
-						const int32 NumCustomDataFloatsFromAsset = bHasValidAssetCustomInstanceData ? AutoInstanceMesh.GetNumDataPerInstance() : 0;
-						const int32 NumCustomDataFloats = NumCustomDataFloatsFromAsset + NumCustomDataFloatsFromComponent;
+				const int32 RootIndex = GetRootIndex();
+				const bool bIsBroken = DynamicCollection ? !DynamicCollection->Active[RootIndex] : false;
+				RendererInterface->UpdateState(*RestCollection, GetComponentTransform(), bIsBroken);
 
-						StaticMeshInstance.NumCustomDataFloats = NumCustomDataFloats;
-
-						if (NumCustomDataFloats > 0)
-						{
-							CustomFloatData.Reset();
-							CustomFloatData.Reserve(AutoInstanceMesh.NumInstances * NumCustomDataFloats);
-							
-							for (int32 InstanceIndex = 0; InstanceIndex < AutoInstanceMesh.NumInstances; ++InstanceIndex)
-							{
-								CustomFloatData.Append(ISMPoolMaterialCustomData);
-								// Append each set of values from AutoInstanceMaterialCustomData after the common values from ISMPoolMaterialCustomData
-								CustomFloatData.Append(&AutoInstanceMesh.CustomData[InstanceIndex], NumCustomDataFloatsFromAsset);
-							}
-						}
-
-						ISMPoolAutoInstancesMeshIds.Add(ISMPoolComp->AddMeshToGroup(ISMPoolMeshGroupIndex, StaticMeshInstance, AutoInstanceMesh.NumInstances, CustomFloatData, bChaos_GC_UseHierarchicalISMForLeafMeshes));
-					}
-					AutoInstanceMeshIndex++;
-				}
-			}
-		}
-	}
-}
-
-void UGeometryCollectionComponent::RefreshISMPoolInstances()
-{
-	if (CanUseISMPool())
-	{
-		if (UGeometryCollectionISMPoolComponent* ISMPoolComp = AssignedISMPool->GetISMPoolComp())
-		{
-			if (RestCollection)
-			{
-				// default to true for editor purposes?
-				//const bool bCollectionIsDirty = DynamicCollection ? DynamicCollection->IsDirty() : true;
-				if (bChaos_GC_UseISMPoolForNonFracturedParts /*&& bCollectionIsDirty*/)
+				if (bIsBroken)
 				{
-					if (RestCollection->GetGeometryCollection())
-					{
-						const TManagedArray<TSet<int32>>& Children = RestCollection->GetGeometryCollection()->Children;
-
-						const GeometryCollection::Facades::FCollectionInstancedMeshFacade InstancedMeshFacade(*RestCollection->GetGeometryCollection());
-						if (InstancedMeshFacade.IsValid())
-						{
-							const int32 NumTransforms = RestCollection->NumElements(FGeometryCollection::TransformAttribute);
-
-							CalculateGlobalMatrices();
-
-							const FTransform& ComponentTransform = GetComponentTransform();
-
-							constexpr bool bWorlSpace = true;
-							constexpr bool bMarkRenderStateDirty = true;
-							constexpr bool bTeleport = true;
-
-							const int32 RootIndex = GetRootIndex();
-							//const bool bIsBroken = DynamicCollection ? (DynamicCollection->Children[RootIndex].Num() != Children[RootIndex].Num()) : false;
-							const bool bIsBroken = DynamicCollection ? !DynamicCollection->Active[RootIndex] : false;
-							const bool bHasRootProxyMesh = !ISMPoolRootProxyMeshIds.IsEmpty();
-
-							if (bHasRootProxyMesh && !bIsBroken)
-							{
-								if (GlobalMatrices.IsValidIndex(RootIndex))
-								{
-									FTransform RootTransform = FTransform(GlobalMatrices[RootIndex]) * ComponentTransform;
-									for (const int32 ISMPoolRootProxyMeshId : ISMPoolRootProxyMeshIds)
-									{
-										ISMPoolComp->BatchUpdateInstancesTransforms(ISMPoolMeshGroupIndex, ISMPoolRootProxyMeshId, 0, { RootTransform }, bWorlSpace, bMarkRenderStateDirty, bTeleport);
-									}
-								}
-							}
-							else if (bChaos_GC_UseISMPoolForNonFracturedParts)
-							{
-								AddAutoInstancesToISMPool();
-
-								// make sure this mesh is invisible 
-								// todo : should be event based instead of doing it every frame
-								if (bHasRootProxyMesh && GlobalMatrices.IsValidIndex(RootIndex))
-								{
-									FTransform RootTransformZeroScale;
-									RootTransformZeroScale.SetIdentityZeroScale();
-									for (const int32 ISMPoolRootProxyMeshId : ISMPoolRootProxyMeshIds)
-									{
-										ISMPoolComp->BatchUpdateInstancesTransforms(ISMPoolMeshGroupIndex, ISMPoolRootProxyMeshId, 0, { RootTransformZeroScale }, bWorlSpace, bMarkRenderStateDirty, bTeleport);
-									}
-								}
-
-								TArray<FTransform> InstanceTransforms;
-								for (int32 MeshIndex = 0; MeshIndex < ISMPoolAutoInstancesMeshIds.Num(); MeshIndex++)
-								{
-									InstanceTransforms.Reset(NumTransforms); // Allocate for worst case
-									for (int32 TransformIndex = 0; TransformIndex < NumTransforms; TransformIndex++)
-									{
-										const int32 AutoInstanceMeshIndex = InstancedMeshFacade.GetIndex(TransformIndex);
-										if (AutoInstanceMeshIndex == MeshIndex && Children[TransformIndex].Num() == 0)
-										{
-											InstanceTransforms.Add(FTransform(GlobalMatrices[TransformIndex]) * ComponentTransform);
-										}
-									}
-									ISMPoolComp->BatchUpdateInstancesTransforms(ISMPoolMeshGroupIndex, ISMPoolAutoInstancesMeshIds[MeshIndex], 0, InstanceTransforms, bWorlSpace, bMarkRenderStateDirty, bTeleport);
-								}
-							}
-						}
-					}
+					CalculateGlobalMatrices();
+					RendererInterface->UpdateTransforms(*RestCollection, GetComponentTransform(), GlobalMatrices);
 				}
 			}
 		}
 	}
 }
 
-UInstancedStaticMeshComponent* UGeometryCollectionComponent::GetRootProxyISM() const
-{
-	if (AssignedISMPool != nullptr)
-	{
-		if (UGeometryCollectionISMPoolComponent* ISMPoolComp = AssignedISMPool->GetISMPoolComp())
-		{
-			if (RestCollection->AutoInstanceMeshes.Num())
-			{
-				return ISMPoolComp->GetISMForMeshId(ISMPoolMeshGroupIndex, 0);
-			}
-		}
-	}
-
-	return nullptr;
-}
-
 TArray<UStaticMeshComponent*> UGeometryCollectionComponent::CreateProxyComponents() const
 {
 	TArray<UStaticMeshComponent*> Components;
diff --git a/Engine/Source/Runtime/Experimental/GeometryCollectionEngine/Private/GeometryCollection/GeometryCollectionISMPoolRenderer.cpp b/Engine/Source/Runtime/Experimental/GeometryCollectionEngine/Private/GeometryCollection/GeometryCollectionISMPoolRenderer.cpp
new file mode 100644
--- /dev/null
+++ b/Engine/Source/Runtime/Experimental/GeometryCollectionEngine/Private/GeometryCollection/GeometryCollectionISMPoolRenderer.cpp
@@ -0,0 +1,195 @@
+// Copyright Epic Games, Inc. All Rights Reserved.
+#include "GeometryCollection/GeometryCollectionISMPoolRenderer.h"
+
+#include "Engine/StaticMesh.h"
+#include "GeometryCollection/Facades/CollectionInstancedMeshFacade.h"
+#include "GeometryCollection/GeometryCollection.h"
+#include "GeometryCollection/GeometryCollectionISMPoolActor.h"
+#include "GeometryCollection/GeometryCollectionISMPoolComponent.h"
+#include "GeometryCollection/GeometryCollectionISMPoolSubSystem.h"
+#include "GeometryCollection/GeometryCollectionObject.h"
+
+#include UE_INLINE_GENERATED_CPP_BY_NAME(GeometryCollectionISMPoolRenderer)
+
+void UGeometryCollectionCustomRendererISMPool::UpdateState(UGeometryCollection const& InGeometryCollection, FTransform const& InBaseTransform, bool bIsBroken)
+{
+	if (!bIsBroken && MergedMeshGroup.GroupIndex == INDEX_NONE)
+	{
+		// Remove broken primitives.
+		ReleaseGroup(InstancesGroup);
+
+		// Add merged mesh.
+		InitMergedMeshFromGeometryCollection(InGeometryCollection);
+		UpdateMergedMeshTransforms(InBaseTransform);
+	}
+
+	if (bIsBroken && InstancesGroup.GroupIndex == INDEX_NONE)
+	{
+		// Remove merged mesh.
+		ReleaseGroup(MergedMeshGroup);
+
+		// Add broken primitives.
+		InitInstancesFromGeometryCollection(InGeometryCollection);
+	}
+}
+
+void UGeometryCollectionCustomRendererISMPool::UpdateTransforms(UGeometryCollection const& InGeometryCollection, FTransform const& InBaseTransform, TArrayView<const FMatrix> InMatrices)
+{
+	UpdateInstanceTransforms(InGeometryCollection, InBaseTransform, InMatrices);
+}
+
+void UGeometryCollectionCustomRendererISMPool::OnUnregisterGeometryCollection()
+{
+	ReleaseGroup(MergedMeshGroup);
+	ReleaseGroup(InstancesGroup);
+
+	ISMPoolActor = nullptr;
+}
+
+void UGeometryCollectionCustomRendererISMPool::InitMergedMeshFromGeometryCollection(UGeometryCollection const& InGeometryCollection)
+{
+	if (InGeometryCollection.RootProxyData.ProxyMeshes.Num() == 0)
+	{
+		return;
+	}
+
+	UGeometryCollectionISMPoolComponent* ISMPoolComponent = ISMPoolActor != nullptr ? ISMPoolActor->GetISMPoolComp() : nullptr;
+	MergedMeshGroup.GroupIndex = ISMPoolComponent != nullptr ? ISMPoolComponent->CreateMeshGroup() : INDEX_NONE;
+
+	if (MergedMeshGroup.GroupIndex == INDEX_NONE)
+	{
+		return;
+	}
+
+	for (UStaticMesh* StaticMesh : InGeometryCollection.RootProxyData.ProxyMeshes)
+	{
+		if (StaticMesh == nullptr)
+		{
+			continue;
+		}
+
+		FGeometryCollectionStaticMeshInstance StaticMeshInstance;
+		StaticMeshInstance.StaticMesh = StaticMesh;
+
+		TArray<float> DummyCustomData;
+		MergedMeshGroup.MeshIds.Add(ISMPoolComponent->AddMeshToGroup(MergedMeshGroup.GroupIndex, StaticMeshInstance, 1, DummyCustomData, true));
+	}
+}
+
+void UGeometryCollectionCustomRendererISMPool::InitInstancesFromGeometryCollection(UGeometryCollection const& InGeometryCollection)
+{
+	const int32 NumMeshes = InGeometryCollection.AutoInstanceMeshes.Num();
+	if (NumMeshes == 0)
+	{
+		return;
+	}
+
+	UGeometryCollectionISMPoolComponent* ISMPoolComponent = ISMPoolActor != nullptr ? ISMPoolActor->GetISMPoolComp() : nullptr;
+	InstancesGroup.GroupIndex = ISMPoolComponent != nullptr ? ISMPoolComponent->CreateMeshGroup() : INDEX_NONE;
+
+	if (InstancesGroup.GroupIndex == INDEX_NONE)
+	{
+		return;
+	}
+
+	InstancesGroup.MeshIds.Reserve(NumMeshes);
+
+	for (const FGeometryCollectionAutoInstanceMesh& AutoInstanceMesh : InGeometryCollection.AutoInstanceMeshes)
+	{
+		if (const UStaticMesh* StaticMesh = AutoInstanceMesh.Mesh)
+		{
+			bool bMaterialOverride = false;
+			for (int32 MatIndex = 0; MatIndex < AutoInstanceMesh.Materials.Num(); MatIndex++)
+			{
+				const UMaterialInterface* OriginalMaterial = StaticMesh->GetMaterial(MatIndex);
+				if (OriginalMaterial != AutoInstanceMesh.Materials[MatIndex])
+				{
+					bMaterialOverride = true;
+					break;
+				}
+			}
+			FGeometryCollectionStaticMeshInstance StaticMeshInstance;
+			StaticMeshInstance.StaticMesh = const_cast<UStaticMesh*>(StaticMesh);
+			if (bMaterialOverride)
+			{
+				StaticMeshInstance.MaterialsOverrides = AutoInstanceMesh.Materials;
+			}
+
+			TArray<float> DummyCustomData;
+			InstancesGroup.MeshIds.Add(ISMPoolComponent->AddMeshToGroup(InstancesGroup.GroupIndex, StaticMeshInstance, AutoInstanceMesh.NumInstances, DummyCustomData, true));
+		}
+	}
+}
+
+void UGeometryCollectionCustomRendererISMPool::UpdateMergedMeshTransforms(FTransform const& InBaseTransform)
+{
+	if (MergedMeshGroup.GroupIndex == INDEX_NONE)
+	{
+		return;
+	}
+
+	UGeometryCollectionISMPoolComponent* ISMPoolComponent = ISMPoolActor != nullptr ? ISMPoolActor->GetISMPoolComp() : nullptr;
+	if (ISMPoolComponent == nullptr)
+	{
+		return;
+	}
+
+	TArray<FTransform> InstanceTransforms;
+	InstanceTransforms.Add(InBaseTransform);
+
+	for (int32 MeshIndex = 0; MeshIndex < MergedMeshGroup.MeshIds.Num(); MeshIndex++)
+	{
+		ISMPoolComponent->BatchUpdateInstancesTransforms(MergedMeshGroup.GroupIndex, MergedMeshGroup.MeshIds[MeshIndex], 0, InstanceTransforms, true/*bWorldSpace*/, false/*bMarkRenderStateDirty*/, false/*bTeleport*/);
+	}
+}
+
+void UGeometryCollectionCustomRendererISMPool::UpdateInstanceTransforms(UGeometryCollection const& InGeometryCollection, FTransform const& InBaseTransform, TArrayView<const FMatrix> InMatrices)
+{
+	if (InstancesGroup.GroupIndex == INDEX_NONE)
+	{
+		return;
+	}
+
+	UGeometryCollectionISMPoolComponent* ISMPoolComponent = ISMPoolActor != nullptr ? ISMPoolActor->GetISMPoolComp() : nullptr;
+	if (ISMPoolComponent == nullptr)
+	{
+		return;
+	}
+
+	const GeometryCollection::Facades::FCollectionInstancedMeshFacade InstancedMeshFacade(*InGeometryCollection.GetGeometryCollection());
+	if (!InstancedMeshFacade.IsValid())
+	{
+		return;
+	}
+
+	const int32 NumTransforms = InGeometryCollection.NumElements(FGeometryCollection::TransformAttribute);
+	const TManagedArray<TSet<int32>>& Children = InGeometryCollection.GetGeometryCollection()->Children;
+
+	TArray<FTransform> InstanceTransforms;
+	for (int32 MeshIndex = 0; MeshIndex < InstancesGroup.MeshIds.Num(); MeshIndex++)
+	{
+		InstanceTransforms.Reset(NumTransforms); // Allocate for worst case
+		for (int32 TransformIndex = 0; TransformIndex < NumTransforms; TransformIndex++)
+		{
+			const int32 AutoInstanceMeshIndex = InstancedMeshFacade.GetIndex(TransformIndex);
+			if (AutoInstanceMeshIndex == MeshIndex && Children[TransformIndex].Num() == 0)
+			{
+				InstanceTransforms.Add(FTransform(InMatrices[TransformIndex]) * InBaseTransform);
+			}
+		}
+		
+		ISMPoolComponent->BatchUpdateInstancesTransforms(InstancesGroup.GroupIndex, InstancesGroup.MeshIds[MeshIndex], 0, InstanceTransforms, true/*bWorldSpace*/, false/*bMarkRenderStateDirty*/, false/*bTeleport*/);
+	}
+}
+
+void UGeometryCollectionCustomRendererISMPool::ReleaseGroup(FISMPoolGroup& InOutGroup)
+{
+	UGeometryCollectionISMPoolComponent* ISMPoolComponent = ISMPoolActor != nullptr ? ISMPoolActor->GetISMPoolComp() : nullptr;
+	if (ISMPoolComponent != nullptr)
+	{
+		ISMPoolComponent->DestroyMeshGroup(InOutGroup.GroupIndex);
+	}
+
+	InOutGroup.GroupIndex = INDEX_NONE;
+	InOutGroup.MeshIds.Empty();
+}
diff --git a/Engine/Source/Runtime/Experimental/GeometryCollectionEngine/Private/GeometryCollection/GeometryCollectionISMPoolRenderer.h b/Engine/Source/Runtime/Experimental/GeometryCollectionEngine/Private/GeometryCollection/GeometryCollectionISMPoolRenderer.h
new file mode 100644
--- /dev/null
+++ b/Engine/Source/Runtime/Experimental/GeometryCollectionEngine/Private/GeometryCollection/GeometryCollectionISMPoolRenderer.h
@@ -0,0 +1,46 @@
+// Copyright Epic Games, Inc. All Rights Reserved.
+
+#pragma once
+
+#include "GeometryCollection/GeometryCollectionExternalRenderInterface.h"
+
+#include "GeometryCollectionISMPoolRenderer.generated.h"
+
+class AGeometryCollectionISMPoolActor;
+
+/** Implementation of a custom renderer that pushes AutoInstanceMeshes to an ISMPool. */
+UCLASS()
+class UGeometryCollectionCustomRendererISMPool : public UObject, public IGeometryCollectionExternalRenderInterface
+{
+	GENERATED_BODY()
+
+public:
+	//~ Begin IGeometryCollectionExternalRenderInterface Interface.
+	virtual void UpdateState(UGeometryCollection const& InGeometryCollection, FTransform const& InBaseTransform, bool bInIsBroken) override;
+	virtual void UpdateTransforms(UGeometryCollection const& InGeometryCollection, FTransform const& InBaseTransform, TArrayView<const FMatrix> InMatrices) override;
+	virtual void OnUnregisterGeometryCollection() override;
+	//~ End IGeometryCollectionExternalRenderInterface Interface.
+
+	/** Instanced Static Mesh Pool actor that is used to render our meshes. */
+	UPROPERTY(Transient)
+	TObjectPtr<AGeometryCollectionISMPoolActor> ISMPoolActor = nullptr;
+
+	/** Description for a group of meshes that are added/updated together. */
+	struct FISMPoolGroup
+	{
+		int32 GroupIndex = INDEX_NONE;
+		TArray<int32> MeshIds;
+	};
+
+protected:
+	/** ISM pool groups per rendering element type. */
+	FISMPoolGroup MergedMeshGroup;
+	FISMPoolGroup InstancesGroup;
+
+private:
+	void InitMergedMeshFromGeometryCollection(UGeometryCollection const& InGeometryCollection);
+	void InitInstancesFromGeometryCollection(UGeometryCollection const& InGeometryCollection);
+	void UpdateMergedMeshTransforms(FTransform const& InBaseTransform);
+	void UpdateInstanceTransforms(UGeometryCollection const& InGeometryCollection, FTransform const& InBaseTransform, TArrayView<const FMatrix> InMatrices);
+	void ReleaseGroup(FISMPoolGroup& InOutGroup);
+};
diff --git a/Engine/Source/Runtime/Experimental/GeometryCollectionEngine/Public/GeometryCollection/GeometryCollectionComponent.h b/Engine/Source/Runtime/Experimental/GeometryCollectionEngine/Public/GeometryCollection/GeometryCollectionComponent.h
--- a/Engine/Source/Runtime/Experimental/GeometryCollectionEngine/Public/GeometryCollection/GeometryCollectionComponent.h
+++ b/Engine/Source/Runtime/Experimental/GeometryCollectionEngine/Public/GeometryCollection/GeometryCollectionComponent.h
@@ -47,6 +47,7 @@ struct FGeometryCollectionSection;
 struct FDamageCollector;
 class FPhysScene_Chaos;
 class AGeometryCollectionISMPoolActor;
+class UGeometryCollectionExternalRenderInterface;
 enum ESimulationInitializationState : uint8;
 enum class EClusterConnectionTypeEnum : uint8;
 enum class EInitialVelocityTypeEnum : uint8;
@@ -1012,7 +1013,6 @@ public:
 	const FGeometryDynamicCollection* GetDynamicCollection() const;
 	FGeometryDynamicCollection* GetDynamicCollection();  // TEMP HACK?
 
-	UInstancedStaticMeshComponent* GetRootProxyISM() const;
 	TArray<UStaticMeshComponent*> CreateProxyComponents() const;
 
 public:
@@ -1053,6 +1053,10 @@ public:
 	UPROPERTY(EditAnywhere, BlueprintReadOnly, Category = "ChaosPhysics|General")
 	bool bStoreVelocities;
 
+	/** An custom renderer object. */
+	UPROPERTY(Transient)
+	TObjectPtr<UGeometryCollectionExternalRenderInterface> CustomRenderer;
+
 protected:
 	/** Display Bone Colors instead of assigned materials */
 	UPROPERTY(EditAnywhere, BlueprintReadWrite, Category = "ChaosPhysics|General")
@@ -1076,21 +1080,6 @@ protected:
 	UPROPERTY(EditAnywhere, BlueprintReadWrite, Category = "ChaosPhysics|Rendering", meta = (DisplayName = "ISM Pool"))
 	TObjectPtr<AGeometryCollectionISMPoolActor> ISMPool;
 
-	UPROPERTY(EditAnywhere, BlueprintReadWrite, Category = "ChaosPhysics|Rendering")
-	bool bAutoAssignISMPool;
-
-	/** 
-	 * Values that will be added to the Per Instance Custom Data for every instance in all ISMPools. 
-	 * This will be added in slots after values that come from the Geometry Collection AutoInstanceMaterialCustomData. 
-	 */
-	UPROPERTY(EditAnywhere, BlueprintReadWrite, Category = Rendering)
-	TArray<float> ISMPoolMaterialCustomData;
-
-	TObjectPtr<AGeometryCollectionISMPoolActor> AssignedISMPool = nullptr;
-	int32 ISMPoolMeshGroupIndex = INDEX_NONE;
-	TArray<int32> ISMPoolRootProxyMeshIds;
-	TArray<int32> ISMPoolAutoInstancesMeshIds;
-
 	/** Populate the dynamic particle data for the render thread. */
 	FGeometryCollectionDynamicData* InitDynamicData(bool bInitialization);
 
@@ -1253,17 +1242,10 @@ private:
 	bool IsEmbeddedGeometryValid() const;
 	void ClearEmbeddedGeometry();
 
-	/** return true is an ISM pool is set and the feature is enabled */
-	bool CanUseISMPool() const;
-
-	/** If an ISM pool is set, register to it  */
-	void RegisterToISMPool();
-
-	/** If an ISM pool is set, unregister to it  */
-	void UnregisterFromISMPool();
-
-	/** update ISM transforms */
-	void RefreshISMPoolInstances();
+	/** return true if a a custom renderer has been set and the feature is enabled */
+	bool CanUseCustomRenderer() const;
+	/** update the custom renderer */
+	void RefreshCustomRenderer();
 
 	/** return true if the root cluster is not longer active at runtime */
 	UFUNCTION(BlueprintCallable, Category = "ChaosPhysics")
@@ -1286,8 +1268,6 @@ private:
 
 	void OnPostPhysicsSync();
 
-	void AddAutoInstancesToISMPool();
-
 	bool HasVisibleGeometry() const;
 
 	/** The clusters we need to replicate */
diff --git a/Engine/Source/Runtime/Experimental/GeometryCollectionEngine/Public/GeometryCollection/GeometryCollectionExternalRenderInterface.h b/Engine/Source/Runtime/Experimental/GeometryCollectionEngine/Public/GeometryCollection/GeometryCollectionExternalRenderInterface.h
new file mode 100644
--- /dev/null
+++ b/Engine/Source/Runtime/Experimental/GeometryCollectionEngine/Public/GeometryCollection/GeometryCollectionExternalRenderInterface.h
@@ -0,0 +1,25 @@
+// Copyright Epic Games, Inc. All Rights Reserved.
+
+#pragma once
+
+#include "UObject/Interface.h"
+
+#include "GeometryCollectionExternalRenderInterface.generated.h"
+
+class UGeometryCollection;
+
+UINTERFACE()
+class GEOMETRYCOLLECTIONENGINE_API UGeometryCollectionExternalRenderInterface : public UInterface
+{
+	GENERATED_BODY()
+};
+
+class GEOMETRYCOLLECTIONENGINE_API IGeometryCollectionExternalRenderInterface
+{
+	GENERATED_BODY()
+
+public:
+	virtual void UpdateState(UGeometryCollection const& InGeometryCollection, FTransform const& InBaseTransform, bool bInIsBroken) = 0;
+	virtual void UpdateTransforms(UGeometryCollection const& InGeometryCollection, FTransform const& InBaseTransform, TArrayView<const FMatrix> InMatrices) = 0;
+	virtual void OnUnregisterGeometryCollection() = 0;
+};
diff --git a/Engine/Source/Runtime/Experimental/GeometryCollectionEngine/Public/GeometryCollection/GeometryCollectionISMPoolDebugDrawComponent.cpp b/Engine/Source/Runtime/Experimental/GeometryCollectionEngine/Public/GeometryCollection/GeometryCollectionISMPoolDebugDrawComponent.cpp
--- a/Engine/Source/Runtime/Experimental/GeometryCollectionEngine/Public/GeometryCollection/GeometryCollectionISMPoolDebugDrawComponent.cpp
+++ b/Engine/Source/Runtime/Experimental/GeometryCollectionEngine/Public/GeometryCollection/GeometryCollectionISMPoolDebugDrawComponent.cpp
@@ -10,7 +10,6 @@
 #include "EngineUtils.h"
 #include "GameFramework/Pawn.h"
 #include "GameFramework/PlayerController.h"
-#include "GeometryCollectionComponent.h"
 #include "GeometryCollectionISMPoolActor.h"
 #include "GeometryCollectionISMPoolComponent.h"
 #include "Misc/CoreDelegates.h"
@@ -178,13 +177,7 @@ void UGeometryCollectionISMPoolDebugDrawComponent::TickComponent(float DeltaTime
 					const bool bHit = World->LineTraceSingleByChannel(HitResult, TraceStart, TraceEnd, ECC_Visibility, TraceParams);
 					if (bHit && HitResult.GetComponent() != nullptr)
 					{
-						if (UGeometryCollectionComponent* GeometryCollectionComponent = Cast<UGeometryCollectionComponent>(HitResult.GetComponent()))
-						{
-							// We hit the GC but don't have any information which part.
-							// For now we just select the ISM of the root proxy.
-							FoundISM = GeometryCollectionComponent->GetRootProxyISM();
-						}
-						else if (UInstancedStaticMeshComponent* ISMComponent = Cast<UInstancedStaticMeshComponent>(HitResult.GetComponent()))
+						if (UInstancedStaticMeshComponent* ISMComponent = Cast<UInstancedStaticMeshComponent>(HitResult.GetComponent()))
 						{
 							FoundISM = ISMComponent;
 						}
