Chaos - various optimizations for ApplyKinematicTargets
- eliminate need for MarkTransientDirtyParticle for moving kinematics (except the frame when they stop moving)
- run in a parallel-for
- don't calculate angular velocity if rotation has not changed

We now require the MovingKinematics optimization that was previously on aa cvar. Also we rely on the fact that Kinematics are actually Dynamics in a kinematic state (required because the DirtyParticlesView is only for Dynamics).

diff --git a/Engine/Source/Programs/HeadlessChaos/Private/HeadlessChaosTestPhysicsScene.cpp b/Engine/Source/Programs/HeadlessChaos/Private/HeadlessChaosTestPhysicsScene.cpp
--- a/Engine/Source/Programs/HeadlessChaos/Private/HeadlessChaosTestPhysicsScene.cpp
+++ b/Engine/Source/Programs/HeadlessChaos/Private/HeadlessChaosTestPhysicsScene.cpp
@@ -1394,8 +1394,12 @@ namespace ChaosTest {
 			Scene.StartFrame();
 			Scene.EndFrame();
 			// Test X and V on GT
-			EXPECT_EQ(Particle.X(), CorrectX);
-			EXPECT_EQ(Particle.V(), CorrectV);
+			EXPECT_EQ(Particle.X().X, CorrectX.X);
+			EXPECT_EQ(Particle.X().Y, CorrectX.Y);
+			EXPECT_EQ(Particle.X().Z, CorrectX.Z);
+			EXPECT_EQ(Particle.V().X, CorrectV.X);
+			EXPECT_EQ(Particle.V().Y, CorrectV.Y);
+			EXPECT_EQ(Particle.V().Z, CorrectV.Z);
 			// Test X and V on PT, this is going to be used in OnPreSimulate_Internal in next frame.
 			Callback->GetProducerInputData_External()->CorrectX = CorrectX;
 			Callback->GetProducerInputData_External()->CorrectV = CorrectV;
diff --git a/Engine/Source/Runtime/Experimental/Chaos/Private/Chaos/PBDRigidsEvolution.cpp b/Engine/Source/Runtime/Experimental/Chaos/Private/Chaos/PBDRigidsEvolution.cpp
--- a/Engine/Source/Runtime/Experimental/Chaos/Private/Chaos/PBDRigidsEvolution.cpp
+++ b/Engine/Source/Runtime/Experimental/Chaos/Private/Chaos/PBDRigidsEvolution.cpp
@@ -21,9 +21,6 @@ FAutoConsoleVariableRef CVarChaosNumPushOutIterationsOverride(TEXT("p.ChaosNumPu
 int32 ChaosNumContactIterationsOverride = -1;
 FAutoConsoleVariableRef CVarChaosNumContactIterationsOverride(TEXT("p.ChaosNumContactIterationsOverride"), ChaosNumContactIterationsOverride, TEXT("Override for num contact iterations if >= 0. [def:-1]"));
 
-int32 ChaosNonMovingKinematicUpdateOptimization = 1;
-FAutoConsoleVariableRef CVarChaosNonMovingKinematicUpdateOptimization(TEXT("p.ChaosNonMovingKinematicUpdateOptimization"), ChaosNonMovingKinematicUpdateOptimization, TEXT("When enabled (1), keep track of moving kinematics and only call ApplyKinematicTargets for those ones. [def:1]"));
-
 namespace Chaos
 {
 	CHAOS_API int32 FixBadAccelerationStructureRemoval = 1;
diff --git a/Engine/Source/Runtime/Experimental/Chaos/Public/Chaos/Collision/SpatialAccelerationBroadPhase.h b/Engine/Source/Runtime/Experimental/Chaos/Public/Chaos/Collision/SpatialAccelerationBroadPhase.h
--- a/Engine/Source/Runtime/Experimental/Chaos/Public/Chaos/Collision/SpatialAccelerationBroadPhase.h
+++ b/Engine/Source/Runtime/Experimental/Chaos/Public/Chaos/Collision/SpatialAccelerationBroadPhase.h
@@ -449,7 +449,7 @@ namespace Chaos
 			if(!bNeedsResim)
 			{
 				const TParticleView<TPBDRigidParticles<FReal, 3>>& DynamicSleepingView = Particles.GetNonDisabledDynamicView();
-				const TParticleView<TKinematicGeometryParticles<FReal, 3>>& DynamicMovingKinematicView = Particles.GetActiveDynamicMovingKinematicParticlesView();
+				const TParticleView<TPBDRigidParticles<FReal, 3>>& DynamicMovingKinematicView = Particles.GetActiveDynamicMovingKinematicParticlesView();
 
 				// Usually we ignore sleeping particles in the outer loop and iterate over awake-dynamics and moving-kinematics. 
 				// However, for scenes with a very large number of moving kinematics, it is faster to loop over awake-dynamics
diff --git a/Engine/Source/Runtime/Experimental/Chaos/Public/Chaos/PBDRigidsEvolution.h b/Engine/Source/Runtime/Experimental/Chaos/Public/Chaos/PBDRigidsEvolution.h
--- a/Engine/Source/Runtime/Experimental/Chaos/Public/Chaos/PBDRigidsEvolution.h
+++ b/Engine/Source/Runtime/Experimental/Chaos/Public/Chaos/PBDRigidsEvolution.h
@@ -24,7 +24,6 @@ extern int32 ChaosRigidsEvolutionApplyAllowEarlyOutCVar;
 extern int32 ChaosRigidsEvolutionApplyPushoutAllowEarlyOutCVar;
 extern int32 ChaosNumPushOutIterationsOverride;
 extern int32 ChaosNumContactIterationsOverride;
-extern int32 ChaosNonMovingKinematicUpdateOptimization;
 
 namespace Chaos
 {
@@ -414,14 +413,11 @@ public:
 	{
 		if (KinematicHandle)
 		{
-			if (ChaosNonMovingKinematicUpdateOptimization)
+			// optimization : we keep track of moving kinematic targets ( list gets clear every frame )
+			if (NewKinematicTarget.GetMode() != EKinematicTargetMode::None)
 			{
-				// optimization : we keep track of moving kinematic targets ( list gets clear every frame )
-				if (NewKinematicTarget.GetMode() != EKinematicTargetMode::None)
-				{
-					// move particle from "non-moving" kinematics to "moving" kinematics
-					Particles.MarkMovingKinematic(KinematicHandle);
-				}
+				// move particle from "non-moving" kinematics to "moving" kinematics
+				Particles.MarkMovingKinematic(KinematicHandle);
 			}
 			KinematicHandle->SetKinematicTarget(NewKinematicTarget);
 		}
@@ -756,24 +752,16 @@ public:
 
 		const bool IsLastStep = (FMath::IsNearlyEqual(StepFraction, (FReal)1, (FReal)UE_KINDA_SMALL_NUMBER));
 
-		const FReal MinDt = 1e-6f;
-		auto GetKinematicView = [this]()
-		{
-			if (ChaosNonMovingKinematicUpdateOptimization)
-			{
-				return Particles.GetActiveMovingKinematicParticlesView();
-			}
-			return Particles.GetActiveKinematicParticlesView();
-		};
-
-		// We do not update the particle views every time we change a particle - we do it once at the end
-		const bool bUpdateParticleViews = false;
-
-		for (auto& Particle : GetKinematicView())
+		// NOTE: ApplyKinematicTargetForParticle is run in a parallel-for. We only write to particle state
+		const auto& ApplyParticleKinematicTarget = 
+		[Dt, StepFraction, IsLastStep]
+		(FTransientPBDRigidParticleHandle& Particle, const int32 ParticleIndex)
+		-> void
 		{
 			TKinematicTarget<FReal, 3>& KinematicTarget = Particle.KinematicTarget();
 			const FVec3 CurrentX = Particle.X();
 			const FRotation3 CurrentR = Particle.R();
+			constexpr FReal MinDt = 1e-6f;
 
 			bool bMoved = false;
 			switch (KinematicTarget.GetMode())
@@ -788,7 +776,6 @@ public:
 				Particle.V() = FVec3(0, 0, 0);
 				Particle.W() = FVec3(0, 0, 0);
 				KinematicTarget.SetMode(EKinematicTargetMode::None);
-				Particles.MarkTransientDirtyParticle(Particle.Handle(), bUpdateParticleViews);
 				break;
 			}
 
@@ -812,26 +799,26 @@ public:
 					NewR = FRotation3::Slerp(CurrentR, KinematicTarget.GetTarget().GetRotation(), decltype(FQuat::X)(StepFraction));
 				}
 
-				bMoved = !FVec3::IsNearlyEqual(NewX, CurrentX, UE_SMALL_NUMBER) || !FRotation3::IsNearlyEqual(NewR, CurrentR, UE_SMALL_NUMBER);
-				FVec3 V = FVec3(0);
-				FVec3 W = FVec3(0);
-				if (bMoved)
+				const bool bPositionChanged = !FVec3::IsNearlyEqual(NewX, CurrentX, UE_SMALL_NUMBER);
+				const bool bRotationChanged = !FRotation3::IsNearlyEqual(NewR, CurrentR, UE_SMALL_NUMBER);
+				bMoved = bPositionChanged || bRotationChanged;
+				FVec3 NewV = FVec3(0);
+				FVec3 NewW = FVec3(0);
+				if (Dt > MinDt)
 				{
-					if (Dt > MinDt)
+					if (bPositionChanged)
 					{
-						V = FVec3::CalculateVelocity(CurrentX, NewX, Dt);
-						Particle.V() = V;
-
-						W = FRotation3::CalculateAngularVelocity(CurrentR, NewR, Dt);
-						Particle.W() = W;
+						NewV = FVec3::CalculateVelocity(CurrentX, NewX, Dt);
+					}
+					if (bRotationChanged)
+					{
+						NewW = FRotation3::CalculateAngularVelocity(CurrentR, NewR, Dt);
 					}
-
-					Particle.X() = NewX;
-					Particle.R() = NewR;
 				}
-				Particle.V() = V;
-				Particle.W() = W;
-				Particles.MarkTransientDirtyParticle(Particle.Handle(), bUpdateParticleViews);
+				Particle.X() = NewX;
+				Particle.R() = NewR;
+				Particle.V() = NewV;
+				Particle.W() = NewW;
 
 				break;
 			}
@@ -842,37 +829,35 @@ public:
 				bMoved = true;
 				Particle.X() = Particle.X() + Particle.V() * Dt;
 				Particle.R() = FRotation3::IntegrateRotationWithAngularVelocity(Particle.R(), Particle.W(), Dt);
-				Particles.MarkTransientDirtyParticle(Particle.Handle(), bUpdateParticleViews);
 				break;
 			}
 			}
 			
 			// Set positions and previous velocities if we can
 			// Note: At present kininematics are in fact rigid bodies
-			auto* Rigid = Particle.CastToRigidParticle();
-			if (Rigid)
-			{
-				Rigid->P() = Rigid->X();
-				Rigid->Q() = Rigid->R();
-				Rigid->PreV() = Rigid->V();
-				Rigid->PreW() = Rigid->W();
+			Particle.P() = Particle.X();
+			Particle.Q() = Particle.R();
+			Particle.PreV() = Particle.V();
+			Particle.PreW() = Particle.W();
 
-				if (bMoved)
+			if (bMoved)
+			{
+				if (!Particle.CCDEnabled())
 				{
-					if (!Rigid->CCDEnabled())
-					{
-						Rigid->UpdateWorldSpaceState(FRigidTransform3(Rigid->P(), Rigid->Q()), FVec3(0));
-					}
-					else
-					{
-						Rigid->UpdateWorldSpaceStateSwept(FRigidTransform3(Rigid->P(), Rigid->Q()), FVec3(0), -Rigid->V() * Dt);
-					}
+					Particle.UpdateWorldSpaceState(FRigidTransform3(Particle.P(), Particle.Q()), FVec3(0));
+				}
+				else
+				{
+					Particle.UpdateWorldSpaceStateSwept(FRigidTransform3(Particle.P(), Particle.Q()), FVec3(0), -Particle.V() * Dt);
 				}
 			}
-		}
+		};
+
+		// Apply kinematic targets in parallel
+		Particles.GetActiveMovingKinematicParticlesView().ParallelFor(ApplyParticleKinematicTarget);
 
 		// done with update, let's clear the tracking structures
-		if (IsLastStep && ChaosNonMovingKinematicUpdateOptimization)
+		if (IsLastStep)
 		{
 			Particles.UpdateAllMovingKinematic();
 		}
diff --git a/Engine/Source/Runtime/Experimental/Chaos/Public/Chaos/PBDRigidsSOAs.h b/Engine/Source/Runtime/Experimental/Chaos/Public/Chaos/PBDRigidsSOAs.h
--- a/Engine/Source/Runtime/Experimental/Chaos/Public/Chaos/PBDRigidsSOAs.h
+++ b/Engine/Source/Runtime/Experimental/Chaos/Public/Chaos/PBDRigidsSOAs.h
@@ -372,7 +372,7 @@ public:
 			}
 
 			RemoveFromActiveArray(PBDRigid, /*bStillDirty=*/ false);
-			MovingKinematicsMapArray.Remove(PBDRigid->CastToKinematicParticle());
+			MovingKinematicsMapArray.Remove(PBDRigid);
 
 			if (auto PBDRigidClustered = Particle->CastToClustered())
 			{
@@ -654,7 +654,7 @@ public:
 
 			if (Particle->ObjectState() != EObjectStateType::Kinematic)
 			{
-				MovingKinematicsMapArray.Remove(Particle->CastToKinematicParticle());
+				MovingKinematicsMapArray.Remove(Particle);
 			}
 
 			// Move to appropriate dynamic SOA
@@ -720,7 +720,10 @@ public:
 
 	void MarkMovingKinematic(FKinematicGeometryParticleHandle* Particle)
 	{
-		MovingKinematicsMapArray.Insert(Particle);
+		if (FPBDRigidParticleHandle* Rigid = Particle->CastToRigidParticle())
+		{
+			MovingKinematicsMapArray.Insert(Rigid);
+		}
 	}
 
 	void UpdateAllMovingKinematic()
@@ -734,7 +737,13 @@ public:
 			if (MovingKinematicsMapArray.GetArray()[Index]->KinematicTarget().GetMode() == EKinematicTargetMode::None)
 			{
 				// remove and do not increment Index
-				MovingKinematicsMapArray.Remove(MovingKinematicsMapArray.GetArray()[Index]);
+				FPBDRigidParticleHandle* Rigid = MovingKinematicsMapArray.GetArray()[Index];
+				MovingKinematicsMapArray.Remove(Rigid);
+
+				// Particle may not be moving, but its velocity was just changed (to zero), so it
+				// needs to be in a dirty list. Since it are no longer in MovingKinematicsMapArray
+				// we must put it the transient dirty list
+				MarkTransientDirtyParticle(Rigid);
 			}
 			else
 			{
@@ -824,14 +833,14 @@ public:
 	const TParticleView<FKinematicGeometryParticles>& GetActiveKinematicParticlesView() const { return ActiveKinematicParticlesView; }
 	TParticleView<FKinematicGeometryParticles>& GetActiveKinematicParticlesView() { return ActiveKinematicParticlesView; }
 
-	const TParticleView<FKinematicGeometryParticles>& GetActiveMovingKinematicParticlesView() const { return ActiveMovingKinematicParticlesView; }
-	TParticleView<FKinematicGeometryParticles>& GetActiveMovingKinematicParticlesView() { return ActiveMovingKinematicParticlesView; }
+	const TParticleView<FPBDRigidParticles>& GetActiveMovingKinematicParticlesView() const { return ActiveMovingKinematicParticlesView; }
+	TParticleView<FPBDRigidParticles>& GetActiveMovingKinematicParticlesView() { return ActiveMovingKinematicParticlesView; }
 
 	const TParticleView<FGeometryParticles>& GetActiveStaticParticlesView() const { return ActiveStaticParticlesView; }
 	TParticleView<FGeometryParticles>& GetActiveStaticParticlesView() { return ActiveStaticParticlesView; }
 
-	const TParticleView<FKinematicGeometryParticles>& GetActiveDynamicMovingKinematicParticlesView() const { return ActiveDynamicMovingKinematicParticlesView; }
-	TParticleView<FKinematicGeometryParticles>& GetActiveDynamicMovingKinematicParticlesView() { return ActiveDynamicMovingKinematicParticlesView; }
+	const TParticleView<FPBDRigidParticles>& GetActiveDynamicMovingKinematicParticlesView() const { return ActiveDynamicMovingKinematicParticlesView; }
+	TParticleView<FPBDRigidParticles>& GetActiveDynamicMovingKinematicParticlesView() { return ActiveDynamicMovingKinematicParticlesView; }
 
 	const TGeometryParticleHandles<FReal, 3>& GetParticleHandles() const { return ParticleHandles; }
 	TGeometryParticleHandles<FReal, 3>& GetParticleHandles() { return ParticleHandles; }
@@ -1113,7 +1122,7 @@ private:
 			NonDisabledDynamicView = MakeParticleView(MoveTemp(TmpArray));
 		}
 		{
-			TArray<TSOAView<FKinematicGeometryParticles>> TmpArray = 
+			TArray<TSOAView<FPBDRigidParticles>> TmpArray =
 			{ 
 				{&ActiveParticlesMapArray.GetArray()},
 				{&MovingKinematicsMapArray.GetArray()},
@@ -1144,12 +1153,13 @@ private:
 			? TArray<TSOAView<FPBDRigidParticles>>
 			{
 				{&ResimActiveParticlesMapArray.GetArray()},
-				{&TransientDirtyMapArray.GetArray() }
+				{&TransientDirtyMapArray.GetArray()}
 			}
 			: TArray<TSOAView<FPBDRigidParticles>>
 			{
 				{&ActiveParticlesMapArray.GetArray()},
 				//	{&DynamicClusteredMapArray.GetArray()},  Cluster particles appear in the ActiveParticlesArray
+				(&MovingKinematicsMapArray.GetArray()),
 				{&KinematicGeometryCollectionArray},
 				{&DynamicGeometryCollectionArray},
 				{&FellAsleepGeometryCollectionArray.GetArray()},
@@ -1190,7 +1200,7 @@ private:
 		}
 		{
 			// todo(chaos) handle resim ?
-			TArray<TSOAView<FKinematicGeometryParticles>> TmpArray =
+			TArray<TSOAView<FPBDRigidParticles>> TmpArray =
 			{
 				{&MovingKinematicsMapArray.GetArray()},
 			};
@@ -1302,7 +1312,7 @@ private:
 	TParticleMapArray<FPBDRigidParticleHandle> TransientDirtyMapArray;
 	
 	// keep track of kinematic that have their kinematic target set for this current frame
-	TParticleMapArray<FKinematicGeometryParticleHandle> MovingKinematicsMapArray;
+	TParticleMapArray<FPBDRigidParticleHandle> MovingKinematicsMapArray;
 
 	// Structures for maintaining a subset view during a resim (TParticleMapArray used when we need to dynamically add/remove during resim)
 	TParticleMapArray<FPBDRigidParticleHandle> ResimActiveParticlesMapArray;
@@ -1318,15 +1328,15 @@ private:
 	TParticleMapArray<FPBDRigidClusteredParticleHandle> DynamicClusteredMapArray;
 
 	//Particle Views
-	TParticleView<FGeometryParticles> NonDisabledView;							//all particles that are not disabled
+	TParticleView<FGeometryParticles> NonDisabledView;								//all particles that are not disabled
 	TParticleView<FPBDRigidParticles> NonDisabledDynamicView;						//all dynamic particles that are not disabled
 	TParticleView<FPBDRigidClusteredParticles> NonDisabledClusteredView;			//all clustered particles that are not disabled
-	TParticleView<FPBDRigidParticles> ActiveParticlesView;						//all particles that are active
+	TParticleView<FPBDRigidParticles> ActiveParticlesView;							//all particles that are active
 	TParticleView<FPBDRigidParticles> DirtyParticlesView;							//all particles that are active + any that were put to sleep this frame
-	TParticleView<FGeometryParticles> AllParticlesView;							//all particles
+	TParticleView<FGeometryParticles> AllParticlesView;								//all particles
 	TParticleView<FKinematicGeometryParticles> ActiveKinematicParticlesView;		//all kinematic particles that are not disabled
-	TParticleView<FKinematicGeometryParticles> ActiveMovingKinematicParticlesView;//all moving kinematic particles that are not disabled
-	TParticleView<FKinematicGeometryParticles> ActiveDynamicMovingKinematicParticlesView;//all moving kinematic particles that are not disabled + all dynamic particles
+	TParticleView<FPBDRigidParticles> ActiveMovingKinematicParticlesView;			//all moving kinematic particles that are not disabled
+	TParticleView<FPBDRigidParticles> ActiveDynamicMovingKinematicParticlesView;	//all moving kinematic particles that are not disabled + all dynamic particles
 	TParticleView<FGeometryParticles> ActiveStaticParticlesView;					//all static particles that are not disabled
 	TParticleView<TPBDGeometryCollectionParticles<FReal, 3>> ActiveGeometryCollectionParticlesView; // all geom collection particles that are not disabled
 
