add proximity contact metric based on convex hull intersections, and expose to proximity tool + dataflow as a way to filter connections

diff --git a/Engine/Plugins/Experimental/ChaosEditor/Source/FractureEditor/Private/FractureToolProximity.cpp b/Engine/Plugins/Experimental/ChaosEditor/Source/FractureEditor/Private/FractureToolProximity.cpp
--- a/Engine/Plugins/Experimental/ChaosEditor/Source/FractureEditor/Private/FractureToolProximity.cpp
+++ b/Engine/Plugins/Experimental/ChaosEditor/Source/FractureEditor/Private/FractureToolProximity.cpp
@@ -188,8 +188,12 @@ void UFractureToolProximity::Render(const FSceneView* View, FViewport* Viewport,
 
 			for (int32 CenterIndex = 0; CenterIndex < Vis.GeoCenters.Num(); CenterIndex++)
 			{
-				FVector P = WorldTransform.TransformPosition(Vis.GeoCenters[CenterIndex] + GetExplodedOffset(CenterIndex));
-				PDI->DrawPoint(P, ProximitySettings->CenterColor, ProximitySettings->CenterSize, SDPG_Foreground);
+				// Draw centers for geometry (clusters can have geometry but it is not visible and won't be connected in the proximity graph)
+				if (Collection.SimulationType[Collection.TransformIndex[CenterIndex]] == FGeometryCollection::ESimulationTypes::FST_Rigid)
+				{
+					FVector P = WorldTransform.TransformPosition(Vis.GeoCenters[CenterIndex] + GetExplodedOffset(CenterIndex));
+					PDI->DrawPoint(P, ProximitySettings->CenterColor, ProximitySettings->CenterSize, SDPG_Foreground);
+				}
 			}
 		}
 	}
diff --git a/Engine/Plugins/Experimental/ChaosEditor/Source/FractureEditor/Private/FractureToolProximity.h b/Engine/Plugins/Experimental/ChaosEditor/Source/FractureEditor/Private/FractureToolProximity.h
--- a/Engine/Plugins/Experimental/ChaosEditor/Source/FractureEditor/Private/FractureToolProximity.h
+++ b/Engine/Plugins/Experimental/ChaosEditor/Source/FractureEditor/Private/FractureToolProximity.h
@@ -27,7 +27,8 @@ public:
 	EProximityMethod Method = EProximityMethod::Precise;
 
 	// If hull-based proximity detection is enabled, amount to expand hulls when searching for overlapping neighbors
-	UPROPERTY(EditAnywhere, Category = Automatic, meta = (ClampMin = "0", EditCondition = "Method == EProximityMethod::ConvexHull"))
+	UPROPERTY(EditAnywhere, Category = Automatic, meta = (ClampMin = "0", 
+		EditCondition = "Method == EProximityMethod::ConvexHull || ContactMethod == EProximityContactMethod::ConvexHullSharpContact || ContactMethod == EProximityContactMethod::ConvexHullAreaContact"))
 	double DistanceThreshold = 1;
 
 	// Method to use to determine the contact between two pieces, if the Contact Threshold is greater than 0
diff --git a/Engine/Plugins/Experimental/GeometryCollectionPlugin/Source/GeometryCollectionNodes/Private/Dataflow/GeometryCollectionNodes.cpp b/Engine/Plugins/Experimental/GeometryCollectionPlugin/Source/GeometryCollectionNodes/Private/Dataflow/GeometryCollectionNodes.cpp
--- a/Engine/Plugins/Experimental/GeometryCollectionPlugin/Source/GeometryCollectionNodes/Private/Dataflow/GeometryCollectionNodes.cpp
+++ b/Engine/Plugins/Experimental/GeometryCollectionPlugin/Source/GeometryCollectionNodes/Private/Dataflow/GeometryCollectionNodes.cpp
@@ -711,6 +711,7 @@ void FProximityDataflowNode::Evaluate(Dataflow::FContext& Context, const FDatafl
 			FGeometryCollectionProximityPropertiesInterface::FProximityProperties Properties = GeomCollection->GetProximityProperties();
 
 			Properties.Method = (EProximityMethod)ProximityMethod;
+			Properties.ContactMethod = (EProximityContactMethod)FilterContactMethod;
 			Properties.DistanceThreshold = DistanceThreshold;
 			Properties.bUseAsConnectionGraph = bUseAsConnectionGraph;
 			Properties.RequireContactAmount = ContactThreshold;
diff --git a/Engine/Plugins/Experimental/GeometryCollectionPlugin/Source/GeometryCollectionNodes/Public/Dataflow/GeometryCollectionNodes.h b/Engine/Plugins/Experimental/GeometryCollectionPlugin/Source/GeometryCollectionNodes/Public/Dataflow/GeometryCollectionNodes.h
--- a/Engine/Plugins/Experimental/GeometryCollectionPlugin/Source/GeometryCollectionNodes/Public/Dataflow/GeometryCollectionNodes.h
+++ b/Engine/Plugins/Experimental/GeometryCollectionPlugin/Source/GeometryCollectionNodes/Public/Dataflow/GeometryCollectionNodes.h
@@ -990,6 +990,20 @@ enum class EProximityMethodEnum : uint8
 	Dataflow_Max                UMETA(Hidden)
 };
 
+UENUM(BlueprintType)
+enum class EProximityContactFilteringMethodEnum : uint8
+{
+	/** Rejects proximity if the bounding boxes do not overlap by more than Contact Threshold centimeters in any major axis direction (or at least half the max possible). This can filter out corner connections of box-like shapes. */
+	Dataflow_ProximityContactFitlteringMethod_ProjectedBoundsOverlap UMETA(DisplayName = "Projected Bounds Overlap"),
+	/** Rejects proximity if the intersection of convex hulls (allowing for optional offset) follows a sharp, thin region which is not wider than Contact Threshold centimeters (or at least half the max possible). */
+	Dataflow_ProximityContactFitlteringMethod_ConvexHullSharp UMETA(DisplayName = "Convex Hull Sharp Contact"),
+	/** Rejects proximity if the surface area of the intersection of convex hulls (allowing for optional offset) is smaller than Contact Threshold squared (or at least half the max possible). */
+	Dataflow_ProximityContactFitlteringMethod_ConvexHullArea UMETA(DisplayName = "Convex Hull Area Contact"),
+	//~~~
+	//256th entry
+	Dataflow_Max                UMETA(Hidden)
+};
+
 /**
  *
  * Update the proximity (contact) graph for the bones in a Collection
@@ -1014,6 +1028,9 @@ public:
 	UPROPERTY(EditAnywhere, Category = "Proximity", meta = (ClampMin = "0"))
 	float ContactThreshold = 0;
 
+	/** How to use the Contact Threshold (if > 0) to filter out unwanted small or corner contacts from the proximity graph. If contact threshold is zero, no filtering is applied. */
+	EProximityContactFilteringMethodEnum FilterContactMethod = EProximityContactFilteringMethodEnum::Dataflow_ProximityContactFitlteringMethod_ProjectedBoundsOverlap;
+
 	/** Whether to automatically transform the proximity graph into a connection graph to be used for simulation */
 	UPROPERTY(EditAnywhere, Category = "Proximity")
 	bool bUseAsConnectionGraph = false;
diff --git a/Engine/Source/Runtime/Experimental/Chaos/Private/GeometryCollection/GeometryCollectionProximityUtility.cpp b/Engine/Source/Runtime/Experimental/Chaos/Private/GeometryCollection/GeometryCollectionProximityUtility.cpp
--- a/Engine/Source/Runtime/Experimental/Chaos/Private/GeometryCollection/GeometryCollectionProximityUtility.cpp
+++ b/Engine/Source/Runtime/Experimental/Chaos/Private/GeometryCollection/GeometryCollectionProximityUtility.cpp
@@ -7,6 +7,8 @@
 #include "GeometryCollection/GeometryCollectionConvexUtility.h"
 #include "Chaos/Convex.h"
 #include "Chaos/GJK.h"
+#include "CompGeom/ConvexHull3.h"
+#include "VectorUtil.h"
 
 #include "Spatial/SparseDynamicOctree3.h"
 #include "Async/ParallelFor.h"
@@ -600,22 +602,25 @@ void FGeometryCollectionProximityUtility::UpdateProximity(UE::GeometryCollection
 
 	FGeometryCollectionProximityPropertiesInterface::FProximityProperties Properties = Collection->GetProximityProperties();
 
-	if (Properties.Method == EProximityMethod::ConvexHull)
-	{
-		FGeometryCollectionConvexPropertiesInterface::FConvexCreationProperties ConvexProperties = Collection->GetConvexProperties();
+	bool bWantConvexContactEdges = Properties.RequireContactAmount > 0.0f &&
+		(Properties.ContactMethod == EProximityContactMethod::ConvexHullSharpContact || Properties.ContactMethod == EProximityContactMethod::ConvexHullAreaContact);
+	bool bWantLocalHulls = Properties.Method == EProximityMethod::ConvexHull || bWantConvexContactEdges;
 
-		UE::GeometryCollectionConvexUtility::FConvexHulls LocalComputedHulls;
-		UE::GeometryCollectionConvexUtility::FConvexHulls* UseComputedHulls = OptionalComputedHulls;
-		if (!UseComputedHulls || UseComputedHulls->OverlapRemovalShrinkPercent > 0) // If we don't have precomputed hulls or if they're shrunk, compute new hulls to use for proximity detection
-		{
-			TArray<FTransform> GlobalTransformArray;
-			GeometryCollectionAlgo::GlobalMatrices(Collection->Transform, Collection->Parent, GlobalTransformArray);
-			LocalComputedHulls = FGeometryCollectionConvexUtility::ComputeLeafHulls(Collection, GlobalTransformArray, ConvexProperties.SimplificationThreshold, 0.0f /* Never shrink hulls for proximity detection */);
-			UseComputedHulls = &LocalComputedHulls;
-		}
+	FGeometryCollectionConvexPropertiesInterface::FConvexCreationProperties ConvexProperties = Collection->GetConvexProperties();
 
-		UE::GeometryCollectionInternal::BuildProximityFromConvexHulls(Collection, *UseComputedHulls, Properties.DistanceThreshold);
+	UE::GeometryCollectionConvexUtility::FConvexHulls LocalComputedHulls;
+	UE::GeometryCollectionConvexUtility::FConvexHulls* UseComputedHulls = OptionalComputedHulls;
+	if (bWantLocalHulls && (!UseComputedHulls || UseComputedHulls->OverlapRemovalShrinkPercent > 0)) // If we don't have precomputed hulls or if they're shrunk, compute new hulls to use for proximity detection
+	{
+		TArray<FTransform> GlobalTransformArray;
+		GeometryCollectionAlgo::GlobalMatrices(Collection->Transform, Collection->Parent, GlobalTransformArray);
+		LocalComputedHulls = FGeometryCollectionConvexUtility::ComputeLeafHulls(Collection, GlobalTransformArray, ConvexProperties.SimplificationThreshold, 0.0f /* Never shrink hulls for proximity detection */);
+		UseComputedHulls = &LocalComputedHulls;
+	}
 
+	if (Properties.Method == EProximityMethod::ConvexHull)
+	{
+		UE::GeometryCollectionInternal::BuildProximityFromConvexHulls(Collection, *UseComputedHulls, Properties.DistanceThreshold);
 	}
 	else
 	{
@@ -626,6 +631,12 @@ void FGeometryCollectionProximityUtility::UpdateProximity(UE::GeometryCollection
 		Spatial.MoveProximityToCollection(Collection);
 	}
 
+	TArray<FGeometryContactEdge> ContactEdges;
+	if (bWantConvexContactEdges)
+	{
+		ContactEdges = ComputeConvexGeometryContactFromProximity(Collection, Properties.DistanceThreshold, *UseComputedHulls);
+	}
+
 	if (Properties.RequireContactAmount > 0.0f)
 	{
 		if (Properties.ContactMethod == EProximityContactMethod::MinOverlapInProjectionToMajorAxes)
@@ -695,6 +706,31 @@ void FGeometryCollectionProximityUtility::UpdateProximity(UE::GeometryCollection
 				}
 			}
 		}
+		else if (Properties.ContactMethod == EProximityContactMethod::ConvexHullSharpContact)
+		{
+			TManagedArray<TSet<int32>>& Proximity = Collection->ModifyAttribute<TSet<int32>>("Proximity", FGeometryCollection::GeometryGroup);
+			for (const FGeometryContactEdge& Edge : ContactEdges)
+			{
+				if (Edge.SharpContactWidth < Properties.RequireContactAmount && Edge.SharpContactWidth < Edge.MaxSharpContact * .5)
+				{
+					Proximity[Edge.GeometryIndices[0]].Remove(Edge.GeometryIndices[1]);
+					Proximity[Edge.GeometryIndices[1]].Remove(Edge.GeometryIndices[0]);
+				}
+			}
+		}
+		else if (Properties.ContactMethod == EProximityContactMethod::ConvexHullAreaContact)
+		{
+			TManagedArray<TSet<int32>>& Proximity = Collection->ModifyAttribute<TSet<int32>>("Proximity", FGeometryCollection::GeometryGroup);
+			for (const FGeometryContactEdge& Edge : ContactEdges)
+			{
+				float RequireArea = Properties.RequireContactAmount * Properties.RequireContactAmount;
+				if (Edge.ContactArea < RequireArea && Edge.ContactArea < Edge.MaxContactArea * .5)
+				{
+					Proximity[Edge.GeometryIndices[0]].Remove(Edge.GeometryIndices[1]);
+					Proximity[Edge.GeometryIndices[1]].Remove(Edge.GeometryIndices[0]);
+				}
+			}
+		}
 	}
 
 	if (Properties.bUseAsConnectionGraph)
@@ -707,3 +743,501 @@ void FGeometryCollectionProximityUtility::UpdateProximity(UE::GeometryCollection
 		ClearConnectionGraph();
 	}
 }
+
+namespace
+{
+	// Helpful struct to represent a convex hull with a local editable representation, to perform intersections / clipping on it
+	// TODO: Consider exposing this in a more generally useful place, or at least the functions to compute intersections and volume/area of intersection
+	struct FHullPolygons
+	{
+		// simple packed representation for convex hull faces, where each polygon's indices are listed sequentially,
+		// and negative values indicate the number of vertices in the next polygon. If no negative value is listed, polygon is a triangle.
+		TArray<int32> PackedPolygons;
+
+		// Copy of hull vertices, to be refined through plane cuts
+		TArray<Chaos::FVec3f> Vertices;
+		Chaos::FAABB3f Bounds;
+
+		FHullPolygons(const Chaos::FConvex& HullIn)
+		{
+			Vertices = HullIn.GetVertices();
+			const Chaos::FConvexStructureData& HullData = HullIn.GetStructureData();
+			int32 NumPlanes = HullIn.NumPlanes();
+			PackedPolygons.Reserve(NumPlanes * 3);
+			for (int PlaneIdx = 0; PlaneIdx < NumPlanes; PlaneIdx++)
+			{
+				int32 NumPlaneVerts = HullData.NumPlaneVertices(PlaneIdx);
+				if (NumPlaneVerts > 3)
+				{
+					PackedPolygons.Add(-NumPlaneVerts);
+				}
+				for (int32 PlaneVertexIdx = 0; PlaneVertexIdx < NumPlaneVerts; PlaneVertexIdx++)
+				{
+					PackedPolygons.Add(HullData.GetPlaneVertex(PlaneIdx, PlaneVertexIdx));
+				}
+			}
+			Bounds = HullIn.GetLocalBoundingBox();
+		}
+
+		void Intersect(const Chaos::FConvex& OtherHull, float ExpandAmount)
+		{
+			// Arrays to store intermediate plane cut data
+			TArray<int32> NewPolygons; 
+			TMap<FIntVector2, int> NewVertices; // mapping from edges to new vertices
+			TArray<float> SignedDist; // signed distance from vertices to a cutting plane
+			TArray<int32> VertexRemap;
+			TMap<int32, int32> OpenEdgeVertMap;
+
+			// Cut by the plane through PlanePt with PlaneNormal. Note: Must pre-apply ExpandAmount; we do not assume the plane is shifted here.
+			// Uses the intermediate data above as temp storage
+			auto PlaneCut = [&NewPolygons, &NewVertices, &SignedDist, &VertexRemap, &OpenEdgeVertMap, this](Chaos::FVec3f PlanePt, Chaos::FVec3f PlaneNormal)
+			{
+				NewPolygons.Reset(PackedPolygons.Num());
+				NewVertices.Reset();
+				OpenEdgeVertMap.Reset();
+				SignedDist.SetNum(Vertices.Num(), false);
+				VertexRemap.SetNum(Vertices.Num(), false);
+				int32 OpenEdgeStart = -1;
+
+				// Possible optimization: if many vertices, check plane vs AABB corners first?
+
+				int32 OutTotal = 0;
+				for (int32 VertIdx = 0; VertIdx < Vertices.Num(); ++VertIdx)
+				{
+					float SD = (float)(Vertices[VertIdx] - PlanePt).Dot(PlaneNormal);
+					SignedDist[VertIdx] = SD;
+					OutTotal += int32(SD > 0);
+				}
+				// hull is fully outside plane's half-space
+				if (OutTotal == Vertices.Num())
+				{
+					Vertices.Empty();
+					PackedPolygons.Empty();
+					return;
+				}
+				// hull is fully inside plane's half-space
+				if (OutTotal == 0)
+				{
+					return;
+				}
+
+				for (int32 Idx = 0, PolyLen = 3; Idx < PackedPolygons.Num(); Idx += PolyLen)
+				{
+					// extract length of current polygon
+					PolyLen = 3;
+					int32 OrigStart = Idx;
+					if (PackedPolygons[Idx] < 0)
+					{
+						PolyLen = -PackedPolygons[Idx];
+						Idx++;
+					}
+					int32 Start = Idx;
+
+					// helper to convert index within polygon to index within vertices array
+					auto ToV = [this, Start, PolyLen](int32 SubIdx) -> int32
+					{
+						checkSlow(SubIdx >= 0 && SubIdx < PolyLen);
+						int32 VertIdx = PackedPolygons[Start + SubIdx];
+						checkSlow(VertIdx >= 0);
+						return VertIdx;
+					};
+					// track where the polygon crosses the plane to decide what to do with it
+					int32 OutCount = 0;
+					int32 FirstIn = -1, FirstOut = -1, LastIn = -1;
+					for (int32 SubIdx = 0; SubIdx < PolyLen; ++SubIdx)
+					{
+						float SD = SignedDist[ToV(SubIdx)];
+						bool IsOut = SD > 0;
+						if (FirstIn == -1)
+						{
+							if (!IsOut)
+							{
+								FirstIn = SubIdx;
+							}
+						}
+						else if (FirstOut == -1 && IsOut)
+						{
+							LastIn = SubIdx - 1;
+							FirstOut = SubIdx;
+						}
+						OutCount += int32(IsOut);
+					}
+					if (FirstOut == -1)
+					{
+						FirstOut = 0;
+						LastIn = PolyLen - 1;
+					}
+					if (OutCount == PolyLen)
+					{
+						continue;
+					}
+					if (OutCount == 0)
+					{
+						// copy original polygon data
+						for (int32 CopyIdx = OrigStart; CopyIdx < Start + PolyLen; ++CopyIdx)
+						{
+							NewPolygons.Add(PackedPolygons[CopyIdx]);
+						}
+						continue;
+					}
+					int32 NewPolyLen = LastIn + 1 - FirstIn;
+					if (FirstIn == 0)
+					{
+						int32 WalkBack = PolyLen - 1;
+						while (WalkBack > 0 && SignedDist[ToV(WalkBack)] <= 0)
+						{
+							FirstIn = WalkBack--;
+							NewPolyLen++;
+						}
+					}
+					auto GetCrossVertIdx = [&NewVertices, &SignedDist, this](int32 InsideVertIdx, int32 OutsideVertIdx) -> int32
+					{
+						checkSlow(InsideVertIdx != OutsideVertIdx);
+						// If within zero tolerance of plane, snap to plane
+						float InsideSD = SignedDist[InsideVertIdx];
+						checkSlow(InsideSD <= 0);
+						if (InsideSD > -FMathf::ZeroTolerance)
+						{
+							return -1;
+						}
+						FIntVector2 Key(InsideVertIdx, OutsideVertIdx);
+						int32* FoundVert = NewVertices.Find(Key);
+						if (!FoundVert)
+						{
+							float OutsideSD = SignedDist[OutsideVertIdx];
+							checkSlow(OutsideSD >= 0);
+							Chaos::FVec3f NewVert = FMath::Lerp(Vertices[InsideVertIdx], Vertices[OutsideVertIdx], InsideSD / (InsideSD - OutsideSD));
+							int32 NewVertIdx = Vertices.Add(NewVert);
+							NewVertices.Add(Key, NewVertIdx);
+							return NewVertIdx;
+						}
+						return *FoundVert;
+					};
+					int32 FirstCross = GetCrossVertIdx(ToV(FirstIn), ToV((FirstIn + PolyLen - 1) % PolyLen));
+					int32 LastCross = GetCrossVertIdx(ToV(LastIn), ToV(FirstOut));
+					int32 OpenPlaneEdgeVA = FirstCross;
+					int32 OpenPlaneEdgeVB = LastCross;
+					if (FirstCross != -1)
+					{
+						NewPolyLen++;
+					}
+					else
+					{
+						OpenPlaneEdgeVA = ToV(FirstIn);
+					}
+					if (LastCross != -1)
+					{
+						NewPolyLen++;
+					}
+					else
+					{
+						OpenPlaneEdgeVB = ToV(LastIn);
+					}
+
+					if (NewPolyLen < 2) // single co-incident vertex; no open edge here
+					{
+						continue;
+					}
+					OpenEdgeStart = OpenPlaneEdgeVB;
+					OpenEdgeVertMap.Add(OpenPlaneEdgeVA, OpenPlaneEdgeVB);
+					if (NewPolyLen == 2)
+					{
+						continue;
+					}
+					if (NewPolyLen > 3)
+					{
+						NewPolygons.Add(-NewPolyLen);
+					}
+					int32 NewPolygonStart = NewPolygons.Num();
+					if (FirstCross != -1)
+					{
+						NewPolygons.Add(FirstCross);
+					}
+					int32 AddStart = FirstIn;
+					if (FirstIn > LastIn)
+					{
+						for (int32 SubIdx = FirstIn; SubIdx < PolyLen; ++SubIdx)
+						{
+							NewPolygons.Add(ToV(SubIdx));
+						}
+						AddStart = 0;
+					}
+					for (int32 SubIdx = AddStart; SubIdx <= LastIn; ++SubIdx)
+					{
+						NewPolygons.Add(ToV(SubIdx));
+					}
+					if (LastCross != -1)
+					{
+						NewPolygons.Add(LastCross);
+					}
+
+					check(NewPolygons.Num() - NewPolygonStart == NewPolyLen);
+				}
+
+				// add the closing polygon
+				if (OpenEdgeStart != -1 && OpenEdgeVertMap.Num() > 2)
+				{
+					int32 OrigEnd = NewPolygons.Num();
+					int32 PolyEdges = OpenEdgeVertMap.Num();
+					if (PolyEdges > 3)
+					{
+						NewPolygons.Add(-PolyEdges);
+					}
+					int32 TraverseIdx = OpenEdgeStart;
+					int32 Added = 0;
+					do
+					{
+						NewPolygons.Add(TraverseIdx);
+						int32* FoundNext = OpenEdgeVertMap.Find(TraverseIdx);
+						if (!FoundNext)
+						{
+							break;
+						}
+						TraverseIdx = *FoundNext;
+						Added++;
+					} while (TraverseIdx != OpenEdgeStart && Added < PolyEdges);
+					if (Added != PolyEdges || TraverseIdx != OpenEdgeStart)
+					{
+						// failsafe if we didn't find a closed loop covering all edges:
+						// add a triangle fan closing off the edges that we did find
+						NewPolygons.SetNum(OrigEnd, false);
+						Chaos::FVec3f Center(0, 0, 0);
+						float CenterWt = 0;
+						int32 CenterIdx = Vertices.Num();
+						for (TPair<int32, int32> KV : OpenEdgeVertMap)
+						{
+							NewPolygons.Add(KV.Key);
+							NewPolygons.Add(KV.Value);
+							NewPolygons.Add(CenterIdx);
+							Center += Vertices[KV.Key];
+							CenterWt += 1.f;
+						}
+						Center /= CenterWt;
+						Vertices.Add(Center);
+					}
+				}
+
+				// NewPolygons now contains the updated polygon data
+				Swap(PackedPolygons, NewPolygons);
+				// Compress the vertex array to only include the vertices that weren't outside
+				// and track how the indices were remapped
+				int32 NumKept = 0;
+				const int32 OldVertCount = SignedDist.Num();
+				for (int32 OldV = 0; OldV < OldVertCount; ++OldV)
+				{
+					if (SignedDist[OldV] <= 0)
+					{
+						int32 UseNewV = NumKept++;
+						VertexRemap[OldV] = UseNewV;
+						checkSlow(OldV >= UseNewV);
+						Vertices[UseNewV] = Vertices[OldV];
+					}
+				}
+				// Translate back the new vertices
+				if (NumKept < OldVertCount)
+				{
+					for (int32 OldIdx = OldVertCount, AddedIdx = 0; OldIdx < Vertices.Num(); ++OldIdx, ++AddedIdx)
+					{
+						Vertices[NumKept + AddedIdx] = Vertices[OldIdx];
+					}
+					int32 NumNew = Vertices.Num() - OldVertCount;
+					Vertices.SetNum(NumKept + NumNew, false);
+				}
+				// Update the polygons w/ the compressed vertex indices
+				for (int32& VIdx : PackedPolygons)
+				{
+					if (VIdx >= 0) // Only remap vertices, not polygon sizes
+					{
+						if (VIdx < OldVertCount)
+						{
+							VIdx = VertexRemap[VIdx];
+						}
+						else // newly-created vertices are kept in the same order at the end of the array
+						{
+							VIdx = NumKept + (VIdx - OldVertCount);
+						}
+					}
+				}
+			};
+
+			// TODO: For performance, consider also pre-cutting with (some of) OtherHull's (expanded) bounding box planes
+			//Chaos::FConvex::FAABB3Type OtherBounds = OtherHull.GetLocalBoundingBox();
+			//OtherBounds.Thicken(ExpandAmount);
+
+			// Cut with each convex plane
+			const int32 NumPlanes = OtherHull.NumPlanes();
+			for (int32 PlaneIdx = 0; PlaneIdx < NumPlanes; ++PlaneIdx)
+			{
+				Chaos::TPlaneConcrete<float, 3> Plane = OtherHull.GetPlaneRaw(PlaneIdx);
+				Chaos::FVec3f N = Plane.Normal();
+				Chaos::FVec3f X = Plane.X();
+				X += N * ExpandAmount;
+				PlaneCut(X, N);
+			}
+
+			// When ExpandAmount is positive, also clip the hull at offsets of average edge planes for 'sharp' edges
+			if (ExpandAmount > 0)
+			{
+				const int32 NumEdges = OtherHull.NumEdges();
+				for (int32 EdgeIdx = 0; EdgeIdx < NumEdges; ++EdgeIdx)
+				{
+					Chaos::TPlaneConcrete<float, 3> EPlane0 = OtherHull.GetPlaneRaw(OtherHull.GetEdgePlane(EdgeIdx, 0));
+					Chaos::TPlaneConcrete<float, 3> EPlane1 = OtherHull.GetPlaneRaw(OtherHull.GetEdgePlane(EdgeIdx, 1));
+					float NormalDot = EPlane0.Normal().Dot(EPlane1.Normal());
+					if (NormalDot < -.1) // add an extra plane when not doing so would leave ~1.5x more space than the expected offset across from the edge, due to the miter
+					{
+						Chaos::FVec3f AvgNormal = EPlane0.Normal() + EPlane1.Normal();
+						if (AvgNormal.Normalize())
+						{
+							Chaos::FVec3f EdgeVert = OtherHull.GetVertex(OtherHull.GetEdgeVertex(EdgeIdx, 0));
+							PlaneCut(EdgeVert + AvgNormal * ExpandAmount, AvgNormal);
+						}
+					}
+
+				}
+			}
+		}
+
+		float ComputeArea()
+		{
+			float Area = 0;
+			for (int32 Idx = 0, PolyLen = 3; Idx < PackedPolygons.Num(); Idx += PolyLen)
+			{
+				// extract length of current polygon
+				PolyLen = 3;
+				if (PackedPolygons[Idx] < 0)
+				{
+					PolyLen = -PackedPolygons[Idx];
+					Idx++;
+				}
+				int32 Start = Idx;
+
+				// Add area of triangle fan covering the polygon
+				Chaos::FVec3f V0 = Vertices[PackedPolygons[Start]];
+				for (int32 SubIdx = 1; SubIdx + 1 < PolyLen; ++SubIdx)
+				{
+					Chaos::FVec3f V1 = Vertices[PackedPolygons[Start + SubIdx]];
+					Chaos::FVec3f V2 = Vertices[PackedPolygons[Start + SubIdx + 1]];
+					Area += UE::Geometry::VectorUtil::Area<float>(V0, V1, V2);
+				}
+			}
+			return Area;
+		}
+
+		void EstimateSharpContact(const Chaos::FConvex* HullA, const Chaos::FConvex* HullB, float& OutSharpContact, float& OutMaxSharpContact)
+		{
+			UE::Geometry::FExtremePoints3f ExtremePts(Vertices.Num(), [this](int32 Idx) {return Vertices[Idx];});
+			if (ExtremePts.Dimension < 1)
+			{
+				OutSharpContact = 0;
+				OutMaxSharpContact = 1;
+				return; // degenerate/empty contact
+			}
+			UE::Geometry::FInterval1f IntersectionIntervals[2];
+			UE::Geometry::FInterval1f HullAIntervals[2], HullBIntervals[2];
+			if (ExtremePts.Dimension > 1)
+			{
+				auto SetIntervals = [&ExtremePts](const TArray<Chaos::FVec3f>& UseVertices, UE::Geometry::FInterval1f* Intervals) -> void
+				{
+					for (FVector3f Vertex : UseVertices)
+					{
+						Intervals[0].Contain(Vertex.Dot(ExtremePts.Basis[1]));
+						if (ExtremePts.Dimension > 2)
+						{
+							Intervals[1].Contain(Vertex.Dot(ExtremePts.Basis[2]));
+						}
+					}
+				};
+				SetIntervals(Vertices, IntersectionIntervals);
+				SetIntervals(HullA->GetVertices(), HullAIntervals);
+				SetIntervals(HullB->GetVertices(), HullBIntervals);
+			}
+			auto IntervalsMaxLen = [](UE::Geometry::FInterval1f* Intervals)
+			{
+				return FMath::Max(Intervals[0].Length(), Intervals[1].Length());
+			};
+			OutSharpContact = IntervalsMaxLen(IntersectionIntervals);
+			OutMaxSharpContact = FMath::Min(IntervalsMaxLen(HullAIntervals), IntervalsMaxLen(HullBIntervals));
+		}
+	};
+
+	static float ComputeHullArea(const Chaos::FConvex& Hull)
+	{
+		float Area = 0;
+		const Chaos::FConvexStructureData& HullData = Hull.GetStructureData();
+		int32 NumPlanes = Hull.NumPlanes();
+		for (int PlaneIdx = 0; PlaneIdx < NumPlanes; PlaneIdx++)
+		{
+			int32 NumPlaneVerts = HullData.NumPlaneVertices(PlaneIdx);
+			Chaos::FVec3f V0 = Hull.GetVertex(HullData.GetPlaneVertex(PlaneIdx, 0));
+			for (int32 PlaneVertexIdx = 1; PlaneVertexIdx + 1 < NumPlaneVerts; PlaneVertexIdx++)
+			{
+				Chaos::FVec3f V1 = Hull.GetVertex(HullData.GetPlaneVertex(PlaneIdx, PlaneVertexIdx));
+				Chaos::FVec3f V2 = Hull.GetVertex(HullData.GetPlaneVertex(PlaneIdx, PlaneVertexIdx + 1));
+				Area += UE::Geometry::VectorUtil::Area<float>(V0, V1, V2);
+			}
+		}
+		return Area;
+	}
+
+	static void HullIntersectionStats(int32 DbgA, int32 DbgB, const Chaos::FConvex* HullA, const Chaos::FConvex* HullB, float HullBExpansion, float& OutArea, float& OutMaxArea, float& OutSharpContact, float& OutMaxSharpContact)
+	{
+		FHullPolygons HullPolygons(*HullA);
+		HullPolygons.Intersect(*HullB, HullBExpansion);
+		OutArea = HullPolygons.ComputeArea();
+		// The maximum intersection area is ~ the minimum of the two hull areas
+		float MaxIntersectionArea = FMath::Min(ComputeHullArea(*HullA), ComputeHullArea(*HullB));
+		OutMaxArea = MaxIntersectionArea;
+		HullPolygons.EstimateSharpContact(HullA, HullB, OutSharpContact, OutMaxSharpContact);
+	}
+}
+
+
+TArray<FGeometryCollectionProximityUtility::FGeometryContactEdge>
+FGeometryCollectionProximityUtility::ComputeConvexGeometryContactFromProximity(
+	FGeometryCollection* Collection, float DistanceTolerance,
+	UE::GeometryCollectionConvexUtility::FConvexHulls& LocalHulls)
+{
+	// We must already have proximity
+	const TManagedArray<TSet<int32>>* Proximity = Collection->FindAttribute<TSet<int32>>("Proximity", FGeometryCollection::GeometryGroup);
+	if (!ensureMsgf(Proximity, TEXT("ComputeGeometryContactFromProximity should only be called when a Proximity attribute has already been added")))
+	{
+		return TArray<FGeometryCollectionProximityUtility::FGeometryContactEdge>();
+	}
+	
+	TArray<FGeometryCollectionProximityUtility::FGeometryContactEdge> ContactEdges;
+	for (int32 GeoIdx = 0; GeoIdx < Proximity->Num(); ++GeoIdx)
+	{
+		int32 TransformIdx = Collection->TransformIndex[GeoIdx];
+		const TSet<int32>& GeoHulls = LocalHulls.TransformToHullsIndices[TransformIdx];
+		for (int32 NbrGeoIdx : (*Proximity)[GeoIdx])
+		{
+			// all connections should be symmetric, and no connections should be to self, so only compute lower idx->higher idx
+			if (GeoIdx >= NbrGeoIdx)
+			{
+				continue;
+			}
+			int32 NbrTransformIdx = Collection->TransformIndex[NbrGeoIdx];
+			const TSet<int32>& NbrGeoHulls = LocalHulls.TransformToHullsIndices[NbrTransformIdx];
+			float OverlapAreas = 0, OverlapMaxAreas = 0, OverlapSharpContact = 0, OverlapMaxSharpContact = 0;
+			for (int32 GeoHullIdx : GeoHulls)
+			{
+				for (int32 NbrHullIdx : NbrGeoHulls)
+				{
+					float Area, MaxArea, SharpContact, MaxSharpContact;
+					HullIntersectionStats(TransformIdx, NbrTransformIdx, LocalHulls.Hulls[GeoHullIdx].Get(), LocalHulls.Hulls[NbrHullIdx].Get(), DistanceTolerance, Area, MaxArea, SharpContact, MaxSharpContact);
+					// Note: The sharp contact concept would be most accurately implemented by combining the intersection/overlaps of all involved hulls
+					// but taking the max values should be a reasonable conservative approximation in most cases
+					OverlapSharpContact = FMath::Max(OverlapSharpContact, SharpContact);
+					OverlapMaxSharpContact = FMath::Max(OverlapMaxSharpContact, MaxSharpContact);
+					OverlapMaxAreas += MaxArea;
+					OverlapAreas += Area;
+				}
+			}
+			// to approximate contact surface area, use half of surface area of convex hull intersection
+			ContactEdges.Emplace(GeoIdx, NbrGeoIdx, OverlapAreas * .5f, OverlapMaxAreas * .5f, OverlapSharpContact, OverlapMaxSharpContact);
+		}
+	}
+
+	return ContactEdges;
+}
diff --git a/Engine/Source/Runtime/Experimental/Chaos/Public/GeometryCollection/GeometryCollectionProximityUtility.h b/Engine/Source/Runtime/Experimental/Chaos/Public/GeometryCollection/GeometryCollectionProximityUtility.h
--- a/Engine/Source/Runtime/Experimental/Chaos/Public/GeometryCollection/GeometryCollectionProximityUtility.h
+++ b/Engine/Source/Runtime/Experimental/Chaos/Public/GeometryCollection/GeometryCollectionProximityUtility.h
@@ -25,8 +25,12 @@ enum class EProximityMethod : int32
 UENUM()
 enum class EProximityContactMethod : uint8
 {
-	// Rejects proximity if the bounding boxes do not overlap by more than this many centimeters in any major axis direction. This can filter out corner connections of box-like shapes.
-	MinOverlapInProjectionToMajorAxes
+	// Rejects proximity if the bounding boxes do not overlap by more than Contact Threshold centimeters in any major axis direction (or at least half the max possible). This can filter out corner connections of box-like shapes.
+	MinOverlapInProjectionToMajorAxes,
+	// Rejects proximity if the intersection of convex hulls (allowing for optional offset) follows a sharp, thin region which is not wider than Contact Threshold centimeters (or at least half the max possible).
+	ConvexHullSharpContact,
+	// Rejects proximity if the surface area of the intersection of convex hulls (allowing for optional offset) is smaller than Contact Threshold squared (or at least half the max possible).
+	ConvexHullAreaContact
 	//~ TODO: Add other methods for filtering overlaps, e.g. based on approximate surface area of the contact
 };
 
@@ -45,6 +49,28 @@ public:
 	void CopyProximityToConnectionGraph();
 	void ClearConnectionGraph();
 
+
+	// Stores stats about the contact between two geometries
+	struct FGeometryContactEdge
+	{
+		int32 GeometryIndices[2];
+		// Area estimate for contact region
+		float ContactArea;
+		// Maximum area for contact for this pair (half the smallest surface area)
+		float MaxContactArea;
+		// 'Sharp contact' width estimate
+		float SharpContactWidth;
+		// Estimate of maximum possible 'sharp contact' for this pair
+		float MaxSharpContact;
+
+		FGeometryContactEdge() = default;
+		FGeometryContactEdge(int32 GeoIdxA, int32 GeoIdxB, float ContactArea, float MaxContactArea, float SharpContactWidth, float MaxSharpContact) :
+			GeometryIndices{ GeoIdxA, GeoIdxB }, ContactArea(ContactArea), MaxContactArea(MaxContactArea), SharpContactWidth(SharpContactWidth), MaxSharpContact(MaxSharpContact)
+		{}
+	};
+	// Note: This computes connections from lower to higher geometry indices, assuming connections are symmetric
+	static TArray<FGeometryContactEdge> ComputeConvexGeometryContactFromProximity(FGeometryCollection* Collection, float DistanceTolerance, UE::GeometryCollectionConvexUtility::FConvexHulls& LocalHulls);
+
 private:
 	FGeometryCollection* Collection;
 };
