ChaosFlesh: Expanded copyrange to allow 0 sized arrays

diff --git a/Engine/Source/Runtime/Experimental/Chaos/Public/GeometryCollection/ManagedArray.h b/Engine/Source/Runtime/Experimental/Chaos/Public/GeometryCollection/ManagedArray.h
--- a/Engine/Source/Runtime/Experimental/Chaos/Public/GeometryCollection/ManagedArray.h
+++ b/Engine/Source/Runtime/Experimental/Chaos/Public/GeometryCollection/ManagedArray.h
@@ -390,7 +390,7 @@ public:
 	virtual void CopyRange(const FManagedArrayBase& ConstArray, int32 Start, int32 Stop, int32 Offset = 0) override
 	{
 		ensureMsgf(ConstArray.GetTypeSize() == GetTypeSize(), TEXT("TManagedArrayBase<T>::Init : Invalid array types."));
-		if (ensureMsgf(Stop + Offset < Array.Num(), TEXT("Error : Index out of bounds")))
+		if (ensureMsgf(Stop + Offset <= Array.Num(), TEXT("Error : Index out of bounds")))
 		{
 			const TManagedArrayBase<ElementType>& TypedConstArray = static_cast<const TManagedArrayBase<ElementType>&>(ConstArray);
 			CopyRangeHelper(Array, TypedConstArray, Start, Stop, Offset);
@@ -738,7 +738,7 @@ public:
 	virtual void CopyRange(const FManagedArrayBase& ConstArray, int32 Start, int32 Stop, int32 Offset = 0) override
 	{
 		ensureMsgf(ConstArray.GetTypeSize() == GetTypeSize(), TEXT("TManagedArrayBase<T>::Init : Invalid array types."));
-		if (ensureMsgf(Stop + Offset < Array.Num(), TEXT("Error : Index out of bounds")))
+		if (ensureMsgf(Stop + Offset <= Array.Num(), TEXT("Error : Index out of bounds")))
 		{
 			const FManagedBitArrayBase& TypedConstArray = static_cast<const FManagedBitArrayBase&>(ConstArray);
 			for (int32 Sdx = Start, Tdx = Start + Offset; Sdx < ConstArray.Num() && Tdx < Array.Num() && Sdx < Stop; Sdx++, Tdx++)
