FReadPhysicsObjectInterface<>::GetAllParticles should not crash in the case where we pass in a null physics object handle

diff --git a/Engine/Source/Runtime/Experimental/Chaos/Private/Chaos/PhysicsObjectInterface.cpp b/Engine/Source/Runtime/Experimental/Chaos/Private/Chaos/PhysicsObjectInterface.cpp
--- a/Engine/Source/Runtime/Experimental/Chaos/Private/Chaos/PhysicsObjectInterface.cpp
+++ b/Engine/Source/Runtime/Experimental/Chaos/Private/Chaos/PhysicsObjectInterface.cpp
@@ -191,6 +191,11 @@ namespace Chaos
 
 		for (const FConstPhysicsObjectHandle& Handle : InObjects)
 		{
+			if (!Handle)
+			{
+				continue;
+			}
+
 			if (TThreadParticle<Id>* Particle = Handle->GetParticle<Id>())
 			{
 				Particles.Add(Particle);
