Add back bAutoAssignISMPool to GeometryComponent.

diff --git a/Engine/Source/Runtime/Experimental/GeometryCollectionEngine/Private/GeometryCollection/GeometryCollectionComponent.cpp b/Engine/Source/Runtime/Experimental/GeometryCollectionEngine/Private/GeometryCollection/GeometryCollectionComponent.cpp
--- a/Engine/Source/Runtime/Experimental/GeometryCollectionEngine/Private/GeometryCollection/GeometryCollectionComponent.cpp
+++ b/Engine/Source/Runtime/Experimental/GeometryCollectionEngine/Private/GeometryCollection/GeometryCollectionComponent.cpp
@@ -30,6 +30,7 @@
 #include "GeometryCollection/GeometryCollectionISMPoolActor.h"
 #include "GeometryCollection/GeometryCollectionISMPoolComponent.h"
 #include "GeometryCollection/GeometryCollectionISMPoolRenderer.h"
+#include "GeometryCollection/GeometryCollectionISMPoolSubSystem.h"
 #include "Math/Sphere.h"
 #include "Modules/ModuleManager.h"
 #include "Net/Core/PushModel/PushModel.h"
@@ -2406,6 +2407,15 @@ void UGeometryCollectionComponent::OnRegister()
 		ISMPoolRenderer->ISMPoolActor = ISMPool;
 		CustomRenderer = (UGeometryCollectionExternalRenderInterface*)ISMPoolRenderer;
 	}
+	if (!CustomRenderer && bAutoAssignISMPool)
+	{
+		if (UGeometryCollectionISMPoolSubSystem* ISMPoolSubSystem = UWorld::GetSubsystem<UGeometryCollectionISMPoolSubSystem>(GetWorld()))
+		{
+			UGeometryCollectionCustomRendererISMPool* ISMPoolRenderer = NewObject<UGeometryCollectionCustomRendererISMPool>(this);
+			ISMPoolRenderer->ISMPoolActor = ISMPoolSubSystem->FindISMPoolActor(*this);
+			CustomRenderer = (UGeometryCollectionExternalRenderInterface*)ISMPoolRenderer;
+		}
+	}
 
 	Super::OnRegister();
 }
diff --git a/Engine/Source/Runtime/Experimental/GeometryCollectionEngine/Public/GeometryCollection/GeometryCollectionComponent.h b/Engine/Source/Runtime/Experimental/GeometryCollectionEngine/Public/GeometryCollection/GeometryCollectionComponent.h
--- a/Engine/Source/Runtime/Experimental/GeometryCollectionEngine/Public/GeometryCollection/GeometryCollectionComponent.h
+++ b/Engine/Source/Runtime/Experimental/GeometryCollectionEngine/Public/GeometryCollection/GeometryCollectionComponent.h
@@ -1080,6 +1080,9 @@ protected:
 	UPROPERTY(EditAnywhere, BlueprintReadWrite, Category = "ChaosPhysics|Rendering", meta = (DisplayName = "ISM Pool"))
 	TObjectPtr<AGeometryCollectionISMPoolActor> ISMPool;
 
+	UPROPERTY(EditAnywhere, BlueprintReadWrite, Category = "ChaosPhysics|Rendering")
+	bool bAutoAssignISMPool = false;
+
 	/** Populate the dynamic particle data for the render thread. */
 	FGeometryCollectionDynamicData* InitDynamicData(bool bInitialization);
 
