Chaos Cloth: Fix ispc map access for new Anisotropic bending constraints.

diff --git a/Engine/Source/Runtime/Experimental/Chaos/Private/Chaos/XPBDAnisotropicBendingConstraints.ispc b/Engine/Source/Runtime/Experimental/Chaos/Private/Chaos/XPBDAnisotropicBendingConstraints.ispc
--- a/Engine/Source/Runtime/Experimental/Chaos/Private/Chaos/XPBDAnisotropicBendingConstraints.ispc
+++ b/Engine/Source/Runtime/Experimental/Chaos/Private/Chaos/XPBDAnisotropicBendingConstraints.ispc
@@ -268,18 +268,18 @@ export void ApplyXPBDAnisotropicBendingConstraintsWithMaps(uniform FVector4f Pan
 		#pragma ignore warning(perf)
 		const varying float Stiffness = StiffnessTable[StiffnessIndex];
 		#pragma ignore warning(perf)
-		const varying float StiffnessWeft = StiffnessWeftTable[StiffnessIndex];
+		const varying float StiffnessWeft = StiffnessWeftTable[StiffnessWeftIndex];
 		#pragma ignore warning(perf)
-		const varying float StiffnessBias = StiffnessBiasTable[StiffnessIndex];
+		const varying float StiffnessBias = StiffnessBiasTable[StiffnessBiasIndex];
 
 		const varying FVector3f StiffnessVector = {{Stiffness, StiffnessWeft, StiffnessBias}};
 
 		#pragma ignore warning(perf)
 		const varying float BucklingStiffness = BucklingStiffnessTable[BucklingStiffnessIndex];
 		#pragma ignore warning(perf)
-		const varying float BucklingStiffnessWeft = BucklingStiffnessWeftTable[BucklingStiffnessIndex];
+		const varying float BucklingStiffnessWeft = BucklingStiffnessWeftTable[BucklingStiffnessWeftIndex];
 		#pragma ignore warning(perf)
-		const varying float BucklingStiffnessBias = BucklingStiffnessBiasTable[BucklingStiffnessIndex];
+		const varying float BucklingStiffnessBias = BucklingStiffnessBiasTable[BucklingStiffnessBiasIndex];
 		
 		const varying FVector3f BucklingStiffnessVector = {{BucklingStiffness, BucklingStiffnessWeft, BucklingStiffnessBias}};
 		
@@ -384,18 +384,18 @@ export void ApplyXPBDAnisotropicBendingConstraintsWithDampingAndMaps(uniform FVe
 		#pragma ignore warning(perf)
 		const varying float Stiffness = StiffnessTable[StiffnessIndex];
 		#pragma ignore warning(perf)
-		const varying float StiffnessWeft = StiffnessWeftTable[StiffnessIndex];
+		const varying float StiffnessWeft = StiffnessWeftTable[StiffnessWeftIndex];
 		#pragma ignore warning(perf)
-		const varying float StiffnessBias = StiffnessBiasTable[StiffnessIndex];
+		const varying float StiffnessBias = StiffnessBiasTable[StiffnessBiasIndex];
 
 		const varying FVector3f StiffnessVector = {{Stiffness, StiffnessWeft, StiffnessBias}};
 
 		#pragma ignore warning(perf)
 		const varying float BucklingStiffness = BucklingStiffnessTable[BucklingStiffnessIndex];
 		#pragma ignore warning(perf)
-		const varying float BucklingStiffnessWeft = BucklingStiffnessWeftTable[BucklingStiffnessIndex];
+		const varying float BucklingStiffnessWeft = BucklingStiffnessWeftTable[BucklingStiffnessWeftIndex];
 		#pragma ignore warning(perf)
-		const varying float BucklingStiffnessBias = BucklingStiffnessBiasTable[BucklingStiffnessIndex];
+		const varying float BucklingStiffnessBias = BucklingStiffnessBiasTable[BucklingStiffnessBiasIndex];
 		
 		const varying FVector3f BucklingStiffnessVector = {{BucklingStiffness, BucklingStiffnessWeft, BucklingStiffnessBias}};
 		
