Fix root component not visually updating with Geometry Collection ISMPool renderer.

diff --git a/Engine/Source/Runtime/Experimental/GeometryCollectionEngine/Private/GeometryCollection/GeometryCollectionComponent.cpp b/Engine/Source/Runtime/Experimental/GeometryCollectionEngine/Private/GeometryCollection/GeometryCollectionComponent.cpp
--- a/Engine/Source/Runtime/Experimental/GeometryCollectionEngine/Private/GeometryCollection/GeometryCollectionComponent.cpp
+++ b/Engine/Source/Runtime/Experimental/GeometryCollectionEngine/Private/GeometryCollection/GeometryCollectionComponent.cpp
@@ -4309,14 +4309,22 @@ void UGeometryCollectionComponent::RefreshCustomRenderer()
 		{
 			if (RestCollection != nullptr)
 			{
+				CalculateGlobalMatrices();
+				
+				const FTransform ComponentTransform = GetComponentTransform();
 				const int32 RootIndex = GetRootIndex();
 				const bool bIsBroken = DynamicCollection ? !DynamicCollection->Active[RootIndex] : false;
-				RendererInterface->UpdateState(*RestCollection, GetComponentTransform(), bIsBroken);
+
+				RendererInterface->UpdateState(*RestCollection, bIsBroken);
 
 				if (bIsBroken)
 				{
-					CalculateGlobalMatrices();
-					RendererInterface->UpdateTransforms(*RestCollection, GetComponentTransform(), GlobalMatrices);
+					RendererInterface->UpdateTransforms(*RestCollection, ComponentTransform, GlobalMatrices);
+				}
+				else
+				{
+					const FTransform RootTransform = GlobalMatrices.IsValidIndex(RootIndex) ? FTransform(GlobalMatrices[RootIndex]) : FTransform::Identity;
+					RendererInterface->UpdateRootTransform(*RestCollection, ComponentTransform, RootTransform);
 				}
 			}
 		}
diff --git a/Engine/Source/Runtime/Experimental/GeometryCollectionEngine/Private/GeometryCollection/GeometryCollectionISMPoolRenderer.cpp b/Engine/Source/Runtime/Experimental/GeometryCollectionEngine/Private/GeometryCollection/GeometryCollectionISMPoolRenderer.cpp
--- a/Engine/Source/Runtime/Experimental/GeometryCollectionEngine/Private/GeometryCollection/GeometryCollectionISMPoolRenderer.cpp
+++ b/Engine/Source/Runtime/Experimental/GeometryCollectionEngine/Private/GeometryCollection/GeometryCollectionISMPoolRenderer.cpp
@@ -11,7 +11,7 @@
 
 #include UE_INLINE_GENERATED_CPP_BY_NAME(GeometryCollectionISMPoolRenderer)
 
-void UGeometryCollectionCustomRendererISMPool::UpdateState(UGeometryCollection const& InGeometryCollection, FTransform const& InBaseTransform, bool bIsBroken)
+void UGeometryCollectionCustomRendererISMPool::UpdateState(UGeometryCollection const& InGeometryCollection, bool bIsBroken)
 {
 	if (!bIsBroken && MergedMeshGroup.GroupIndex == INDEX_NONE)
 	{
@@ -20,7 +20,6 @@ void UGeometryCollectionCustomRendererISMPool::UpdateState(UGeometryCollection c
 
 		// Add merged mesh.
 		InitMergedMeshFromGeometryCollection(InGeometryCollection);
-		UpdateMergedMeshTransforms(InBaseTransform);
 	}
 
 	if (bIsBroken && InstancesGroup.GroupIndex == INDEX_NONE)
@@ -33,6 +32,11 @@ void UGeometryCollectionCustomRendererISMPool::UpdateState(UGeometryCollection c
 	}
 }
 
+void UGeometryCollectionCustomRendererISMPool::UpdateRootTransform(UGeometryCollection const& InGeometryCollection, FTransform const& InBaseTransform, FTransform const& InRootTransform)
+{
+	UpdateMergedMeshTransforms(InRootTransform * InBaseTransform);
+}
+
 void UGeometryCollectionCustomRendererISMPool::UpdateTransforms(UGeometryCollection const& InGeometryCollection, FTransform const& InBaseTransform, TArrayView<const FMatrix> InMatrices)
 {
 	UpdateInstanceTransforms(InGeometryCollection, InBaseTransform, InMatrices);
diff --git a/Engine/Source/Runtime/Experimental/GeometryCollectionEngine/Private/GeometryCollection/GeometryCollectionISMPoolRenderer.h b/Engine/Source/Runtime/Experimental/GeometryCollectionEngine/Private/GeometryCollection/GeometryCollectionISMPoolRenderer.h
--- a/Engine/Source/Runtime/Experimental/GeometryCollectionEngine/Private/GeometryCollection/GeometryCollectionISMPoolRenderer.h
+++ b/Engine/Source/Runtime/Experimental/GeometryCollectionEngine/Private/GeometryCollection/GeometryCollectionISMPoolRenderer.h
@@ -16,7 +16,8 @@ class UGeometryCollectionCustomRendererISMPool : public UObject, public IGeometr
 
 public:
 	//~ Begin IGeometryCollectionExternalRenderInterface Interface.
-	virtual void UpdateState(UGeometryCollection const& InGeometryCollection, FTransform const& InBaseTransform, bool bInIsBroken) override;
+	virtual void UpdateState(UGeometryCollection const& InGeometryCollection, bool bInIsBroken) override;
+	virtual void UpdateRootTransform(UGeometryCollection const& InGeometryCollection, FTransform const& InBaseTransform, FTransform const& InRootTransform) override;
 	virtual void UpdateTransforms(UGeometryCollection const& InGeometryCollection, FTransform const& InBaseTransform, TArrayView<const FMatrix> InMatrices) override;
 	virtual void OnUnregisterGeometryCollection() override;
 	//~ End IGeometryCollectionExternalRenderInterface Interface.
diff --git a/Engine/Source/Runtime/Experimental/GeometryCollectionEngine/Public/GeometryCollection/GeometryCollectionExternalRenderInterface.h b/Engine/Source/Runtime/Experimental/GeometryCollectionEngine/Public/GeometryCollection/GeometryCollectionExternalRenderInterface.h
--- a/Engine/Source/Runtime/Experimental/GeometryCollectionEngine/Public/GeometryCollection/GeometryCollectionExternalRenderInterface.h
+++ b/Engine/Source/Runtime/Experimental/GeometryCollectionEngine/Public/GeometryCollection/GeometryCollectionExternalRenderInterface.h
@@ -19,7 +19,8 @@ class GEOMETRYCOLLECTIONENGINE_API IGeometryCollectionExternalRenderInterface
 	GENERATED_BODY()
 
 public:
-	virtual void UpdateState(UGeometryCollection const& InGeometryCollection, FTransform const& InBaseTransform, bool bInIsBroken) = 0;
+	virtual void UpdateState(UGeometryCollection const& InGeometryCollection, bool bInIsBroken) = 0;
+	virtual void UpdateRootTransform(UGeometryCollection const& InGeometryCollection, FTransform const& InBaseTransform, FTransform const& InRootTransform) = 0;
 	virtual void UpdateTransforms(UGeometryCollection const& InGeometryCollection, FTransform const& InBaseTransform, TArrayView<const FMatrix> InMatrices) = 0;
 	virtual void OnUnregisterGeometryCollection() = 0;
 };
