Chaos - collision solver uses scratch buffer (groundwork to allow SIMD solve when beneficial)
- reshare code between shock prop between SIMD and non-SIMD collision solvers

diff --git a/Engine/Source/Runtime/Experimental/Chaos/Private/Chaos/Collision/PBDCollisionContainerSolver.cpp b/Engine/Source/Runtime/Experimental/Chaos/Private/Chaos/Collision/PBDCollisionContainerSolver.cpp
--- a/Engine/Source/Runtime/Experimental/Chaos/Private/Chaos/Collision/PBDCollisionContainerSolver.cpp
+++ b/Engine/Source/Runtime/Experimental/Chaos/Private/Chaos/Collision/PBDCollisionContainerSolver.cpp
@@ -3,6 +3,7 @@
 
 #include "Chaos/Collision/CollisionApplyType.h"
 #include "Chaos/Collision/PBDCollisionConstraint.h"
+#include "Chaos/Collision/PBDCollisionSolverUtilities.h"
 #include "Chaos/CollisionResolution.h"
 #include "Chaos/DebugDrawQueue.h"
 #include "Chaos/Evolution/SolverBodyContainer.h"
@@ -11,6 +12,7 @@
 #include "Chaos/PBDCollisionConstraints.h"
 #include "Chaos/PBDCollisionConstraintsContact.h"
 #include "Chaos/Utilities.h"
+#include "Templates/AlignmentTemplates.h"
 
 // Private includes
 
@@ -24,6 +26,8 @@ namespace Chaos
 	{
 		extern bool bChaos_PBDCollisionSolver_Position_SolveEnabled;
 		extern bool bChaos_PBDCollisionSolver_Velocity_SolveEnabled;
+		extern float Chaos_PBDCollisionSolver_Position_MinInvMassScale;
+		extern float Chaos_PBDCollisionSolver_Velocity_MinInvMassScale;
 
 		// If one body is more than MassRatio1 times the mass of the other, adjust the solver stiffness when the lighter body is underneath.
 		// Solver stiffness will be equal to 1 when the mass ratio is MassRatio1.
@@ -56,7 +60,14 @@ namespace Chaos
 	//////////////////////////////////////////////////////////////////////////////////////////////////
 
 	// Transform the Constraint's local-space data into world space for use by  the collision solver and also calculate tangents, errors, etc
-	FORCEINLINE_DEBUGGABLE void UpdateCollisionSolverContactPointFromConstraint(Private::FPBDCollisionSolver& Solver, const int32 SolverPointIndex, const FPBDCollisionConstraint* Constraint, const int32 ConstraintPointIndex, const FRealSingle Dt, const FSolverBody& Body0, const FSolverBody& Body1)
+	void UpdateCollisionSolverContactPointFromConstraint(
+		Private::FPBDCollisionSolver& Solver, 
+		const int32 SolverPointIndex, 
+		const FPBDCollisionConstraint* Constraint, 
+		const int32 ConstraintPointIndex, 
+		const FRealSingle Dt, 
+		const FConstraintSolverBody& Body0,
+		const FConstraintSolverBody& Body1)
 	{
 		const FManifoldPoint& ManifoldPoint = Constraint->GetManifoldPoint(ConstraintPointIndex);
 
@@ -131,22 +142,29 @@ namespace Chaos
 			}
 		}
 
-		FWorldContactPoint& WorldContactPoint = Solver.GetWorldContactPoint(SolverPointIndex);
-		WorldContactPoint.RelativeContactPoints[0] = WorldRelativeContact0;
-		WorldContactPoint.RelativeContactPoints[1] = WorldRelativeContact1;
-		WorldContactPoint.ContactNormal = WorldContactNormal;
-		WorldContactPoint.ContactTangentU = WorldContactTangentU;
-		WorldContactPoint.ContactTangentV = WorldContactTangentV;
-		WorldContactPoint.ContactDeltaNormal = WorldContactDeltaNormal;
-		WorldContactPoint.ContactDeltaTangentU = WorldContactDeltaTangentU;
-		WorldContactPoint.ContactDeltaTangentV = WorldContactDeltaTangentV;
-		WorldContactPoint.ContactTargetVelocityNormal = WorldContactTargetVelocityNormal;
+		Solver.InitManifoldPoint(
+			SolverPointIndex,
+			Dt,
+			WorldRelativeContact0,
+			WorldRelativeContact1,
+			WorldContactNormal,
+			WorldContactTangentU,
+			WorldContactTangentV,
+			WorldContactDeltaNormal,
+			WorldContactDeltaTangentU,
+			WorldContactDeltaTangentV,
+			WorldContactTargetVelocityNormal);
 	}
 
-	FORCEINLINE_DEBUGGABLE void UpdateCollisionSolverManifoldFromConstraint(Private::FPBDCollisionSolver& Solver, const FPBDCollisionConstraint* Constraint, const FSolverReal Dt, const int32 ConstraintPointBeginIndex, const int32 ConstraintPointEndIndex)
+	void UpdateCollisionSolverManifoldFromConstraint(
+		Private::FPBDCollisionSolver& Solver, 
+		const FPBDCollisionConstraint* Constraint, 
+		const FSolverReal Dt, 
+		const int32 ConstraintPointBeginIndex, 
+		const int32 ConstraintPointEndIndex)
 	{
-		const FSolverBody& Body0 = Solver.SolverBody0().SolverBody();
-		const FSolverBody& Body1 = Solver.SolverBody1().SolverBody();
+		const FConstraintSolverBody& Body0 = Solver.SolverBody0();
+		const FConstraintSolverBody& Body1 = Solver.SolverBody1();
 
 		// Only calculate state for newly added contacts. Normally this is all of them, but maybe not if incremental collision is used by RBAN.
 		// Also we only add active points to the solver's manifold points list
@@ -159,14 +177,18 @@ namespace Chaos
 				// Transform the constraint contact data into world space for use by the solver
 				// We build this data directly into the solver's world-space contact data which looks a bit odd with "Init" called after but there you go
 				UpdateCollisionSolverContactPointFromConstraint(Solver, SolverManifoldPointIndex, Constraint, ConstraintManifoldPointIndex, Dt, Body0, Body1);
-
-				// Prepare the solver (set constraint-space mass etc)
-				Solver.FinalizeManifoldPoint(SolverManifoldPointIndex, Dt);
 			}
 		}
+
+		Solver.FinalizeManifold();
 	}
 
-	FORCEINLINE_DEBUGGABLE void UpdateCollisionSolverFromConstraint(Private::FPBDCollisionSolver& Solver, const FPBDCollisionConstraint* Constraint, const FSolverReal Dt, const FPBDCollisionSolverSettings& SolverSettings, bool& bOutPerIterationCollision)
+	void UpdateCollisionSolverFromConstraint(
+		Private::FPBDCollisionSolver& Solver, 
+		const FPBDCollisionConstraint* Constraint, 
+		const FSolverReal Dt, 
+		const FPBDCollisionSolverSettings& SolverSettings, 
+		bool& bOutPerIterationCollision)
 	{
 		// Friction values. Static and Dynamic friction are applied in the position solve for most shapes.
 		// We can also run in a mode without static friction at all. This is faster but stacking is not possible.
@@ -264,6 +286,11 @@ namespace Chaos
 	FPBDCollisionContainerSolver::FPBDCollisionContainerSolver(const FPBDCollisionConstraints& InConstraintContainer, const int32 InPriority)
 		: FConstraintContainerSolver(InPriority)
 		, ConstraintContainer(InConstraintContainer)
+		, CollisionConstraints()
+		, AppliedShockPropagation(1)
+		, Scratch()
+		, CollisionSolvers()
+		, bCollisionConstraintPerIterationCollisionDetection()
 		, bPerIterationCollisionDetection(false)
 	{
 #if !UE_BUILD_SHIPPING && !UE_BUILD_TEST && INTEL_ISPC
@@ -277,22 +304,32 @@ namespace Chaos
 
 	void FPBDCollisionContainerSolver::Reset(const int32 MaxCollisions)
 	{
+		AppliedShockPropagation = FSolverReal(1);
+
 		// A over-allocation policy to avoid reallocation every frame in the common case where a pile of objects is dropped
 		// and the number of contacts increases every tick.
-		if (MaxCollisions > CollisionSolvers.Max())
+		int CollisionBufferNum = MaxCollisions;
+		if (CollisionBufferNum > CollisionConstraints.Max())
 		{
-			CollisionSolvers.Reset((5 * MaxCollisions) / 4);	// +25%
+			CollisionBufferNum = (5 * MaxCollisions) / 4; // +25%
 		}
 
-		CollisionSolvers.Reset(MaxCollisions);
-		CollisionConstraints.Reset(CollisionSolvers.Max());
-		bCollisionConstraintPerIterationCollisionDetection.Reset(CollisionSolvers.Max());
-		CollisionSolverManifoldPoints.Reset(CollisionSolvers.Max() * Private::FPBDCollisionSolver::MaxPointsPerConstraint);
+		CollisionConstraints.Reset(CollisionBufferNum);
+		bCollisionConstraintPerIterationCollisionDetection.Reset(CollisionBufferNum);
+
+		// Just set the array size for these right away - all data will be initialized later
+		bCollisionConstraintPerIterationCollisionDetection.SetNumUninitialized(MaxCollisions, false);
 
-		// Just set the array size for these right away - all data will be initialized later and there is no constructor
-		// We use CollisionConstraints.Num() to track how many constraints we currently have
-		CollisionSolvers.SetNumUninitialized(MaxCollisions);
-		bCollisionConstraintPerIterationCollisionDetection.SetNumUninitialized(MaxCollisions);
+		// Prepare the scratch buffer
+		const size_t AlignedSolverSize = Align(sizeof(Private::FPBDCollisionSolver), alignof(Private::FPBDCollisionSolver));
+		const size_t AlignedPointSize = Align(sizeof(Private::FPBDCollisionSolverManifoldPoint), alignof(Private::FPBDCollisionSolverManifoldPoint));
+		const size_t ScratchSize = CollisionBufferNum * (AlignedSolverSize + Private::FPBDCollisionSolver::MaxPointsPerConstraint * AlignedPointSize);
+		Scratch.Reset(ScratchSize);
+
+		// Allocate scratch space for the collision solvers and manifold points
+		CollisionSolvers = Scratch.AllocArray<Private::FPBDCollisionSolver>(MaxCollisions);
+		CollisionSolverManifoldPoints = Scratch.AllocArray<Private::FPBDCollisionSolverManifoldPoint>(MaxCollisions);
+		NumCollisionSolverManifoldPoints = 0;
 	}
 
 	void FPBDCollisionContainerSolver::AddConstraints()
@@ -328,19 +365,15 @@ namespace Chaos
 		// @todo(chaos): see if we can do better - we don't want to pay for this logic when it is only used rarely
 		const bool bDeferredCollisionDetection = ConstraintContainer.GetDetectorSettings().bDeferNarrowPhase;
 		const int32 ManifoldPointMax = (bDeferredCollisionDetection || Constraint.GetUseIncrementalManifold()) ? Private::FPBDCollisionSolver::MaxPointsPerConstraint : Constraint.NumManifoldPoints();
+
 		if (ManifoldPointMax > 0)
 		{
-			const int32 ManifoldPointStart = CollisionSolverManifoldPoints.Num();
-			check(CollisionSolverManifoldPoints.Max() >= ManifoldPointStart + ManifoldPointMax);	// No resize following reset
-			CollisionSolverManifoldPoints.SetNum(ManifoldPointStart + ManifoldPointMax);
-
-			CollisionSolvers[Index].SetManifoldPointsBuffer(&CollisionSolverManifoldPoints[ManifoldPointStart], ManifoldPointMax);
+			GetSolver(Index).SetManifoldPointsBuffer(&CollisionSolverManifoldPoints[NumCollisionSolverManifoldPoints], ManifoldPointMax);
+			NumCollisionSolverManifoldPoints += ManifoldPointMax;
 		}
 		else
 		{
-			// Ideally we wouldn't have any active constraint with no manifold points, but this can
-			// currently happen for CCD objects and incremental collisions
-			CollisionSolvers[Index].Reset();
+			GetSolver(Index).Reset();
 		}
 	}
 
@@ -349,8 +382,8 @@ namespace Chaos
 	{
 		for (int32 SolverIndex = 0, SolverEndIndex = NumSolvers(); SolverIndex < SolverEndIndex; ++SolverIndex)
 		{
-			Private::FPBDCollisionSolver& CollisionSolver = CollisionSolvers[SolverIndex];
-			FPBDCollisionConstraint* Constraint = CollisionConstraints[SolverIndex];
+			Private::FPBDCollisionSolver& CollisionSolver = GetSolver(SolverIndex);
+			FPBDCollisionConstraint* Constraint = GetConstraint(SolverIndex);
 			check(Constraint != nullptr);
 
 			// Find the solver bodies for the particles we constrain. This will add them to the container
@@ -373,12 +406,8 @@ namespace Chaos
 	{
 		if (ConstraintIndex < NumSolvers())
 		{
-			const int32 NumManifoldPoints = CollisionSolvers[ConstraintIndex].NumManifoldPoints();
-			if (NumManifoldPoints > 0)
-			{
-				FPlatformMisc::PrefetchBlock(&CollisionSolvers[ConstraintIndex].GetManifoldPoint(0), NumManifoldPoints * sizeof(Private::FPBDCollisionSolverManifoldPoint));
-			}
-			FPlatformMisc::PrefetchBlock(&CollisionSolvers[ConstraintIndex], sizeof(FPBDCollisionContainerSolver));
+			FPlatformMisc::PrefetchBlock(GetConstraint(ConstraintIndex), sizeof(FPBDCollisionConstraint));
+			FPlatformMisc::PrefetchBlock(&GetSolver(ConstraintIndex), sizeof(Private::FPBDCollisionSolver));
 		}
 	}
 
@@ -391,7 +420,7 @@ namespace Chaos
 
 		const FSolverReal Dt = FSolverReal(InDt);
 
-		const int32 PrefetchCount = 4;
+		const int32 PrefetchCount = 2;
 		for (int32 PrefetchIndex = 0; PrefetchIndex < PrefetchCount; ++PrefetchIndex)
 		{
 			CachePrefetchSolver(PrefetchIndex);
@@ -400,16 +429,16 @@ namespace Chaos
 		bool bWantPerIterationCollisionDetection = false;
 		for (int32 ConstraintIndex = BeginIndex; ConstraintIndex < EndIndex; ++ConstraintIndex)
 		{
-			Private::FPBDCollisionSolver& CollisionSolver = CollisionSolvers[ConstraintIndex];
-			FPBDCollisionConstraint* Constraint = CollisionConstraints[ConstraintIndex];
+			CachePrefetchSolver(ConstraintIndex + PrefetchCount);
+
+			Private::FPBDCollisionSolver& CollisionSolver = GetSolver(ConstraintIndex);
+			FPBDCollisionConstraint* Constraint = GetConstraint(ConstraintIndex);
 			bool& bPerIterationCollision = bCollisionConstraintPerIterationCollisionDetection[ConstraintIndex];
 
 			UpdateCollisionSolverFromConstraint(CollisionSolver, Constraint, Dt, ConstraintContainer.GetSolverSettings(), bPerIterationCollision);
 
 			// We need to run collision every iteration if we are not using manifolds, or are using incremental manifolds
 			bWantPerIterationCollisionDetection = bWantPerIterationCollisionDetection || bPerIterationCollision;
-
-			CachePrefetchSolver(ConstraintIndex + PrefetchCount);
 		}
 
 		if (bWantPerIterationCollisionDetection)
@@ -435,8 +464,8 @@ namespace Chaos
 
 		for (int32 ConstraintIndex = BeginIndex; ConstraintIndex < EndIndex; ++ConstraintIndex)
 		{
-			Private::FPBDCollisionSolver& CollisionSolver = CollisionSolvers[ConstraintIndex];
-			FPBDCollisionConstraint* Constraint = CollisionConstraints[ConstraintIndex];
+			Private::FPBDCollisionSolver& CollisionSolver = GetSolver(ConstraintIndex);
+			FPBDCollisionConstraint* Constraint = GetConstraint(ConstraintIndex);
 
 			UpdateCollisionConstraintFromSolver(Constraint, CollisionSolver, Dt);
 
@@ -461,38 +490,46 @@ namespace Chaos
 		// Not supported for collisions
 	}
 
-	void FPBDCollisionContainerSolver::UpdatePositionShockPropagation(const FReal Dt, const int32 It, const int32 NumIts, const int32 BeginIndex, const int32 EndIndex, const FPBDCollisionSolverSettings& SolverSettings)
+	void FPBDCollisionContainerSolver::ApplyShockPropagation(const FSolverReal ShockPropagation)
 	{
-		// If this is the first shock propagation iteration, enable it on each solver
-		const bool bEnableShockPropagation = (It == NumIts - SolverSettings.NumPositionShockPropagationIterations);
-		if (bEnableShockPropagation)
+		// @todo(chaos): cache the mass scales so we don't have to look in the constraint again
+		if (ShockPropagation != AppliedShockPropagation)
 		{
-			for (int32 SolverIndex = BeginIndex; SolverIndex < EndIndex; ++SolverIndex)
+			for (int32 SolverIndex = 0; SolverIndex < NumSolvers(); ++SolverIndex)
 			{
-				CollisionSolvers[SolverIndex].EnablePositionShockPropagation();
+				const FPBDCollisionConstraint* Constraint = GetConstraint(SolverIndex);
+				Private::FPBDCollisionSolver& Solver = GetSolver(SolverIndex);
+				if (Constraint != nullptr)
+				{
+					FConstraintSolverBody& Body0 = Solver.SolverBody0();
+					FConstraintSolverBody& Body1 = Solver.SolverBody1();
+
+					FSolverReal ShockPropagation0, ShockPropagation1;
+					if (Private::CalculateBodyShockPropagation(Body0.SolverBody(), Body1.SolverBody(), ShockPropagation, ShockPropagation0, ShockPropagation1))
+					{
+						Body0.SetShockPropagationScale(ShockPropagation0);
+						Body1.SetShockPropagationScale(ShockPropagation1);
+						Solver.UpdateMassNormal();
+					}
+				}
 			}
+
+			AppliedShockPropagation = ShockPropagation;
 		}
 	}
 
+	void FPBDCollisionContainerSolver::UpdatePositionShockPropagation(const FReal Dt, const int32 It, const int32 NumIts, const int32 BeginIndex, const int32 EndIndex, const FPBDCollisionSolverSettings& SolverSettings)
+	{
+		const bool bEnableShockPropagation = (It >= NumIts - SolverSettings.NumPositionShockPropagationIterations);
+		const FSolverReal ShockPropagation = (bEnableShockPropagation) ? CVars::Chaos_PBDCollisionSolver_Position_MinInvMassScale : FSolverReal(1);
+		ApplyShockPropagation(ShockPropagation);
+	}
+
 	void FPBDCollisionContainerSolver::UpdateVelocityShockPropagation(const FReal Dt, const int32 It, const int32 NumIts, const int32 BeginIndex, const int32 EndIndex, const FPBDCollisionSolverSettings& SolverSettings)
 	{
-		// Set/reset the shock propagation based on current iteration. The position solve may
-		// have left the bodies with a mass scale and we want to change or reset it.
-		const bool bEnableShockPropagation = (It == NumIts - SolverSettings.NumVelocityShockPropagationIterations);
-		if (bEnableShockPropagation)
-		{
-			for (int32 SolverIndex = BeginIndex; SolverIndex < EndIndex; ++SolverIndex)
-			{
-				CollisionSolvers[SolverIndex].EnableVelocityShockPropagation();
-			}
-		}
-		else if (It == 0)
-		{
-			for (int32 SolverIndex = BeginIndex; SolverIndex < EndIndex; ++SolverIndex)
-			{
-				CollisionSolvers[SolverIndex].DisableShockPropagation();
-			}
-		}
+		const bool bEnableShockPropagation = (It >= NumIts - SolverSettings.NumVelocityShockPropagationIterations);
+		const FSolverReal ShockPropagation = (bEnableShockPropagation) ? CVars::Chaos_PBDCollisionSolver_Velocity_MinInvMassScale : FSolverReal(1);
+		ApplyShockPropagation(ShockPropagation);
 	}
 
 	bool FPBDCollisionContainerSolver::SolvePositionImpl(const FReal InDt, const int32 It, const int32 NumIts, const int32 BeginIndex, const int32 EndIndex, const FPBDCollisionSolverSettings& SolverSettings)
@@ -528,11 +565,11 @@ namespace Chaos
 		// Apply the position correction
 		if (bApplyStaticFriction)
 		{
-			Private::FPBDCollisionSolverHelper::SolvePositionWithFriction(MakeArrayView(&CollisionSolvers[BeginIndex], EndIndex - BeginIndex), Dt, MaxPushOut);
+			Private::FPBDCollisionSolverHelper::SolvePositionWithFriction(MakeArrayView(&GetSolver(BeginIndex), EndIndex - BeginIndex), Dt, MaxPushOut);
 		}
 		else
 		{
-			Private::FPBDCollisionSolverHelper::SolvePositionNoFriction(MakeArrayView(&CollisionSolvers[BeginIndex], EndIndex - BeginIndex), Dt, MaxPushOut);
+			Private::FPBDCollisionSolverHelper::SolvePositionNoFriction(MakeArrayView(&GetSolver(BeginIndex), EndIndex - BeginIndex), Dt, MaxPushOut);
 		}
 
 		return true;
@@ -556,7 +593,7 @@ namespace Chaos
 		const FSolverReal Dt = FSolverReal(InDt);
 		const bool bApplyDynamicFriction = (It >= NumIts - SolverSettings.NumVelocityFrictionIterations);
 
-		Private::FPBDCollisionSolverHelper::SolveVelocity(MakeArrayView(&CollisionSolvers[BeginIndex], EndIndex - BeginIndex), Dt, bApplyDynamicFriction);
+		Private::FPBDCollisionSolverHelper::SolveVelocity(MakeArrayView(&GetSolver(BeginIndex), EndIndex - BeginIndex), Dt, bApplyDynamicFriction);
 
 		return true;
 	}
@@ -571,12 +608,12 @@ namespace Chaos
 		{
 			if (bDeferredCollisionDetection || bCollisionConstraintPerIterationCollisionDetection[SolverIndex])
 			{
-				Private::FPBDCollisionSolver& CollisionSolver = CollisionSolvers[SolverIndex];
-				FPBDCollisionConstraint* Constraint = CollisionConstraints[SolverIndex];
+				Private::FPBDCollisionSolver& CollisionSolver = GetSolver(SolverIndex);
+				FPBDCollisionConstraint* Constraint = GetConstraint(SolverIndex);
 
 				// Run collision detection at the current transforms including any correction from previous iterations
-				const FSolverBody& Body0 = CollisionSolver.SolverBody0().SolverBody();
-				const FSolverBody& Body1 = CollisionSolver.SolverBody1().SolverBody();
+				const FConstraintSolverBody& Body0 = CollisionSolver.SolverBody0();
+				const FConstraintSolverBody& Body1 = CollisionSolver.SolverBody1();
 				const FRigidTransform3 CorrectedActorWorldTransform0 = FRigidTransform3(Body0.CorrectedActorP(), Body0.CorrectedActorQ());
 				const FRigidTransform3 CorrectedActorWorldTransform1 = FRigidTransform3(Body1.CorrectedActorP(), Body1.CorrectedActorQ());
 				const FRigidTransform3 CorrectedShapeWorldTransform0 = Constraint->GetShapeRelativeTransform0() * CorrectedActorWorldTransform0;
@@ -584,7 +621,7 @@ namespace Chaos
 
 				// @todo(chaos): this is ugly - pass these to the required functions instead and remove from the constraint class
 				// This is now only needed for LevelSet collision (see UpdateLevelsetLevelsetConstraint)
-				Constraint->SetSolverBodies(&Body0, &Body1);
+				Constraint->SetSolverBodies(&Body0.SolverBody(), &Body1.SolverBody());
 
 				// Reset the manifold if we are not using manifolds (we just use the first manifold point)
 				if (!Constraint->GetUseManifold())
diff --git a/Engine/Source/Runtime/Experimental/Chaos/Private/Chaos/Collision/PBDCollisionContainerSolverSimd.cpp b/Engine/Source/Runtime/Experimental/Chaos/Private/Chaos/Collision/PBDCollisionContainerSolverSimd.cpp
--- a/Engine/Source/Runtime/Experimental/Chaos/Private/Chaos/Collision/PBDCollisionContainerSolverSimd.cpp
+++ b/Engine/Source/Runtime/Experimental/Chaos/Private/Chaos/Collision/PBDCollisionContainerSolverSimd.cpp
@@ -3,6 +3,7 @@
 
 #include "Chaos/Collision/CollisionApplyType.h"
 #include "Chaos/Collision/PBDCollisionConstraint.h"
+#include "Chaos/Collision/PBDCollisionSolverUtilities.h"
 #include "Chaos/CollisionResolution.h"
 #include "Chaos/DebugDrawQueue.h"
 #include "Chaos/Evolution/SolverBodyContainer.h"
@@ -34,37 +35,6 @@ namespace Chaos
 		//////////////////////////////////////////////////////////////////////////////////////////////////
 		//////////////////////////////////////////////////////////////////////////////////////////////////
 
-		bool CalculateBodyShockPropagation(
-			const FSolverBody& Body0,
-			const FSolverBody& Body1,
-			const FSolverReal ShockPropagation,
-			FSolverReal& OutShockPropagation0,
-			FSolverReal& OutShockPropagation1)
-		{
-			OutShockPropagation0 = FSolverReal(1);
-			OutShockPropagation1 = FSolverReal(1);
-
-			// Shock propagation decreases the inverse mass of bodies that are lower in the pile
-			// of objects. This significantly improves stability of heaps and stacks. Height in the pile is indictaed by the "level". 
-			// No need to set an inverse mass scale if the other body is kinematic (with inv mass of 0).
-			// Bodies at the same level do not take part in shock propagation.
-			if (Body0.IsDynamic() && Body1.IsDynamic() && (Body0.Level() != Body1.Level()))
-			{
-				// Set the inv mass scale of the "lower" body to make it heavier
-				if (Body0.Level() < Body1.Level())
-				{
-					OutShockPropagation0 = ShockPropagation;
-				}
-				else
-				{
-					OutShockPropagation1 = ShockPropagation;
-				}
-				return true;
-			}
-
-			return false;
-		}
-
 		// Transform and copy a single manifold point for use in the solver
 		template<int TNumLanes>
 		void UpdateSolverManifoldPointFromConstraint(
diff --git a/Engine/Source/Runtime/Experimental/Chaos/Private/Chaos/Collision/PBDCollisionSolver.cpp b/Engine/Source/Runtime/Experimental/Chaos/Private/Chaos/Collision/PBDCollisionSolver.cpp
--- a/Engine/Source/Runtime/Experimental/Chaos/Private/Chaos/Collision/PBDCollisionSolver.cpp
+++ b/Engine/Source/Runtime/Experimental/Chaos/Private/Chaos/Collision/PBDCollisionSolver.cpp
@@ -69,72 +69,11 @@ namespace Chaos
 
 	namespace Private
 	{
-		//////////////////////////////////////////////////////////////////////////////////////////////////
-		//////////////////////////////////////////////////////////////////////////////////////////////////
-		//////////////////////////////////////////////////////////////////////////////////////////////////
-		//////////////////////////////////////////////////////////////////////////////////////////////////
-
-
-		void FPBDCollisionSolver::EnablePositionShockPropagation()
-		{
-			SetShockPropagationInvMassScale(Chaos_PBDCollisionSolver_Position_MinInvMassScale);
-		}
-
-		void FPBDCollisionSolver::EnableVelocityShockPropagation()
-		{
-			SetShockPropagationInvMassScale(Chaos_PBDCollisionSolver_Velocity_MinInvMassScale);
-		}
-
-		void FPBDCollisionSolver::DisableShockPropagation()
-		{
-			SetShockPropagationInvMassScale(FReal(1));
-		}
-
-		void FPBDCollisionSolver::SetShockPropagationInvMassScale(const FSolverReal InvMassScale)
-		{
-			FConstraintSolverBody& Body0 = SolverBody0();
-			FConstraintSolverBody& Body1 = SolverBody1();
-
-			// Shock propagation decreases the inverse mass of bodies that are lower in the pile
-			// of objects. This significantly improves stability of heaps and stacks. Height in the pile is indictaed by the "level". 
-			// No need to set an inverse mass scale if the other body is kinematic (with inv mass of 0).
-			// Bodies at the same level do not take part in shock propagation.
-			if (Body0.IsDynamic() && Body1.IsDynamic() && (Body0.Level() != Body1.Level()))
-			{
-				// Set the inv mass scale of the "lower" body to make it heavier
-				bool bInvMassUpdated = false;
-				if (Body0.Level() < Body1.Level())
-				{
-					if (Body0.ShockPropagationScale() != InvMassScale)
-					{
-						Body0.SetShockPropagationScale(InvMassScale);
-						bInvMassUpdated = true;
-					}
-				}
-				else
-				{
-					if (Body1.ShockPropagationScale() != InvMassScale)
-					{
-						Body1.SetShockPropagationScale(InvMassScale);
-						bInvMassUpdated = true;
-					}
-				}
-
-				// If the masses changed, we need to rebuild the contact mass for each manifold point
-				if (bInvMassUpdated)
-				{
-					for (int32 PointIndex = 0; PointIndex < NumManifoldPoints(); ++PointIndex)
-					{
-						State.ManifoldPoints[PointIndex].UpdateMassNormal(Body0, Body1);
-					}
-				}
-			}
-		}
-
 		void FPBDCollisionSolver::SolveVelocityAverage(const FSolverReal Dt)
 		{
-			FConstraintSolverBody& Body0 = SolverBody0();
-			FConstraintSolverBody& Body1 = SolverBody1();
+			FSolverReal InvM0, InvM1;
+			FSolverMatrix33 InvI0, InvI1;
+			GetDynamicMassProperties(InvM0, InvI0, InvM1, InvI1);
 
 			// Generate a new contact point at the average of all the active contacts
 			int32 NumActiveManifoldPoints = 0;
@@ -146,7 +85,7 @@ namespace Chaos
 			for (int32 PointIndex = 0; PointIndex < NumManifoldPoints(); ++PointIndex)
 			{
 				FPBDCollisionSolverManifoldPoint& SolverManifoldPoint = State.ManifoldPoints[PointIndex];
-				if (SolverManifoldPoint.ShouldSolveVelocity())
+				if (ShouldSolveVelocity(SolverManifoldPoint))
 				{
 					RelativeContactPosition0 += SolverManifoldPoint.WorldContact.RelativeContactPoints[0];
 					RelativeContactPosition1 += SolverManifoldPoint.WorldContact.RelativeContactPoints[1];
@@ -183,17 +122,17 @@ namespace Chaos
 				FSolverVec3 WorldContactNormalAngular0 = FSolverVec3(0);
 				FSolverVec3 WorldContactNormalAngular1 = FSolverVec3(0);
 				FSolverReal ContactMassInvNormal = FSolverReal(0);
-				if (Body0.IsDynamic())
+				if (IsDynamic(0))
 				{
 					const FSolverVec3 R0xN = FSolverVec3::CrossProduct(RelativeContactPosition0, WorldContactNormal);
-					WorldContactNormalAngular0 = Body0.InvI() * R0xN;
-					ContactMassInvNormal += FSolverVec3::DotProduct(R0xN, WorldContactNormalAngular0) + Body0.InvM();
+					WorldContactNormalAngular0 = InvI0 * R0xN;
+					ContactMassInvNormal += FSolverVec3::DotProduct(R0xN, WorldContactNormalAngular0) + InvM0;
 				}
-				if (Body1.IsDynamic())
+				if (IsDynamic(1))
 				{
 					const FSolverVec3 R1xN = FSolverVec3::CrossProduct(RelativeContactPosition1, WorldContactNormal);
-					WorldContactNormalAngular1 = Body1.InvI() * R1xN;
-					ContactMassInvNormal += FSolverVec3::DotProduct(R1xN, WorldContactNormalAngular1) + Body1.InvM();
+					WorldContactNormalAngular1 = InvI1 * R1xN;
+					ContactMassInvNormal += FSolverVec3::DotProduct(R1xN, WorldContactNormalAngular1) + InvM1;
 				}
 				AverageManifoldPoint.ContactMassNormal = (ContactMassInvNormal > FSolverReal(UE_SMALL_NUMBER)) ? FSolverReal(1) / ContactMassInvNormal : FSolverReal(0);
 				AverageManifoldPoint.WorldContactNormalAngular0 = WorldContactNormalAngular0;
@@ -219,7 +158,7 @@ namespace Chaos
 				AverageManifoldPoint.StaticFrictionRatio = FSolverReal(0);
 
 				FSolverReal ContactVelocityDeltaNormal;
-				AverageManifoldPoint.CalculateContactVelocityErrorNormal(Body0, Body1, ContactVelocityDeltaNormal);
+				CalculateContactVelocityErrorNormal(AverageManifoldPoint, ContactVelocityDeltaNormal);
 
 				const FSolverReal MinImpulseNormal = FMath::Min(FSolverReal(0), -AverageManifoldPoint.NetPushOutNormal / Dt);
 
@@ -227,15 +166,13 @@ namespace Chaos
 					State.Stiffness,
 					ContactVelocityDeltaNormal,
 					MinImpulseNormal,
-					AverageManifoldPoint,
-					Body0,
-					Body1);
+					AverageManifoldPoint);
 
 				// Now distribute the net impulse among the active points so we don't over-correct pushout from initial overlaps
 				for (int32 PointIndex = 0; PointIndex < NumManifoldPoints(); ++PointIndex)
 				{
 					FPBDCollisionSolverManifoldPoint& SolverManifoldPoint = State.ManifoldPoints[PointIndex];
-					if (SolverManifoldPoint.ShouldSolveVelocity())
+					if (ShouldSolveVelocity(SolverManifoldPoint))
 					{
 						SolverManifoldPoint.NetImpulseNormal += AverageManifoldPoint.NetImpulseNormal * InvCount;
 					}
@@ -248,54 +185,65 @@ namespace Chaos
 		//////////////////////////////////////////////////////////////////////////////////////////////////
 		//////////////////////////////////////////////////////////////////////////////////////////////////
 
+		void CachePrefetchSolver(const TArrayView<FPBDCollisionSolver>& CollisionSolvers, const int32 ConstraintIndex)
+		{
+			//if (ConstraintIndex < CollisionSolvers.Num())
+			//{
+			//	if (ConstraintIndex < CollisionSolvers.Num() - 1)
+			//	{
+			//		FPlatformMisc::PrefetchBlock(&CollisionSolvers[ConstraintIndex + 1], sizeof(Private::FPBDCollisionSolver));
+			//	}
+
+			//	const int32 NumManifoldPoints = CollisionSolvers[ConstraintIndex].NumManifoldPoints();
+			//	if (NumManifoldPoints > 0)
+			//	{
+			//		FPlatformMisc::PrefetchBlock(&CollisionSolvers[ConstraintIndex].GetManifoldPoint(0), NumManifoldPoints * sizeof(Private::FPBDCollisionSolverManifoldPoint));
+			//	}
+			//}
+		}
+
 		void FPBDCollisionSolverHelper::SolvePositionNoFriction(const TArrayView<FPBDCollisionSolver>& CollisionSolvers, const FSolverReal Dt, const FSolverReal MaxPushOut)
 		{
 			SCOPE_CYCLE_COUNTER(STAT_SolvePositionNoFriction);
 
-#if INTEL_ISPC
-			if (CVars::bChaos_PBDCollisionSolver_ISPC)
+			for (int32 SolverIndex = 0; SolverIndex < CollisionSolvers.Num(); ++SolverIndex)
 			{
-				ispc::SolvePositionNoFriction(
-					(ispc::FPBDCollisionSolver*)CollisionSolvers.GetData(), 
-					CollisionSolvers.Num(), 
-					Dt, 
-					MaxPushOut);
-				return;
-			}
-#endif
+				CachePrefetchSolver(CollisionSolvers, SolverIndex + 1);
 
-			for (FPBDCollisionSolver& CollisionSolver : CollisionSolvers)
-			{
-				CollisionSolver.SolvePositionNoFriction(Dt, MaxPushOut);
+				CollisionSolvers[SolverIndex].SolvePositionNoFriction(Dt, MaxPushOut);
 			}
 		}
 
 		void FPBDCollisionSolverHelper::SolvePositionWithFriction(const TArrayView<FPBDCollisionSolver>& CollisionSolvers, const FSolverReal Dt, const FSolverReal MaxPushOut)
 		{
-			for (FPBDCollisionSolver& CollisionSolver : CollisionSolvers)
+			for (int32 SolverIndex = 0; SolverIndex < CollisionSolvers.Num(); ++SolverIndex)
 			{
-				CollisionSolver.SolvePositionWithFriction(Dt, MaxPushOut);
+				CachePrefetchSolver(CollisionSolvers, SolverIndex + 1);
+
+				CollisionSolvers[SolverIndex].SolvePositionWithFriction(Dt, MaxPushOut);
 			}
 		}
 
 		void FPBDCollisionSolverHelper::SolveVelocity(const TArrayView<FPBDCollisionSolver>& CollisionSolvers, const FSolverReal Dt, const bool bApplyDynamicFriction)
 		{
-			for (FPBDCollisionSolver& CollisionSolver : CollisionSolvers)
+			for (int32 SolverIndex = 0; SolverIndex < CollisionSolvers.Num(); ++SolverIndex)
 			{
-				CollisionSolver.SolveVelocity(Dt, bApplyDynamicFriction);
+				CachePrefetchSolver(CollisionSolvers, SolverIndex + 1);
+
+				CollisionSolvers[SolverIndex].SolveVelocity(Dt, bApplyDynamicFriction);
 			}
 		}
 
 		void FPBDCollisionSolverHelper::CheckISPC()
 		{
-#if INTEL_ISPC
-			check(sizeof(ispc::FPBDCollisionSolver) == sizeof(Private::FPBDCollisionSolver));
-			check(sizeof(ispc::FPBDCollisionSolverManifoldPoint) == sizeof(Private::FPBDCollisionSolverManifoldPoint));
-			check(sizeof(ispc::FWorldContactPoint) == sizeof(FWorldContactPoint));
-			check(sizeof(ispc::FConstraintSolverBody) == sizeof(FConstraintSolverBody));
-			check(sizeof(ispc::FSolverBody) == sizeof(FSolverBody));
-#endif
+//#if INTEL_ISPC
+//			check(sizeof(ispc::FPBDCollisionSolver) == sizeof(Private::FPBDCollisionSolver));
+//			check(sizeof(ispc::FPBDCollisionSolverManifoldPoint) == sizeof(Private::FPBDCollisionSolverManifoldPoint));
+//			check(sizeof(ispc::FWorldContactPoint) == sizeof(FWorldContactPoint));
+//			check(sizeof(ispc::FConstraintSolverBody) == sizeof(FConstraintSolverBody));
+//			check(sizeof(ispc::FSolverBody) == sizeof(FSolverBody));
+//#endif
 		}
 
 	}	// namespace Private
-}	// namespace Chaos
\ No newline at end of file
+}	// namespace Chaos
diff --git a/Engine/Source/Runtime/Experimental/Chaos/Private/Chaos/Collision/PBDCollisionSolver.ispc b/Engine/Source/Runtime/Experimental/Chaos/Private/Chaos/Collision/PBDCollisionSolver.ispc
--- a/Engine/Source/Runtime/Experimental/Chaos/Private/Chaos/Collision/PBDCollisionSolver.ispc
+++ b/Engine/Source/Runtime/Experimental/Chaos/Private/Chaos/Collision/PBDCollisionSolver.ispc
@@ -30,7 +30,6 @@ export struct FSolverBody
 	float WPad;
 
 	// World-space inverse inertia
-	// @todo(chaos): do we need this, or should we force all systems to use the FConstraintSolverBody decorator?
 	FMatrix33f InvI;
 
 	// Actor-space center of mass rotation
@@ -72,17 +71,6 @@ export struct FSolverBody
 	float AlignmentPadding[2];
 };
 
-export struct FConstraintSolverBody
-{
-	// The body we decorate
-	FSolverBody* Body;
-
-	// The body modifiers
-	float InvMScale;
-	float InvIScale;
-	float ShockPropagationScale;
-};
-
 export struct FWorldContactPoint
 {
 	// World-space contact point relative to each particle's center of mass
@@ -149,213 +137,13 @@ export struct FPBDCollisionSolver
 	float Stiffness;
 
 	// Bodies and contacts
-	FConstraintSolverBody SolverBodies[2];
+	FSolverBody* SolverBodies[2];
+	float InvMs[2];
 	FPBDCollisionSolverManifoldPoint* ManifoldPoints;
 	int32 NumManifoldPoints;
 	int32 MaxManifoldPoints;
 };
 
-
-//
-//
-// GENERAL FUNCTIONS
-//
-//
-
-
-static inline uniform bool Body_IsDynamic(const uniform FSolverBody* uniform Body)
-{ 
-	return (Body->InvM > FLOAT_SMALL_NUMBER); 
-}
-
-static inline void Body_ApplyPositionDelta(uniform FSolverBody* uniform Body, const uniform FVector3f DP)
-{
-	Body->DP = Body->DP + DP;
-}
-
-static inline void Body_ApplyRotationDelta(uniform FSolverBody* uniform Body, const uniform FVector3f DQ)
-{
-	Body->DQ = Body->DQ + DQ;
-}
-
-static inline uniform float CalculatePositionErrorNormal(
-	uniform FPBDCollisionSolverManifoldPoint* uniform SolverManifoldPoint, 
-	const uniform FSolverBody* uniform Body0, 
-	const uniform FSolverBody* uniform Body1, 
-	const uniform float MaxPushOut)
-{
-	// Linear version: calculate the contact delta assuming linear motion after applying a positional impulse at the contact point. There will be an error that depends on the size of the rotation.
-	const uniform FVector3f ContactDelta0 = Body0->DP + VectorCross(Body0->DQ, SolverManifoldPoint->WorldContact.RelativeContactPoints[0]);
-	const uniform FVector3f ContactDelta1 = Body1->DP + VectorCross(Body1->DQ, SolverManifoldPoint->WorldContact.RelativeContactPoints[1]);
-	const uniform FVector3f ContactDelta = ContactDelta0 - ContactDelta1;
-
-	uniform float ContactDeltaNormal = SolverManifoldPoint->WorldContact.ContactDeltaNormal + VectorDot(ContactDelta, SolverManifoldPoint->WorldContact.ContactNormal);
-
-	// NOTE: ContactDeltaNormal is negative for penetration
-	// NOTE: MaxPushOut == 0 disables the pushout limits
-	if ((MaxPushOut > 0) && (ContactDeltaNormal < -MaxPushOut))
-	{
-		ContactDeltaNormal = -MaxPushOut;
-	}
-
-	return ContactDeltaNormal;
-}
-
-static inline uniform float CalculatePositionCorrectionNormal(
-	const uniform float Stiffness,
-	const uniform float ContactDeltaNormal,
-	const uniform float ContactMassNormal,
-	const uniform float NetPushOutNormal)
-{
-	const uniform float PushOutNormal = -Stiffness * ContactDeltaNormal * ContactMassNormal;
-
-	// The total pushout so far this sub-step
-	// Unilateral constraint: Net-negative impulses not allowed (negative incremental impulses are allowed as long as the net is positive)
-	return (NetPushOutNormal + PushOutNormal) > 0 ? PushOutNormal : -NetPushOutNormal;
-}
-
-
-//
-//
-// GAUSS SEIDEL SOLVER
-//
-//
-
-static inline void CollisionSolverApplyPositionCorrectionNormal(
-	uniform FPBDCollisionSolver* uniform Solver,
-	const uniform float Stiffness,
-	const uniform float ContactDeltaNormal,
-	uniform FPBDCollisionSolverManifoldPoint* uniform ManifoldPoint,
-	uniform FSolverBody* uniform Body0,
-	uniform FSolverBody* uniform Body1)
-{
-	const uniform float PushOutNormal = CalculatePositionCorrectionNormal(Stiffness, ContactDeltaNormal, ManifoldPoint->ContactMassNormal, ManifoldPoint->NetPushOutNormal);
-
-	ManifoldPoint->NetPushOutNormal = ManifoldPoint->NetPushOutNormal + PushOutNormal;
-
-	// Update the particle state based on the pushout
-	if (Body_IsDynamic(Body0))
-	{
-		const uniform FVector3f DX0 = ManifoldPoint->WorldContact.ContactNormal * (Body0->InvM * PushOutNormal);
-		const uniform FVector3f DR0 = ManifoldPoint->WorldContactNormalAngular0 * PushOutNormal;
-		Body_ApplyPositionDelta(Body0, DX0);
-		Body_ApplyRotationDelta(Body0, DR0);
-	}
-	if (Body_IsDynamic(Body1))
-	{
-		const uniform FVector3f DX1 = ManifoldPoint->WorldContact.ContactNormal * (Body1->InvM * -PushOutNormal);
-		const uniform FVector3f DR1 = ManifoldPoint->WorldContactNormalAngular1 * -PushOutNormal;
-		Body_ApplyPositionDelta(Body1, DX1);
-		Body_ApplyRotationDelta(Body1, DR1);
-	}
-}
-
-static inline void CollisionSolverSolvePositionNoFriction(
-	uniform FPBDCollisionSolver* uniform Solver, 
-	const uniform float Dt, 
-	const uniform float MaxPushOut)
-{
-	for (uniform int PointIndex = 0; PointIndex < Solver->NumManifoldPoints; ++PointIndex)
-	{
-		uniform FPBDCollisionSolverManifoldPoint* uniform SolverManifoldPoint = &Solver->ManifoldPoints[PointIndex];
-
-		const uniform float ContactDeltaNormal = CalculatePositionErrorNormal(SolverManifoldPoint, Solver->SolverBodies[0].Body, Solver->SolverBodies[1].Body, MaxPushOut);
-
-		const uniform bool bProcessManifoldPoint = (ContactDeltaNormal < 0) || (SolverManifoldPoint->NetPushOutNormal > FLOAT_SMALL_NUMBER);
-		if (bProcessManifoldPoint)
-		{
-			CollisionSolverApplyPositionCorrectionNormal(Solver, Solver->Stiffness, ContactDeltaNormal, SolverManifoldPoint, Solver->SolverBodies[0].Body, Solver->SolverBodies[1].Body);
-		}
-	}
-}
-
-export void SolvePositionNoFriction(
-	uniform FPBDCollisionSolver* uniform Solvers, 
-	const uniform int NumSolvers, 
-	const uniform float Dt, 
-	const uniform float MaxPushOut)
-{
-	for (uniform int SolverIndex = 0; SolverIndex < NumSolvers; ++SolverIndex)
-	{
-		CollisionSolverSolvePositionNoFriction(&Solvers[SolverIndex], Dt, MaxPushOut);
-	}
-}
-
-
-//
-//
-// PARTIAL-JACOBI SOLVER
-//
-//
-
-static void CollisionSolver_SolvePositionNoFriction_Jacobi(
-	uniform FPBDCollisionSolver* uniform Solver, 
-	const uniform float Dt, 
-	const uniform float MaxPushOut,
-	const uniform float JacobiStiffness)
-{
-	uniform FSolverBody* uniform Body0 = Solver->SolverBodies[0].Body;
-	uniform FSolverBody* uniform Body1 = Solver->SolverBodies[1].Body;
-
-	uniform FVector3f DX0 = FloatZeroVector;
-	uniform FVector3f DX1 = FloatZeroVector;
-	uniform FVector3f DR0 = FloatZeroVector;
-	uniform FVector3f DR1 = FloatZeroVector;
-
-	const uniform float Stiffness = JacobiStiffness * Solver->Stiffness;
-
-	for (uniform int PointIndex = 0; PointIndex < Solver->NumManifoldPoints; ++PointIndex)
-	{
-		uniform FPBDCollisionSolverManifoldPoint* uniform ManifoldPoint = &Solver->ManifoldPoints[PointIndex];
-
-		const uniform float ContactDeltaNormal = CalculatePositionErrorNormal(ManifoldPoint, Solver->SolverBodies[0].Body, Solver->SolverBodies[1].Body, MaxPushOut);
-
-		const uniform bool bProcessManifoldPoint = (ContactDeltaNormal < 0) || (ManifoldPoint->NetPushOutNormal > FLOAT_SMALL_NUMBER);
-		if (bProcessManifoldPoint)
-		{
-			const uniform float PushOutNormal = CalculatePositionCorrectionNormal(Stiffness, ContactDeltaNormal, ManifoldPoint->ContactMassNormal, ManifoldPoint->NetPushOutNormal);
-
-			ManifoldPoint->NetPushOutNormal = ManifoldPoint->NetPushOutNormal + PushOutNormal;
-
-			if (Body_IsDynamic(Body0))
-			{
-				const uniform FVector3f DDX0 = ManifoldPoint->WorldContact.ContactNormal * (Body0->InvM * PushOutNormal);
-				const uniform FVector3f DDR0 = ManifoldPoint->WorldContactNormalAngular0 * PushOutNormal;
-				DX0 = DX0 + DDX0;
-				DR0 = DR0 + DDR0;
-			}
-			if (Body_IsDynamic(Body1))
-			{
-				const uniform FVector3f DDX1 = ManifoldPoint->WorldContact.ContactNormal * (Body1->InvM * PushOutNormal);
-				const uniform FVector3f DDR1 = ManifoldPoint->WorldContactNormalAngular1 * PushOutNormal;
-				DX1 = DX1 - DDX1;
-				DR1 = DR1 - DDR1;
-			}
-		}
-	}
-
-	if (Body_IsDynamic(Body0))
-	{
-		Body_ApplyPositionDelta(Body0, DX0);
-		Body_ApplyRotationDelta(Body0, DR0);
-	}
-	if (Body_IsDynamic(Body1))
-	{
-		Body_ApplyPositionDelta(Body1, DX1);
-		Body_ApplyRotationDelta(Body1, DR1);
-	}
-}
-
-
-export void SolvePositionNoFrictionJacobi(
-	uniform FPBDCollisionSolver* uniform Solvers, 
-	const uniform int NumSolvers, 
-	const uniform float Dt, 
-	const uniform float MaxPushOut,
-	const uniform float JacobiStiffness)
+export void SolvePositionNoFriction()
 {
-	for (uniform int SolverIndex = 0; SolverIndex < NumSolvers; ++SolverIndex)
-	{
-		CollisionSolver_SolvePositionNoFriction_Jacobi(&Solvers[SolverIndex], Dt, MaxPushOut, JacobiStiffness);
-	}
-}
+}
\ No newline at end of file
diff --git a/Engine/Source/Runtime/Experimental/Chaos/Private/Chaos/PBDSuspensionConstraints.cpp b/Engine/Source/Runtime/Experimental/Chaos/Private/Chaos/PBDSuspensionConstraints.cpp
--- a/Engine/Source/Runtime/Experimental/Chaos/Private/Chaos/PBDSuspensionConstraints.cpp
+++ b/Engine/Source/Runtime/Experimental/Chaos/Private/Chaos/PBDSuspensionConstraints.cpp
@@ -274,7 +274,7 @@ namespace Chaos
 
 				// inject a manifold for our suspension Hard-stop - behaves like a regular friction-less collision, prevents the vehicle chassis from ever hitting the ground
 				Solver->AddManifoldPoint();
-				Solver->SetManifoldPoint(
+				Solver->InitManifoldPoint(
 					0,										// ManifoldIndex
 					FSolverReal(Dt),						// Delta Time
 					FSolverVec3(WorldArm),					// RelativeContactPosition0,
@@ -287,7 +287,7 @@ namespace Chaos
 					FSolverReal(0),							// WorldContactDeltaTangentV
 					FSolverReal(0)							// WorldTargetContactVelocityNormal
 					);
-
+				Solver->FinalizeManifold();
 			}
 		}
 
diff --git a/Engine/Source/Runtime/Experimental/Chaos/Public/Chaos/Collision/PBDCollisionContainerSolver.h b/Engine/Source/Runtime/Experimental/Chaos/Public/Chaos/Collision/PBDCollisionContainerSolver.h
--- a/Engine/Source/Runtime/Experimental/Chaos/Public/Chaos/Collision/PBDCollisionContainerSolver.h
+++ b/Engine/Source/Runtime/Experimental/Chaos/Public/Chaos/Collision/PBDCollisionContainerSolver.h
@@ -6,6 +6,7 @@
 #include "Chaos/Collision/PBDCollisionSolver.h"
 #include "Chaos/Collision/PBDCollisionSolverSettings.h"
 #include "Chaos/Evolution/SolverConstraintContainer.h"
+#include "Chaos/Framework/ScratchBuffer.h"
 
 namespace Chaos
 {
@@ -35,7 +36,7 @@ namespace Chaos
 
 		virtual void Reset(const int32 InMaxCollisions) override final;
 
-		virtual int32 GetNumConstraints() const override final { return NumSolvers(); }
+		virtual int32 GetNumConstraints() const override final { return CollisionConstraints.Num(); }
 
 		//
 		// IslandGroup API
@@ -52,23 +53,43 @@ namespace Chaos
 		virtual void ApplyProjectionConstraints(const FReal Dt, const int32 It, const int32 NumIts) override final;
 
 		// For testing
-		const Private::FPBDCollisionSolver& GetConstraintSolver(const int32 ConstraintIndex) const { return CollisionSolvers[ConstraintIndex]; }
-		const Private::FPBDCollisionSolverManifoldPoint& GetManifoldPointSolver(const int32 RowIndex) const { return CollisionSolverManifoldPoints[RowIndex]; }
+		const Private::FPBDCollisionSolver& GetConstraintSolver(const int32 ConstraintIndex) const { return GetSolver(ConstraintIndex); }
 
 	private:
+		FPBDCollisionConstraint* GetConstraint(const int32 Index) { return CollisionConstraints[Index]; }
+		const FPBDCollisionConstraint* GetConstraint(const int32 Index) const { return CollisionConstraints[Index]; }
+		Private::FPBDCollisionSolver& GetSolver(const int32 Index) { check(Index < NumSolvers()); return CollisionSolvers[Index]; }
+		const Private::FPBDCollisionSolver& GetSolver(const int32 Index) const { check(Index < NumSolvers()); return CollisionSolvers[Index]; }
+
 		void CachePrefetchSolver(const int32 ConstraintIndex) const;
 		void AddConstraint(FPBDCollisionConstraint& Constraint);
 		void UpdatePositionShockPropagation(const FReal Dt, const int32 It, const int32 NumIts, const int32 BeginIndex, const int32 EndIndex, const FPBDCollisionSolverSettings& SolverSettings);
 		void UpdateVelocityShockPropagation(const FReal Dt, const int32 It, const int32 NumIts, const int32 BeginIndex, const int32 EndIndex, const FPBDCollisionSolverSettings& SolverSettings);
+		void ApplyShockPropagation(const FSolverReal ShockPropagation);
 		bool SolvePositionImpl(const FReal Dt, const int32 It, const int32 NumIts, const int32 BeginIndex, const int32 EndIndex, const FPBDCollisionSolverSettings& SolverSettings);
 		bool SolveVelocityImpl(const FReal Dt, const int32 It, const int32 NumIts, const int32 BeginIndex, const int32 EndIndex, const FPBDCollisionSolverSettings& SolverSettings);
 		void UpdateCollisions(const FReal InDt, const int32 BeginIndex, const int32 EndIndex);
 
+		// The constraints we are solving and the container to which they belong
 		const FPBDCollisionConstraints& ConstraintContainer;
-		TArray<Private::FPBDCollisionSolver> CollisionSolvers;
-		TArray<Private::FPBDCollisionSolverManifoldPoint> CollisionSolverManifoldPoints;
 		TArray<FPBDCollisionConstraint*> CollisionConstraints;
+
+		// The last shock propagation factor we applied
+		FSolverReal AppliedShockPropagation;
+
+		// Buffer for all allocations for this tick (solvers and manifold points)
+		Private::FScratchBuffer Scratch;
+
+		// The start of the solver array in the scratch buffer
+		Private::FPBDCollisionSolver* CollisionSolvers;
+
+		// The start of the manifold points array in the scratch buffer
+		Private::FPBDCollisionSolverManifoldPoint* CollisionSolverManifoldPoints;
+		int32 NumCollisionSolverManifoldPoints;
+
+		// Whether we need to run incremental collision for each constraint (LevelSets only now)
 		TArray<bool> bCollisionConstraintPerIterationCollisionDetection;
 		bool bPerIterationCollisionDetection;
+
 	};
 }
diff --git a/Engine/Source/Runtime/Experimental/Chaos/Public/Chaos/Collision/PBDCollisionContainerSolverSimd.h b/Engine/Source/Runtime/Experimental/Chaos/Public/Chaos/Collision/PBDCollisionContainerSolverSimd.h
--- a/Engine/Source/Runtime/Experimental/Chaos/Public/Chaos/Collision/PBDCollisionContainerSolverSimd.h
+++ b/Engine/Source/Runtime/Experimental/Chaos/Public/Chaos/Collision/PBDCollisionContainerSolverSimd.h
@@ -80,7 +80,6 @@ namespace Chaos
 			// For testing
 			const FConstraintSolverId& GetConstraintSolverId(const int32 ConstraintIndex) const { return ConstraintSolverIds[ConstraintIndex]; }
 			const TPBDCollisionSolverSimd<4>& GetConstraintSolver(const int32 SolverIndex) const { return SimdData.SimdSolvers[SolverIndex]; }
-			const TPBDCollisionSolverManifoldPointsSimd<4>& GetManifoldPointSolver(const int32 PointIndex) const { return SimdData.SimdManifoldPoints[PointIndex]; }
 			TArrayView<const TPBDCollisionSolverManifoldPointsSimd<4>> GetManifoldPointBuffer() const { return MakeArrayView(SimdData.SimdManifoldPoints); }
 
 		private:
diff --git a/Engine/Source/Runtime/Experimental/Chaos/Public/Chaos/Collision/PBDCollisionSolver.h b/Engine/Source/Runtime/Experimental/Chaos/Public/Chaos/Collision/PBDCollisionSolver.h
--- a/Engine/Source/Runtime/Experimental/Chaos/Public/Chaos/Collision/PBDCollisionSolver.h
+++ b/Engine/Source/Runtime/Experimental/Chaos/Public/Chaos/Collision/PBDCollisionSolver.h
@@ -32,56 +32,6 @@ namespace Chaos
 		class FPBDCollisionSolverManifoldPoint
 		{
 		public:
-			/**
-			 * @brief Initialize the geometric data for the contact
-			*/
-			void InitContact(
-				const FSolverReal Dt,
-				const FConstraintSolverBody& Body0,
-				const FConstraintSolverBody& Body1);
-
-			/**
-			 * @brief Initialize the material related properties of the contact
-			*/
-			//void InitMaterial(
-			//	const FConstraintSolverBody& Body0,
-			//	const FConstraintSolverBody& Body1,
-			//	const FSolverReal InRestitution,
-			//	const FSolverReal InRestitutionVelocityThreshold);
-
-			/**
-			 * @brief Update the cached mass properties based on the current body transforms
-			*/
-			void UpdateMass(const FConstraintSolverBody& Body0, const FConstraintSolverBody& Body1);
-
-			/**
-			 * @brief Update the contact mass for the normal correction
-			 * This is used by shock propagation.
-			*/
-			void UpdateMassNormal(const FConstraintSolverBody& Body0, const FConstraintSolverBody& Body1);
-
-			/**
-			 * @brief Calculate the position error at the current transforms
-			 * @param MaxPushOut a limit on the position error for this iteration to prevent initial-penetration explosion (a common PBD problem)
-			*/
-			void CalculateContactPositionErrorNormal(const FConstraintSolverBody& Body0, const FConstraintSolverBody& Body1, const FSolverReal MaxPushOut, FSolverReal& OutContactDeltaNormal) const;
-			void CalculateContactPositionErrorTangential(const FConstraintSolverBody& Body0, const FConstraintSolverBody& Body1, FSolverReal& OutContactDeltaTanget0, FSolverReal& OutContactDeltaTangent1) const;
-
-			/**
-			 * @brief Calculate the velocity error at the current transforms
-			*/
-			void CalculateContactVelocityErrorNormal(const FConstraintSolverBody& Body0, const FConstraintSolverBody& Body1, FSolverReal& OutContactVelocityDeltaNormal) const;
-			void CalculateContactVelocityError(const FConstraintSolverBody& Body0, const FConstraintSolverBody& Body1, const FSolverReal DynamicFriction, const FSolverReal Dt, FSolverReal& OutContactVelocityDeltaNormal, FSolverReal& OutContactVelocityDeltaTangent0, FSolverReal& OutContactVelocityDeltaTangent1) const;
-
-			// @todo(chaos): make private
-		public:
-			friend class FPBDCollisionSolver;
-
-			/**
-			 * @brief Whether we need to solve velocity for this manifold point (only if we were penetrating or applied a pushout)
-			*/
-			bool ShouldSolveVelocity() const;
-
 			// World-space contact data
 			FWorldContactPoint WorldContact;
 
@@ -156,9 +106,9 @@ namespace Chaos
 				State.NumManifoldPoints = 0;
 			}
 
-			FSolverReal StaticFriction() const { return State.StaticFriction; }
-			FSolverReal DynamicFriction() const { return State.DynamicFriction; }
-			FSolverReal VelocityFriction() const { return State.VelocityFriction; }
+			FSolverReal GetStaticFriction() const { return State.StaticFriction; }
+			FSolverReal GetDynamicFriction() const { return State.DynamicFriction; }
+			FSolverReal GetVelocityFriction() const { return State.VelocityFriction; }
 
 			void SetFriction(const FSolverReal InStaticFriction, const FSolverReal InDynamicFriction, const FSolverReal InVelocityFriction)
 			{
@@ -183,6 +133,11 @@ namespace Chaos
 				return State.NumManifoldPoints;
 			}
 
+			int32 MaxManifoldPoints() const
+			{
+				return State.MaxManifoldPoints;
+			}
+
 			int32 AddManifoldPoint()
 			{
 				check(State.NumManifoldPoints < State.MaxManifoldPoints);
@@ -202,16 +157,11 @@ namespace Chaos
 				return State.ManifoldPoints[ManifoldPointIndex];
 			}
 
-			FWorldContactPoint& GetWorldContactPoint(const int32 ManifoldPointIndex)
-			{
-				return State.ManifoldPoints[ManifoldPointIndex].WorldContact;
-			}
-
 			/**
-			 * Set up a manifold point (also calls FinalizeManifoldPoint)
+			 * Set up a manifold point (except mass properties. @see UpdateMass())
 			*/
-			void SetManifoldPoint(
-				const int32 ManifoldPoiontIndex,
+			void InitManifoldPoint(
+				const int32 PointIndex,
 				const FSolverReal Dt,
 				const FSolverVec3& InRelativeContactPosition0,
 				const FSolverVec3& InRelativeContactPosition1,
@@ -223,60 +173,160 @@ namespace Chaos
 				const FSolverReal InWorldContactDeltaTangentV,
 				const FSolverReal InWorldContactVelocityTargetNormal);
 
-			/** 
-			 * Finish manifold point setup.
-			 * NOTE: Can only be called after the WorldContact has been set up (e.g., see SetManifoldPoint)
-			*/
-			void FinalizeManifoldPoint(const int32 ManifoldPoiontIndex, const FSolverReal Dt);
+			/**
+			 * Call once all manifold points have been initialized. Calculate mass properties etc.
+			 */
+			void FinalizeManifold();
 
 			/**
-			 * @brief Get the first (decorated) solver body
-			 * The decorator add a possible mass scale
+			 * @brief Get the first solver body
+			 * NOTE: This will not include any shock propagation mass scaling
 			*/
 			FConstraintSolverBody& SolverBody0() { return State.SolverBodies[0]; }
 			const FConstraintSolverBody& SolverBody0() const { return State.SolverBodies[0]; }
 
+
 			/**
-			 * @brief Get the second (decorated) solver body
-			 * The decorator add a possible mass scale
+			 * @brief Get the second solver body
+			 * NOTE: This will not include any shock propagation mass scaling
 			*/
 			FConstraintSolverBody& SolverBody1() { return State.SolverBodies[1]; }
 			const FConstraintSolverBody& SolverBody1() const { return State.SolverBodies[1]; }
 
-			/**
-			 * @brief Set up the mass scaling for shock propagation, using the position-phase mass scale
-			*/
-			void EnablePositionShockPropagation();
-
-			/**
-			 * @brief Set up the mass scaling for shock propagation, using the velocity-phase mass scale
-			*/
-			void EnableVelocityShockPropagation();
-
-			/**
-			 * @brief Disable mass scaling
-			*/
-			void DisableShockPropagation();
 
 			/**
 			 * @brief Calculate and apply the position correction for this iteration
 			 * @return true if we need to run more iterations, false if we did not apply any correction
 			*/
-			bool SolvePositionWithFriction(const FSolverReal Dt, const FSolverReal MaxPushOut);
-			bool SolvePositionNoFriction(const FSolverReal Dt, const FSolverReal MaxPushOut);
+			void SolvePositionNoFriction(
+				const FSolverReal Dt,
+				const FSolverReal MaxPushOut);
+
+			void SolvePositionWithFriction(
+				const FSolverReal Dt, 
+				const FSolverReal MaxPushOut);
 
 			/**
 			 * @brief Calculate and apply the velocity correction for this iteration
 			 * @return true if we need to run more iterations, false if we did not apply any correction
 			*/
-			bool SolveVelocity(const FSolverReal Dt, const bool bApplyDynamicFriction);
+			void SolveVelocity(
+				const FSolverReal Dt, 
+				const bool bApplyDynamicFriction);
+
+			void UpdateMass();
+
+			void UpdateMassNormal();
 
 		private:
-			/**
-			 * @brief Apply the inverse mass scale the body with the lower level
-			 * @param InvMassScale
-			*/
-			void SetShockPropagationInvMassScale(const FSolverReal InvMassScale);
+			// Get the mass properties for the bodies if they are dynamic
+			FORCEINLINE_DEBUGGABLE void GetDynamicMassProperties(
+				FSolverReal& OutInvM0,
+				FSolverMatrix33& OutInvI0, 
+				FSolverReal& OutInvM1, 
+				FSolverMatrix33& OutInvI1)
+			{
+				OutInvM0 = 0;
+				OutInvM1 = 0;
+
+				FConstraintSolverBody& Body0 = SolverBody0();
+				if (Body0.IsDynamic())
+				{
+					OutInvM0 = Body0.InvM();
+					OutInvI0 = Body0.InvI();
+				}
+
+				FConstraintSolverBody& Body1 = SolverBody1();
+				if (Body1.IsDynamic())
+				{
+					OutInvM1 = Body1.InvM();
+					OutInvI1 = Body1.InvI();
+				}
+			}
+
+			bool IsDynamic(const int32 BodyIndex) const
+			{
+				return (State.InvMs[BodyIndex] > FSolverBody::ZeroMassThreshold());
+			}
+
+			void CalculateContactPositionErrorNormal(
+				const FPBDCollisionSolverManifoldPoint& ManifoldPoint, 
+				const FSolverReal MaxPushOut, 
+				FSolverReal& OutContactDeltaNormal) const;
+
+			void CalculateContactPositionErrorTangential(
+				const FPBDCollisionSolverManifoldPoint& ManifoldPoint, 
+				FSolverReal& OutContactDeltaTangentU, 
+				FSolverReal& OutContactDeltaTangentV) const;
+
+			void CalculateContactVelocityError(
+				const FPBDCollisionSolverManifoldPoint& ManifoldPoint, 
+				const FSolverReal DynamicFriction, 
+				const FSolverReal Dt, 
+				FSolverReal& OutContactVelocityDeltaNormal, 
+				FSolverReal& OutContactVelocityDeltaTangent0, 
+				FSolverReal& OutContactVelocityDeltaTangent1) const;
+
+			void CalculateContactVelocityErrorNormal(
+				const FPBDCollisionSolverManifoldPoint& ManifoldPoint, 
+				FSolverReal& OutContactVelocityDeltaNormal) const;
+
+			bool ShouldSolveVelocity(
+				const FPBDCollisionSolverManifoldPoint& ManifoldPoint) const;
+
+			void CalculatePositionCorrectionNormal(
+				const FSolverReal Stiffness,
+				const FSolverReal ContactDeltaNormal,
+				const FSolverReal ContactMassNormal,
+				const FSolverReal NetPushOutNormal,
+				FSolverReal& OutPushOutNormal);
+
+			void CalculatePositionCorrectionTangent(
+				const FSolverReal Stiffness,
+				const FSolverReal ContactDeltaTangent,
+				const FSolverReal ContactMassTangent,
+				const FSolverReal NetPushOutTangent,
+				FSolverReal& OutPushOutTangent);
+
+			void ApplyFrictionCone(
+				const FSolverReal StaticFriction,
+				const FSolverReal DynamicFriction,
+				const FSolverReal MaxFrictionPushOut,
+				FSolverReal& InOutPushOutTangentU,
+				FSolverReal& InOutPushOutTangentV,
+				FSolverReal& InOutNetPushOutTangentU,
+				FSolverReal& InOutNetPushOutTangentV,
+				FSolverReal& OutStaticFrictionRatio);
+
+			void ApplyPositionCorrectionTangential(
+				const FSolverReal Stiffness,
+				const FSolverReal StaticFriction,
+				const FSolverReal DynamicFriction,
+				const FSolverReal MaxFrictionPushOut,
+				const FSolverReal ContactDeltaTangentU,
+				const FSolverReal ContactDeltaTangentV,
+				FPBDCollisionSolverManifoldPoint& ManifoldPoint);
+
+			void ApplyPositionCorrectionNormal(
+				const FSolverReal Stiffness,
+				const FSolverReal ContactDeltaNormal,
+				FPBDCollisionSolverManifoldPoint& ManifoldPoint);
+
+			void ApplyVelocityCorrection(
+				const FSolverReal Stiffness,
+				const FSolverReal Dt,
+				const FSolverReal DynamicFriction,
+				const FSolverReal ContactVelocityDeltaNormal,
+				const FSolverReal ContactVelocityDeltaTangent0,
+				const FSolverReal ContactVelocityDeltaTangent1,
+				const FSolverReal MinImpulseNormal,
+				FPBDCollisionSolverManifoldPoint& ManifoldPoint);
+
+			void ApplyVelocityCorrectionNormal(
+				const FSolverReal Stiffness,
+				const FSolverReal ContactVelocityDeltaNormal,
+				const FSolverReal MinImpulseNormal,
+				FPBDCollisionSolverManifoldPoint& ManifoldPoint);
 
 			/**
 			 * @brief Run a velocity solve on the average point from all the points that received a position impulse
@@ -294,12 +344,12 @@ namespace Chaos
 
 				void Init()
 				{
-					SolverBodies[0].Init();
-					SolverBodies[1].Init();
 					StaticFriction = 0;
 					DynamicFriction = 0;
 					VelocityFriction = 0;
 					Stiffness = 1;
+					SolverBodies[0].Reset();
+					SolverBodies[1].Reset();
 					ManifoldPoints = nullptr;
 					NumManifoldPoints = 0;
 					MaxManifoldPoints = 0;
@@ -317,8 +367,11 @@ namespace Chaos
 				// Solver stiffness (scales all pushout and impulses)
 				FSolverReal Stiffness;
 
-				// Bodies and contacts
+				// Bodies
 				FConstraintSolverBody SolverBodies[MaxConstrainedBodies];
+				FSolverReal InvMs[2];
+
+				// Manifold Points
 				FPBDCollisionSolverManifoldPoint* ManifoldPoints;
 				int32 NumManifoldPoints;
 				int32 MaxManifoldPoints;
@@ -333,21 +386,7 @@ namespace Chaos
 		//////////////////////////////////////////////////////////////////////////////////////////////////
 		//////////////////////////////////////////////////////////////////////////////////////////////////
 
-		// If the vector size squared is below the threshold, zero the vector and return false, otherwise leave vector as-is and return true.
-		// This is useful for maintaining perfect restitution bounces - we remove any rounding errors in the impulses before applying to the rotation
-		// @todo(chaos): consider using a test of absolute element values instead
-		// @todo(chaos): disabled for now - too expensive :(
-		FORCEINLINE_DEBUGGABLE bool AboveThresholdElseZero(FSolverVec3& InOutV, const FSolverReal ThresholdSq)
-		{
-			//if (InOutV.SizeSquared() < ThresholdSq)
-			//{
-			//	InOutV = FSolverVec3(0);
-			//	return false;
-			//}
-			return true;
-		}
-
-		FORCEINLINE_DEBUGGABLE void CalculatePositionCorrectionNormal(
+		FORCEINLINE_DEBUGGABLE void FPBDCollisionSolver::CalculatePositionCorrectionNormal(
 			const FSolverReal Stiffness,
 			const FSolverReal ContactDeltaNormal,
 			const FSolverReal ContactMassNormal,
@@ -368,8 +407,7 @@ namespace Chaos
 			}
 		}
 
-
-		FORCEINLINE_DEBUGGABLE void CalculatePositionCorrectionTangent(
+		FORCEINLINE_DEBUGGABLE void FPBDCollisionSolver::CalculatePositionCorrectionTangent(
 			const FSolverReal Stiffness,
 			const FSolverReal ContactDeltaTangent,
 			const FSolverReal ContactMassTangent,
@@ -380,8 +418,7 @@ namespace Chaos
 			OutPushOutTangent = -Stiffness * ContactMassTangent * ContactDeltaTangent;
 		}
 
-
-		FORCEINLINE_DEBUGGABLE void ApplyFrictionCone(
+		FORCEINLINE_DEBUGGABLE void FPBDCollisionSolver::ApplyFrictionCone(
 			const FSolverReal StaticFriction,
 			const FSolverReal DynamicFriction,
 			const FSolverReal MaxFrictionPushOut,
@@ -429,17 +466,14 @@ namespace Chaos
 			OutStaticFrictionRatio = ClampedStaticFrictionRatio;
 		}
 
-
-		FORCEINLINE_DEBUGGABLE void ApplyPositionCorrectionTangential(
+		FORCEINLINE_DEBUGGABLE void FPBDCollisionSolver::ApplyPositionCorrectionTangential(
 			const FSolverReal Stiffness,
 			const FSolverReal StaticFriction,
 			const FSolverReal DynamicFriction,
 			const FSolverReal MaxFrictionPushOut,
 			const FSolverReal ContactDeltaTangentU,
 			const FSolverReal ContactDeltaTangentV,
-			FPBDCollisionSolverManifoldPoint& ManifoldPoint,
-			FConstraintSolverBody& Body0,
-			FConstraintSolverBody& Body1)
+			FPBDCollisionSolverManifoldPoint& ManifoldPoint)
 		{
 			FSolverReal PushOutTangentU = FSolverReal(0);
 			FSolverReal PushOutTangentV = FSolverReal(0);
@@ -471,28 +505,28 @@ namespace Chaos
 
 			// Update the particle state based on the pushout
 			const FVec3 PushOut = PushOutTangentU * ManifoldPoint.WorldContact.ContactTangentU + PushOutTangentV * ManifoldPoint.WorldContact.ContactTangentV;
-			if (Body0.IsDynamic())
+			if (IsDynamic(0))
 			{
-				const FSolverVec3 DX0 = Body0.InvM() * PushOut;
+				const FSolverVec3 DX0 = State.InvMs[0] * PushOut;
 				const FSolverVec3 DR0 = ManifoldPoint.WorldContactTangentUAngular0 * PushOutTangentU + ManifoldPoint.WorldContactTangentVAngular0 * PushOutTangentV;
+				FConstraintSolverBody& Body0 = SolverBody0();
 				Body0.ApplyPositionDelta(DX0);
 				Body0.ApplyRotationDelta(DR0);
 			}
-			if (Body1.IsDynamic())
+			if (IsDynamic(1))
 			{
-				const FSolverVec3 DX1 = -Body1.InvM() * PushOut;
+				const FSolverVec3 DX1 = -State.InvMs[1]* PushOut;
 				const FSolverVec3 DR1 = ManifoldPoint.WorldContactTangentUAngular1 * -PushOutTangentU + ManifoldPoint.WorldContactTangentVAngular1 * -PushOutTangentV;
+				FConstraintSolverBody& Body1 = SolverBody1();
 				Body1.ApplyPositionDelta(DX1);
 				Body1.ApplyRotationDelta(DR1);
 			}
 		}
 
-		FORCEINLINE_DEBUGGABLE void ApplyPositionCorrectionNormal(
+		FORCEINLINE_DEBUGGABLE void FPBDCollisionSolver::ApplyPositionCorrectionNormal(
 			const FSolverReal Stiffness,
 			const FSolverReal ContactDeltaNormal,
-			FPBDCollisionSolverManifoldPoint& ManifoldPoint,
-			FConstraintSolverBody& Body0,
-			FConstraintSolverBody& Body1)
+			FPBDCollisionSolverManifoldPoint& ManifoldPoint)
 		{
 			FSolverReal PushOutNormal = FSolverReal(0);
 
@@ -506,24 +540,26 @@ namespace Chaos
 			ManifoldPoint.NetPushOutNormal += PushOutNormal;
 
 			// Update the particle state based on the pushout
-			if (Body0.IsDynamic())
+			if (IsDynamic(0))
 			{
-				const FSolverVec3 DX0 = (Body0.InvM() * PushOutNormal) * ManifoldPoint.WorldContact.ContactNormal;
+				const FSolverVec3 DX0 = (State.InvMs[0] * PushOutNormal) * ManifoldPoint.WorldContact.ContactNormal;
 				const FSolverVec3 DR0 = ManifoldPoint.WorldContactNormalAngular0 * PushOutNormal;
+				FConstraintSolverBody& Body0 = SolverBody0();
 				Body0.ApplyPositionDelta(DX0);
 				Body0.ApplyRotationDelta(DR0);
 			}
-			if (Body1.IsDynamic())
+			if (IsDynamic(1))
 			{
-				const FSolverVec3 DX1 = (Body1.InvM() * -PushOutNormal) * ManifoldPoint.WorldContact.ContactNormal;
+				const FSolverVec3 DX1 = (State.InvMs[1] * -PushOutNormal) * ManifoldPoint.WorldContact.ContactNormal;
 				const FSolverVec3 DR1 = ManifoldPoint.WorldContactNormalAngular1 * -PushOutNormal;
+				FConstraintSolverBody& Body1 = SolverBody1();
 				Body1.ApplyPositionDelta(DX1);
 				Body1.ApplyRotationDelta(DR1);
 			}
 		}
 
 
-		FORCEINLINE_DEBUGGABLE void ApplyVelocityCorrection(
+		FORCEINLINE_DEBUGGABLE void FPBDCollisionSolver::ApplyVelocityCorrection(
 			const FSolverReal Stiffness,
 			const FSolverReal Dt,
 			const FSolverReal DynamicFriction,
@@ -531,9 +567,7 @@ namespace Chaos
 			const FSolverReal ContactVelocityDeltaTangent0,
 			const FSolverReal ContactVelocityDeltaTangent1,
 			const FSolverReal MinImpulseNormal,
-			FPBDCollisionSolverManifoldPoint& ManifoldPoint,
-			FConstraintSolverBody& Body0,
-			FConstraintSolverBody& Body1)
+			FPBDCollisionSolverManifoldPoint& ManifoldPoint)
 		{
 			FSolverReal ImpulseNormal = -Stiffness * ManifoldPoint.ContactMassNormal * ContactVelocityDeltaNormal;
 			FSolverReal ImpulseTangentU = -Stiffness * ManifoldPoint.ContactMassTangentU * ContactVelocityDeltaTangent0;
@@ -568,27 +602,27 @@ namespace Chaos
 
 			// Apply the velocity deltas from the impulse
 			const FSolverVec3 Impulse = ImpulseNormal * ManifoldPoint.WorldContact.ContactNormal + ImpulseTangentU * ManifoldPoint.WorldContact.ContactTangentU + ImpulseTangentV * ManifoldPoint.WorldContact.ContactTangentV;
-			if (Body0.IsDynamic())
+			if (IsDynamic(0))
 			{
-				const FSolverVec3 DV0 = Body0.InvM() * Impulse;
+				const FSolverVec3 DV0 = State.InvMs[0] * Impulse;
 				const FSolverVec3 DW0 = ManifoldPoint.WorldContactNormalAngular0 * ImpulseNormal + ManifoldPoint.WorldContactTangentUAngular0 * ImpulseTangentU + ManifoldPoint.WorldContactTangentVAngular0 * ImpulseTangentV;
+				FConstraintSolverBody& Body0 = SolverBody0();
 				Body0.ApplyVelocityDelta(DV0, DW0);
 			}
-			if (Body1.IsDynamic())
+			if (IsDynamic(1))
 			{
-				const FSolverVec3 DV1 = -Body1.InvM() * Impulse;
+				const FSolverVec3 DV1 = -State.InvMs[1]* Impulse;
 				const FSolverVec3 DW1 = ManifoldPoint.WorldContactNormalAngular1 * -ImpulseNormal + ManifoldPoint.WorldContactTangentUAngular1 * -ImpulseTangentU + ManifoldPoint.WorldContactTangentVAngular1 * -ImpulseTangentV;
+				FConstraintSolverBody& Body1 = SolverBody1();
 				Body1.ApplyVelocityDelta(DV1, DW1);
 			}
 		}
 
-		FORCEINLINE_DEBUGGABLE void ApplyVelocityCorrectionNormal(
+		FORCEINLINE_DEBUGGABLE void FPBDCollisionSolver::ApplyVelocityCorrectionNormal(
 			const FSolverReal Stiffness,
 			const FSolverReal ContactVelocityDeltaNormal,
 			const FSolverReal MinImpulseNormal,
-			FPBDCollisionSolverManifoldPoint& ManifoldPoint,
-			FConstraintSolverBody& Body0,
-			FConstraintSolverBody& Body1)
+			FPBDCollisionSolverManifoldPoint& ManifoldPoint)
 		{
 			FSolverReal ImpulseNormal = -(Stiffness * ManifoldPoint.ContactMassNormal) * ContactVelocityDeltaNormal;
 
@@ -602,121 +636,131 @@ namespace Chaos
 
 			// Calculate the velocity deltas from the impulse
 			FSolverVec3 Impulse = ImpulseNormal * ManifoldPoint.WorldContact.ContactNormal;
-			if (Body0.IsDynamic())
+			if (IsDynamic(0))
 			{
-				const FSolverVec3 DV0 = Body0.InvM() * Impulse;
+				const FSolverVec3 DV0 = State.InvMs[0] * Impulse;
 				const FSolverVec3 DW0 = ManifoldPoint.WorldContactNormalAngular0 * ImpulseNormal;
+				FConstraintSolverBody& Body0 = SolverBody0();
 				Body0.ApplyVelocityDelta(DV0, DW0);
 			}
-			if (Body1.IsDynamic())
+			if (IsDynamic(1))
 			{
-				const FSolverVec3 DV1 = -Body1.InvM() * Impulse;
+				const FSolverVec3 DV1 = -State.InvMs[1] * Impulse;
 				const FSolverVec3 DW1 = ManifoldPoint.WorldContactNormalAngular1 * -ImpulseNormal;
+				FConstraintSolverBody& Body1 = SolverBody1();
 				Body1.ApplyVelocityDelta(DV1, DW1);
 			}
 		}
 
-
-		//////////////////////////////////////////////////////////////////////////////////////////////////
-		//////////////////////////////////////////////////////////////////////////////////////////////////
-		//////////////////////////////////////////////////////////////////////////////////////////////////
-		//////////////////////////////////////////////////////////////////////////////////////////////////
-
-		FORCEINLINE_DEBUGGABLE void FPBDCollisionSolverManifoldPoint::InitContact(
-			const FSolverReal Dt,
-			const FConstraintSolverBody& Body0,
-			const FConstraintSolverBody& Body1)
+		FORCEINLINE_DEBUGGABLE void FPBDCollisionSolver::UpdateMass()
 		{
-			NetPushOutNormal = FSolverReal(0);
-			NetPushOutTangentU = FSolverReal(0);
-			NetPushOutTangentV = FSolverReal(0);
-			NetImpulseNormal = FSolverReal(0);
-			NetImpulseTangentU = FSolverReal(0);
-			NetImpulseTangentV = FSolverReal(0);
-			StaticFrictionRatio = FSolverReal(0);
-
-			UpdateMass(Body0, Body1);
-		}
+			const FConstraintSolverBody& Body0 = SolverBody0();
+			const FConstraintSolverBody& Body1 = SolverBody1();
 
-		FORCEINLINE_DEBUGGABLE void FPBDCollisionSolverManifoldPoint::UpdateMass(const FConstraintSolverBody& Body0, const FConstraintSolverBody& Body1)
-		{
-			FSolverReal ContactMassInvNormal = FSolverReal(0);
-			FSolverReal ContactMassInvTangentU = FSolverReal(0);
-			FSolverReal ContactMassInvTangentV = FSolverReal(0);
+			FSolverReal InvM0, InvM1;
+			FSolverMatrix33 InvI0, InvI1;
+			GetDynamicMassProperties(InvM0, InvI0, InvM1, InvI1);
 
-			// These are not used if not initialized below so no need to clear
-			//WorldContactNormalAngular0 = FSolverVec3(0);
-			//WorldContactTangentUAngular0 = FSolverVec3(0);
-			//WorldContactTangentVAngular0 = FSolverVec3(0);
-			//WorldContactNormalAngular1 = FSolverVec3(0);
-			//WorldContactTangentUAngular1 = FSolverVec3(0);
-			//WorldContactTangentVAngular1 = FSolverVec3(0);
+			State.InvMs[0] = InvM0;
+			State.InvMs[1] = InvM1;
 
-			if (Body0.IsDynamic())
+			for (int32 PointIndex = 0; PointIndex < NumManifoldPoints(); ++PointIndex)
 			{
-				const FSolverVec3 R0xN = FSolverVec3::CrossProduct(WorldContact.RelativeContactPoints[0], WorldContact.ContactNormal);
-				const FSolverVec3 R0xU = FSolverVec3::CrossProduct(WorldContact.RelativeContactPoints[0], WorldContact.ContactTangentU);
-				const FSolverVec3 R0xV = FSolverVec3::CrossProduct(WorldContact.RelativeContactPoints[0], WorldContact.ContactTangentV);
+				FPBDCollisionSolverManifoldPoint& ManifoldPoint = State.ManifoldPoints[PointIndex];
 
-				const FSolverMatrix33 InvI0 = Body0.InvI();
+				FSolverReal ContactMassInvNormal = FSolverReal(0);
+				FSolverReal ContactMassInvTangentU = FSolverReal(0);
+				FSolverReal ContactMassInvTangentV = FSolverReal(0);
 
-				WorldContactNormalAngular0 = InvI0 * R0xN;
-				WorldContactTangentUAngular0 = InvI0 * R0xU;
-				WorldContactTangentVAngular0 = InvI0 * R0xV;
+				// These are not used if not initialized below so no need to clear
+				//WorldContactNormalAngular0 = FSolverVec3(0);
+				//WorldContactTangentUAngular0 = FSolverVec3(0);
+				//WorldContactTangentVAngular0 = FSolverVec3(0);
+				//WorldContactNormalAngular1 = FSolverVec3(0);
+				//WorldContactTangentUAngular1 = FSolverVec3(0);
+				//WorldContactTangentVAngular1 = FSolverVec3(0);
 
-				ContactMassInvNormal += FSolverVec3::DotProduct(R0xN, WorldContactNormalAngular0) + Body0.InvM();
-				ContactMassInvTangentU += FSolverVec3::DotProduct(R0xU, WorldContactTangentUAngular0) + Body0.InvM();
-				ContactMassInvTangentV += FSolverVec3::DotProduct(R0xV, WorldContactTangentVAngular0) + Body0.InvM();
-			}
-			if (Body1.IsDynamic())
-			{
-				const FSolverVec3 R1xN = FSolverVec3::CrossProduct(WorldContact.RelativeContactPoints[1], WorldContact.ContactNormal);
-				const FSolverVec3 R1xU = FSolverVec3::CrossProduct(WorldContact.RelativeContactPoints[1], WorldContact.ContactTangentU);
-				const FSolverVec3 R1xV = FSolverVec3::CrossProduct(WorldContact.RelativeContactPoints[1], WorldContact.ContactTangentV);
+				if (IsDynamic(0))
+				{
+					const FSolverVec3 R0xN = FSolverVec3::CrossProduct(ManifoldPoint.WorldContact.RelativeContactPoints[0], ManifoldPoint.WorldContact.ContactNormal);
+					const FSolverVec3 R0xU = FSolverVec3::CrossProduct(ManifoldPoint.WorldContact.RelativeContactPoints[0], ManifoldPoint.WorldContact.ContactTangentU);
+					const FSolverVec3 R0xV = FSolverVec3::CrossProduct(ManifoldPoint.WorldContact.RelativeContactPoints[0], ManifoldPoint.WorldContact.ContactTangentV);
 
-				const FSolverMatrix33 InvI1 = Body1.InvI();
+					ManifoldPoint.WorldContactNormalAngular0 = InvI0 * R0xN;
+					ManifoldPoint.WorldContactTangentUAngular0 = InvI0 * R0xU;
+					ManifoldPoint.WorldContactTangentVAngular0 = InvI0 * R0xV;
 
-				WorldContactNormalAngular1 = InvI1 * R1xN;
-				WorldContactTangentUAngular1 = InvI1 * R1xU;
-				WorldContactTangentVAngular1 = InvI1 * R1xV;
+					ContactMassInvNormal += FSolverVec3::DotProduct(R0xN, ManifoldPoint.WorldContactNormalAngular0) + InvM0;
+					ContactMassInvTangentU += FSolverVec3::DotProduct(R0xU, ManifoldPoint.WorldContactTangentUAngular0) + InvM0;
+					ContactMassInvTangentV += FSolverVec3::DotProduct(R0xV, ManifoldPoint.WorldContactTangentVAngular0) + InvM0;
+				}
+				if (IsDynamic(1))
+				{
+					const FSolverVec3 R1xN = FSolverVec3::CrossProduct(ManifoldPoint.WorldContact.RelativeContactPoints[1], ManifoldPoint.WorldContact.ContactNormal);
+					const FSolverVec3 R1xU = FSolverVec3::CrossProduct(ManifoldPoint.WorldContact.RelativeContactPoints[1], ManifoldPoint.WorldContact.ContactTangentU);
+					const FSolverVec3 R1xV = FSolverVec3::CrossProduct(ManifoldPoint.WorldContact.RelativeContactPoints[1], ManifoldPoint.WorldContact.ContactTangentV);
 
-				ContactMassInvNormal += FSolverVec3::DotProduct(R1xN, WorldContactNormalAngular1) + Body1.InvM();
-				ContactMassInvTangentU += FSolverVec3::DotProduct(R1xU, WorldContactTangentUAngular1) + Body1.InvM();
-				ContactMassInvTangentV += FSolverVec3::DotProduct(R1xV, WorldContactTangentVAngular1) + Body1.InvM();
-			}
+					ManifoldPoint.WorldContactNormalAngular1 = InvI1 * R1xN;
+					ManifoldPoint.WorldContactTangentUAngular1 = InvI1 * R1xU;
+					ManifoldPoint.WorldContactTangentVAngular1 = InvI1 * R1xV;
+
+					ContactMassInvNormal += FSolverVec3::DotProduct(R1xN, ManifoldPoint.WorldContactNormalAngular1) + InvM1;
+					ContactMassInvTangentU += FSolverVec3::DotProduct(R1xU, ManifoldPoint.WorldContactTangentUAngular1) + InvM1;
+					ContactMassInvTangentV += FSolverVec3::DotProduct(R1xV, ManifoldPoint.WorldContactTangentVAngular1) + InvM1;
+				}
 
-			ContactMassNormal = (ContactMassInvNormal > FSolverReal(UE_SMALL_NUMBER)) ? FSolverReal(1) / ContactMassInvNormal : FSolverReal(0);
-			ContactMassTangentU = (ContactMassInvTangentU > FSolverReal(UE_SMALL_NUMBER)) ? FSolverReal(1) / ContactMassInvTangentU : FSolverReal(0);
-			ContactMassTangentV = (ContactMassInvTangentV > FSolverReal(UE_SMALL_NUMBER)) ? FSolverReal(1) / ContactMassInvTangentV : FSolverReal(0);
+				ManifoldPoint.ContactMassNormal = (ContactMassInvNormal > FSolverReal(UE_SMALL_NUMBER)) ? FSolverReal(1) / ContactMassInvNormal : FSolverReal(0);
+				ManifoldPoint.ContactMassTangentU = (ContactMassInvTangentU > FSolverReal(UE_SMALL_NUMBER)) ? FSolverReal(1) / ContactMassInvTangentU : FSolverReal(0);
+				ManifoldPoint.ContactMassTangentV = (ContactMassInvTangentV > FSolverReal(UE_SMALL_NUMBER)) ? FSolverReal(1) / ContactMassInvTangentV : FSolverReal(0);
+			}
 		}
 
-		FORCEINLINE_DEBUGGABLE void FPBDCollisionSolverManifoldPoint::UpdateMassNormal(const FConstraintSolverBody& Body0, const FConstraintSolverBody& Body1)
+		FORCEINLINE_DEBUGGABLE void FPBDCollisionSolver::UpdateMassNormal()
 		{
-			FSolverReal ContactMassInvNormal = FSolverReal(0);
-			if (Body0.IsDynamic())
-			{
-				const FSolverVec3 R0xN = FSolverVec3::CrossProduct(WorldContact.RelativeContactPoints[0], WorldContact.ContactNormal);
-				const FSolverMatrix33 InvI0 = Body0.InvI();
-				WorldContactNormalAngular0 = InvI0 * R0xN;
-				ContactMassInvNormal += FSolverVec3::DotProduct(R0xN, WorldContactNormalAngular0) + Body0.InvM();
-			}
-			if (Body1.IsDynamic())
+			const FConstraintSolverBody& Body0 = SolverBody0();
+			const FConstraintSolverBody& Body1 = SolverBody1();
+
+			FSolverReal InvM0, InvM1;
+			FSolverMatrix33 InvI0, InvI1;
+			GetDynamicMassProperties(InvM0, InvI0, InvM1, InvI1);
+
+			State.InvMs[0] = InvM0;
+			State.InvMs[1] = InvM1;
+
+			for (int32 PointIndex = 0; PointIndex < NumManifoldPoints(); ++PointIndex)
 			{
-				const FSolverVec3 R1xN = FSolverVec3::CrossProduct(WorldContact.RelativeContactPoints[1], WorldContact.ContactNormal);
-				const FSolverMatrix33 InvI1 = Body1.InvI();
-				WorldContactNormalAngular1 = InvI1 * R1xN;
-				ContactMassInvNormal += FSolverVec3::DotProduct(R1xN, WorldContactNormalAngular1) + Body1.InvM();
+				FPBDCollisionSolverManifoldPoint& ManifoldPoint = State.ManifoldPoints[PointIndex];
+
+				FSolverReal ContactMassInvNormal = FSolverReal(0);
+				if (IsDynamic(0))
+				{
+					const FSolverVec3 R0xN = FSolverVec3::CrossProduct(ManifoldPoint.WorldContact.RelativeContactPoints[0], ManifoldPoint.WorldContact.ContactNormal);
+					ManifoldPoint.WorldContactNormalAngular0 = InvI0 * R0xN;
+					ContactMassInvNormal += FSolverVec3::DotProduct(R0xN, ManifoldPoint.WorldContactNormalAngular0) + InvM0;
+				}
+				if (IsDynamic(1))
+				{
+					const FSolverVec3 R1xN = FSolverVec3::CrossProduct(ManifoldPoint.WorldContact.RelativeContactPoints[1], ManifoldPoint.WorldContact.ContactNormal);
+					ManifoldPoint.WorldContactNormalAngular1 = InvI1 * R1xN;
+					ContactMassInvNormal += FSolverVec3::DotProduct(R1xN, ManifoldPoint.WorldContactNormalAngular1) + InvM1;
+				}
+				ManifoldPoint.ContactMassNormal = (ContactMassInvNormal > FSolverReal(UE_SMALL_NUMBER)) ? FSolverReal(1) / ContactMassInvNormal : FSolverReal(0);
 			}
-			ContactMassNormal = (ContactMassInvNormal > FSolverReal(UE_SMALL_NUMBER)) ? FSolverReal(1) / ContactMassInvNormal : FSolverReal(0);
 		}
 
-		FORCEINLINE_DEBUGGABLE void FPBDCollisionSolverManifoldPoint::CalculateContactPositionErrorNormal(const FConstraintSolverBody& Body0, const FConstraintSolverBody& Body1, const FSolverReal MaxPushOut, FSolverReal& OutContactDeltaNormal) const
+		FORCEINLINE_DEBUGGABLE void FPBDCollisionSolver::CalculateContactPositionErrorNormal(
+			const FPBDCollisionSolverManifoldPoint& ManifoldPoint,
+			const FSolverReal MaxPushOut,
+			FSolverReal& OutContactDeltaNormal) const
 		{
+			const FConstraintSolverBody& Body0 = SolverBody0();
+			const FConstraintSolverBody& Body1 = SolverBody1();
+
 			// Linear version: calculate the contact delta assuming linear motion after applying a positional impulse at the contact point. There will be an error that depends on the size of the rotation.
-			const FSolverVec3 ContactDelta0 = Body0.DP() + FSolverVec3::CrossProduct(Body0.DQ(), WorldContact.RelativeContactPoints[0]);
-			const FSolverVec3 ContactDelta1 = Body1.DP() + FSolverVec3::CrossProduct(Body1.DQ(), WorldContact.RelativeContactPoints[1]);
+			const FSolverVec3 ContactDelta0 = Body0.DP() + FSolverVec3::CrossProduct(Body0.DQ(), ManifoldPoint.WorldContact.RelativeContactPoints[0]);
+			const FSolverVec3 ContactDelta1 = Body1.DP() + FSolverVec3::CrossProduct(Body1.DQ(), ManifoldPoint.WorldContact.RelativeContactPoints[1]);
 			const FSolverVec3 ContactDelta = ContactDelta0 - ContactDelta1;
-			OutContactDeltaNormal = WorldContact.ContactDeltaNormal + FSolverVec3::DotProduct(ContactDelta, WorldContact.ContactNormal);
+			OutContactDeltaNormal = ManifoldPoint.WorldContact.ContactDeltaNormal + FSolverVec3::DotProduct(ContactDelta, ManifoldPoint.WorldContact.ContactNormal);
 
 			// NOTE: OutContactDeltaNormal is negative for penetration
 			// NOTE: MaxPushOut == 0 disables the pushout limits
@@ -726,62 +770,70 @@ namespace Chaos
 			}
 		}
 
-		FORCEINLINE_DEBUGGABLE void FPBDCollisionSolverManifoldPoint::CalculateContactPositionErrorTangential(const FConstraintSolverBody& Body0, const FConstraintSolverBody& Body1, FSolverReal& OutContactDeltaTangentU, FSolverReal& OutContactDeltaTangentV) const
+		FORCEINLINE_DEBUGGABLE void FPBDCollisionSolver::CalculateContactPositionErrorTangential(
+			const FPBDCollisionSolverManifoldPoint& ManifoldPoint,
+			FSolverReal& OutContactDeltaTangentU,
+			FSolverReal& OutContactDeltaTangentV) const
 		{
+			const FConstraintSolverBody& Body0 = SolverBody0();
+			const FConstraintSolverBody& Body1 = SolverBody1();
+
 			// Linear version: calculate the contact delta assuming linear motion after applying a positional impulse at the contact point. There will be an error that depends on the size of the rotation.
-			const FSolverVec3 ContactDelta0 = Body0.DP() + FSolverVec3::CrossProduct(Body0.DQ(), WorldContact.RelativeContactPoints[0]);
-			const FSolverVec3 ContactDelta1 = Body1.DP() + FSolverVec3::CrossProduct(Body1.DQ(), WorldContact.RelativeContactPoints[1]);
+			const FSolverVec3 ContactDelta0 = Body0.DP() + FSolverVec3::CrossProduct(Body0.DQ(), ManifoldPoint.WorldContact.RelativeContactPoints[0]);
+			const FSolverVec3 ContactDelta1 = Body1.DP() + FSolverVec3::CrossProduct(Body1.DQ(), ManifoldPoint.WorldContact.RelativeContactPoints[1]);
 			const FSolverVec3 ContactDelta = ContactDelta0 - ContactDelta1;
-			OutContactDeltaTangentU = WorldContact.ContactDeltaTangentU + FSolverVec3::DotProduct(ContactDelta, WorldContact.ContactTangentU);
-			OutContactDeltaTangentV = WorldContact.ContactDeltaTangentV + FSolverVec3::DotProduct(ContactDelta, WorldContact.ContactTangentV);
+			OutContactDeltaTangentU = ManifoldPoint.WorldContact.ContactDeltaTangentU + FSolverVec3::DotProduct(ContactDelta, ManifoldPoint.WorldContact.ContactTangentU);
+			OutContactDeltaTangentV = ManifoldPoint.WorldContact.ContactDeltaTangentV + FSolverVec3::DotProduct(ContactDelta, ManifoldPoint.WorldContact.ContactTangentV);
 		}
 
-		FORCEINLINE_DEBUGGABLE void FPBDCollisionSolverManifoldPoint::CalculateContactVelocityError(const FConstraintSolverBody& Body0, const FConstraintSolverBody& Body1, const FSolverReal DynamicFriction, const FSolverReal Dt, FSolverReal& OutContactVelocityDeltaNormal, FSolverReal& OutContactVelocityDeltaTangent0, FSolverReal& OutContactVelocityDeltaTangent1) const
+		FORCEINLINE_DEBUGGABLE void FPBDCollisionSolver::CalculateContactVelocityError(
+			const FPBDCollisionSolverManifoldPoint& ManifoldPoint,
+			const FSolverReal DynamicFriction,
+			const FSolverReal Dt,
+			FSolverReal& OutContactVelocityDeltaNormal,
+			FSolverReal& OutContactVelocityDeltaTangent0,
+			FSolverReal& OutContactVelocityDeltaTangent1) const
 		{
-			const FSolverVec3 ContactVelocity0 = Body0.V() + FSolverVec3::CrossProduct(Body0.W(), WorldContact.RelativeContactPoints[0]);
-			const FSolverVec3 ContactVelocity1 = Body1.V() + FSolverVec3::CrossProduct(Body1.W(), WorldContact.RelativeContactPoints[1]);
+			const FConstraintSolverBody& Body0 = SolverBody0();
+			const FConstraintSolverBody& Body1 = SolverBody1();
+
+			const FSolverVec3 ContactVelocity0 = Body0.V() + FSolverVec3::CrossProduct(Body0.W(), ManifoldPoint.WorldContact.RelativeContactPoints[0]);
+			const FSolverVec3 ContactVelocity1 = Body1.V() + FSolverVec3::CrossProduct(Body1.W(), ManifoldPoint.WorldContact.RelativeContactPoints[1]);
 			const FSolverVec3 ContactVelocity = ContactVelocity0 - ContactVelocity1;
-			const FSolverReal ContactVelocityNormal = FSolverVec3::DotProduct(ContactVelocity, WorldContact.ContactNormal);
-			const FSolverReal ContactVelocityTangent0 = FSolverVec3::DotProduct(ContactVelocity, WorldContact.ContactTangentU);
-			const FSolverReal ContactVelocityTangent1 = FSolverVec3::DotProduct(ContactVelocity, WorldContact.ContactTangentV);
+			const FSolverReal ContactVelocityNormal = FSolverVec3::DotProduct(ContactVelocity, ManifoldPoint.WorldContact.ContactNormal);
+			const FSolverReal ContactVelocityTangent0 = FSolverVec3::DotProduct(ContactVelocity, ManifoldPoint.WorldContact.ContactTangentU);
+			const FSolverReal ContactVelocityTangent1 = FSolverVec3::DotProduct(ContactVelocity, ManifoldPoint.WorldContact.ContactTangentV);
 
-			OutContactVelocityDeltaNormal = (ContactVelocityNormal - WorldContact.ContactTargetVelocityNormal);
+			OutContactVelocityDeltaNormal = (ContactVelocityNormal - ManifoldPoint.WorldContact.ContactTargetVelocityNormal);
 			OutContactVelocityDeltaTangent0 = ContactVelocityTangent0;
 			OutContactVelocityDeltaTangent1 = ContactVelocityTangent1;
 		}
 
-		FORCEINLINE_DEBUGGABLE void FPBDCollisionSolverManifoldPoint::CalculateContactVelocityErrorNormal(const FConstraintSolverBody& Body0, const FConstraintSolverBody& Body1, FSolverReal& OutContactVelocityDeltaNormal) const
+		FORCEINLINE_DEBUGGABLE void FPBDCollisionSolver::CalculateContactVelocityErrorNormal(
+			const FPBDCollisionSolverManifoldPoint& ManifoldPoint,
+			FSolverReal& OutContactVelocityDeltaNormal) const
 		{
-			const FSolverVec3 ContactVelocity0 = Body0.V() + FSolverVec3::CrossProduct(Body0.W(), WorldContact.RelativeContactPoints[0]);
-			const FSolverVec3 ContactVelocity1 = Body1.V() + FSolverVec3::CrossProduct(Body1.W(), WorldContact.RelativeContactPoints[1]);
+			const FConstraintSolverBody& Body0 = SolverBody0();
+			const FConstraintSolverBody& Body1 = SolverBody1();
+
+			const FSolverVec3 ContactVelocity0 = Body0.V() + FSolverVec3::CrossProduct(Body0.W(), ManifoldPoint.WorldContact.RelativeContactPoints[0]);
+			const FSolverVec3 ContactVelocity1 = Body1.V() + FSolverVec3::CrossProduct(Body1.W(), ManifoldPoint.WorldContact.RelativeContactPoints[1]);
 			const FSolverVec3 ContactVelocity = ContactVelocity0 - ContactVelocity1;
-			const FSolverReal ContactVelocityNormal = FSolverVec3::DotProduct(ContactVelocity, WorldContact.ContactNormal);
+			const FSolverReal ContactVelocityNormal = FSolverVec3::DotProduct(ContactVelocity, ManifoldPoint.WorldContact.ContactNormal);
 
 			// Add up the errors in the velocity (current velocity - desired velocity)
-			OutContactVelocityDeltaNormal = (ContactVelocityNormal - WorldContact.ContactTargetVelocityNormal);
+			OutContactVelocityDeltaNormal = (ContactVelocityNormal - ManifoldPoint.WorldContact.ContactTargetVelocityNormal);
 		}
 
-		FORCEINLINE_DEBUGGABLE bool FPBDCollisionSolverManifoldPoint::ShouldSolveVelocity() const
+		FORCEINLINE_DEBUGGABLE bool FPBDCollisionSolver::ShouldSolveVelocity(
+			const FPBDCollisionSolverManifoldPoint& ManifoldPoint) const
 		{
 			// We ensure positive separating velocity for close contacts even if they didn't receive a pushout
-			return (NetPushOutNormal > FSolverReal(0)) || (WorldContact.ContactDeltaNormal < FSolverReal(0));
-		}
-
-		//////////////////////////////////////////////////////////////////////////////////////////////////
-		//////////////////////////////////////////////////////////////////////////////////////////////////
-		//////////////////////////////////////////////////////////////////////////////////////////////////
-		//////////////////////////////////////////////////////////////////////////////////////////////////
-
-		FORCEINLINE_DEBUGGABLE  void FPBDCollisionSolver::FinalizeManifoldPoint(const int32 ManifoldPoiontIndex, const FSolverReal Dt)
-		{
-			State.ManifoldPoints[ManifoldPoiontIndex].InitContact(
-				FSolverReal(Dt),
-				State.SolverBodies[0],
-				State.SolverBodies[1]);
+			return (ManifoldPoint.NetPushOutNormal > FSolverReal(0)) || (ManifoldPoint.WorldContact.ContactDeltaNormal < FSolverReal(0));
 		}
 
-		FORCEINLINE_DEBUGGABLE void FPBDCollisionSolver::SetManifoldPoint(
-			const int32 ManifoldPoiontIndex,
+		FORCEINLINE_DEBUGGABLE void FPBDCollisionSolver::InitManifoldPoint(
+			const int32 PointIndex,
 			const FSolverReal Dt,
 			const FSolverVec3& InRelativeContactPosition0,
 			const FSolverVec3& InRelativeContactPosition1,
@@ -793,29 +845,36 @@ namespace Chaos
 			const FSolverReal InWorldContactDeltaTangentV,
 			const FSolverReal InWorldContactVelocityTargetNormal)
 		{
-			FWorldContactPoint& WorldContactPoint = State.ManifoldPoints[ManifoldPoiontIndex].WorldContact;
-			WorldContactPoint.RelativeContactPoints[0] = InRelativeContactPosition0;
-			WorldContactPoint.RelativeContactPoints[1] = InRelativeContactPosition1;
-			WorldContactPoint.ContactNormal = InWorldContactNormal;
-			WorldContactPoint.ContactTangentU = InWorldContactTangentU;
-			WorldContactPoint.ContactTangentV = InWorldContactTangentV;
-			WorldContactPoint.ContactDeltaNormal = InWorldContactDeltaNormal;
-			WorldContactPoint.ContactDeltaTangentU = InWorldContactDeltaTangentU;
-			WorldContactPoint.ContactDeltaTangentV = InWorldContactDeltaTangentV;
-			WorldContactPoint.ContactTargetVelocityNormal = InWorldContactVelocityTargetNormal;
-
-			State.ManifoldPoints[ManifoldPoiontIndex].InitContact(
-				FSolverReal(Dt),
-				State.SolverBodies[0],
-				State.SolverBodies[1]);
+			FPBDCollisionSolverManifoldPoint& ManifoldPoint = State.ManifoldPoints[PointIndex];
+
+			ManifoldPoint.WorldContact.RelativeContactPoints[0] = InRelativeContactPosition0;
+			ManifoldPoint.WorldContact.RelativeContactPoints[1] = InRelativeContactPosition1;
+			ManifoldPoint.WorldContact.ContactNormal = InWorldContactNormal;
+			ManifoldPoint.WorldContact.ContactTangentU = InWorldContactTangentU;
+			ManifoldPoint.WorldContact.ContactTangentV = InWorldContactTangentV;
+			ManifoldPoint.WorldContact.ContactDeltaNormal = InWorldContactDeltaNormal;
+			ManifoldPoint.WorldContact.ContactDeltaTangentU = InWorldContactDeltaTangentU;
+			ManifoldPoint.WorldContact.ContactDeltaTangentV = InWorldContactDeltaTangentV;
+			ManifoldPoint.WorldContact.ContactTargetVelocityNormal = InWorldContactVelocityTargetNormal;
+
+			ManifoldPoint.NetPushOutNormal = FSolverReal(0);
+			ManifoldPoint.NetPushOutTangentU = FSolverReal(0);
+			ManifoldPoint.NetPushOutTangentV = FSolverReal(0);
+			ManifoldPoint.NetImpulseNormal = FSolverReal(0);
+			ManifoldPoint.NetImpulseTangentU = FSolverReal(0);
+			ManifoldPoint.NetImpulseTangentV = FSolverReal(0);
+			ManifoldPoint.StaticFrictionRatio = FSolverReal(0);
 		}
 
-		FORCEINLINE_DEBUGGABLE bool FPBDCollisionSolver::SolvePositionWithFriction(const FSolverReal Dt, const FSolverReal MaxPushOut)
+		FORCEINLINE_DEBUGGABLE void FPBDCollisionSolver::FinalizeManifold()
 		{
-			// SolverBody decorator used to add mass scaling
-			FConstraintSolverBody& Body0 = SolverBody0();
-			FConstraintSolverBody& Body1 = SolverBody1();
+			UpdateMass();
+		}
 
+		FORCEINLINE_DEBUGGABLE void FPBDCollisionSolver::SolvePositionWithFriction(
+			const FSolverReal Dt,
+			const FSolverReal MaxPushOut)
+		{
 			// Accumulate net pushout for friction limits below
 			bool bApplyFriction = false;
 			bool bApplyPointFriction[MaxPointsPerConstraint] = { false, };
@@ -828,7 +887,7 @@ namespace Chaos
 				FPBDCollisionSolverManifoldPoint& SolverManifoldPoint = State.ManifoldPoints[PointIndex];
 
 				FSolverReal ContactDeltaNormal;
-				SolverManifoldPoint.CalculateContactPositionErrorNormal(Body0.SolverBody(), Body1.SolverBody(), MaxPushOut, ContactDeltaNormal);
+				CalculateContactPositionErrorNormal(SolverManifoldPoint, MaxPushOut, ContactDeltaNormal);
 
 				// Apply a normal correction if we still have penetration or if we are now separated but have previously applied a correction that we may want to undo
 				const bool bProcessManifoldPoint = (ContactDeltaNormal < FSolverReal(0)) || (SolverManifoldPoint.NetPushOutNormal > FSolverReal(UE_SMALL_NUMBER));
@@ -837,9 +896,7 @@ namespace Chaos
 					ApplyPositionCorrectionNormal(
 						State.Stiffness,
 						ContactDeltaNormal,
-						SolverManifoldPoint,
-						Body0,
-						Body1);
+						SolverManifoldPoint);
 				}
 
 				// Friction gets updated for any point with a net normal correction or where we have previously had a normal correction and 
@@ -872,7 +929,7 @@ namespace Chaos
 					if (bApplyPointFriction[PointIndex])
 					{
 						FSolverReal ContactDeltaTangentU, ContactDeltaTangentV;
-						SolverManifoldPoint.CalculateContactPositionErrorTangential(Body0.SolverBody(), Body1.SolverBody(), ContactDeltaTangentU, ContactDeltaTangentV);
+						CalculateContactPositionErrorTangential(SolverManifoldPoint, ContactDeltaTangentU, ContactDeltaTangentV);
 
 						ApplyPositionCorrectionTangential(
 							FrictionStiffness,
@@ -881,29 +938,23 @@ namespace Chaos
 							FrictionMaxPushOut,
 							ContactDeltaTangentU,
 							ContactDeltaTangentV,
-							SolverManifoldPoint,
-							Body0,
-							Body1);
+							SolverManifoldPoint);
 					}
 				}
 			}
-
-			return false;
 		}
 
-		FORCEINLINE_DEBUGGABLE bool FPBDCollisionSolver::SolvePositionNoFriction(const FSolverReal Dt, const FSolverReal MaxPushOut)
+		FORCEINLINE_DEBUGGABLE void FPBDCollisionSolver::SolvePositionNoFriction(
+			const FSolverReal Dt,
+			const FSolverReal MaxPushOut)
 		{
-			// SolverBody decorator used to add mass scaling
-			FConstraintSolverBody& Body0 = SolverBody0();
-			FConstraintSolverBody& Body1 = SolverBody1();
-
 			// Apply the position correction so that all contacts have zero separation
 			for (int32 PointIndex = 0; PointIndex < NumManifoldPoints(); ++PointIndex)
 			{
 				FPBDCollisionSolverManifoldPoint& SolverManifoldPoint = State.ManifoldPoints[PointIndex];
 
 				FSolverReal ContactDeltaNormal;
-				SolverManifoldPoint.CalculateContactPositionErrorNormal(Body0.SolverBody(), Body1.SolverBody(), MaxPushOut, ContactDeltaNormal);
+				CalculateContactPositionErrorNormal(SolverManifoldPoint, MaxPushOut, ContactDeltaNormal);
 
 				const bool bProcessManifoldPoint = (ContactDeltaNormal < FSolverReal(0)) || (SolverManifoldPoint.NetPushOutNormal > FSolverReal(UE_SMALL_NUMBER));
 				if (bProcessManifoldPoint)
@@ -911,16 +962,18 @@ namespace Chaos
 					ApplyPositionCorrectionNormal(
 						State.Stiffness,
 						ContactDeltaNormal,
-						SolverManifoldPoint,
-						Body0,
-						Body1);
+						SolverManifoldPoint);
 				}
 			}
-
-			return false;
 		}
 
-		FORCEINLINE_DEBUGGABLE bool FPBDCollisionSolver::SolveVelocity(const FSolverReal Dt, const bool bApplyDynamicFriction)
+		/**
+		 * @brief Calculate and apply the velocity correction for this iteration
+		 * @return true if we need to run more iterations, false if we did not apply any correction
+		*/
+		FORCEINLINE_DEBUGGABLE void FPBDCollisionSolver::SolveVelocity(
+			const FSolverReal Dt,
+			const bool bApplyDynamicFriction)
 		{
 			// Apply restitution at the average contact point
 			// This means we don't need to run as many iterations to get stable bouncing
@@ -931,9 +984,6 @@ namespace Chaos
 				SolveVelocityAverage(Dt);
 			}
 
-			FConstraintSolverBody& Body0 = SolverBody0();
-			FConstraintSolverBody& Body1 = SolverBody1();
-
 			// NOTE: this dynamic friction implementation is iteration-count sensitive
 			// @todo(chaos): fix iteration count dependence of dynamic friction
 			const FSolverReal DynamicFriction = (bApplyDynamicFriction && (Dt > 0) && CVars::bChaos_PBDCollisionSolver_Velocity_FrictionEnabled) ? State.VelocityFriction : FSolverReal(0);
@@ -942,14 +992,14 @@ namespace Chaos
 			{
 				FPBDCollisionSolverManifoldPoint& SolverManifoldPoint = State.ManifoldPoints[PointIndex];
 
-				if (SolverManifoldPoint.ShouldSolveVelocity())
+				if (ShouldSolveVelocity(SolverManifoldPoint))
 				{
 					const FSolverReal MinImpulseNormal = FMath::Min(FSolverReal(0), -SolverManifoldPoint.NetPushOutNormal / Dt);
 
 					if (DynamicFriction > 0)
 					{
 						FSolverReal ContactVelocityDeltaNormal, ContactVelocityDeltaTangentU, ContactVelocityDeltaTangentV;
-						SolverManifoldPoint.CalculateContactVelocityError(Body0, Body1, DynamicFriction, Dt, ContactVelocityDeltaNormal, ContactVelocityDeltaTangentU, ContactVelocityDeltaTangentV);
+						CalculateContactVelocityError(SolverManifoldPoint, DynamicFriction, Dt, ContactVelocityDeltaNormal, ContactVelocityDeltaTangentU, ContactVelocityDeltaTangentV);
 
 						ApplyVelocityCorrection(
 							State.Stiffness,
@@ -959,32 +1009,28 @@ namespace Chaos
 							ContactVelocityDeltaTangentU,
 							ContactVelocityDeltaTangentV,
 							MinImpulseNormal,
-							SolverManifoldPoint,
-							Body0,
-							Body1);
+							SolverManifoldPoint);
 					}
 					else
 					{
 						FSolverReal ContactVelocityDeltaNormal;
-						SolverManifoldPoint.CalculateContactVelocityErrorNormal(Body0, Body1, ContactVelocityDeltaNormal);
+						CalculateContactVelocityErrorNormal(SolverManifoldPoint, ContactVelocityDeltaNormal);
 
 						ApplyVelocityCorrectionNormal(
 							State.Stiffness,
 							ContactVelocityDeltaNormal,
 							MinImpulseNormal,
-							SolverManifoldPoint,
-							Body0,
-							Body1);
+							SolverManifoldPoint);
 					}
 				}
 			}
-
-			// Early-out support for the velocity solve is not currently very important because we
-			// only run one iteration in the velocity solve phase.
-			// @todo(chaos): support early-out in velocity solve if necessary
-			return true;
 		}
 
+		//////////////////////////////////////////////////////////////////////////////////////////////////
+		//////////////////////////////////////////////////////////////////////////////////////////////////
+		//////////////////////////////////////////////////////////////////////////////////////////////////
+		//////////////////////////////////////////////////////////////////////////////////////////////////
+
 		/**
 		 * A helper for solving arrays of constraints
 		 */
diff --git a/Engine/Source/Runtime/Experimental/Chaos/Public/Chaos/Collision/PBDCollisionSolverSimd.h b/Engine/Source/Runtime/Experimental/Chaos/Public/Chaos/Collision/PBDCollisionSolverSimd.h
--- a/Engine/Source/Runtime/Experimental/Chaos/Public/Chaos/Collision/PBDCollisionSolverSimd.h
+++ b/Engine/Source/Runtime/Experimental/Chaos/Public/Chaos/Collision/PBDCollisionSolverSimd.h
@@ -407,9 +407,10 @@ namespace Chaos
 				const FVec3f ContactTangentU = ManifoldPoint.SimdContactTangentU.GetValue(LaneIndex);
 				const FVec3f ContactTangentV = ManifoldPoint.SimdContactTangentV.GetValue(LaneIndex);
 
-				if (Body0.IsDynamic())
+				const FSolverReal InvM0 = InvMScale0 * Body0.InvM();
+				const FSolverReal InvM1 = InvMScale1 * Body1.InvM();
+				if (InvM0 > 0)
 				{
-					const FSolverReal InvM0 = InvMScale0 * Body0.InvM();
 					const FSolverMatrix33 InvI0 = InvIScale0 * Body0.InvI();
 
 					const FVec3f RelativeContactPoint0 = ManifoldPoint.SimdRelativeContactPoint0.GetValue(LaneIndex);
@@ -430,9 +431,8 @@ namespace Chaos
 					ContactMassInvTangentU += FSolverVec3::DotProduct(R0xU, IR0xU) + InvM0;
 					ContactMassInvTangentV += FSolverVec3::DotProduct(R0xV, IR0xV) + InvM0;
 				}
-				if (Body1.IsDynamic())
+				if (InvM1 > 0)
 				{
-					const FSolverReal InvM1 = InvMScale1 * Body1.InvM();
 					const FSolverMatrix33 InvI1 = InvIScale1 * Body1.InvI();
 
 					const FVec3f RelativeContactPoint1 = ManifoldPoint.SimdRelativeContactPoint1.GetValue(LaneIndex);
@@ -487,9 +487,10 @@ namespace Chaos
 
 				const FVec3f ContactNormal = ManifoldPoint.SimdContactNormal.GetValue(LaneIndex);
 
-				if (Body0.IsDynamic())
+				const FSolverReal InvM0 = InvMScale0 * Body0.InvM();
+				const FSolverReal InvM1 = InvMScale1 * Body1.InvM();
+				if (InvM0 > 0)
 				{
-					const FSolverReal InvM0 = InvMScale0 * Body0.InvM();
 					const FSolverMatrix33 InvI0 = InvIScale0 * Body0.InvI();
 
 					const FVec3f RelativeContactPoint0 = ManifoldPoint.SimdRelativeContactPoint0.GetValue(LaneIndex);
@@ -502,9 +503,8 @@ namespace Chaos
 
 					ContactMassInvNormal += FSolverVec3::DotProduct(R0xN, IR0xN) + InvM0;
 				}
-				if (Body1.IsDynamic())
+				if (InvM1 > 0)
 				{
-					const FSolverReal InvM1 = InvMScale1 * Body1.InvM();
 					const FSolverMatrix33 InvI1 = InvIScale1 * Body1.InvI();
 
 					const FVec3f RelativeContactPoint1 = ManifoldPoint.SimdRelativeContactPoint1.GetValue(LaneIndex);
diff --git a/Engine/Source/Runtime/Experimental/Chaos/Public/Chaos/Collision/PBDCollisionSolverUtilities.h b/Engine/Source/Runtime/Experimental/Chaos/Public/Chaos/Collision/PBDCollisionSolverUtilities.h
new file mode 100644
--- /dev/null
+++ b/Engine/Source/Runtime/Experimental/Chaos/Public/Chaos/Collision/PBDCollisionSolverUtilities.h
@@ -0,0 +1,46 @@
+// Copyright Epic Games, Inc. All Rights Reserved.
+#pragma once
+
+#include "Chaos/Core.h"
+#include "Chaos/Defines.h"
+#include "Chaos/Evolution/SolverBody.h"
+
+
+namespace Chaos
+{
+	namespace Private
+	{
+
+		inline bool CalculateBodyShockPropagation(
+			const FSolverBody& Body0,
+			const FSolverBody& Body1,
+			const FSolverReal ShockPropagation,
+			FSolverReal& OutShockPropagation0,
+			FSolverReal& OutShockPropagation1)
+		{
+			OutShockPropagation0 = FSolverReal(1);
+			OutShockPropagation1 = FSolverReal(1);
+
+			// Shock propagation decreases the inverse mass of bodies that are lower in the pile
+			// of objects. This significantly improves stability of heaps and stacks. Height in the pile is indictaed by the "level". 
+			// No need to set an inverse mass scale if the other body is kinematic (with inv mass of 0).
+			// Bodies at the same level do not take part in shock propagation.
+			if (Body0.IsDynamic() && Body1.IsDynamic() && (Body0.Level() != Body1.Level()))
+			{
+				// Set the inv mass scale of the "lower" body to make it heavier
+				if (Body0.Level() < Body1.Level())
+				{
+					OutShockPropagation0 = ShockPropagation;
+				}
+				else
+				{
+					OutShockPropagation1 = ShockPropagation;
+				}
+				return true;
+			}
+
+			return false;
+		}
+
+	} // namespace Private
+} // namespace Chaos
\ No newline at end of file
diff --git a/Engine/Source/Runtime/Experimental/Chaos/Public/Chaos/Evolution/SolverBody.h b/Engine/Source/Runtime/Experimental/Chaos/Public/Chaos/Evolution/SolverBody.h
--- a/Engine/Source/Runtime/Experimental/Chaos/Public/Chaos/Evolution/SolverBody.h
+++ b/Engine/Source/Runtime/Experimental/Chaos/Public/Chaos/Evolution/SolverBody.h
@@ -98,6 +98,8 @@ namespace Chaos
 	class FSolverBody
 	{
 	public:
+		static constexpr FSolverReal ZeroMassThreshold() { return std::numeric_limits<FSolverReal>::min(); }
+
 
 		/**
 		 * A factory method to create a safely initialized Solverbody. 
@@ -332,7 +334,7 @@ namespace Chaos
 		 * @brief Whether the body has a finite mass
 		 * @note This is based on the current inverse mass, so a "dynamic" particle with 0 inverse mass will return true here.
 		*/
-		inline bool IsDynamic() const { return (State.InvM > UE_SMALL_NUMBER); }
+		inline bool IsDynamic() const { return (State.InvM > FSolverBody::ZeroMassThreshold()); }
 
 		/**
 		 * @brief Apply a world-space position and rotation delta to the body center of mass, and update inverse mass
@@ -617,7 +619,7 @@ namespace Chaos
 		/**
 		 * @brief Whether the body is dynamic (i.e., has a finite mass) after scaling is applied
 		*/
-		inline bool IsDynamic() const { return (InvM() > TNumericLimits<FSolverReal>::Min()); }
+		inline bool IsDynamic() const { return (InvM() > FSolverBody::ZeroMassThreshold()); }
 
 		//
 		// From here all methods just forward to the FSolverBody
