Chaos - minor cleanup in collision solver

diff --git a/Engine/Source/Runtime/Experimental/Chaos/Private/Chaos/Collision/PBDCollisionContainerSolver.cpp b/Engine/Source/Runtime/Experimental/Chaos/Private/Chaos/Collision/PBDCollisionContainerSolver.cpp
--- a/Engine/Source/Runtime/Experimental/Chaos/Private/Chaos/Collision/PBDCollisionContainerSolver.cpp
+++ b/Engine/Source/Runtime/Experimental/Chaos/Private/Chaos/Collision/PBDCollisionContainerSolver.cpp
@@ -252,14 +252,14 @@ namespace Chaos
 				const Private::FPBDCollisionSolverManifoldPoint& SolverManifoldPoint = Solver.GetManifoldPoint(SolverManifoldPointIndex);
 
 				NetPushOut =
-					SolverManifoldPoint.NetPushOutNormal * SolverManifoldPoint.WorldContact.ContactNormal +
-					SolverManifoldPoint.NetPushOutTangentU * SolverManifoldPoint.WorldContact.ContactTangentU +
-					SolverManifoldPoint.NetPushOutTangentV * SolverManifoldPoint.WorldContact.ContactTangentV;
+					SolverManifoldPoint.NetPushOutNormal * SolverManifoldPoint.ContactNormal +
+					SolverManifoldPoint.NetPushOutTangentU * SolverManifoldPoint.ContactTangentU +
+					SolverManifoldPoint.NetPushOutTangentV * SolverManifoldPoint.ContactTangentV;
 
 				NetImpulse =
-					SolverManifoldPoint.NetImpulseNormal * SolverManifoldPoint.WorldContact.ContactNormal +
-					SolverManifoldPoint.NetImpulseTangentU * SolverManifoldPoint.WorldContact.ContactTangentU +
-					SolverManifoldPoint.NetImpulseTangentV * SolverManifoldPoint.WorldContact.ContactTangentV;
+					SolverManifoldPoint.NetImpulseNormal * SolverManifoldPoint.ContactNormal +
+					SolverManifoldPoint.NetImpulseTangentU * SolverManifoldPoint.ContactTangentU +
+					SolverManifoldPoint.NetImpulseTangentV * SolverManifoldPoint.ContactTangentV;
 
 				StaticFrictionRatio = SolverManifoldPoint.StaticFrictionRatio;
 
diff --git a/Engine/Source/Runtime/Experimental/Chaos/Private/Chaos/Collision/PBDCollisionSolver.cpp b/Engine/Source/Runtime/Experimental/Chaos/Private/Chaos/Collision/PBDCollisionSolver.cpp
--- a/Engine/Source/Runtime/Experimental/Chaos/Private/Chaos/Collision/PBDCollisionSolver.cpp
+++ b/Engine/Source/Runtime/Experimental/Chaos/Private/Chaos/Collision/PBDCollisionSolver.cpp
@@ -87,10 +87,10 @@ namespace Chaos
 				FPBDCollisionSolverManifoldPoint& SolverManifoldPoint = State.ManifoldPoints[PointIndex];
 				if (ShouldSolveVelocity(SolverManifoldPoint))
 				{
-					RelativeContactPosition0 += SolverManifoldPoint.WorldContact.RelativeContactPoints[0];
-					RelativeContactPosition1 += SolverManifoldPoint.WorldContact.RelativeContactPoints[1];
-					WorldContactVelocityTargetNormal += SolverManifoldPoint.WorldContact.ContactTargetVelocityNormal;
-					WorldContactNormal = SolverManifoldPoint.WorldContact.ContactNormal;	// Take last value - should all be similar
+					RelativeContactPosition0 += SolverManifoldPoint.RelativeContactPoints[0];
+					RelativeContactPosition1 += SolverManifoldPoint.RelativeContactPoints[1];
+					WorldContactVelocityTargetNormal += SolverManifoldPoint.ContactTargetVelocityNormal;
+					WorldContactNormal = SolverManifoldPoint.ContactNormal;	// Take last value - should all be similar
 					NetPushOutNormal += SolverManifoldPoint.NetPushOutNormal;
 					++NumActiveManifoldPoints;
 				}
@@ -107,10 +107,10 @@ namespace Chaos
 				const FSolverReal InvCount = FSolverReal(1) / FSolverReal(NumActiveManifoldPoints);
 
 				FPBDCollisionSolverManifoldPoint AverageManifoldPoint;
-				AverageManifoldPoint.WorldContact.RelativeContactPoints[0] = RelativeContactPosition0 * InvCount;
-				AverageManifoldPoint.WorldContact.RelativeContactPoints[1] = RelativeContactPosition1 * InvCount;
-				AverageManifoldPoint.WorldContact.ContactNormal = WorldContactNormal;
-				AverageManifoldPoint.WorldContact.ContactTargetVelocityNormal = WorldContactVelocityTargetNormal * InvCount;
+				AverageManifoldPoint.RelativeContactPoints[0] = RelativeContactPosition0 * InvCount;
+				AverageManifoldPoint.RelativeContactPoints[1] = RelativeContactPosition1 * InvCount;
+				AverageManifoldPoint.ContactNormal = WorldContactNormal;
+				AverageManifoldPoint.ContactTargetVelocityNormal = WorldContactVelocityTargetNormal * InvCount;
 			
 				// Total pushout (not average) which is correct if the average point is also the centroid but otherwise probably an overestimate.
 				// This is used to limit the possibly-attractive impulse that corrects implicit velocity errors from the PBD solve, but it
@@ -135,19 +135,19 @@ namespace Chaos
 					ContactMassInvNormal += FSolverVec3::DotProduct(R1xN, WorldContactNormalAngular1) + InvM1;
 				}
 				AverageManifoldPoint.ContactMassNormal = (ContactMassInvNormal > FSolverReal(UE_SMALL_NUMBER)) ? FSolverReal(1) / ContactMassInvNormal : FSolverReal(0);
-				AverageManifoldPoint.WorldContactNormalAngular0 = WorldContactNormalAngular0;
-				AverageManifoldPoint.WorldContactNormalAngular1 = WorldContactNormalAngular1;
+				AverageManifoldPoint.ContactNormalAngular0 = WorldContactNormalAngular0;
+				AverageManifoldPoint.ContactNormalAngular1 = WorldContactNormalAngular1;
 
 				// @todo(chaos): we don't use these - maybe do the calculation without an actual manifold point object...
-				AverageManifoldPoint.WorldContact.ContactTangentU = FSolverVec3(0);
-				AverageManifoldPoint.WorldContact.ContactTangentV = FSolverVec3(0);
-				AverageManifoldPoint.WorldContact.ContactDeltaNormal = FSolverReal(0);
-				AverageManifoldPoint.WorldContact.ContactDeltaTangentU = FSolverReal(0);
-				AverageManifoldPoint.WorldContact.ContactDeltaTangentV = FSolverReal(0);
-				AverageManifoldPoint.WorldContactTangentUAngular0 = FSolverVec3(0);
-				AverageManifoldPoint.WorldContactTangentVAngular0 = FSolverVec3(0);
-				AverageManifoldPoint.WorldContactTangentUAngular1 = FSolverVec3(0);
-				AverageManifoldPoint.WorldContactTangentVAngular1 = FSolverVec3(0);
+				AverageManifoldPoint.ContactTangentU = FSolverVec3(0);
+				AverageManifoldPoint.ContactTangentV = FSolverVec3(0);
+				AverageManifoldPoint.ContactDeltaNormal = FSolverReal(0);
+				AverageManifoldPoint.ContactDeltaTangentU = FSolverReal(0);
+				AverageManifoldPoint.ContactDeltaTangentV = FSolverReal(0);
+				AverageManifoldPoint.ContactTangentUAngular0 = FSolverVec3(0);
+				AverageManifoldPoint.ContactTangentVAngular0 = FSolverVec3(0);
+				AverageManifoldPoint.ContactTangentUAngular1 = FSolverVec3(0);
+				AverageManifoldPoint.ContactTangentVAngular1 = FSolverVec3(0);
 				AverageManifoldPoint.ContactMassTangentU = FSolverReal(0);
 				AverageManifoldPoint.ContactMassTangentV = FSolverReal(0);
 				AverageManifoldPoint.NetPushOutTangentU = FSolverReal(0);
diff --git a/Engine/Source/Runtime/Experimental/Chaos/Private/Chaos/PBDCollisionConstraints.cpp b/Engine/Source/Runtime/Experimental/Chaos/Private/Chaos/PBDCollisionConstraints.cpp
--- a/Engine/Source/Runtime/Experimental/Chaos/Private/Chaos/PBDCollisionConstraints.cpp
+++ b/Engine/Source/Runtime/Experimental/Chaos/Private/Chaos/PBDCollisionConstraints.cpp
@@ -133,27 +133,27 @@ namespace Chaos
 				const Private::TPBDCollisionSolverManifoldPointsSimd<4>& ManifoldPointSolverB = ConstraintSolverB.GetManifoldPoint(ManifoldPointIndex, SolverB.GetManifoldPointBuffer());
 				bool bIsError = false;
 
-				if (ManifoldPointSolverA.WorldContact.ContactNormal != ManifoldPointSolverB.SimdContactNormal.GetValue(LaneIndexB))
+				if (ManifoldPointSolverA.ContactNormal != ManifoldPointSolverB.SimdContactNormal.GetValue(LaneIndexB))
 				{
 					UE_LOG(LogChaos, Warning, TEXT("ContactNormal mismatch"));
 					bIsError = true;
 				}
-				if (ManifoldPointSolverA.WorldContact.ContactTangentU != ManifoldPointSolverB.SimdContactTangentU.GetValue(LaneIndexB))
+				if (ManifoldPointSolverA.ContactTangentU != ManifoldPointSolverB.SimdContactTangentU.GetValue(LaneIndexB))
 				{
 					UE_LOG(LogChaos, Warning, TEXT("ContactTangentU mismatch"));
 					bIsError = true;
 				}
-				if (ManifoldPointSolverA.WorldContact.ContactTangentV != ManifoldPointSolverB.SimdContactTangentV.GetValue(LaneIndexB))
+				if (ManifoldPointSolverA.ContactTangentV != ManifoldPointSolverB.SimdContactTangentV.GetValue(LaneIndexB))
 				{
 					UE_LOG(LogChaos, Warning, TEXT("ContactTangentV mismatch"));
 					bIsError = true;
 				}
-				if (ManifoldPointSolverA.WorldContact.RelativeContactPoints[0] != ManifoldPointSolverB.SimdRelativeContactPoint0.GetValue(LaneIndexB))
+				if (ManifoldPointSolverA.RelativeContactPoints[0] != ManifoldPointSolverB.SimdRelativeContactPoint0.GetValue(LaneIndexB))
 				{
 					UE_LOG(LogChaos, Warning, TEXT("RelativeContactPoints mismatch"));
 					bIsError = true;
 				}
-				if (ManifoldPointSolverA.WorldContact.RelativeContactPoints[1] != ManifoldPointSolverB.SimdRelativeContactPoint1.GetValue(LaneIndexB))
+				if (ManifoldPointSolverA.RelativeContactPoints[1] != ManifoldPointSolverB.SimdRelativeContactPoint1.GetValue(LaneIndexB))
 				{
 					UE_LOG(LogChaos, Warning, TEXT("RelativeContactPoints mismatch"));
 					bIsError = true;
@@ -175,12 +175,12 @@ namespace Chaos
 				}
 				if (Body0.IsDynamic())
 				{
-					if (ManifoldPointSolverA.WorldContactTangentUAngular0 != ManifoldPointSolverB.SimdContactTangentUAngular0.GetValue(LaneIndexB))
+					if (ManifoldPointSolverA.ContactTangentUAngular0 != ManifoldPointSolverB.SimdContactTangentUAngular0.GetValue(LaneIndexB))
 					{
 						UE_LOG(LogChaos, Warning, TEXT("ContactTangentUAngular0 mismatch"));
 						bIsError = true;
 					}
-					if (ManifoldPointSolverA.WorldContactTangentVAngular0 != ManifoldPointSolverB.SimdContactTangentVAngular0.GetValue(LaneIndexB))
+					if (ManifoldPointSolverA.ContactTangentVAngular0 != ManifoldPointSolverB.SimdContactTangentVAngular0.GetValue(LaneIndexB))
 					{
 						UE_LOG(LogChaos, Warning, TEXT("ContactTangentVAngular0 mismatch"));
 						bIsError = true;
@@ -188,12 +188,12 @@ namespace Chaos
 				}
 				if (Body1.IsDynamic())
 				{
-					if (ManifoldPointSolverA.WorldContactTangentUAngular1 != ManifoldPointSolverB.SimdContactTangentUAngular1.GetValue(LaneIndexB))
+					if (ManifoldPointSolverA.ContactTangentUAngular1 != ManifoldPointSolverB.SimdContactTangentUAngular1.GetValue(LaneIndexB))
 					{
 						UE_LOG(LogChaos, Warning, TEXT("ContactTangentUAngular1 mismatch"));
 						bIsError = true;
 					}
-					if (ManifoldPointSolverA.WorldContactTangentVAngular1 != ManifoldPointSolverB.SimdContactTangentVAngular1.GetValue(LaneIndexB))
+					if (ManifoldPointSolverA.ContactTangentVAngular1 != ManifoldPointSolverB.SimdContactTangentVAngular1.GetValue(LaneIndexB))
 					{
 						UE_LOG(LogChaos, Warning, TEXT("ContactTangentVAngular1 mismatch"));
 						bIsError = true;
diff --git a/Engine/Source/Runtime/Experimental/Chaos/Private/Chaos/PBDSuspensionConstraints.cpp b/Engine/Source/Runtime/Experimental/Chaos/Private/Chaos/PBDSuspensionConstraints.cpp
--- a/Engine/Source/Runtime/Experimental/Chaos/Private/Chaos/PBDSuspensionConstraints.cpp
+++ b/Engine/Source/Runtime/Experimental/Chaos/Private/Chaos/PBDSuspensionConstraints.cpp
@@ -300,8 +300,8 @@ namespace Chaos
 		if ((CollisionSolvers.Num() > 0) && (CollisionSolvers[ConstraintIndex].NumManifoldPoints() > 0))
 		{
 			const FPBDCollisionSolverManifoldPoint& ManifoldPoint = CollisionSolvers[ConstraintIndex].GetManifoldPoint(0);
-			ConstraintResults[ConstraintIndex].HardStopNetPushOut = ManifoldPoint.NetPushOutNormal * ManifoldPoint.WorldContact.ContactNormal;
-			ConstraintResults[ConstraintIndex].HardStopNetImpulse = ManifoldPoint.NetImpulseNormal * ManifoldPoint.WorldContact.ContactNormal;
+			ConstraintResults[ConstraintIndex].HardStopNetPushOut = ManifoldPoint.NetPushOutNormal * ManifoldPoint.ContactNormal;
+			ConstraintResults[ConstraintIndex].HardStopNetImpulse = ManifoldPoint.NetImpulseNormal * ManifoldPoint.ContactNormal;
 		}
 
 		ConstraintSolverBodies[ConstraintIndex] = nullptr;
diff --git a/Engine/Source/Runtime/Experimental/Chaos/Public/Chaos/Collision/PBDCollisionSolver.h b/Engine/Source/Runtime/Experimental/Chaos/Public/Chaos/Collision/PBDCollisionSolver.h
--- a/Engine/Source/Runtime/Experimental/Chaos/Public/Chaos/Collision/PBDCollisionSolver.h
+++ b/Engine/Source/Runtime/Experimental/Chaos/Public/Chaos/Collision/PBDCollisionSolver.h
@@ -32,16 +32,29 @@ namespace Chaos
 		class FPBDCollisionSolverManifoldPoint
 		{
 		public:
-			// World-space contact data
-			FWorldContactPoint WorldContact;
+			// World-space contact point relative to each particle's center of mass
+			FVec3f RelativeContactPoints[2];
+
+			// World-space contact normal and tangents
+			FVec3f ContactNormal;
+			FVec3f ContactTangentU;
+			FVec3f ContactTangentV;
+
+			// Errors to correct along each of the contact axes
+			FRealSingle ContactDeltaNormal;
+			FRealSingle ContactDeltaTangentU;
+			FRealSingle ContactDeltaTangentV;
+
+			// Target velocity along the normal direction
+			FRealSingle ContactTargetVelocityNormal;
 
 			// I^-1.(R x A) for each body where A is each axis (Normal, TangentU, TangentV)
-			FSolverVec3 WorldContactNormalAngular0;
-			FSolverVec3 WorldContactTangentUAngular0;
-			FSolverVec3 WorldContactTangentVAngular0;
-			FSolverVec3 WorldContactNormalAngular1;
-			FSolverVec3 WorldContactTangentUAngular1;
-			FSolverVec3 WorldContactTangentVAngular1;
+			FSolverVec3 ContactNormalAngular0;
+			FSolverVec3 ContactTangentUAngular0;
+			FSolverVec3 ContactTangentVAngular0;
+			FSolverVec3 ContactNormalAngular1;
+			FSolverVec3 ContactTangentUAngular1;
+			FSolverVec3 ContactTangentVAngular1;
 
 			// Contact mass (for non-friction)
 			FSolverReal ContactMassNormal;
@@ -504,11 +517,11 @@ namespace Chaos
 				ManifoldPoint.StaticFrictionRatio);					// Out
 
 			// Update the particle state based on the pushout
-			const FVec3 PushOut = PushOutTangentU * ManifoldPoint.WorldContact.ContactTangentU + PushOutTangentV * ManifoldPoint.WorldContact.ContactTangentV;
+			const FVec3 PushOut = PushOutTangentU * ManifoldPoint.ContactTangentU + PushOutTangentV * ManifoldPoint.ContactTangentV;
 			if (IsDynamic(0))
 			{
 				const FSolverVec3 DX0 = State.InvMs[0] * PushOut;
-				const FSolverVec3 DR0 = ManifoldPoint.WorldContactTangentUAngular0 * PushOutTangentU + ManifoldPoint.WorldContactTangentVAngular0 * PushOutTangentV;
+				const FSolverVec3 DR0 = ManifoldPoint.ContactTangentUAngular0 * PushOutTangentU + ManifoldPoint.ContactTangentVAngular0 * PushOutTangentV;
 				FConstraintSolverBody& Body0 = SolverBody0();
 				Body0.ApplyPositionDelta(DX0);
 				Body0.ApplyRotationDelta(DR0);
@@ -516,7 +529,7 @@ namespace Chaos
 			if (IsDynamic(1))
 			{
 				const FSolverVec3 DX1 = -State.InvMs[1]* PushOut;
-				const FSolverVec3 DR1 = ManifoldPoint.WorldContactTangentUAngular1 * -PushOutTangentU + ManifoldPoint.WorldContactTangentVAngular1 * -PushOutTangentV;
+				const FSolverVec3 DR1 = ManifoldPoint.ContactTangentUAngular1 * -PushOutTangentU + ManifoldPoint.ContactTangentVAngular1 * -PushOutTangentV;
 				FConstraintSolverBody& Body1 = SolverBody1();
 				Body1.ApplyPositionDelta(DX1);
 				Body1.ApplyRotationDelta(DR1);
@@ -542,16 +555,16 @@ namespace Chaos
 			// Update the particle state based on the pushout
 			if (IsDynamic(0))
 			{
-				const FSolverVec3 DX0 = (State.InvMs[0] * PushOutNormal) * ManifoldPoint.WorldContact.ContactNormal;
-				const FSolverVec3 DR0 = ManifoldPoint.WorldContactNormalAngular0 * PushOutNormal;
+				const FSolverVec3 DX0 = (State.InvMs[0] * PushOutNormal) * ManifoldPoint.ContactNormal;
+				const FSolverVec3 DR0 = ManifoldPoint.ContactNormalAngular0 * PushOutNormal;
 				FConstraintSolverBody& Body0 = SolverBody0();
 				Body0.ApplyPositionDelta(DX0);
 				Body0.ApplyRotationDelta(DR0);
 			}
 			if (IsDynamic(1))
 			{
-				const FSolverVec3 DX1 = (State.InvMs[1] * -PushOutNormal) * ManifoldPoint.WorldContact.ContactNormal;
-				const FSolverVec3 DR1 = ManifoldPoint.WorldContactNormalAngular1 * -PushOutNormal;
+				const FSolverVec3 DX1 = (State.InvMs[1] * -PushOutNormal) * ManifoldPoint.ContactNormal;
+				const FSolverVec3 DR1 = ManifoldPoint.ContactNormalAngular1 * -PushOutNormal;
 				FConstraintSolverBody& Body1 = SolverBody1();
 				Body1.ApplyPositionDelta(DX1);
 				Body1.ApplyRotationDelta(DR1);
@@ -601,18 +614,18 @@ namespace Chaos
 			ManifoldPoint.NetImpulseTangentV += ImpulseTangentV;
 
 			// Apply the velocity deltas from the impulse
-			const FSolverVec3 Impulse = ImpulseNormal * ManifoldPoint.WorldContact.ContactNormal + ImpulseTangentU * ManifoldPoint.WorldContact.ContactTangentU + ImpulseTangentV * ManifoldPoint.WorldContact.ContactTangentV;
+			const FSolverVec3 Impulse = ImpulseNormal * ManifoldPoint.ContactNormal + ImpulseTangentU * ManifoldPoint.ContactTangentU + ImpulseTangentV * ManifoldPoint.ContactTangentV;
 			if (IsDynamic(0))
 			{
 				const FSolverVec3 DV0 = State.InvMs[0] * Impulse;
-				const FSolverVec3 DW0 = ManifoldPoint.WorldContactNormalAngular0 * ImpulseNormal + ManifoldPoint.WorldContactTangentUAngular0 * ImpulseTangentU + ManifoldPoint.WorldContactTangentVAngular0 * ImpulseTangentV;
+				const FSolverVec3 DW0 = ManifoldPoint.ContactNormalAngular0 * ImpulseNormal + ManifoldPoint.ContactTangentUAngular0 * ImpulseTangentU + ManifoldPoint.ContactTangentVAngular0 * ImpulseTangentV;
 				FConstraintSolverBody& Body0 = SolverBody0();
 				Body0.ApplyVelocityDelta(DV0, DW0);
 			}
 			if (IsDynamic(1))
 			{
 				const FSolverVec3 DV1 = -State.InvMs[1]* Impulse;
-				const FSolverVec3 DW1 = ManifoldPoint.WorldContactNormalAngular1 * -ImpulseNormal + ManifoldPoint.WorldContactTangentUAngular1 * -ImpulseTangentU + ManifoldPoint.WorldContactTangentVAngular1 * -ImpulseTangentV;
+				const FSolverVec3 DW1 = ManifoldPoint.ContactNormalAngular1 * -ImpulseNormal + ManifoldPoint.ContactTangentUAngular1 * -ImpulseTangentU + ManifoldPoint.ContactTangentVAngular1 * -ImpulseTangentV;
 				FConstraintSolverBody& Body1 = SolverBody1();
 				Body1.ApplyVelocityDelta(DV1, DW1);
 			}
@@ -635,18 +648,18 @@ namespace Chaos
 			ManifoldPoint.NetImpulseNormal += ImpulseNormal;
 
 			// Calculate the velocity deltas from the impulse
-			FSolverVec3 Impulse = ImpulseNormal * ManifoldPoint.WorldContact.ContactNormal;
+			FSolverVec3 Impulse = ImpulseNormal * ManifoldPoint.ContactNormal;
 			if (IsDynamic(0))
 			{
 				const FSolverVec3 DV0 = State.InvMs[0] * Impulse;
-				const FSolverVec3 DW0 = ManifoldPoint.WorldContactNormalAngular0 * ImpulseNormal;
+				const FSolverVec3 DW0 = ManifoldPoint.ContactNormalAngular0 * ImpulseNormal;
 				FConstraintSolverBody& Body0 = SolverBody0();
 				Body0.ApplyVelocityDelta(DV0, DW0);
 			}
 			if (IsDynamic(1))
 			{
 				const FSolverVec3 DV1 = -State.InvMs[1] * Impulse;
-				const FSolverVec3 DW1 = ManifoldPoint.WorldContactNormalAngular1 * -ImpulseNormal;
+				const FSolverVec3 DW1 = ManifoldPoint.ContactNormalAngular1 * -ImpulseNormal;
 				FConstraintSolverBody& Body1 = SolverBody1();
 				Body1.ApplyVelocityDelta(DV1, DW1);
 			}
@@ -673,40 +686,40 @@ namespace Chaos
 				FSolverReal ContactMassInvTangentV = FSolverReal(0);
 
 				// These are not used if not initialized below so no need to clear
-				//WorldContactNormalAngular0 = FSolverVec3(0);
-				//WorldContactTangentUAngular0 = FSolverVec3(0);
-				//WorldContactTangentVAngular0 = FSolverVec3(0);
-				//WorldContactNormalAngular1 = FSolverVec3(0);
-				//WorldContactTangentUAngular1 = FSolverVec3(0);
-				//WorldContactTangentVAngular1 = FSolverVec3(0);
+				//ContactNormalAngular0 = FSolverVec3(0);
+				//ContactTangentUAngular0 = FSolverVec3(0);
+				//ContactTangentVAngular0 = FSolverVec3(0);
+				//ContactNormalAngular1 = FSolverVec3(0);
+				//ContactTangentUAngular1 = FSolverVec3(0);
+				//ContactTangentVAngular1 = FSolverVec3(0);
 
 				if (IsDynamic(0))
 				{
-					const FSolverVec3 R0xN = FSolverVec3::CrossProduct(ManifoldPoint.WorldContact.RelativeContactPoints[0], ManifoldPoint.WorldContact.ContactNormal);
-					const FSolverVec3 R0xU = FSolverVec3::CrossProduct(ManifoldPoint.WorldContact.RelativeContactPoints[0], ManifoldPoint.WorldContact.ContactTangentU);
-					const FSolverVec3 R0xV = FSolverVec3::CrossProduct(ManifoldPoint.WorldContact.RelativeContactPoints[0], ManifoldPoint.WorldContact.ContactTangentV);
+					const FSolverVec3 R0xN = FSolverVec3::CrossProduct(ManifoldPoint.RelativeContactPoints[0], ManifoldPoint.ContactNormal);
+					const FSolverVec3 R0xU = FSolverVec3::CrossProduct(ManifoldPoint.RelativeContactPoints[0], ManifoldPoint.ContactTangentU);
+					const FSolverVec3 R0xV = FSolverVec3::CrossProduct(ManifoldPoint.RelativeContactPoints[0], ManifoldPoint.ContactTangentV);
 
-					ManifoldPoint.WorldContactNormalAngular0 = InvI0 * R0xN;
-					ManifoldPoint.WorldContactTangentUAngular0 = InvI0 * R0xU;
-					ManifoldPoint.WorldContactTangentVAngular0 = InvI0 * R0xV;
+					ManifoldPoint.ContactNormalAngular0 = InvI0 * R0xN;
+					ManifoldPoint.ContactTangentUAngular0 = InvI0 * R0xU;
+					ManifoldPoint.ContactTangentVAngular0 = InvI0 * R0xV;
 
-					ContactMassInvNormal += FSolverVec3::DotProduct(R0xN, ManifoldPoint.WorldContactNormalAngular0) + InvM0;
-					ContactMassInvTangentU += FSolverVec3::DotProduct(R0xU, ManifoldPoint.WorldContactTangentUAngular0) + InvM0;
-					ContactMassInvTangentV += FSolverVec3::DotProduct(R0xV, ManifoldPoint.WorldContactTangentVAngular0) + InvM0;
+					ContactMassInvNormal += FSolverVec3::DotProduct(R0xN, ManifoldPoint.ContactNormalAngular0) + InvM0;
+					ContactMassInvTangentU += FSolverVec3::DotProduct(R0xU, ManifoldPoint.ContactTangentUAngular0) + InvM0;
+					ContactMassInvTangentV += FSolverVec3::DotProduct(R0xV, ManifoldPoint.ContactTangentVAngular0) + InvM0;
 				}
 				if (IsDynamic(1))
 				{
-					const FSolverVec3 R1xN = FSolverVec3::CrossProduct(ManifoldPoint.WorldContact.RelativeContactPoints[1], ManifoldPoint.WorldContact.ContactNormal);
-					const FSolverVec3 R1xU = FSolverVec3::CrossProduct(ManifoldPoint.WorldContact.RelativeContactPoints[1], ManifoldPoint.WorldContact.ContactTangentU);
-					const FSolverVec3 R1xV = FSolverVec3::CrossProduct(ManifoldPoint.WorldContact.RelativeContactPoints[1], ManifoldPoint.WorldContact.ContactTangentV);
+					const FSolverVec3 R1xN = FSolverVec3::CrossProduct(ManifoldPoint.RelativeContactPoints[1], ManifoldPoint.ContactNormal);
+					const FSolverVec3 R1xU = FSolverVec3::CrossProduct(ManifoldPoint.RelativeContactPoints[1], ManifoldPoint.ContactTangentU);
+					const FSolverVec3 R1xV = FSolverVec3::CrossProduct(ManifoldPoint.RelativeContactPoints[1], ManifoldPoint.ContactTangentV);
 
-					ManifoldPoint.WorldContactNormalAngular1 = InvI1 * R1xN;
-					ManifoldPoint.WorldContactTangentUAngular1 = InvI1 * R1xU;
-					ManifoldPoint.WorldContactTangentVAngular1 = InvI1 * R1xV;
+					ManifoldPoint.ContactNormalAngular1 = InvI1 * R1xN;
+					ManifoldPoint.ContactTangentUAngular1 = InvI1 * R1xU;
+					ManifoldPoint.ContactTangentVAngular1 = InvI1 * R1xV;
 
-					ContactMassInvNormal += FSolverVec3::DotProduct(R1xN, ManifoldPoint.WorldContactNormalAngular1) + InvM1;
-					ContactMassInvTangentU += FSolverVec3::DotProduct(R1xU, ManifoldPoint.WorldContactTangentUAngular1) + InvM1;
-					ContactMassInvTangentV += FSolverVec3::DotProduct(R1xV, ManifoldPoint.WorldContactTangentVAngular1) + InvM1;
+					ContactMassInvNormal += FSolverVec3::DotProduct(R1xN, ManifoldPoint.ContactNormalAngular1) + InvM1;
+					ContactMassInvTangentU += FSolverVec3::DotProduct(R1xU, ManifoldPoint.ContactTangentUAngular1) + InvM1;
+					ContactMassInvTangentV += FSolverVec3::DotProduct(R1xV, ManifoldPoint.ContactTangentVAngular1) + InvM1;
 				}
 
 				ManifoldPoint.ContactMassNormal = (ContactMassInvNormal > FSolverReal(UE_SMALL_NUMBER)) ? FSolverReal(1) / ContactMassInvNormal : FSolverReal(0);
@@ -734,15 +747,15 @@ namespace Chaos
 				FSolverReal ContactMassInvNormal = FSolverReal(0);
 				if (IsDynamic(0))
 				{
-					const FSolverVec3 R0xN = FSolverVec3::CrossProduct(ManifoldPoint.WorldContact.RelativeContactPoints[0], ManifoldPoint.WorldContact.ContactNormal);
-					ManifoldPoint.WorldContactNormalAngular0 = InvI0 * R0xN;
-					ContactMassInvNormal += FSolverVec3::DotProduct(R0xN, ManifoldPoint.WorldContactNormalAngular0) + InvM0;
+					const FSolverVec3 R0xN = FSolverVec3::CrossProduct(ManifoldPoint.RelativeContactPoints[0], ManifoldPoint.ContactNormal);
+					ManifoldPoint.ContactNormalAngular0 = InvI0 * R0xN;
+					ContactMassInvNormal += FSolverVec3::DotProduct(R0xN, ManifoldPoint.ContactNormalAngular0) + InvM0;
 				}
 				if (IsDynamic(1))
 				{
-					const FSolverVec3 R1xN = FSolverVec3::CrossProduct(ManifoldPoint.WorldContact.RelativeContactPoints[1], ManifoldPoint.WorldContact.ContactNormal);
-					ManifoldPoint.WorldContactNormalAngular1 = InvI1 * R1xN;
-					ContactMassInvNormal += FSolverVec3::DotProduct(R1xN, ManifoldPoint.WorldContactNormalAngular1) + InvM1;
+					const FSolverVec3 R1xN = FSolverVec3::CrossProduct(ManifoldPoint.RelativeContactPoints[1], ManifoldPoint.ContactNormal);
+					ManifoldPoint.ContactNormalAngular1 = InvI1 * R1xN;
+					ContactMassInvNormal += FSolverVec3::DotProduct(R1xN, ManifoldPoint.ContactNormalAngular1) + InvM1;
 				}
 				ManifoldPoint.ContactMassNormal = (ContactMassInvNormal > FSolverReal(UE_SMALL_NUMBER)) ? FSolverReal(1) / ContactMassInvNormal : FSolverReal(0);
 			}
@@ -757,10 +770,10 @@ namespace Chaos
 			const FConstraintSolverBody& Body1 = SolverBody1();
 
 			// Linear version: calculate the contact delta assuming linear motion after applying a positional impulse at the contact point. There will be an error that depends on the size of the rotation.
-			const FSolverVec3 ContactDelta0 = Body0.DP() + FSolverVec3::CrossProduct(Body0.DQ(), ManifoldPoint.WorldContact.RelativeContactPoints[0]);
-			const FSolverVec3 ContactDelta1 = Body1.DP() + FSolverVec3::CrossProduct(Body1.DQ(), ManifoldPoint.WorldContact.RelativeContactPoints[1]);
+			const FSolverVec3 ContactDelta0 = Body0.DP() + FSolverVec3::CrossProduct(Body0.DQ(), ManifoldPoint.RelativeContactPoints[0]);
+			const FSolverVec3 ContactDelta1 = Body1.DP() + FSolverVec3::CrossProduct(Body1.DQ(), ManifoldPoint.RelativeContactPoints[1]);
 			const FSolverVec3 ContactDelta = ContactDelta0 - ContactDelta1;
-			OutContactDeltaNormal = ManifoldPoint.WorldContact.ContactDeltaNormal + FSolverVec3::DotProduct(ContactDelta, ManifoldPoint.WorldContact.ContactNormal);
+			OutContactDeltaNormal = ManifoldPoint.ContactDeltaNormal + FSolverVec3::DotProduct(ContactDelta, ManifoldPoint.ContactNormal);
 
 			// NOTE: OutContactDeltaNormal is negative for penetration
 			// NOTE: MaxPushOut == 0 disables the pushout limits
@@ -779,11 +792,11 @@ namespace Chaos
 			const FConstraintSolverBody& Body1 = SolverBody1();
 
 			// Linear version: calculate the contact delta assuming linear motion after applying a positional impulse at the contact point. There will be an error that depends on the size of the rotation.
-			const FSolverVec3 ContactDelta0 = Body0.DP() + FSolverVec3::CrossProduct(Body0.DQ(), ManifoldPoint.WorldContact.RelativeContactPoints[0]);
-			const FSolverVec3 ContactDelta1 = Body1.DP() + FSolverVec3::CrossProduct(Body1.DQ(), ManifoldPoint.WorldContact.RelativeContactPoints[1]);
+			const FSolverVec3 ContactDelta0 = Body0.DP() + FSolverVec3::CrossProduct(Body0.DQ(), ManifoldPoint.RelativeContactPoints[0]);
+			const FSolverVec3 ContactDelta1 = Body1.DP() + FSolverVec3::CrossProduct(Body1.DQ(), ManifoldPoint.RelativeContactPoints[1]);
 			const FSolverVec3 ContactDelta = ContactDelta0 - ContactDelta1;
-			OutContactDeltaTangentU = ManifoldPoint.WorldContact.ContactDeltaTangentU + FSolverVec3::DotProduct(ContactDelta, ManifoldPoint.WorldContact.ContactTangentU);
-			OutContactDeltaTangentV = ManifoldPoint.WorldContact.ContactDeltaTangentV + FSolverVec3::DotProduct(ContactDelta, ManifoldPoint.WorldContact.ContactTangentV);
+			OutContactDeltaTangentU = ManifoldPoint.ContactDeltaTangentU + FSolverVec3::DotProduct(ContactDelta, ManifoldPoint.ContactTangentU);
+			OutContactDeltaTangentV = ManifoldPoint.ContactDeltaTangentV + FSolverVec3::DotProduct(ContactDelta, ManifoldPoint.ContactTangentV);
 		}
 
 		FORCEINLINE_DEBUGGABLE void FPBDCollisionSolver::CalculateContactVelocityError(
@@ -797,14 +810,14 @@ namespace Chaos
 			const FConstraintSolverBody& Body0 = SolverBody0();
 			const FConstraintSolverBody& Body1 = SolverBody1();
 
-			const FSolverVec3 ContactVelocity0 = Body0.V() + FSolverVec3::CrossProduct(Body0.W(), ManifoldPoint.WorldContact.RelativeContactPoints[0]);
-			const FSolverVec3 ContactVelocity1 = Body1.V() + FSolverVec3::CrossProduct(Body1.W(), ManifoldPoint.WorldContact.RelativeContactPoints[1]);
+			const FSolverVec3 ContactVelocity0 = Body0.V() + FSolverVec3::CrossProduct(Body0.W(), ManifoldPoint.RelativeContactPoints[0]);
+			const FSolverVec3 ContactVelocity1 = Body1.V() + FSolverVec3::CrossProduct(Body1.W(), ManifoldPoint.RelativeContactPoints[1]);
 			const FSolverVec3 ContactVelocity = ContactVelocity0 - ContactVelocity1;
-			const FSolverReal ContactVelocityNormal = FSolverVec3::DotProduct(ContactVelocity, ManifoldPoint.WorldContact.ContactNormal);
-			const FSolverReal ContactVelocityTangent0 = FSolverVec3::DotProduct(ContactVelocity, ManifoldPoint.WorldContact.ContactTangentU);
-			const FSolverReal ContactVelocityTangent1 = FSolverVec3::DotProduct(ContactVelocity, ManifoldPoint.WorldContact.ContactTangentV);
+			const FSolverReal ContactVelocityNormal = FSolverVec3::DotProduct(ContactVelocity, ManifoldPoint.ContactNormal);
+			const FSolverReal ContactVelocityTangent0 = FSolverVec3::DotProduct(ContactVelocity, ManifoldPoint.ContactTangentU);
+			const FSolverReal ContactVelocityTangent1 = FSolverVec3::DotProduct(ContactVelocity, ManifoldPoint.ContactTangentV);
 
-			OutContactVelocityDeltaNormal = (ContactVelocityNormal - ManifoldPoint.WorldContact.ContactTargetVelocityNormal);
+			OutContactVelocityDeltaNormal = (ContactVelocityNormal - ManifoldPoint.ContactTargetVelocityNormal);
 			OutContactVelocityDeltaTangent0 = ContactVelocityTangent0;
 			OutContactVelocityDeltaTangent1 = ContactVelocityTangent1;
 		}
@@ -816,20 +829,20 @@ namespace Chaos
 			const FConstraintSolverBody& Body0 = SolverBody0();
 			const FConstraintSolverBody& Body1 = SolverBody1();
 
-			const FSolverVec3 ContactVelocity0 = Body0.V() + FSolverVec3::CrossProduct(Body0.W(), ManifoldPoint.WorldContact.RelativeContactPoints[0]);
-			const FSolverVec3 ContactVelocity1 = Body1.V() + FSolverVec3::CrossProduct(Body1.W(), ManifoldPoint.WorldContact.RelativeContactPoints[1]);
+			const FSolverVec3 ContactVelocity0 = Body0.V() + FSolverVec3::CrossProduct(Body0.W(), ManifoldPoint.RelativeContactPoints[0]);
+			const FSolverVec3 ContactVelocity1 = Body1.V() + FSolverVec3::CrossProduct(Body1.W(), ManifoldPoint.RelativeContactPoints[1]);
 			const FSolverVec3 ContactVelocity = ContactVelocity0 - ContactVelocity1;
-			const FSolverReal ContactVelocityNormal = FSolverVec3::DotProduct(ContactVelocity, ManifoldPoint.WorldContact.ContactNormal);
+			const FSolverReal ContactVelocityNormal = FSolverVec3::DotProduct(ContactVelocity, ManifoldPoint.ContactNormal);
 
 			// Add up the errors in the velocity (current velocity - desired velocity)
-			OutContactVelocityDeltaNormal = (ContactVelocityNormal - ManifoldPoint.WorldContact.ContactTargetVelocityNormal);
+			OutContactVelocityDeltaNormal = (ContactVelocityNormal - ManifoldPoint.ContactTargetVelocityNormal);
 		}
 
 		FORCEINLINE_DEBUGGABLE bool FPBDCollisionSolver::ShouldSolveVelocity(
 			const FPBDCollisionSolverManifoldPoint& ManifoldPoint) const
 		{
 			// We ensure positive separating velocity for close contacts even if they didn't receive a pushout
-			return (ManifoldPoint.NetPushOutNormal > FSolverReal(0)) || (ManifoldPoint.WorldContact.ContactDeltaNormal < FSolverReal(0));
+			return (ManifoldPoint.NetPushOutNormal > FSolverReal(0)) || (ManifoldPoint.ContactDeltaNormal < FSolverReal(0));
 		}
 
 		FORCEINLINE_DEBUGGABLE void FPBDCollisionSolver::InitManifoldPoint(
@@ -847,15 +860,15 @@ namespace Chaos
 		{
 			FPBDCollisionSolverManifoldPoint& ManifoldPoint = State.ManifoldPoints[PointIndex];
 
-			ManifoldPoint.WorldContact.RelativeContactPoints[0] = InRelativeContactPosition0;
-			ManifoldPoint.WorldContact.RelativeContactPoints[1] = InRelativeContactPosition1;
-			ManifoldPoint.WorldContact.ContactNormal = InWorldContactNormal;
-			ManifoldPoint.WorldContact.ContactTangentU = InWorldContactTangentU;
-			ManifoldPoint.WorldContact.ContactTangentV = InWorldContactTangentV;
-			ManifoldPoint.WorldContact.ContactDeltaNormal = InWorldContactDeltaNormal;
-			ManifoldPoint.WorldContact.ContactDeltaTangentU = InWorldContactDeltaTangentU;
-			ManifoldPoint.WorldContact.ContactDeltaTangentV = InWorldContactDeltaTangentV;
-			ManifoldPoint.WorldContact.ContactTargetVelocityNormal = InWorldContactVelocityTargetNormal;
+			ManifoldPoint.RelativeContactPoints[0] = InRelativeContactPosition0;
+			ManifoldPoint.RelativeContactPoints[1] = InRelativeContactPosition1;
+			ManifoldPoint.ContactNormal = InWorldContactNormal;
+			ManifoldPoint.ContactTangentU = InWorldContactTangentU;
+			ManifoldPoint.ContactTangentV = InWorldContactTangentV;
+			ManifoldPoint.ContactDeltaNormal = InWorldContactDeltaNormal;
+			ManifoldPoint.ContactDeltaTangentU = InWorldContactDeltaTangentU;
+			ManifoldPoint.ContactDeltaTangentV = InWorldContactDeltaTangentV;
+			ManifoldPoint.ContactTargetVelocityNormal = InWorldContactVelocityTargetNormal;
 
 			ManifoldPoint.NetPushOutNormal = FSolverReal(0);
 			ManifoldPoint.NetPushOutTangentU = FSolverReal(0);
