Disable sleeping event handling in geometry collection event handler. This avoids multiple dispatches of identical sleeping events as they are handled by the scene.

diff --git a/Engine/Source/Runtime/Experimental/ChaosSolverEngine/Private/Chaos/ChaosGameplayEventDispatcher.cpp b/Engine/Source/Runtime/Experimental/ChaosSolverEngine/Private/Chaos/ChaosGameplayEventDispatcher.cpp
--- a/Engine/Source/Runtime/Experimental/ChaosSolverEngine/Private/Chaos/ChaosGameplayEventDispatcher.cpp
+++ b/Engine/Source/Runtime/Experimental/ChaosSolverEngine/Private/Chaos/ChaosGameplayEventDispatcher.cpp
@@ -322,7 +322,6 @@ void UChaosGameplayEventDispatcher::RegisterChaosEvents()
 			Chaos::FEventManager* EventManager = Solver->GetEventManager();
 			EventManager->RegisterHandler<Chaos::FCollisionEventData>(Chaos::EEventType::Collision, this, &UChaosGameplayEventDispatcher::HandleCollisionEvents);
 			EventManager->RegisterHandler<Chaos::FBreakingEventData>(Chaos::EEventType::Breaking, this, &UChaosGameplayEventDispatcher::HandleBreakingEvents);
-			EventManager->RegisterHandler<Chaos::FSleepingEventData>(Chaos::EEventType::Sleeping, this, &UChaosGameplayEventDispatcher::HandleSleepingEvents);
 			EventManager->RegisterHandler<Chaos::FRemovalEventData>(Chaos::EEventType::Removal, this, &UChaosGameplayEventDispatcher::HandleRemovalEvents);
 			EventManager->RegisterHandler<Chaos::FCrumblingEventData>(Chaos::EEventType::Crumbling, this, &UChaosGameplayEventDispatcher::HandleCrumblingEvents);
 		}
