Add support for setting custom renderer class on geom collection asset or component.
Deprecate oldd ISMPool setting so that it uses UGeometryCollectionISMPoolRenderer which has old behavior.

diff --git a/Engine/Source/Runtime/Experimental/GeometryCollectionEngine/Private/GeometryCollection/GeometryCollectionComponent.cpp b/Engine/Source/Runtime/Experimental/GeometryCollectionEngine/Private/GeometryCollection/GeometryCollectionComponent.cpp
--- a/Engine/Source/Runtime/Experimental/GeometryCollectionEngine/Private/GeometryCollection/GeometryCollectionComponent.cpp
+++ b/Engine/Source/Runtime/Experimental/GeometryCollectionEngine/Private/GeometryCollection/GeometryCollectionComponent.cpp
@@ -30,7 +30,6 @@
 #include "GeometryCollection/GeometryCollectionISMPoolActor.h"
 #include "GeometryCollection/GeometryCollectionISMPoolComponent.h"
 #include "GeometryCollection/GeometryCollectionISMPoolRenderer.h"
-#include "GeometryCollection/GeometryCollectionISMPoolSubSystem.h"
 #include "Math/Sphere.h"
 #include "Modules/ModuleManager.h"
 #include "Net/Core/PushModel/PushModel.h"
@@ -108,7 +107,7 @@ bool bChaos_GC_CacheComponentSpaceBounds = true;
 FAutoConsoleVariableRef CVarChaosGCCacheComponentSpaceBounds(TEXT("p.Chaos.GC.CacheComponentSpaceBounds"), bChaos_GC_CacheComponentSpaceBounds, TEXT("Cache component space bounds for performance"));
 
 bool bChaos_GC_UseCustomRenderer = true;
-FAutoConsoleVariableRef CVarChaosGCUseISMPool(TEXT("p.Chaos.GC.UseCustomRenderer"), bChaos_GC_UseCustomRenderer, TEXT("When enabled, use a custom renderer if specified"));
+FAutoConsoleVariableRef CVarChaosGCUseCustomRenderer(TEXT("p.Chaos.GC.UseCustomRenderer"), bChaos_GC_UseCustomRenderer, TEXT("When enabled, use a custom renderer if specified"));
 
 bool bChaos_GC_InitConstantDataUseParallelFor = true;
 FAutoConsoleVariableRef CVarChaosGCInitConstantDataUseParallelFor(TEXT("p.Chaos.GC.InitConstantDataUseParallelFor"), bChaos_GC_InitConstantDataUseParallelFor, TEXT("When enabled, InitConstant data will use parallelFor for copying some of the data"));
@@ -2398,22 +2397,15 @@ void UGeometryCollectionComponent::OnRegister()
 
 	InitializeEmbeddedGeometry();
 
-	// Look for any attached custom renderer.
-	CustomRenderer = GetOwner()->FindComponentByInterface<UGeometryCollectionExternalRenderInterface>();
-	// Use the ISMPool custom renderer if it was set.
-	if (!CustomRenderer && ISMPool)
+	// Create a custom renderer object if a type is set on the component or the GC asset.
+	UClass* Type = bOverrideCustomRenderer ? CustomRendererType : (RestCollection ? RestCollection->CustomRendererType : nullptr);
+	if (Type && Type->ImplementsInterface(UGeometryCollectionExternalRenderInterface::StaticClass()))
 	{
-		UGeometryCollectionCustomRendererISMPool* ISMPoolRenderer = NewObject<UGeometryCollectionCustomRendererISMPool>(this);
-		ISMPoolRenderer->ISMPoolActor = ISMPool;
-		CustomRenderer = (UGeometryCollectionExternalRenderInterface*)ISMPoolRenderer;
-	}
-	if (!CustomRenderer && bAutoAssignISMPool)
-	{
-		if (UGeometryCollectionISMPoolSubSystem* ISMPoolSubSystem = UWorld::GetSubsystem<UGeometryCollectionISMPoolSubSystem>(GetWorld()))
+		CustomRenderer = (UGeometryCollectionExternalRenderInterface*)NewObject<UObject>(this, Type);
+
+		if (IGeometryCollectionExternalRenderInterface* RendererInterface = Cast<IGeometryCollectionExternalRenderInterface>(CustomRenderer))
 		{
-			UGeometryCollectionCustomRendererISMPool* ISMPoolRenderer = NewObject<UGeometryCollectionCustomRendererISMPool>(this);
-			ISMPoolRenderer->ISMPoolActor = ISMPoolSubSystem->FindISMPoolActor(*this);
-			CustomRenderer = (UGeometryCollectionExternalRenderInterface*)ISMPoolRenderer;
+			RendererInterface->OnRegisterGeometryCollection(*this);
 		}
 	}
 
@@ -2424,11 +2416,15 @@ void UGeometryCollectionComponent::OnUnregister()
 {
 	Super::OnUnregister();
 
-	if (IGeometryCollectionExternalRenderInterface* RendererInterface = Cast<IGeometryCollectionExternalRenderInterface>(CustomRenderer))
+	// Remove any custom renderer.
+	if (CustomRenderer)
 	{
-		RendererInterface->OnUnregisterGeometryCollection();
+		if (IGeometryCollectionExternalRenderInterface* RendererInterface = Cast<IGeometryCollectionExternalRenderInterface>(CustomRenderer))
+		{
+			RendererInterface->OnUnregisterGeometryCollection();
+		}
+		CustomRenderer = nullptr;
 	}
-	CustomRenderer = nullptr;
 }
 
 void UGeometryCollectionComponent::ResetDynamicCollection()
@@ -4841,6 +4837,15 @@ void UGeometryCollectionComponent::PostLoad()
 		BodyInstance.SetPhysMaterialOverride(PhysicalMaterialOverride_DEPRECATED.Get());
 		PhysicalMaterialOverride_DEPRECATED = nullptr;
 	}
+
+	// Convert deprecated ISMPool and bAutoAssignISMPool settings to use the a CustomRendererType.
+	if ((ISMPool_DEPRECATED || bAutoAssignISMPool_DEPRECATED) && !(bOverrideCustomRenderer || CustomRendererType))
+	{
+		bOverrideCustomRenderer = true;
+		CustomRendererType = UGeometryCollectionISMPoolRenderer::StaticClass();
+		ISMPool_DEPRECATED = nullptr;
+		bAutoAssignISMPool_DEPRECATED = false;
+	}
 }
 
 Chaos::FPhysicsObject* UGeometryCollectionComponent::GetPhysicsObjectById(Chaos::FPhysicsObjectId Id) const
diff --git a/Engine/Source/Runtime/Experimental/GeometryCollectionEngine/Private/GeometryCollection/GeometryCollectionISMPoolRenderer.cpp b/Engine/Source/Runtime/Experimental/GeometryCollectionEngine/Private/GeometryCollection/GeometryCollectionISMPoolRenderer.cpp
--- a/Engine/Source/Runtime/Experimental/GeometryCollectionEngine/Private/GeometryCollection/GeometryCollectionISMPoolRenderer.cpp
+++ b/Engine/Source/Runtime/Experimental/GeometryCollectionEngine/Private/GeometryCollection/GeometryCollectionISMPoolRenderer.cpp
@@ -2,8 +2,10 @@
 #include "GeometryCollection/GeometryCollectionISMPoolRenderer.h"
 
 #include "Engine/StaticMesh.h"
+#include "Engine/World.h"
 #include "GeometryCollection/Facades/CollectionInstancedMeshFacade.h"
 #include "GeometryCollection/GeometryCollection.h"
+#include "GeometryCollection/GeometryCollectionComponent.h"
 #include "GeometryCollection/GeometryCollectionISMPoolActor.h"
 #include "GeometryCollection/GeometryCollectionISMPoolComponent.h"
 #include "GeometryCollection/GeometryCollectionISMPoolSubSystem.h"
@@ -11,7 +13,23 @@
 
 #include UE_INLINE_GENERATED_CPP_BY_NAME(GeometryCollectionISMPoolRenderer)
 
-void UGeometryCollectionCustomRendererISMPool::UpdateState(UGeometryCollection const& InGeometryCollection, bool bIsBroken)
+void UGeometryCollectionISMPoolRenderer::OnRegisterGeometryCollection(UGeometryCollectionComponent const& InComponent)
+{
+	if (UGeometryCollectionISMPoolSubSystem* ISMPoolSubSystem = UWorld::GetSubsystem<UGeometryCollectionISMPoolSubSystem>(GetWorld()))
+	{
+		ISMPoolActor = ISMPoolSubSystem->FindISMPoolActor(InComponent);
+	}
+}
+
+void UGeometryCollectionISMPoolRenderer::OnUnregisterGeometryCollection()
+{
+	ReleaseGroup(MergedMeshGroup);
+	ReleaseGroup(InstancesGroup);
+
+	ISMPoolActor = nullptr;
+}
+
+void UGeometryCollectionISMPoolRenderer::UpdateState(UGeometryCollection const& InGeometryCollection, bool bIsBroken)
 {
 	if (!bIsBroken && MergedMeshGroup.GroupIndex == INDEX_NONE)
 	{
@@ -32,25 +50,17 @@ void UGeometryCollectionCustomRendererISMPool::UpdateState(UGeometryCollection c
 	}
 }
 
-void UGeometryCollectionCustomRendererISMPool::UpdateRootTransform(UGeometryCollection const& InGeometryCollection, FTransform const& InBaseTransform, FTransform const& InRootTransform)
+void UGeometryCollectionISMPoolRenderer::UpdateRootTransform(UGeometryCollection const& InGeometryCollection, FTransform const& InBaseTransform, FTransform const& InRootTransform)
 {
 	UpdateMergedMeshTransforms(InRootTransform * InBaseTransform);
 }
 
-void UGeometryCollectionCustomRendererISMPool::UpdateTransforms(UGeometryCollection const& InGeometryCollection, FTransform const& InBaseTransform, TArrayView<const FMatrix> InMatrices)
+void UGeometryCollectionISMPoolRenderer::UpdateTransforms(UGeometryCollection const& InGeometryCollection, FTransform const& InBaseTransform, TArrayView<const FMatrix> InMatrices)
 {
 	UpdateInstanceTransforms(InGeometryCollection, InBaseTransform, InMatrices);
 }
 
-void UGeometryCollectionCustomRendererISMPool::OnUnregisterGeometryCollection()
-{
-	ReleaseGroup(MergedMeshGroup);
-	ReleaseGroup(InstancesGroup);
-
-	ISMPoolActor = nullptr;
-}
-
-void UGeometryCollectionCustomRendererISMPool::InitMergedMeshFromGeometryCollection(UGeometryCollection const& InGeometryCollection)
+void UGeometryCollectionISMPoolRenderer::InitMergedMeshFromGeometryCollection(UGeometryCollection const& InGeometryCollection)
 {
 	if (InGeometryCollection.RootProxyData.ProxyMeshes.Num() == 0)
 	{
@@ -80,7 +90,7 @@ void UGeometryCollectionCustomRendererISMPool::InitMergedMeshFromGeometryCollect
 	}
 }
 
-void UGeometryCollectionCustomRendererISMPool::InitInstancesFromGeometryCollection(UGeometryCollection const& InGeometryCollection)
+void UGeometryCollectionISMPoolRenderer::InitInstancesFromGeometryCollection(UGeometryCollection const& InGeometryCollection)
 {
 	const int32 NumMeshes = InGeometryCollection.AutoInstanceMeshes.Num();
 	if (NumMeshes == 0)
@@ -125,7 +135,7 @@ void UGeometryCollectionCustomRendererISMPool::InitInstancesFromGeometryCollecti
 	}
 }
 
-void UGeometryCollectionCustomRendererISMPool::UpdateMergedMeshTransforms(FTransform const& InBaseTransform)
+void UGeometryCollectionISMPoolRenderer::UpdateMergedMeshTransforms(FTransform const& InBaseTransform)
 {
 	if (MergedMeshGroup.GroupIndex == INDEX_NONE)
 	{
@@ -147,7 +157,7 @@ void UGeometryCollectionCustomRendererISMPool::UpdateMergedMeshTransforms(FTrans
 	}
 }
 
-void UGeometryCollectionCustomRendererISMPool::UpdateInstanceTransforms(UGeometryCollection const& InGeometryCollection, FTransform const& InBaseTransform, TArrayView<const FMatrix> InMatrices)
+void UGeometryCollectionISMPoolRenderer::UpdateInstanceTransforms(UGeometryCollection const& InGeometryCollection, FTransform const& InBaseTransform, TArrayView<const FMatrix> InMatrices)
 {
 	if (InstancesGroup.GroupIndex == INDEX_NONE)
 	{
@@ -186,7 +196,7 @@ void UGeometryCollectionCustomRendererISMPool::UpdateInstanceTransforms(UGeometr
 	}
 }
 
-void UGeometryCollectionCustomRendererISMPool::ReleaseGroup(FISMPoolGroup& InOutGroup)
+void UGeometryCollectionISMPoolRenderer::ReleaseGroup(FISMPoolGroup& InOutGroup)
 {
 	UGeometryCollectionISMPoolComponent* ISMPoolComponent = ISMPoolActor != nullptr ? ISMPoolActor->GetISMPoolComp() : nullptr;
 	if (ISMPoolComponent != nullptr)
diff --git a/Engine/Source/Runtime/Experimental/GeometryCollectionEngine/Private/GeometryCollection/GeometryCollectionISMPoolRenderer.h b/Engine/Source/Runtime/Experimental/GeometryCollectionEngine/Private/GeometryCollection/GeometryCollectionISMPoolRenderer.h
--- a/Engine/Source/Runtime/Experimental/GeometryCollectionEngine/Private/GeometryCollection/GeometryCollectionISMPoolRenderer.h
+++ b/Engine/Source/Runtime/Experimental/GeometryCollectionEngine/Private/GeometryCollection/GeometryCollectionISMPoolRenderer.h
@@ -8,18 +8,19 @@
 
 class AGeometryCollectionISMPoolActor;
 
-/** Implementation of a custom renderer that pushes AutoInstanceMeshes to an ISMPool. */
+/** Implementation of a geometry collection custom renderer that pushes AutoInstanceMeshes to an ISMPool. */
 UCLASS()
-class UGeometryCollectionCustomRendererISMPool : public UObject, public IGeometryCollectionExternalRenderInterface
+class UGeometryCollectionISMPoolRenderer : public UObject, public IGeometryCollectionExternalRenderInterface
 {
 	GENERATED_BODY()
 
 public:
 	//~ Begin IGeometryCollectionExternalRenderInterface Interface.
+	virtual void OnRegisterGeometryCollection(UGeometryCollectionComponent const& InComponent) override;
+	virtual void OnUnregisterGeometryCollection() override;
 	virtual void UpdateState(UGeometryCollection const& InGeometryCollection, bool bInIsBroken) override;
 	virtual void UpdateRootTransform(UGeometryCollection const& InGeometryCollection, FTransform const& InBaseTransform, FTransform const& InRootTransform) override;
 	virtual void UpdateTransforms(UGeometryCollection const& InGeometryCollection, FTransform const& InBaseTransform, TArrayView<const FMatrix> InMatrices) override;
-	virtual void OnUnregisterGeometryCollection() override;
 	//~ End IGeometryCollectionExternalRenderInterface Interface.
 
 	/** Instanced Static Mesh Pool actor that is used to render our meshes. */
diff --git a/Engine/Source/Runtime/Experimental/GeometryCollectionEngine/Public/GeometryCollection/GeometryCollectionComponent.h b/Engine/Source/Runtime/Experimental/GeometryCollectionEngine/Public/GeometryCollection/GeometryCollectionComponent.h
--- a/Engine/Source/Runtime/Experimental/GeometryCollectionEngine/Public/GeometryCollection/GeometryCollectionComponent.h
+++ b/Engine/Source/Runtime/Experimental/GeometryCollectionEngine/Public/GeometryCollection/GeometryCollectionComponent.h
@@ -1053,10 +1053,6 @@ public:
 	UPROPERTY(EditAnywhere, BlueprintReadOnly, Category = "ChaosPhysics|General")
 	bool bStoreVelocities;
 
-	/** An custom renderer object. */
-	UPROPERTY(Transient)
-	TObjectPtr<UGeometryCollectionExternalRenderInterface> CustomRenderer;
-
 protected:
 	/** Display Bone Colors instead of assigned materials */
 	UPROPERTY(EditAnywhere, BlueprintReadWrite, Category = "ChaosPhysics|General")
@@ -1075,13 +1071,26 @@ protected:
 	UPROPERTY(BlueprintReadOnly, Category = "ChaosPhysics|General")
 	FGuid RunTimeDataCollectionGuid;
 #endif
-	
-	/** ISM pool to use to render the geometry collection - only works for unfractured geometry collections  */
-	UPROPERTY(EditAnywhere, BlueprintReadWrite, Category = "ChaosPhysics|Rendering", meta = (DisplayName = "ISM Pool"))
-	TObjectPtr<AGeometryCollectionISMPoolActor> ISMPool;
 
-	UPROPERTY(EditAnywhere, BlueprintReadWrite, Category = "ChaosPhysics|Rendering")
-	bool bAutoAssignISMPool = false;
+	/** Deprecated for CustomRendererType. */
+	UPROPERTY()
+	TObjectPtr<AGeometryCollectionISMPoolActor> ISMPool_DEPRECATED;
+
+	/** Deprecated for CustomRendererType. */
+	UPROPERTY()
+	bool bAutoAssignISMPool_DEPRECATED = false;
+
+	/** If true, CustomRendererType will be used. If false, CustomRendererType comes from the RestCollection. */
+	UPROPERTY(EditAnywhere, BlueprintReadWrite, Category = "ChaosPhysics|Rendering", meta = (InlineEditConditionToggle))
+	bool bOverrideCustomRenderer = false;
+
+	/** Custom class type that will be used to render the geometry collection instead of using the native rendering. */
+	UPROPERTY(EditAnywhere, BlueprintReadWrite, Category = "ChaosPhysics|Rendering", meta = (editcondition = "bOverrideCustomRenderer", MustImplement = "/Script/GeometryCollectionEngine.GeometryCollectionExternalRenderInterface"))
+	TObjectPtr<UClass> CustomRendererType;
+
+	/** A custom renderer object created from CustomRenderType. */
+	UPROPERTY(Transient)
+	TObjectPtr<UGeometryCollectionExternalRenderInterface> CustomRenderer;
 
 	/** Populate the dynamic particle data for the render thread. */
 	FGeometryCollectionDynamicData* InitDynamicData(bool bInitialization);
diff --git a/Engine/Source/Runtime/Experimental/GeometryCollectionEngine/Public/GeometryCollection/GeometryCollectionExternalRenderInterface.h b/Engine/Source/Runtime/Experimental/GeometryCollectionEngine/Public/GeometryCollection/GeometryCollectionExternalRenderInterface.h
--- a/Engine/Source/Runtime/Experimental/GeometryCollectionEngine/Public/GeometryCollection/GeometryCollectionExternalRenderInterface.h
+++ b/Engine/Source/Runtime/Experimental/GeometryCollectionEngine/Public/GeometryCollection/GeometryCollectionExternalRenderInterface.h
@@ -7,6 +7,7 @@
 #include "GeometryCollectionExternalRenderInterface.generated.h"
 
 class UGeometryCollection;
+class UGeometryCollectionComponent;
 
 UINTERFACE()
 class GEOMETRYCOLLECTIONENGINE_API UGeometryCollectionExternalRenderInterface : public UInterface
@@ -19,8 +20,9 @@ class GEOMETRYCOLLECTIONENGINE_API IGeometryCollectionExternalRenderInterface
 	GENERATED_BODY()
 
 public:
+	virtual void OnRegisterGeometryCollection(UGeometryCollectionComponent const& InComponent) = 0;
+	virtual void OnUnregisterGeometryCollection() = 0;
 	virtual void UpdateState(UGeometryCollection const& InGeometryCollection, bool bInIsBroken) = 0;
 	virtual void UpdateRootTransform(UGeometryCollection const& InGeometryCollection, FTransform const& InBaseTransform, FTransform const& InRootTransform) = 0;
 	virtual void UpdateTransforms(UGeometryCollection const& InGeometryCollection, FTransform const& InBaseTransform, TArrayView<const FMatrix> InMatrices) = 0;
-	virtual void OnUnregisterGeometryCollection() = 0;
 };
diff --git a/Engine/Source/Runtime/Experimental/GeometryCollectionEngine/Public/GeometryCollection/GeometryCollectionObject.h b/Engine/Source/Runtime/Experimental/GeometryCollectionEngine/Public/GeometryCollection/GeometryCollectionObject.h
--- a/Engine/Source/Runtime/Experimental/GeometryCollectionEngine/Public/GeometryCollection/GeometryCollectionObject.h
+++ b/Engine/Source/Runtime/Experimental/GeometryCollectionEngine/Public/GeometryCollection/GeometryCollectionObject.h
@@ -550,16 +550,20 @@ public:
 
 	/**
 	 * Strip unnecessary render data from the Geometry Collection to keep the memory footprint as small as possible.
-	 * This may be used if the cooked build uses an alternative render path such as ISMPools.
+	 * This may be used if the cooked build uses a custom renderer such as the ISMPool renderer.
 	 */
 	UPROPERTY(EditAnywhere, BlueprintReadWrite, Category = "Rendering")
 	bool bStripRenderDataOnCook;
 
-	/** static mesh to use as a proxy for rendering until the geometry collection is broken */
+	/** Custom class type that will be used to render the geometry collection instead of using the native rendering. */
+	UPROPERTY(EditAnywhere, BlueprintReadWrite, Category = "Rendering", meta = (MustImplement = "/Script/GeometryCollectionEngine.GeometryCollectionExternalRenderInterface"))
+	TObjectPtr<UClass> CustomRendererType;
+
+	/** Static mesh to use as a proxy for rendering until the geometry collection is broken. */
 	UPROPERTY(BlueprintReadWrite, EditAnywhere, Category = "Rendering")
 	FGeometryCollectionProxyMeshData RootProxyData;
 
-	/** list of unique static mesh / materials pairs for auto instancing*/
+	/** List of unique static mesh / materials pairs for auto instancing. */
 	UPROPERTY(EditAnywhere, Category = "Rendering")
 	TArray<FGeometryCollectionAutoInstanceMesh> AutoInstanceMeshes;
 
