UpdateGeometryCollectionViews polls all GC particles

diff --git a/Engine/Source/Runtime/Engine/Private/PhysicsEngine/ImmediatePhysics/ImmediatePhysicsChaos/ImmediatePhysicsActorHandle_Chaos.cpp b/Engine/Source/Runtime/Engine/Private/PhysicsEngine/ImmediatePhysics/ImmediatePhysicsChaos/ImmediatePhysicsActorHandle_Chaos.cpp
--- a/Engine/Source/Runtime/Engine/Private/PhysicsEngine/ImmediatePhysics/ImmediatePhysicsChaos/ImmediatePhysicsActorHandle_Chaos.cpp
+++ b/Engine/Source/Runtime/Engine/Private/PhysicsEngine/ImmediatePhysics/ImmediatePhysicsChaos/ImmediatePhysicsActorHandle_Chaos.cpp
@@ -278,7 +278,7 @@ namespace ImmediatePhysics_Chaos
 						Dynamic->SetAngularEtherDrag(BodyInstance->AngularDamping);
 						Dynamic->SetGravityEnabled(BodyInstance->bEnableGravity);
 					}
-					Dynamic->Disabled() = true;
+					Dynamic->SetDisabled(true);
 				}
 			}
 		}
@@ -320,7 +320,7 @@ namespace ImmediatePhysics_Chaos
 		Chaos::FPBDRigidParticleHandle* Dynamic = ParticleHandle->CastToRigidParticle();
 		if (Dynamic && Dynamic->ObjectState() == Chaos::EObjectStateType::Dynamic)
 		{
-			Dynamic->Disabled() = !bEnabled;
+			Dynamic->SetDisabled(!bEnabled);
 		}
 	}
 
diff --git a/Engine/Source/Runtime/Experimental/Chaos/Private/PBDRigidsSolver.cpp b/Engine/Source/Runtime/Experimental/Chaos/Private/PBDRigidsSolver.cpp
--- a/Engine/Source/Runtime/Experimental/Chaos/Private/PBDRigidsSolver.cpp
+++ b/Engine/Source/Runtime/Experimental/Chaos/Private/PBDRigidsSolver.cpp
@@ -1423,11 +1423,6 @@ namespace Chaos
 			}
 		});
 
-		if(NumInitializedGCProxies > 0)
-		{
-			GetParticles().UpdateGeometryCollectionViews(true);
-		}
-
 		//need to create new constraint handles
 		DirtyProxiesData->ForEachProxy([this, Manager, RewindData](int32 DataIdx, FDirtyProxy& Dirty)
 		{
diff --git a/Engine/Source/Runtime/Experimental/Chaos/Public/Chaos/PBDRigidsSOAs.h b/Engine/Source/Runtime/Experimental/Chaos/Public/Chaos/PBDRigidsSOAs.h
--- a/Engine/Source/Runtime/Experimental/Chaos/Public/Chaos/PBDRigidsSOAs.h
+++ b/Engine/Source/Runtime/Experimental/Chaos/Public/Chaos/PBDRigidsSOAs.h
@@ -210,7 +210,6 @@ public:
 		ClusteredParticles = MakeUnique< FPBDRigidClusteredParticles>();
 
 		GeometryCollectionParticles = MakeUnique<TPBDGeometryCollectionParticles<FReal, 3>>();
-		bGeometryCollectionDirty = false;
 
 		UpdateViews();
 	}
@@ -300,8 +299,13 @@ public:
 				Handle->SetObjectStateLowLevel(Chaos::EObjectStateType::Dynamic);
 				Handle->SetSleeping(false);
 			}
+			if (!Params.bDisabled)
+			{
+				InsertGeometryCollectionParticle(Handle);
+			}
 		}
-		bGeometryCollectionDirty = true;
+		UpdateViews();
+		
 		return Results;
 	}
 
@@ -374,17 +378,13 @@ public:
 			RemoveFromActiveArray(PBDRigid, /*bStillDirty=*/ false);
 			MovingKinematicsMapArray.Remove(PBDRigid);
 
-			if (auto PBDRigidClustered = Particle->CastToClustered())
+			if (TPBDGeometryCollectionParticleHandle<FReal, 3>* GCHandle = Particle->CastToGeometryCollection())
 			{
-				if (Particle->GetParticleType() == Chaos::EParticleType::GeometryCollection)
-				{
-					FellAsleepGeometryCollectionArray.Remove(PBDRigid);
-					bGeometryCollectionDirty = true;
-				}
-				else // clustered
-				{
-					DynamicClusteredMapArray.Remove(PBDRigidClustered);
-				}
+				RemoveGeometryCollectionParticle(GCHandle);
+			}
+			else if (FPBDRigidClusteredParticleHandle* PBDRigidClustered = Particle->CastToClustered())
+			{
+				DynamicClusteredMapArray.Remove(PBDRigidClustered);
 			}
 
 			// Check for sleep events referencing this particle
@@ -415,22 +415,17 @@ public:
 
 		if (auto PBDRigid = Particle->CastToRigidParticle())
 		{
-			PBDRigid->Disabled() = true;
+			PBDRigid->SetDisabled(true);
 			PBDRigid->V() = FVec3(0);
 			PBDRigid->W() = FVec3(0);
 
-			if (auto PBDRigidClustered = Particle->CastToClustered())
+			if (TPBDGeometryCollectionParticleHandle<FReal, 3>* PBDRigidGC = Particle->CastToGeometryCollection())
 			{
-				if (Particle->GetParticleType() == Chaos::EParticleType::GeometryCollection)
-				{
-					FellAsleepGeometryCollectionArray.Remove(PBDRigid);
-					bGeometryCollectionDirty = true;
-					return;
-				}
-				else // clustered
-				{
-					RemoveClusteredParticle(PBDRigidClustered);
-				}
+				RemoveGeometryCollectionParticle(PBDRigidGC);
+			}
+			else if (FPBDRigidClusteredParticleHandle* PBDRigidClustered = Particle->CastToClustered())
+			{
+				RemoveClusteredParticle(PBDRigidClustered);
 			}
 			else
 			{
@@ -460,21 +455,18 @@ public:
 
 		if (auto PBDRigid = Particle->CastToRigidParticle())
 		{
-			PBDRigid->Disabled() = false;
+			PBDRigid->SetDisabled(false);
 			// DisableParticle() zeros V and W.  We do nothing here and assume the client
 			// sets appropriate values.
 
-			if (auto PBDRigidClustered = Particle->CastToClustered())
+			if (TPBDGeometryCollectionParticleHandle<FReal, 3>* PBDRigidGC = Particle->CastToGeometryCollection())
 			{
-				if (Particle->GetParticleType() == Chaos::EParticleType::GeometryCollection)
-				{
-					bGeometryCollectionDirty = true;
-					return;
-				}
-				else // clustered
-				{
-					InsertClusteredParticle(PBDRigidClustered);
-				}
+				RemoveGeometryCollectionParticle(PBDRigidGC);
+				InsertGeometryCollectionParticle(PBDRigidGC);
+			}
+			else if (FPBDRigidClusteredParticleHandle* PBDRigidClustered = Particle->CastToClustered())
+			{
+				InsertClusteredParticle(PBDRigidClustered);
 			}
 			else
 			{
@@ -499,9 +491,8 @@ public:
 
 	/**
 	 * Wake a sleeping dynamic non-disabled particle.
-	 * return true if Geometry collection needs to be updated
 	 */
-	bool ActivateParticle(FGeometryParticleHandle* Particle, const bool DeferUpdateViews=false)
+	void ActivateParticle(FGeometryParticleHandle* Particle, const bool DeferUpdateViews=false)
 	{
 		if (auto PBDRigid = Particle->CastToRigidParticle())
 		{
@@ -514,25 +505,14 @@ public:
 					PBDRigid->SetSleeping(false);
 					PBDRigid->SetObjectStateLowLevel(EObjectStateType::Dynamic);
 		
-					bool bUpdateGeometryCollection = false;
-					if (auto PBDRigidClustered = Particle->CastToClustered())
+					if (TPBDGeometryCollectionParticleHandle<FReal, 3>* PBDRigidGC = Particle->CastToGeometryCollection())
 					{
-						if (Particle->GetParticleType() == Chaos::EParticleType::GeometryCollection)
-						{
-							FellAsleepGeometryCollectionArray.Remove(PBDRigid);
-							bGeometryCollectionDirty = true;
-							return true;
-						}
-						else
-						{
-							AddToActiveArray(PBDRigid);
-						}
-					}
-					else
-					{
-						// Non clustered rigid particles:
-						AddToActiveArray(PBDRigid);
+						RemoveGeometryCollectionParticle(PBDRigidGC);
+						InsertGeometryCollectionParticle(PBDRigidGC);
 					}
+
+					AddToActiveArray(PBDRigid);
+
 					if(!DeferUpdateViews)
 					{
 						UpdateViews();
@@ -540,7 +520,6 @@ public:
 				}
 			}
 		}
-		return false;
 	}
 
 	/**
@@ -548,19 +527,11 @@ public:
 	 */
 	void ActivateParticles(const TArray<FGeometryParticleHandle*>& Particles)
 	{
-		bool bUpdateGeometryCollection = false;
 		for (auto Particle : Particles)
 		{
-			bUpdateGeometryCollection |= ActivateParticle(Particle, true);
+			ActivateParticle(Particle, true);
 		}
-		if (bUpdateGeometryCollection)
-		{
-			UpdateGeometryCollectionViews();
-		}
-		else
-		{
 			UpdateViews();
-		}
 	}
 
 	/**
@@ -584,17 +555,14 @@ public:
 					PBDRigid->SetSleeping(true);
 					PBDRigid->SetObjectStateLowLevel(EObjectStateType::Sleeping);
 
-					if (auto PBDRigidClustered = Particle->CastToClustered())
+					if (TPBDGeometryCollectionParticleHandle<FReal, 3>* PBDRigidGC = Particle->CastToGeometryCollection())
+					{
+						RemoveGeometryCollectionParticle(PBDRigidGC);
+						InsertGeometryCollectionParticle(PBDRigidGC);
+					}
+					else if (FPBDRigidClusteredParticleHandle* PBDRigidClustered = Particle->CastToClustered())
 					{
-						if (Particle->GetParticleType() == Chaos::EParticleType::GeometryCollection)
-						{
-							FellAsleepGeometryCollectionArray.Insert(PBDRigid);
-							bGeometryCollectionDirty = true;
-						}
-						else
-						{
-							RemoveFromActiveArray(PBDRigid, /*bStillDirty=*/true);
-						}
+						RemoveFromActiveArray(PBDRigid, /*bStillDirty=*/true);
 					}
 					else
 					{
@@ -619,7 +587,6 @@ public:
 		{
 			DeactivateParticle(Particle, true);
 		}
-		UpdateIfNeeded();
 		UpdateViews();
 	}
 	
@@ -628,7 +595,6 @@ public:
 	*/
 	void RebuildViews()
 	{
-		UpdateIfNeeded();
 		UpdateViews();
 	}
 
@@ -694,8 +660,15 @@ public:
 
 	void SetClusteredParticleSOA(FPBDRigidClusteredParticleHandle* ClusteredParticle)
 	{
-		if (ClusteredParticle->GetParticleType() != Chaos::EParticleType::GeometryCollection)
+		if (TPBDGeometryCollectionParticleHandle<FReal, 3>* GCParticle = ClusteredParticle->CastToGeometryCollection())
+		{
+			// Geometry collection particles have their own arrays which are also included in the active view
+			RemoveGeometryCollectionParticle(GCParticle);
+			InsertGeometryCollectionParticle(GCParticle);
+		}
+		else 
 		{ 
+			check(ClusteredParticle->GetParticleType() != Chaos::EParticleType::GeometryCollection);
 			RemoveClusteredParticle(ClusteredParticle);
 			InsertClusteredParticle(ClusteredParticle);
 
@@ -709,11 +682,6 @@ public:
 				AddToActiveArray(ClusteredParticle->CastToRigidParticle());
 			}
 		}
-		else
-		{
-			// Geometry collection particles have their own arrays which are also included in the active view
-			bGeometryCollectionDirty = true;
-		}
 		
 		UpdateViews();
 	}
@@ -813,21 +781,21 @@ public:
 	}
 
 
-	const TParticleView<FGeometryParticles>& GetNonDisabledView() const { UpdateIfNeeded();  return NonDisabledView; }
+	const TParticleView<FGeometryParticles>& GetNonDisabledView() const {  return NonDisabledView; }
 
-	const TParticleView<FPBDRigidParticles>& GetNonDisabledDynamicView() const { UpdateIfNeeded(); return NonDisabledDynamicView; }
+	const TParticleView<FPBDRigidParticles>& GetNonDisabledDynamicView() const { return NonDisabledDynamicView; }
 
 	const TParticleView<FPBDRigidClusteredParticles>& GetNonDisabledClusteredView() const { return NonDisabledClusteredView; }
 
-	const TParticleView<FPBDRigidParticles>& GetActiveParticlesView() const { UpdateIfNeeded(); return ActiveParticlesView; }
-	TParticleView<FPBDRigidParticles>& GetActiveParticlesView() { UpdateIfNeeded(); return ActiveParticlesView; }
+	const TParticleView<FPBDRigidParticles>& GetActiveParticlesView() const {  return ActiveParticlesView; }
+	TParticleView<FPBDRigidParticles>& GetActiveParticlesView() { return ActiveParticlesView; }
 
 	const TArray<FPBDRigidParticleHandle*>& GetActiveParticlesArray() const { return ActiveParticlesMapArray.GetArray(); }
 	
-	const TParticleView<FPBDRigidParticles>& GetDirtyParticlesView() const { UpdateIfNeeded(); return DirtyParticlesView; }
-	TParticleView<FPBDRigidParticles>& GetDirtyParticlesView() { UpdateIfNeeded(); return DirtyParticlesView; }
+	const TParticleView<FPBDRigidParticles>& GetDirtyParticlesView() const { return DirtyParticlesView; }
+	TParticleView<FPBDRigidParticles>& GetDirtyParticlesView() { return DirtyParticlesView; }
 
-	const TParticleView<FGeometryParticles>& GetAllParticlesView() const { UpdateIfNeeded(); return AllParticlesView; }
+	const TParticleView<FGeometryParticles>& GetAllParticlesView() const { return AllParticlesView; }
 
 
 	const TParticleView<FKinematicGeometryParticles>& GetActiveKinematicParticlesView() const { return ActiveKinematicParticlesView; }
@@ -861,128 +829,38 @@ public:
 	const TPBDGeometryCollectionParticles<FReal, 3>& GetGeometryCollectionParticles() const { return *GeometryCollectionParticles; }
 	TPBDGeometryCollectionParticles<FReal, 3>& GetGeometryCollectionParticles() { return *GeometryCollectionParticles; }
 
-	void UpdateIfNeeded() const 
+	void InsertGeometryCollectionParticle(TPBDGeometryCollectionParticleHandle<FReal, 3>* GCParticle)
 	{
-		if (bGeometryCollectionDirty) 
+		if (!GCParticle->Disabled())
 		{
-			FPBDRigidsSOAs* NCThis = const_cast<FPBDRigidsSOAs*>(this);
-			NCThis->UpdateGeometryCollectionViews();
-		}
-	}
-
-	/**
-	 * Update which particle arrays geometry collection particles are in based on 
-	 * their object state (static, kinematic, dynamic, sleeping) and their disabled 
-	 * state.
-	 *
-	 * The reason for specializing this function for geometry collections is for 
-	 * scalability.  That is, we try to process many geometry collection particles
-	 * at a time, rather than one by one.
-	 */
-	void UpdateGeometryCollectionViews(const bool ForceUpdateViews=false)
-	{
-		SCOPE_CYCLE_COUNTER(STAT_UpdateGeometryCollectionViews);
-
-		if (!GeometryCollectionParticles)
-		{
-			if (ForceUpdateViews)
-				UpdateViews();
-			return;
-		}
-
-		int32 StaticIdx = 0, KinematicIdx = 0, SleepingIdx = 0, DynamicIdx = 0;
-		int32 ActiveIdx = 0, DisabledIdx = 0;
-
-		for(int32 PIdx = 0; PIdx < (int32)GeometryCollectionParticles->Size(); PIdx++)
-		{
-			const TPBDGeometryCollectionParticleHandle<FReal, 3>* Handle = 
-				GeometryCollectionParticles->Handle(PIdx);
-			if (!Handle)
-				continue;
-
-			const bool bDisabled = Handle->Disabled();
-			if (bDisabled)
-				continue;
-
-			// Count the number of particles in each state.
-			const Chaos::EObjectStateType State = 
-				Handle->Sleeping() ? Chaos::EObjectStateType::Sleeping : Handle->ObjectState();
+			const Chaos::EObjectStateType State = GCParticle->Sleeping() ? Chaos::EObjectStateType::Sleeping : GCParticle->ObjectState();
 			switch (State)
 			{
-			case Chaos::EObjectStateType::Static:
-				StaticIdx++;
-				break;
-			case Chaos::EObjectStateType::Kinematic:
-				KinematicIdx++;
-				break;
-			case Chaos::EObjectStateType::Sleeping:
-				SleepingIdx++;
-				break;
-			case Chaos::EObjectStateType::Dynamic:
-				DynamicIdx++;
-				break;
-			default:
-				break;
-			};
-		}
-
-		// Compare with the previous array sizes, and resize if needed.
-		bool Changed =
-			StaticGeometryCollectionArray.Num() != StaticIdx ||
-			KinematicGeometryCollectionArray.Num() != KinematicIdx ||
-			SleepingGeometryCollectionArray.Num() != SleepingIdx ||
-			DynamicGeometryCollectionArray.Num() != DynamicIdx;
-		if (Changed)
-		{
-			StaticGeometryCollectionArray.SetNumUninitialized(StaticIdx);
-			KinematicGeometryCollectionArray.SetNumUninitialized(KinematicIdx);
-			SleepingGeometryCollectionArray.SetNumUninitialized(SleepingIdx);
-			DynamicGeometryCollectionArray.SetNumUninitialized(DynamicIdx);
-		}
-
-		// (Re)populate the arrays, making note if any prior entires differ from the current.
-		StaticIdx = KinematicIdx = SleepingIdx = DynamicIdx = 0;
-		for(int32 PIdx = 0; PIdx < (int32)GeometryCollectionParticles->Size(); PIdx++)
-		{
-			TPBDGeometryCollectionParticleHandle<FReal, 3>* Handle = 
-				GeometryCollectionParticles->Handle(PIdx);
-			if (!Handle)
-				continue;
-
-			const bool bDisabled = Handle->Disabled();
-			if (bDisabled)
-				continue;
-
-			const Chaos::EObjectStateType State = 
-				Handle->Sleeping() ? Chaos::EObjectStateType::Sleeping : Handle->ObjectState();
-			switch (State)
-			{
-			case Chaos::EObjectStateType::Static:
-				Changed |= StaticGeometryCollectionArray[StaticIdx] != Handle;
-				StaticGeometryCollectionArray[StaticIdx++] = Handle;
+			case EObjectStateType::Uninitialized:
+				ensure(false); // we should probably not be here 
 				break;
-			case Chaos::EObjectStateType::Kinematic:
-				Changed |= KinematicGeometryCollectionArray[KinematicIdx] != Handle;
-				KinematicGeometryCollectionArray[KinematicIdx++] = Handle;
+			case EObjectStateType::Static:
+				StaticGeometryCollectionArray.Insert(GCParticle);
 				break;
-			case Chaos::EObjectStateType::Sleeping:
-				Changed |= SleepingGeometryCollectionArray[SleepingIdx] != Handle;
-				SleepingGeometryCollectionArray[SleepingIdx++] = Handle;
+			case EObjectStateType::Kinematic:
+				KinematicGeometryCollectionArray.Insert(GCParticle);
 				break;
-			case Chaos::EObjectStateType::Dynamic:
-				Changed |= DynamicGeometryCollectionArray[DynamicIdx] != Handle;
-				DynamicGeometryCollectionArray[DynamicIdx++] = Handle;
+			case EObjectStateType::Dynamic:
+				DynamicGeometryCollectionArray.Insert(GCParticle);
 				break;
-			default:
+			case EObjectStateType::Sleeping:
+				SleepingGeometryCollectionArray.Insert(GCParticle);
 				break;
-			};
+			}
 		}
+	}
 
-		if(Changed || ForceUpdateViews)
-		{
-			UpdateViews();
-		}
-		bGeometryCollectionDirty = false;
+	void RemoveGeometryCollectionParticle(TPBDGeometryCollectionParticleHandle<FReal, 3>* GCParticle)
+	{
+		StaticGeometryCollectionArray.Remove(GCParticle);
+		KinematicGeometryCollectionArray.Remove(GCParticle);
+		SleepingGeometryCollectionArray.Remove(GCParticle);
+		DynamicGeometryCollectionArray.Remove(GCParticle);
 	}
 
 	const auto& GetClusteredParticles() const { return *ClusteredParticles; }
@@ -1099,10 +977,10 @@ private:
 				{&StaticClusteredMapArray.GetArray()},
 				{&KinematicClusteredMapArray.GetArray()},
 				{&DynamicClusteredMapArray.GetArray()},
-				{&StaticGeometryCollectionArray},
-				{&KinematicGeometryCollectionArray},
-				{&SleepingGeometryCollectionArray},
-				{&DynamicGeometryCollectionArray}
+				{&StaticGeometryCollectionArray.GetArray()},
+				{&KinematicGeometryCollectionArray.GetArray()},
+				{&SleepingGeometryCollectionArray.GetArray()},
+				{&DynamicGeometryCollectionArray.GetArray()}
 			};
 			NonDisabledView = MakeParticleView(MoveTemp(TmpArray));
 		}
@@ -1116,8 +994,8 @@ private:
 			{
 				DynamicParticles.Get(),
 				{&DynamicClusteredMapArray.GetArray() },
-				{&SleepingGeometryCollectionArray },
-				{&DynamicGeometryCollectionArray }
+				{&SleepingGeometryCollectionArray.GetArray() },
+				{&DynamicGeometryCollectionArray.GetArray() }
 			};
 			NonDisabledDynamicView = MakeParticleView(MoveTemp(TmpArray));
 		}
@@ -1127,8 +1005,6 @@ private:
 				{&ActiveParticlesMapArray.GetArray()},
 				{&MovingKinematicsMapArray.GetArray()},
 				//{&KinematicClusteredMapArray.GetArray()},	// @todo(chaos): we should include moving clustered kinematics in this view
-				//{&DynamicClusteredMapArray.GetArray()},	// Cluster particles appear in the ActiveParticlesArray
-				{&DynamicGeometryCollectionArray},
 			};
 			ActiveDynamicMovingKinematicParticlesView = MakeParticleView(MoveTemp(TmpArray));
 		}
@@ -1141,10 +1017,8 @@ private:
 			: TArray<TSOAView<FPBDRigidParticles>>
 			{ 
 				{&ActiveParticlesMapArray.GetArray()},
-			//	{&DynamicClusteredMapArray.GetArray()},  Cluster particles appear in the ActiveParticlesArray
-				{&StaticGeometryCollectionArray},
-				{&KinematicGeometryCollectionArray},
-				{&DynamicGeometryCollectionArray}
+				{&StaticGeometryCollectionArray.GetArray()},
+				{&KinematicGeometryCollectionArray.GetArray()},
 			};
 			ActiveParticlesView = MakeParticleView(MoveTemp(TmpArray));
 		}
@@ -1158,11 +1032,10 @@ private:
 			: TArray<TSOAView<FPBDRigidParticles>>
 			{
 				{&ActiveParticlesMapArray.GetArray()},
-				//	{&DynamicClusteredMapArray.GetArray()},  Cluster particles appear in the ActiveParticlesArray
-				(&MovingKinematicsMapArray.GetArray()),
-				{&KinematicGeometryCollectionArray},
-				{&DynamicGeometryCollectionArray},
-				{&FellAsleepGeometryCollectionArray.GetArray()},
+				{&MovingKinematicsMapArray.GetArray()},
+				{&KinematicGeometryCollectionArray.GetArray()},
+				{&DynamicGeometryCollectionArray.GetArray()},
+				{&SleepingGeometryCollectionArray.GetArray()},
 				{&TransientDirtyMapArray.GetArray()}
 			};
 			DirtyParticlesView = MakeParticleView(MoveTemp(TmpArray));
@@ -1193,7 +1066,7 @@ private:
 			{ 
 				KinematicParticles.Get(),
 				DynamicKinematicParticles.Get(),
-				{&KinematicGeometryCollectionArray},
+				{&KinematicGeometryCollectionArray.GetArray()},
 				{&KinematicClusteredMapArray.GetArray()}
 			};
 			ActiveKinematicParticlesView = MakeParticleView(MoveTemp(TmpArray));
@@ -1224,9 +1097,9 @@ private:
 			? TArray<TSOAView<TPBDGeometryCollectionParticles<FReal, 3>>> {}	//no geometry collection during resim
 			: TArray<TSOAView<TPBDGeometryCollectionParticles<FReal, 3>>>
 			{ 
-				{&StaticGeometryCollectionArray},
-				{&KinematicGeometryCollectionArray},
-				{&DynamicGeometryCollectionArray}
+				{&StaticGeometryCollectionArray.GetArray()},
+				{&KinematicGeometryCollectionArray.GetArray()},
+				{&DynamicGeometryCollectionArray.GetArray()}
 			};
 			ActiveGeometryCollectionParticlesView = MakeParticleView(MoveTemp(TmpArray));
 		}
@@ -1300,12 +1173,10 @@ private:
 
 	// Geometry collection particle state is controlled via their disabled state and assigned 
 	// EObjectStateType, and are shuffled into these corresponding arrays in UpdateGeometryCollectionViews().
-	TArray<TPBDGeometryCollectionParticleHandle<FReal, 3>*> StaticGeometryCollectionArray;
-	TArray<TPBDGeometryCollectionParticleHandle<FReal, 3>*> KinematicGeometryCollectionArray;
-	TArray<TPBDGeometryCollectionParticleHandle<FReal, 3>*> SleepingGeometryCollectionArray;
-	TArray<TPBDGeometryCollectionParticleHandle<FReal, 3>*> DynamicGeometryCollectionArray;
-	TParticleMapArray<FPBDRigidParticleHandle> FellAsleepGeometryCollectionArray;
-	bool bGeometryCollectionDirty;
+	TParticleMapArray<FPBDGeometryCollectionParticleHandle> StaticGeometryCollectionArray;
+	TParticleMapArray<FPBDGeometryCollectionParticleHandle> KinematicGeometryCollectionArray;
+	TParticleMapArray<FPBDGeometryCollectionParticleHandle> SleepingGeometryCollectionArray;
+	TParticleMapArray<FPBDGeometryCollectionParticleHandle> DynamicGeometryCollectionArray;
 
 	//Utility structures for maintaining an Active particles view
 	TParticleMapArray<FPBDRigidParticleHandle> ActiveParticlesMapArray;
diff --git a/Engine/Source/Runtime/Experimental/Chaos/Public/Chaos/ParticleHandle.h b/Engine/Source/Runtime/Experimental/Chaos/Public/Chaos/ParticleHandle.h
--- a/Engine/Source/Runtime/Experimental/Chaos/Public/Chaos/ParticleHandle.h
+++ b/Engine/Source/Runtime/Experimental/Chaos/Public/Chaos/ParticleHandle.h
@@ -611,6 +611,9 @@ public:
 	const TPBDRigidClusteredParticleHandleImp<T, d, bPersistent>* CastToClustered() const;
 	TPBDRigidClusteredParticleHandleImp<T, d, bPersistent>* CastToClustered();
 
+	const TPBDGeometryCollectionParticleHandleImp<T, d, bPersistent>* CastToGeometryCollection() const;
+	TPBDGeometryCollectionParticleHandleImp<T, d, bPersistent>* CastToGeometryCollection();
+
 	const TGeometryParticleHandle<T, d>* Handle() const { return GetHandleHelper(this);}
 	TGeometryParticleHandle<T, d>* Handle() { return GetHandleHelper(this); }
 	
@@ -921,6 +924,7 @@ public:
 	uint32 CollisionConstraintFlags() const { return PBDRigidParticles->CollisionConstraintFlags(ParticleIdx); }
 
 	bool Disabled() const { return PBDRigidParticles->Disabled(ParticleIdx); }
+	//UE_DEPRECATED(5.3, "This method should not be used anymore. SetDisabled should be used instead.")
 	bool& Disabled() { return PBDRigidParticles->DisabledRef(ParticleIdx); }
 
 	// See Comment on TRigidParticle::SetDisabledLowLevel. State changes in Evolution should accompany this call.
@@ -1402,6 +1406,13 @@ const TPBDRigidClusteredParticleHandleImp<T, d, bPersistent>* TGeometryParticleH
 template <typename T, int d, bool bPersistent>
 TPBDRigidClusteredParticleHandleImp<T, d, bPersistent>* TGeometryParticleHandleImp<T, d, bPersistent>::CastToClustered() { checkSlow( Type <= EParticleType::Clustered || Type == EParticleType::GeometryCollection); return Type >= EParticleType::Clustered ? static_cast<TPBDRigidClusteredParticleHandleImp<T, d, bPersistent>*>(this) : nullptr; }
 
+template <typename T, int d, bool bPersistent>
+const TPBDGeometryCollectionParticleHandleImp<T, d, bPersistent>* TGeometryParticleHandleImp<T, d, bPersistent>::CastToGeometryCollection() const { checkSlow(Type <= EParticleType::GeometryCollection || Type == EParticleType::GeometryCollection); return Type >= EParticleType::GeometryCollection ? static_cast<const TPBDGeometryCollectionParticleHandleImp<T, d, bPersistent>*>(this) : nullptr; }
+
+template <typename T, int d, bool bPersistent>
+TPBDGeometryCollectionParticleHandleImp<T, d, bPersistent>* TGeometryParticleHandleImp<T, d, bPersistent>::CastToGeometryCollection() { checkSlow(Type <= EParticleType::GeometryCollection|| Type == EParticleType::GeometryCollection); return Type >= EParticleType::GeometryCollection ? static_cast<TPBDGeometryCollectionParticleHandleImp<T, d, bPersistent>*>(this) : nullptr; }
+
+
 template <typename T, int d, bool bPersistent>
 EObjectStateType TGeometryParticleHandleImp<T,d, bPersistent>::ObjectState() const
 {
diff --git a/Engine/Source/Runtime/Experimental/Chaos/Public/PhysicsProxy/FieldSystemProxyHelper.h b/Engine/Source/Runtime/Experimental/Chaos/Public/PhysicsProxy/FieldSystemProxyHelper.h
--- a/Engine/Source/Runtime/Experimental/Chaos/Public/PhysicsProxy/FieldSystemProxyHelper.h
+++ b/Engine/Source/Runtime/Experimental/Chaos/Public/PhysicsProxy/FieldSystemProxyHelper.h
@@ -259,9 +259,6 @@ namespace Chaos
 	{
 		if (UpdatedParticleIndices.Num() > 0)
 		{
-			// chaos(todo) : do we really need this ? 
-			RigidSolver->GetParticles().UpdateGeometryCollectionViews(true);
-
 			TSet<FPBDRigidClusteredParticleHandle*> TopParentClusters;
 			for (const FFieldContextIndex& ParticleIndex: UpdatedParticleIndices)
 			{
@@ -403,6 +400,10 @@ namespace Chaos
 				if (bHasStateChanged)
 				{
 					UpdateSolverParticlesState(RigidSolver, EvaluatedSamples, ParticleHandles);
+					for (Chaos::FGeometryParticleHandle* ParticleHandle : ParticleHandles)
+					{
+						RigidSolver->GetEvolution()->EnableParticle(ParticleHandle);
+					}
 				}
 			}
 		}
