Chaos Cloth: make XPBDSpringConstraints resolution independent. This does change behavior of XPBD Springs, but there was no way to author them, so no need to upgrade.

diff --git a/Engine/Source/Runtime/Experimental/Chaos/Private/Chaos/XPBDSpringConstraints.ispc b/Engine/Source/Runtime/Experimental/Chaos/Private/Chaos/XPBDSpringConstraints.ispc
--- a/Engine/Source/Runtime/Experimental/Chaos/Private/Chaos/XPBDSpringConstraints.ispc
+++ b/Engine/Source/Runtime/Experimental/Chaos/Private/Chaos/XPBDSpringConstraints.ispc
@@ -39,8 +39,9 @@ export void ApplyXPBDSpringConstraints(uniform FVector4f PandInvM[],
 		const float Distance = SafeNormalize(Direction);
 
 		const varying float Dist = Dists[i];
+		const varying float FinalStiffness = select(Dist < FLOAT_SMALL_NUMBER, Stiffness, Stiffness / Dist);
 		const varying float Lambda = Lambdas[i];
-		const uniform float Alpha = (uniform float)1.f / (Stiffness * Dt * Dt);
+		const varying float Alpha = (uniform float)1.f / (FinalStiffness * Dt * Dt);
 		const varying float Offset = Distance - Dist;
 
 		const float CombinedInvMass = InvM2 + InvM1;
@@ -95,12 +96,13 @@ export void ApplyXPBDSpringConstraintsWithDamping(uniform FVector4f PandInvM[],
 		const varying FVector3f X2 = VectorGather(&X[i2]);
 
 		const varying float Dist = Dists[i];
+		const varying float FinalStiffness = select(Dist < FLOAT_SMALL_NUMBER, Stiffness, Stiffness / Dist);
 		const varying float Lambda = Lambdas[i];
-		const uniform float Alpha = (uniform float)1.f / (Stiffness * Dt * Dt);
+		const varying float Alpha = (uniform float)1.f / (FinalStiffness * Dt * Dt);
 		const varying float Offset = Distance - Dist;
 
 		const float CombinedInvMass = InvM2 + InvM1;
-		const varying float Damping = (uniform float)2.f * DampingRatio * sqrt(Stiffness / CombinedInvMass);
+		const varying float Damping = (uniform float)2.f * DampingRatio * sqrt(FinalStiffness / CombinedInvMass);
 		const varying float Gamma = Alpha * Damping * Dt;
 		const varying FVector3f RelativeVelocityTimesDt = P1 - X1 - P2 + X2;
 
@@ -160,12 +162,14 @@ export void ApplyXPBDSpringConstraintsWithWeightMaps(uniform FVector4f PandInvM[
 		#pragma ignore warning(perf)
 		const varying float Stiffness = StiffnessTable[StiffnessIndex];
 
-		const varying float Alpha = (uniform float)1.f / (Stiffness * Dt * Dt);
+		const varying float FinalStiffness = select(Dist < FLOAT_SMALL_NUMBER, Stiffness, Stiffness / Dist);
+
+		const varying float Alpha = (uniform float)1.f / (FinalStiffness * Dt * Dt);
 
 		const float CombinedInvMass = InvM2 + InvM1;
 
 		const varying float ComputedDLambda = (Offset - Alpha * Lambda) / (CombinedInvMass + Alpha);
-		const varying float DLambda = select(Stiffness < MinStiffness, 0.f, ComputedDLambda);
+		const varying float DLambda = select(FinalStiffness < MinStiffness, 0.f, ComputedDLambda);
 
 		const FVector3f ComputedDelta = DLambda * Direction;
 
@@ -230,16 +234,18 @@ export void ApplyXPBDSpringConstraintsWithDampingAndWeightMaps(uniform FVector4f
 		const varying float Stiffness = StiffnessTable[StiffnessIndex];
 		#pragma ignore warning(perf)
 		const varying float DampingRatio = DampingTable[DampingIndex];
+		
+		const varying float FinalStiffness = select(Dist < FLOAT_SMALL_NUMBER, Stiffness, Stiffness / Dist);
 
-		const varying float Alpha = (uniform float)1.f / (Stiffness * Dt * Dt);
+		const varying float Alpha = (uniform float)1.f / (FinalStiffness * Dt * Dt);
 
 		const float CombinedInvMass = InvM2 + InvM1;
-		const varying float Damping = (uniform float)2.f * DampingRatio * sqrt(Stiffness / CombinedInvMass);
+		const varying float Damping = (uniform float)2.f * DampingRatio * sqrt(FinalStiffness / CombinedInvMass);
 		const varying float Gamma = Alpha * Damping * Dt;
 		const varying FVector3f RelativeVelocityTimesDt = P1 - X1 - P2 + X2;
 
 		const varying float ComputedDLambda = (Offset - Alpha * Lambda + Gamma * VectorDot(Direction, RelativeVelocityTimesDt)) / ((1.f + Gamma)*CombinedInvMass + Alpha);
-		const varying float DLambda = select(Stiffness < MinStiffness, 0.f, ComputedDLambda);
+		const varying float DLambda = select(FinalStiffness < MinStiffness, 0.f, ComputedDLambda);
 
 		const FVector3f ComputedDelta = DLambda * Direction;
 
diff --git a/Engine/Source/Runtime/Experimental/Chaos/Public/Chaos/XPBDSpringConstraints.h b/Engine/Source/Runtime/Experimental/Chaos/Public/Chaos/XPBDSpringConstraints.h
--- a/Engine/Source/Runtime/Experimental/Chaos/Public/Chaos/XPBDSpringConstraints.h
+++ b/Engine/Source/Runtime/Experimental/Chaos/Public/Chaos/XPBDSpringConstraints.h
@@ -9,7 +9,7 @@
 namespace Chaos::Softs
 {
 
-// Stiffness is in kg/s^2
+// Stiffness is in kg cm /s^2
 UE_DEPRECATED(5.2, "Use FXPBDSpringConstraints::MinStiffness instead.")
 static const FSolverReal XPBDSpringMinStiffness = (FSolverReal)1e-4; // Stiffness below this will be considered 0 since all of our calculations are actually based on 1 / stiffness.
 UE_DEPRECATED(5.2, "Use FXPBDSpringConstraints::MaxStiffness instead.")
@@ -119,7 +119,10 @@ private:
 
 		const FSolverReal CombinedInvMass = Particles.InvM(i2) + Particles.InvM(i1);
 
-		const FSolverReal Damping = DampingRatioValue * 2.f * FMath::Sqrt(StiffnessValue / CombinedInvMass);
+		// This scale factor makes things more resolution independent.
+		const FSolverReal FinalStiffnessValue = Dists[ConstraintIndex] < UE_SMALL_NUMBER ? StiffnessValue : StiffnessValue / Dists[ConstraintIndex];
+
+		const FSolverReal Damping = DampingRatioValue * 2.f * FMath::Sqrt(FinalStiffnessValue / CombinedInvMass);
 
 		const FSolverVec3& P1 = Particles.P(i1);
 		const FSolverVec3& P2 = Particles.P(i2);
@@ -134,7 +137,7 @@ private:
 
 
 		FSolverReal& Lambda = Lambdas[ConstraintIndex];
-		const FSolverReal Alpha = (FSolverReal)1.f / (StiffnessValue * Dt * Dt);
+		const FSolverReal Alpha = (FSolverReal)1.f / (FinalStiffnessValue * Dt * Dt);
 		const FSolverReal Gamma = Alpha * Damping * Dt;
 
 		const FSolverReal DLambda = (Offset - Alpha * Lambda + Gamma * FSolverVec3::DotProduct(Direction, RelativeVelocityTimesDt)) / (((FSolverReal)1.f + Gamma) * CombinedInvMass + Alpha);
