Iris
* Split NetObjectGridFilter into two classes:
* NetObjectGridFragmentLocFilter for objects with tagged locations
* NetObjectGridWorldLocFilter for objects that get their location via the FWorldLocations manager

diff --git a/Engine/Config/BaseEngine.ini b/Engine/Config/BaseEngine.ini
--- a/Engine/Config/BaseEngine.ini
+++ b/Engine/Config/BaseEngine.ini
@@ -1191,7 +1191,7 @@ MaxPartCount=4096
 +NetObjectPrioritizerDefinitions=(PrioritizerName=PlayerState, ClassName=/Script/IrisCore.NetObjectCountLimiter, ConfigClassName=/Script/Engine.PlayerStateCountLimiterConfig)
 
 [/Script/IrisCore.NetObjectFilterDefinitions]
-+NetObjectFilterDefinitions=(FilterName=Spatial, ClassName=/Script/IrisCore.NetObjectGridFilter, ConfigClassName=/Script/IrisCore.NetObjectGridFilterConfig)
++NetObjectFilterDefinitions=(FilterName=Spatial, ClassName=/Script/IrisCore.NetObjectGridWorldLocFilter, ConfigClassName=/Script/IrisCore.NetObjectGridFilterConfig)
 +NetObjectFilterDefinitions=(FilterName=NotRouted, ClassName=/Script/IrisCore.FilterOutNetObjectFilter, ConfigClassName=/Script/IrisCore.FilterOutNetObjectFilterConfig)
 
 [/Script/IrisCore.ObjectReplicationBridgeConfig]
diff --git a/Engine/Source/Runtime/Experimental/Iris/Core/Private/Iris/ReplicationSystem/Filtering/NetObjectGridFilter.cpp b/Engine/Source/Runtime/Experimental/Iris/Core/Private/Iris/ReplicationSystem/Filtering/NetObjectGridFilter.cpp
--- a/Engine/Source/Runtime/Experimental/Iris/Core/Private/Iris/ReplicationSystem/Filtering/NetObjectGridFilter.cpp
+++ b/Engine/Source/Runtime/Experimental/Iris/Core/Private/Iris/ReplicationSystem/Filtering/NetObjectGridFilter.cpp
@@ -18,7 +18,6 @@ void UNetObjectGridFilter::Init(FNetObjectFilterInitParams& Params)
 
 	PerConnectionInfos.SetNum(Params.MaxConnectionCount + 1);
 
-	WorldLocations = &Params.ReplicationSystem->GetWorldLocations();
 	NetCullDistanceOverrides = &Params.ReplicationSystem->GetNetCullDistanceOverrides();
 }
 
@@ -34,22 +33,12 @@ void UNetObjectGridFilter::RemoveConnection(uint32 ConnectionId)
 
 bool UNetObjectGridFilter::AddObject(uint32 ObjectIndex, FNetObjectFilterAddObjectParams& Params)
 {
+	
 	// We support either a world location in the state, tagged with RepTag_WorldLocation, or via the WorldLocations instance.
 	UE::Net::FRepTagFindInfo WorldLocationTagInfo;
-	bool bHasWorldLocation = false;
-	if (WorldLocations->HasWorldLocation(ObjectIndex))
-	{
-		bHasWorldLocation = true;
-		// Craft tag info that will let us know we need to retrieve the location from WorldLocations
-		WorldLocationTagInfo.StateIndex = InvalidStateIndex;
-		WorldLocationTagInfo.ExternalStateOffset = InvalidStateOffset;
-	}
-	else if (!UE::Net::FindRepTag(Params.Protocol, UE::Net::RepTag_WorldLocation, WorldLocationTagInfo))
-	{
-		return false;
-	}
+	const bool bHasLocation = BuildLocationTagInfo(WorldLocationTagInfo, ObjectIndex, Params);
 
-	if (!bHasWorldLocation && ((WorldLocationTagInfo.ExternalStateOffset >= MAX_uint16) || (WorldLocationTagInfo.StateIndex >= MAX_uint16)))
+	if (!bHasLocation)
 	{
 		return false;
 	}
@@ -83,6 +72,7 @@ bool UNetObjectGridFilter::AddObject(uint32 ObjectIndex, FNetObjectFilterAddObje
 	{
 		// Too big an object. We expect it to be costly to move it across cells.
 		RemoveObject(ObjectIndex, ObjectLocationInfo);
+		//TODO: Should we log these objects so they can get flagged and properly set always relevant ?
 		return false;
 	}
 
@@ -372,21 +362,9 @@ void UNetObjectGridFilter::UpdateCellInfoForObject(const FObjectLocationInfo& Ob
 
 void UNetObjectGridFilter::UpdatePositionAndCullDistance(const UNetObjectGridFilter::FObjectLocationInfo& ObjectLocationInfo, UNetObjectGridFilter::FPerObjectInfo& PerObjectInfo, const UE::Net::FReplicationInstanceProtocol* InstanceProtocol)
 {
-	TArrayView<const UE::Net::FReplicationInstanceProtocol::FFragmentData> FragmentDatas = MakeArrayView(InstanceProtocol->FragmentData, InstanceProtocol->FragmentCount);
-
 	// Update position
 	{
-		if (ObjectLocationInfo.IsUsingWorldLocations())
-		{
-			PerObjectInfo.Position = WorldLocations->GetWorldLocation(PerObjectInfo.ObjectIndex);
-		}
-		else
-		{
-			const UE::Net::FReplicationInstanceProtocol::FFragmentData& FragmentData = FragmentDatas[ObjectLocationInfo.GetLocationStateIndex()];
-			const uint8* LocationAddress = FragmentData.ExternalSrcBuffer + ObjectLocationInfo.GetLocationStateOffset();
-			const float* LocationComponents = reinterpret_cast<const float*>(LocationAddress);
-			PerObjectInfo.Position = FVector(LocationComponents[0], LocationComponents[1], LocationComponents[2]);
-		}
+		UpdateObjectPosition(PerObjectInfo, ObjectLocationInfo, InstanceProtocol);
 	}
 
 	// Optionally update cull distance
@@ -397,6 +375,8 @@ void UNetObjectGridFilter::UpdatePositionAndCullDistance(const UNetObjectGridFil
 	}
 	else if (PerObjectInfo.CullDistanceSqrStateOffset != InvalidStateOffset)
 	{
+		TArrayView<const UE::Net::FReplicationInstanceProtocol::FFragmentData> FragmentDatas = MakeArrayView(InstanceProtocol->FragmentData, InstanceProtocol->FragmentCount);
+
 		const UE::Net::FReplicationInstanceProtocol::FFragmentData& FragmentData = FragmentDatas[PerObjectInfo.CullDistanceSqrStateIndex];
 		const uint8* CullDistanceSqrAddress = FragmentData.ExternalSrcBuffer + PerObjectInfo.CullDistanceSqrStateOffset;
 		const float CullDistanceSqr = *reinterpret_cast<const float*>(CullDistanceSqrAddress);
@@ -441,3 +421,49 @@ bool UNetObjectGridFilter::DoesCellContainCoord(const UNetObjectGridFilter::FCel
 {
 	return (Coord.X >= Cell.MinX) & (Coord.X <= Cell.MaxX) & (Coord.Y >= Cell.MinY) & (Coord.Y <= Cell.MaxY);
 }
+
+// UNetObjectGridWorldLocFilter
+
+void UNetObjectGridWorldLocFilter::Init(FNetObjectFilterInitParams& Params)
+{
+	Super::Init(Params);
+
+	WorldLocations = &Params.ReplicationSystem->GetWorldLocations();
+}
+
+void UNetObjectGridWorldLocFilter::UpdateObjectPosition(UNetObjectGridFilter::FPerObjectInfo& PerObjectInfo, const UNetObjectGridFilter::FObjectLocationInfo& ObjectLocationInfo, const UE::Net::FReplicationInstanceProtocol* InstanceProtocol)
+{
+	check(ObjectLocationInfo.IsUsingWorldLocations());
+	
+	PerObjectInfo.Position = WorldLocations->GetWorldLocation(PerObjectInfo.ObjectIndex);
+}
+
+bool UNetObjectGridWorldLocFilter::BuildLocationTagInfo(UE::Net::FRepTagFindInfo& OutWorldLocationTagInfo, uint32 ObjectIndex, FNetObjectFilterAddObjectParams&)
+{
+	check(WorldLocations->HasWorldLocation(ObjectIndex));
+
+	// Craft tag info that will let us know we need to retrieve the location from WorldLocations
+	OutWorldLocationTagInfo.StateIndex = InvalidStateIndex;
+	OutWorldLocationTagInfo.ExternalStateOffset = InvalidStateOffset;
+	
+	return true;
+}
+
+// UNetObjectGridFragmentLocFilter
+
+void UNetObjectGridFragmentLocFilter::UpdateObjectPosition(UNetObjectGridFilter::FPerObjectInfo& PerObjectInfo, const UNetObjectGridFilter::FObjectLocationInfo& ObjectLocationInfo, const UE::Net::FReplicationInstanceProtocol* InstanceProtocol)
+{
+	check(ObjectLocationInfo.IsUsingWorldLocations() == false);
+
+	TArrayView<const UE::Net::FReplicationInstanceProtocol::FFragmentData> FragmentDatas = MakeArrayView(InstanceProtocol->FragmentData, InstanceProtocol->FragmentCount);
+
+	const UE::Net::FReplicationInstanceProtocol::FFragmentData& FragmentData = FragmentDatas[ObjectLocationInfo.GetLocationStateIndex()];
+	const uint8* LocationAddress = FragmentData.ExternalSrcBuffer + ObjectLocationInfo.GetLocationStateOffset();
+	const float* LocationComponents = reinterpret_cast<const float*>(LocationAddress);
+	PerObjectInfo.Position = FVector(LocationComponents[0], LocationComponents[1], LocationComponents[2]);
+}
+
+bool UNetObjectGridFragmentLocFilter::BuildLocationTagInfo(UE::Net::FRepTagFindInfo& OutWorldLocationTagInfo, uint32 ObjectIndex, FNetObjectFilterAddObjectParams& Params)
+{
+	return UE::Net::FindRepTag(Params.Protocol, UE::Net::RepTag_WorldLocation, OutWorldLocationTagInfo);
+}
\ No newline at end of file
diff --git a/Engine/Source/Runtime/Experimental/Iris/Core/Public/Iris/ReplicationSystem/Filtering/NetObjectGridFilter.h b/Engine/Source/Runtime/Experimental/Iris/Core/Public/Iris/ReplicationSystem/Filtering/NetObjectGridFilter.h
--- a/Engine/Source/Runtime/Experimental/Iris/Core/Public/Iris/ReplicationSystem/Filtering/NetObjectGridFilter.h
+++ b/Engine/Source/Runtime/Experimental/Iris/Core/Public/Iris/ReplicationSystem/Filtering/NetObjectGridFilter.h
@@ -13,6 +13,7 @@ namespace UE::Net
 {
 	class FNetCullDistanceOverrides;
 	class FWorldLocations;
+	struct FRepTagFindInfo;
 }
 
 UCLASS(transient, config=Engine, MinimalAPI)
@@ -84,7 +85,8 @@ protected:
 	IRISCORE_API void RemoveCellInfoForObject(const FObjectLocationInfo& ObjectInfo);
 	IRISCORE_API void UpdateCellInfoForObject(const FObjectLocationInfo& ObjectInfo, const UE::Net::FReplicationInstanceProtocol* InstanceProtocol);
 
-private:
+protected:
+
 	enum : unsigned
 	{
 		ObjectInfosChunkSize = 64 * 1024,
@@ -106,7 +108,7 @@ private:
 	// We can't fit all info we need in 4x16bits.
 	struct FPerObjectInfo
 	{
-		FVector Position = {0,0,0};
+		FVector Position = { 0,0,0 };
 		FCellBox CellBox = {};
 		float CullDistance = 0;
 		uint32 ObjectIndex = 0;
@@ -114,6 +116,13 @@ private:
 		uint16 CullDistanceSqrStateOffset = InvalidStateOffset;
 	};
 
+	/** Sets the current position of the object based on how we access it's given location */
+	virtual void UpdateObjectPosition(FPerObjectInfo& PerObjectInfo, const FObjectLocationInfo& ObjectLocationInfo, const UE::Net::FReplicationInstanceProtocol* InstanceProtocol) {}
+
+	virtual bool BuildLocationTagInfo(UE::Net::FRepTagFindInfo& OutWorldLocationTagInfo, uint32 ObjectIndex, FNetObjectFilterAddObjectParams&) { return false; }
+
+private:
+
 	struct FCellCoord
 	{
 		int32 X;
@@ -173,10 +182,48 @@ private:
 	TMap<FCellCoord, FCellObjects> Cells;
 	uint32 FrameIndex = 0;
 
-	const UE::Net::FWorldLocations* WorldLocations = nullptr;
 	const UE::Net::FNetCullDistanceOverrides* NetCullDistanceOverrides = nullptr;
 };
 
+/**
+ * Filter for replicated objects that have a WorldLocation reference (e.g. Actors).
+ * 
+ * This filter is more efficient since it's run before Polling and culls out objects that are not relevant to any connection.
+ */
+UCLASS()
+class UNetObjectGridWorldLocFilter : public UNetObjectGridFilter
+{
+	GENERATED_BODY()
+
+protected:
+
+	virtual void Init(FNetObjectFilterInitParams&) override;
+
+	virtual void UpdateObjectPosition(FPerObjectInfo& PerObjectInfo, const UNetObjectGridFilter::FObjectLocationInfo& ObjectLocationInfo, const UE::Net::FReplicationInstanceProtocol* InstanceProtocol) override;
+
+	virtual bool BuildLocationTagInfo(UE::Net::FRepTagFindInfo& OutWorldLocationTagInfo, uint32 ObjectIndex, FNetObjectFilterAddObjectParams&) override;
+
+private:
+
+	const UE::Net::FWorldLocations* WorldLocations = nullptr;
+};
+
+/**
+ * Filter for replicated objects that have their location stored in their fragment
+ * 
+ * This filter may be less efficient since it's run after Polling and DirtyData copying and cannot cull out objects from those operations.
+ */
+UCLASS()
+class UNetObjectGridFragmentLocFilter : public UNetObjectGridFilter
+{
+	GENERATED_BODY()
+
+protected:
+
+	virtual void UpdateObjectPosition(FPerObjectInfo& PerObjectInfo, const UNetObjectGridFilter::FObjectLocationInfo& ObjectLocationInfo, const UE::Net::FReplicationInstanceProtocol* InstanceProtocol) override;
+
+	virtual bool BuildLocationTagInfo(UE::Net::FRepTagFindInfo& OutWorldLocationTagInfo, uint32 ObjectIndex, FNetObjectFilterAddObjectParams&) override;
+};
 
 //
 inline bool UNetObjectGridFilter::FCellBox::operator==(const UNetObjectGridFilter::FCellBox& Other) const
