Geometry collection : add time scalar for removal ( cvar )

diff --git a/Engine/Source/Runtime/Experimental/GeometryCollectionEngine/Private/GeometryCollection/GeometryCollectionComponent.cpp b/Engine/Source/Runtime/Experimental/GeometryCollectionEngine/Private/GeometryCollection/GeometryCollectionComponent.cpp
--- a/Engine/Source/Runtime/Experimental/GeometryCollectionEngine/Private/GeometryCollection/GeometryCollectionComponent.cpp
+++ b/Engine/Source/Runtime/Experimental/GeometryCollectionEngine/Private/GeometryCollection/GeometryCollectionComponent.cpp
@@ -118,6 +118,9 @@ FAutoConsoleVariableRef CVarChaosGCInitConstantDataParallelForBatchSize(TEXT("p.
 int32 MaxGeometryCollectionAsyncPhysicsTickIdleTimeMs = 30;
 FAutoConsoleVariableRef CVarMaxGeometryCollectionAsyncPhysicsTickIdleTimeMs(TEXT("p.Chaos.GC.MaxGeometryCollectionAsyncPhysicsTickIdleTimeMs"), MaxGeometryCollectionAsyncPhysicsTickIdleTimeMs, TEXT("Amount of time in milliseconds before the async tick turns off when it is otherwise not doing anything."));
 
+float GeometryCollectionRemovalMultiplier = 1.0f;
+FAutoConsoleVariableRef CVarGeometryCollectionRemovalTimerMultiplier(TEXT("p.Chaos.GC.RemovalTimerMultiplier"), GeometryCollectionRemovalMultiplier, TEXT("Multiplier for the removal time evaluation ( > 1 : faster removal , > 1 slower"));
+
 DEFINE_LOG_CATEGORY_STATIC(UGCC_LOG, Error, All);
 
 extern FGeometryCollectionDynamicDataPool GDynamicDataPool;
@@ -2330,9 +2333,10 @@ void UGeometryCollectionComponent::TickComponent(float DeltaTime, enum ELevelTic
 	// todo(chaos) : cache root broken state ? 
 	if (IsRootBroken())
 	{
+		const float AdjustedDeltaTime = FMath::Max(GeometryCollectionRemovalMultiplier, 0.0001f) * DeltaTime;
 		// todo(chaos) : move removal logic on the physics thread
-		IncrementSleepTimer(DeltaTime);
-		IncrementBreakTimer(DeltaTime);
+		IncrementSleepTimer(AdjustedDeltaTime);
+		IncrementBreakTimer(AdjustedDeltaTime);
 	}
 }
 
