properly bookkeep explicit index on cluster unions

diff --git a/Engine/Source/Runtime/Experimental/Chaos/Private/Chaos/ClusterUnionManager.cpp b/Engine/Source/Runtime/Experimental/Chaos/Private/Chaos/ClusterUnionManager.cpp
--- a/Engine/Source/Runtime/Experimental/Chaos/Private/Chaos/ClusterUnionManager.cpp
+++ b/Engine/Source/Runtime/Experimental/Chaos/Private/Chaos/ClusterUnionManager.cpp
@@ -493,7 +493,10 @@ namespace Chaos
 			return *OutIndex;
 		}
 
-		FClusterUnionIndex NewIndex = CreateNewClusterUnion(DefaultClusterCreationParameters());
+		FClusterUnionCreationParameters OtherParams;
+		OtherParams.ExplicitIndex = InIndex;
+
+		FClusterUnionIndex NewIndex = CreateNewClusterUnion(DefaultClusterCreationParameters(), OtherParams);
 		ExplicitIndexMap.Add(InIndex, NewIndex);
 		return NewIndex;
 	}
