Dataflow : add a generate convex hull from direct children hull node

diff --git a/Engine/Plugins/Experimental/GeometryCollectionPlugin/Source/GeometryCollectionNodes/Private/Dataflow/GeometryCollectionUtilityNodes.cpp b/Engine/Plugins/Experimental/GeometryCollectionPlugin/Source/GeometryCollectionNodes/Private/Dataflow/GeometryCollectionUtilityNodes.cpp
--- a/Engine/Plugins/Experimental/GeometryCollectionPlugin/Source/GeometryCollectionNodes/Private/Dataflow/GeometryCollectionUtilityNodes.cpp
+++ b/Engine/Plugins/Experimental/GeometryCollectionPlugin/Source/GeometryCollectionNodes/Private/Dataflow/GeometryCollectionUtilityNodes.cpp
@@ -18,6 +18,7 @@ namespace Dataflow
 
 		DATAFLOW_NODE_REGISTER_CREATION_FACTORY(FCreateNonOverlappingConvexHullsDataflowNode);
 		DATAFLOW_NODE_REGISTER_CREATION_FACTORY(FGenerateClusterConvexHullsFromLeafHullsDataflowNode);
+		DATAFLOW_NODE_REGISTER_CREATION_FACTORY(FGenerateClusterConvexHullsFromChildrenHullsDataflowNode);
 	}
 }
 
@@ -97,6 +98,52 @@ void FGenerateClusterConvexHullsFromLeafHullsDataflowNode::Evaluate(Dataflow::FC
 				);
 			}
 
+			SetValue<FManagedArrayCollection>(Context, static_cast<const FManagedArrayCollection>(*GeomCollection), &Collection);
+		}
+	}
+}
+
+FGenerateClusterConvexHullsFromChildrenHullsDataflowNode::FGenerateClusterConvexHullsFromChildrenHullsDataflowNode(const Dataflow::FNodeParameters& InParam, FGuid InGuid)
+	: FDataflowNode(InParam, InGuid)
+{
+	RegisterInputConnection(&Collection);
+	RegisterInputConnection(&ConvexCount);
+	RegisterInputConnection(&ErrorTolerance);
+	RegisterInputConnection(&OptionalSelectionFilter);
+
+	RegisterOutputConnection(&Collection);
+}
+
+void FGenerateClusterConvexHullsFromChildrenHullsDataflowNode::Evaluate(Dataflow::FContext& Context, const FDataflowOutput* Out) const
+{
+	if (Out->IsA<FManagedArrayCollection>(&Collection))
+	{
+		const FManagedArrayCollection& InCollection = GetValue<FManagedArrayCollection>(Context, &Collection);
+
+		if (TUniquePtr<FGeometryCollection> GeomCollection = TUniquePtr<FGeometryCollection>(InCollection.NewCopy<FGeometryCollection>()))
+		{
+			const int32 InConvexCount = GetValue(Context, &ConvexCount);
+			const double InErrorToleranceInCm = GetValue(Context, &ErrorTolerance);
+			if (IsConnected(&OptionalSelectionFilter))
+			{
+				const FDataflowTransformSelection& InOptionalSelectionFilter = GetValue<FDataflowTransformSelection>(Context, &OptionalSelectionFilter);
+				FGeometryCollectionConvexUtility::GenerateClusterConvexHullsFromChildrenHulls(
+					*GeomCollection,
+					InConvexCount,
+					InErrorToleranceInCm,
+					InOptionalSelectionFilter.AsArray()
+				);
+			}
+			else
+			{
+				FGeometryCollectionConvexUtility::GenerateClusterConvexHullsFromChildrenHulls(
+					*GeomCollection,
+					InConvexCount,
+					InErrorToleranceInCm,
+					TArrayView<const int32>()
+				);
+			}
+
 			SetValue<FManagedArrayCollection>(Context, static_cast<const FManagedArrayCollection>(*GeomCollection), &Collection);
 		}
 	}
diff --git a/Engine/Plugins/Experimental/GeometryCollectionPlugin/Source/GeometryCollectionNodes/Public/Dataflow/GeometryCollectionUtilityNodes.h b/Engine/Plugins/Experimental/GeometryCollectionPlugin/Source/GeometryCollectionNodes/Public/Dataflow/GeometryCollectionUtilityNodes.h
--- a/Engine/Plugins/Experimental/GeometryCollectionPlugin/Source/GeometryCollectionNodes/Public/Dataflow/GeometryCollectionUtilityNodes.h
+++ b/Engine/Plugins/Experimental/GeometryCollectionPlugin/Source/GeometryCollectionNodes/Public/Dataflow/GeometryCollectionUtilityNodes.h
@@ -95,6 +95,43 @@ public:
 
 };
 
+/**
+ *
+ * Generates cluster convex hulls for children hulls
+ *
+ */
+USTRUCT(meta = (DataflowGeometryCollection))
+struct FGenerateClusterConvexHullsFromChildrenHullsDataflowNode : public FDataflowNode
+{
+	GENERATED_USTRUCT_BODY()
+	DATAFLOW_NODE_DEFINE_INTERNAL(FGenerateClusterConvexHullsFromChildrenHullsDataflowNode, "GenerateClusterConvexHullsFromChildrenHullsDataflowNode", "GeometryCollection|Utilities", "")
+
+public:
+	UPROPERTY(meta = (DataflowInput, DataflowOutput))
+	FManagedArrayCollection Collection;
+
+	/** Maximum number of convex to generate for a specific cluster. Will be ignored if error tolerance is used instead */
+	UPROPERTY(EditAnywhere, Category = "Convex", meta = (DataflowInput, EditCondition = "ErrorTolerance == 0"))
+	int32 ConvexCount = 2;
+
+	/**
+	* Error tolerance to use to decide to merge leaf convex together.
+	* This is in centimeters and represents the side of a cube, the volume of which will be used as threshold
+	* to know if the volume of the generated convex is too large compared to the sum of the volume of the leaf convex
+	*/
+	UPROPERTY(EditAnywhere, Category = "Convex", meta = (DataflowInput, UIMin = "0", UIMax = "100.", Units = cm))
+	double ErrorTolerance = 0.0;
+
+	/** Optional transform selection to compute cluster hulls on -- if not provided, all cluster hulls will be computed. */
+	UPROPERTY(meta = (DataflowInput, DataflowIntrinsic))
+	FDataflowTransformSelection OptionalSelectionFilter;
+
+	FGenerateClusterConvexHullsFromChildrenHullsDataflowNode(const Dataflow::FNodeParameters& InParam, FGuid InGuid = FGuid::NewGuid());
+
+	virtual void Evaluate(Dataflow::FContext& Context, const FDataflowOutput* Out) const override;
+
+};
+
 
 namespace Dataflow
 {
diff --git a/Engine/Source/Runtime/Experimental/Chaos/Private/GeometryCollection/GeometryCollectionConvexUtility.cpp b/Engine/Source/Runtime/Experimental/Chaos/Private/GeometryCollection/GeometryCollectionConvexUtility.cpp
--- a/Engine/Source/Runtime/Experimental/Chaos/Private/GeometryCollection/GeometryCollectionConvexUtility.cpp
+++ b/Engine/Source/Runtime/Experimental/Chaos/Private/GeometryCollection/GeometryCollectionConvexUtility.cpp
@@ -1450,16 +1450,20 @@ FGeometryCollectionConvexUtility::FGeometryCollectionConvexData FGeometryCollect
 }
 
 
+void FGeometryCollectionConvexUtility::GenerateClusterConvexHullsFromChildrenHulls(FGeometryCollection& Collection, int32 ConvexCount, double ErrorToleranceInCm, const TArrayView<const int32> TransformSubset)
+{
+	GenerateClusterConvexHullsFromLeafOrChildrenHullsInternal(Collection, ConvexCount, ErrorToleranceInCm, true, true/*bUseDirectChildren*/, TransformSubset);
+}
 void FGeometryCollectionConvexUtility::GenerateClusterConvexHullsFromLeafHulls(FGeometryCollection& Collection, int32 ConvexCount, double ErrorToleranceInCm, const TArrayView<const int32> TransformSubset)
 {
-	GenerateClusterConvexHullsFromLeafHullsInternal(Collection, ConvexCount, ErrorToleranceInCm, true, TransformSubset);
+	GenerateClusterConvexHullsFromLeafOrChildrenHullsInternal(Collection, ConvexCount, ErrorToleranceInCm, true, false/*bUseDirectChildren*/, TransformSubset);
 }
 void FGeometryCollectionConvexUtility::GenerateClusterConvexHullsFromLeafHulls(FGeometryCollection& Collection, int32 ConvexCount, double ErrorToleranceInCm)
 {
-	GenerateClusterConvexHullsFromLeafHullsInternal(Collection, ConvexCount, ErrorToleranceInCm, false, TArrayView<const int32>());
+	GenerateClusterConvexHullsFromLeafOrChildrenHullsInternal(Collection, ConvexCount, ErrorToleranceInCm, false, false/*bUseDirectChildren*/, TArrayView<const int32>());
 }
 
-void FGeometryCollectionConvexUtility::GenerateClusterConvexHullsFromLeafHullsInternal(FGeometryCollection& Collection,	int32 ConvexCount, double ErrorToleranceInCm, bool bOnlySubset, const TArrayView<const int32> TransformSubset)
+void FGeometryCollectionConvexUtility::GenerateClusterConvexHullsFromLeafOrChildrenHullsInternal(FGeometryCollection& Collection,	int32 ConvexCount, double ErrorToleranceInCm, bool bOnlySubset, bool bUseDirectChildren, const TArrayView<const int32> TransformSubset)
 {
 	static FName ConvexGroupName("Convex");
 	static FName ConvexHullAttributeName("ConvexHull");
@@ -1515,7 +1519,15 @@ void FGeometryCollectionConvexUtility::GenerateClusterConvexHullsFromLeafHullsIn
 			// compute using the leaf nodes 
 			// we cannot use the direct cluster children because they do not have proximity data that is stored at the geo level )
 			TArray<int32> SourceTransformIndices;
-			FGeometryCollectionClusteringUtility::GetLeafBones(&Collection, TransformIndex, true, SourceTransformIndices);
+			if (bUseDirectChildren)
+			{
+				SourceTransformIndices = HierarchyFacade.GetChildrenAsArray(TransformIndex);
+			}
+			else
+			{
+				FGeometryCollectionClusteringUtility::GetLeafBones(&Collection, TransformIndex, true, SourceTransformIndices);
+			}
+			
 
 			UE::Geometry::FDynamicMesh3 CombinedGeometry;
 			if (SourceTransformIndices.Num() > 0)
diff --git a/Engine/Source/Runtime/Experimental/Chaos/Public/GeometryCollection/GeometryCollectionConvexUtility.h b/Engine/Source/Runtime/Experimental/Chaos/Public/GeometryCollection/GeometryCollectionConvexUtility.h
--- a/Engine/Source/Runtime/Experimental/Chaos/Public/GeometryCollection/GeometryCollectionConvexUtility.h
+++ b/Engine/Source/Runtime/Experimental/Chaos/Public/GeometryCollection/GeometryCollectionConvexUtility.h
@@ -79,6 +79,7 @@ public:
 	static FGeometryCollectionConvexData CreateNonOverlappingConvexHullData(FGeometryCollection* GeometryCollection, double FractionAllowRemove = .3, double SimplificationDistanceThreshold = 0.0, double CanExceedFraction = .5, 
 		EConvexOverlapRemoval OverlapRemovalMethod = EConvexOverlapRemoval::All, double OverlapRemovalShrinkPercent = 0.0, UE::GeometryCollectionConvexUtility::FConvexHulls* ComputedLeafHullsToModify = nullptr);
 
+	static void GenerateClusterConvexHullsFromChildrenHulls(FGeometryCollection& Collection, int32 ConvexCount, double ErrorToleranceInCm, const TArrayView<const int32> TransformSubset);
 	static void GenerateClusterConvexHullsFromLeafHulls(FGeometryCollection& Collection, int32 ConvexCount, double ErrorToleranceInCm, const TArrayView<const int32> OptionalTransformSubset);
 	static void GenerateClusterConvexHullsFromLeafHulls(FGeometryCollection& Collection, int32 ConvexCount, double ErrorToleranceInCm);
 
@@ -148,6 +149,6 @@ private:
 	static void CreateConvexHullAttributesIfNeeded(FManagedArrayCollection& GeometryCollection);
 
 	// Implementation for GenerateClusterConvexHullsFromLeafHulls, supporting the full-collection and subset cases
-	static void GenerateClusterConvexHullsFromLeafHullsInternal(FGeometryCollection& Collection, int32 ConvexCount, double ErrorToleranceInCm, bool bOnlySubset, const TArrayView<const int32> OptionalTransformSubset);
+	static void GenerateClusterConvexHullsFromLeafOrChildrenHullsInternal(FGeometryCollection& Collection, int32 ConvexCount, double ErrorToleranceInCm, bool bOnlySubset, bool bUseDirectChildren, const TArrayView<const int32> OptionalTransformSubset);
 };
 
