Chaos - Collision Solver Optimizations
- eliminate RelativePosition from contacft data
- cache RxN etc and eliminate some a few cross products and dot products

diff --git a/Engine/Source/Programs/HeadlessChaos/Private/HeadlessChaosTestCollisions.cpp b/Engine/Source/Programs/HeadlessChaos/Private/HeadlessChaosTestCollisions.cpp
--- a/Engine/Source/Programs/HeadlessChaos/Private/HeadlessChaosTestCollisions.cpp
+++ b/Engine/Source/Programs/HeadlessChaos/Private/HeadlessChaosTestCollisions.cpp
@@ -479,7 +479,8 @@ namespace ChaosTest {
 			INVARIANT_XR_END(DynamicCube);
 		}
 
-		Collisions.ApplyPushOut(Dt, 1);
+		// If we need accurate restitution we need more velocity iterations
+		Collisions.ApplyPushOut(Dt, 4);
 
 		Collisions.ScatterOutput(Dt);
 
diff --git a/Engine/Source/Programs/HeadlessChaos/Private/HeadlessChaosTestContactModification.cpp b/Engine/Source/Programs/HeadlessChaos/Private/HeadlessChaosTestContactModification.cpp
--- a/Engine/Source/Programs/HeadlessChaos/Private/HeadlessChaosTestContactModification.cpp
+++ b/Engine/Source/Programs/HeadlessChaos/Private/HeadlessChaosTestContactModification.cpp
@@ -558,6 +558,8 @@ namespace ChaosTest
 			}
 		};
 
+		// If we rely on good restitution, we need more velocity iterations
+		Solver->GetEvolution()->SetNumVelocityIterations(4);
 
 		const float Dt = 0.1f;
 		const int32 Steps = 30;
@@ -568,7 +570,7 @@ namespace ChaosTest
 		}
 
 		
-		EXPECT_NEAR(NoBounceCubeParticle.V().Z, 0.f, KINDA_SMALL_NUMBER);
+		EXPECT_NEAR(NoBounceCubeParticle.V().Z, 0.f, 0.1);
 
 		Solver->UnregisterAndFreeSimCallbackObject_External(Callback);
 		Solver->UnregisterObject(NoBounceCubeProxy);
@@ -635,6 +637,8 @@ namespace ChaosTest
 			}
 		};
 
+		// If we rely on good restitution, we need more velocity iterations
+		Solver->GetEvolution()->SetNumVelocityIterations(4);
 
 		const float Dt = 0.1f;
 		const int32 Steps = 30;
@@ -644,7 +648,7 @@ namespace ChaosTest
 			Solver->UpdateGameThreadStructures();
 		}
 
-		EXPECT_NEAR(BounceCubeParticle.V().Z, 100.f, KINDA_SMALL_NUMBER);
+		EXPECT_NEAR(BounceCubeParticle.V().Z, 100.f, 0.1);
 
 		Solver->UnregisterAndFreeSimCallbackObject_External(Callback);
 		Solver->UnregisterObject(BounceCubeProxy);
diff --git a/Engine/Source/Runtime/Experimental/Chaos/Private/Chaos/Collision/PBDCollisionContainerSolver.cpp b/Engine/Source/Runtime/Experimental/Chaos/Private/Chaos/Collision/PBDCollisionContainerSolver.cpp
--- a/Engine/Source/Runtime/Experimental/Chaos/Private/Chaos/Collision/PBDCollisionContainerSolver.cpp
+++ b/Engine/Source/Runtime/Experimental/Chaos/Private/Chaos/Collision/PBDCollisionContainerSolver.cpp
@@ -293,9 +293,6 @@ namespace Chaos
 		, bCollisionConstraintPerIterationCollisionDetection()
 		, bPerIterationCollisionDetection(false)
 	{
-#if !UE_BUILD_SHIPPING && !UE_BUILD_TEST && INTEL_ISPC
-		Private::FPBDCollisionSolverHelper::CheckISPC();
-#endif
 	}
 
 	FPBDCollisionContainerSolver::~FPBDCollisionContainerSolver()
@@ -532,17 +529,17 @@ namespace Chaos
 		ApplyShockPropagation(ShockPropagation);
 	}
 
-	bool FPBDCollisionContainerSolver::SolvePositionImpl(const FReal InDt, const int32 It, const int32 NumIts, const int32 BeginIndex, const int32 EndIndex, const FPBDCollisionSolverSettings& SolverSettings)
+	void FPBDCollisionContainerSolver::SolvePositionImpl(const FReal InDt, const int32 It, const int32 NumIts, const int32 BeginIndex, const int32 EndIndex, const FPBDCollisionSolverSettings& SolverSettings)
 	{
 		SCOPE_CYCLE_COUNTER(STAT_Collisions_Apply);
 		if (!CVars::bChaos_PBDCollisionSolver_Position_SolveEnabled)
 		{
-			return false;
+			return;
 		}
 
 		if (EndIndex <= BeginIndex)
 		{
-			return false;
+			return;
 		}
 
 		UpdatePositionShockPropagation(InDt, It, NumIts, BeginIndex, EndIndex, SolverSettings);
@@ -565,27 +562,31 @@ namespace Chaos
 		// Apply the position correction
 		if (bApplyStaticFriction)
 		{
-			Private::FPBDCollisionSolverHelper::SolvePositionWithFriction(MakeArrayView(&GetSolver(BeginIndex), EndIndex - BeginIndex), Dt, MaxPushOut);
+			for (int32 SolverIndex = 0; SolverIndex < NumSolvers(); ++SolverIndex)
+			{
+				CollisionSolvers[SolverIndex].SolvePositionWithFriction(Dt, MaxPushOut);
+			}
 		}
 		else
 		{
-			Private::FPBDCollisionSolverHelper::SolvePositionNoFriction(MakeArrayView(&GetSolver(BeginIndex), EndIndex - BeginIndex), Dt, MaxPushOut);
+			for (int32 SolverIndex = 0; SolverIndex < NumSolvers(); ++SolverIndex)
+			{
+				CollisionSolvers[SolverIndex].SolvePositionNoFriction(Dt, MaxPushOut);
+			}
 		}
-
-		return true;
 	}
 
-	bool FPBDCollisionContainerSolver::SolveVelocityImpl(const FReal InDt, const int32 It, const int32 NumIts, const int32 BeginIndex, const int32 EndIndex, const FPBDCollisionSolverSettings& SolverSettings)
+	void FPBDCollisionContainerSolver::SolveVelocityImpl(const FReal InDt, const int32 It, const int32 NumIts, const int32 BeginIndex, const int32 EndIndex, const FPBDCollisionSolverSettings& SolverSettings)
 	{
 		SCOPE_CYCLE_COUNTER(STAT_Collisions_ApplyPushOut);
 		if (!CVars::bChaos_PBDCollisionSolver_Velocity_SolveEnabled)
 		{
-			return false;
+			return;
 		}
 
 		if (EndIndex <= BeginIndex)
 		{
-			return false;
+			return;
 		}
 
 		UpdateVelocityShockPropagation(InDt, It, NumIts, BeginIndex, EndIndex, SolverSettings);
@@ -593,9 +594,10 @@ namespace Chaos
 		const FSolverReal Dt = FSolverReal(InDt);
 		const bool bApplyDynamicFriction = (It >= NumIts - SolverSettings.NumVelocityFrictionIterations);
 
-		Private::FPBDCollisionSolverHelper::SolveVelocity(MakeArrayView(&GetSolver(BeginIndex), EndIndex - BeginIndex), Dt, bApplyDynamicFriction);
-
-		return true;
+		for (int32 SolverIndex = 0; SolverIndex < NumSolvers(); ++SolverIndex)
+		{
+			CollisionSolvers[SolverIndex].SolveVelocity(Dt, bApplyDynamicFriction);
+		}
 	}
 
 	void FPBDCollisionContainerSolver::UpdateCollisions(const FReal InDt, const int32 BeginIndex, const int32 EndIndex)
diff --git a/Engine/Source/Runtime/Experimental/Chaos/Private/Chaos/Collision/PBDCollisionSolver.cpp b/Engine/Source/Runtime/Experimental/Chaos/Private/Chaos/Collision/PBDCollisionSolver.cpp
--- a/Engine/Source/Runtime/Experimental/Chaos/Private/Chaos/Collision/PBDCollisionSolver.cpp
+++ b/Engine/Source/Runtime/Experimental/Chaos/Private/Chaos/Collision/PBDCollisionSolver.cpp
@@ -71,6 +71,7 @@ namespace Chaos
 	{
 		void FPBDCollisionSolver::SolveVelocityAverage(const FSolverReal Dt)
 		{
+			// @todo(chaos): A better solution would be to add an extra manifold point into array and use it for the position solver as well.
 			FSolverReal InvM0, InvM1;
 			FSolverMatrix33 InvI0, InvI1;
 			GetDynamicMassProperties(InvM0, InvI0, InvM1, InvI1);
@@ -118,19 +119,23 @@ namespace Chaos
 				// so I don't think it matters...we could do better though if it is a problem
 				AverageManifoldPoint.NetPushOutNormal = NetPushOutNormal;
 
+				const FSolverVec3 R0xN = FSolverVec3::CrossProduct(RelativeContactPosition0, WorldContactNormal);
+				AverageManifoldPoint.ContactRxNormal0 = R0xN;
+
+				const FSolverVec3 R1xN = FSolverVec3::CrossProduct(RelativeContactPosition1, WorldContactNormal);
+				AverageManifoldPoint.ContactRxNormal1 = R1xN;
+
 				// Calculate the contact mass (and derived properties) for this point
 				FSolverVec3 WorldContactNormalAngular0 = FSolverVec3(0);
 				FSolverVec3 WorldContactNormalAngular1 = FSolverVec3(0);
 				FSolverReal ContactMassInvNormal = FSolverReal(0);
 				if (IsDynamic(0))
 				{
-					const FSolverVec3 R0xN = FSolverVec3::CrossProduct(RelativeContactPosition0, WorldContactNormal);
 					WorldContactNormalAngular0 = InvI0 * R0xN;
 					ContactMassInvNormal += FSolverVec3::DotProduct(R0xN, WorldContactNormalAngular0) + InvM0;
 				}
 				if (IsDynamic(1))
 				{
-					const FSolverVec3 R1xN = FSolverVec3::CrossProduct(RelativeContactPosition1, WorldContactNormal);
 					WorldContactNormalAngular1 = InvI1 * R1xN;
 					ContactMassInvNormal += FSolverVec3::DotProduct(R1xN, WorldContactNormalAngular1) + InvM1;
 				}
@@ -144,6 +149,10 @@ namespace Chaos
 				AverageManifoldPoint.ContactDeltaNormal = FSolverReal(0);
 				AverageManifoldPoint.ContactDeltaTangentU = FSolverReal(0);
 				AverageManifoldPoint.ContactDeltaTangentV = FSolverReal(0);
+				AverageManifoldPoint.ContactRxTangentU0 = FVec3(0);
+				AverageManifoldPoint.ContactRxTangentV0 = FVec3(0);
+				AverageManifoldPoint.ContactRxTangentU1 = FVec3(0);
+				AverageManifoldPoint.ContactRxTangentV1 = FVec3(0);
 				AverageManifoldPoint.ContactTangentUAngular0 = FSolverVec3(0);
 				AverageManifoldPoint.ContactTangentVAngular0 = FSolverVec3(0);
 				AverageManifoldPoint.ContactTangentUAngular1 = FSolverVec3(0);
@@ -180,70 +189,5 @@ namespace Chaos
 			}
 		}
 
-		//////////////////////////////////////////////////////////////////////////////////////////////////
-		//////////////////////////////////////////////////////////////////////////////////////////////////
-		//////////////////////////////////////////////////////////////////////////////////////////////////
-		//////////////////////////////////////////////////////////////////////////////////////////////////
-
-		void CachePrefetchSolver(const TArrayView<FPBDCollisionSolver>& CollisionSolvers, const int32 ConstraintIndex)
-		{
-			//if (ConstraintIndex < CollisionSolvers.Num())
-			//{
-			//	if (ConstraintIndex < CollisionSolvers.Num() - 1)
-			//	{
-			//		FPlatformMisc::PrefetchBlock(&CollisionSolvers[ConstraintIndex + 1], sizeof(Private::FPBDCollisionSolver));
-			//	}
-
-			//	const int32 NumManifoldPoints = CollisionSolvers[ConstraintIndex].NumManifoldPoints();
-			//	if (NumManifoldPoints > 0)
-			//	{
-			//		FPlatformMisc::PrefetchBlock(&CollisionSolvers[ConstraintIndex].GetManifoldPoint(0), NumManifoldPoints * sizeof(Private::FPBDCollisionSolverManifoldPoint));
-			//	}
-			//}
-		}
-
-		void FPBDCollisionSolverHelper::SolvePositionNoFriction(const TArrayView<FPBDCollisionSolver>& CollisionSolvers, const FSolverReal Dt, const FSolverReal MaxPushOut)
-		{
-			SCOPE_CYCLE_COUNTER(STAT_SolvePositionNoFriction);
-
-			for (int32 SolverIndex = 0; SolverIndex < CollisionSolvers.Num(); ++SolverIndex)
-			{
-				CachePrefetchSolver(CollisionSolvers, SolverIndex + 1);
-
-				CollisionSolvers[SolverIndex].SolvePositionNoFriction(Dt, MaxPushOut);
-			}
-		}
-
-		void FPBDCollisionSolverHelper::SolvePositionWithFriction(const TArrayView<FPBDCollisionSolver>& CollisionSolvers, const FSolverReal Dt, const FSolverReal MaxPushOut)
-		{
-			for (int32 SolverIndex = 0; SolverIndex < CollisionSolvers.Num(); ++SolverIndex)
-			{
-				CachePrefetchSolver(CollisionSolvers, SolverIndex + 1);
-
-				CollisionSolvers[SolverIndex].SolvePositionWithFriction(Dt, MaxPushOut);
-			}
-		}
-
-		void FPBDCollisionSolverHelper::SolveVelocity(const TArrayView<FPBDCollisionSolver>& CollisionSolvers, const FSolverReal Dt, const bool bApplyDynamicFriction)
-		{
-			for (int32 SolverIndex = 0; SolverIndex < CollisionSolvers.Num(); ++SolverIndex)
-			{
-				CachePrefetchSolver(CollisionSolvers, SolverIndex + 1);
-
-				CollisionSolvers[SolverIndex].SolveVelocity(Dt, bApplyDynamicFriction);
-			}
-		}
-
-		void FPBDCollisionSolverHelper::CheckISPC()
-		{
-//#if INTEL_ISPC
-//			check(sizeof(ispc::FPBDCollisionSolver) == sizeof(Private::FPBDCollisionSolver));
-//			check(sizeof(ispc::FPBDCollisionSolverManifoldPoint) == sizeof(Private::FPBDCollisionSolverManifoldPoint));
-//			check(sizeof(ispc::FWorldContactPoint) == sizeof(FWorldContactPoint));
-//			check(sizeof(ispc::FConstraintSolverBody) == sizeof(FConstraintSolverBody));
-//			check(sizeof(ispc::FSolverBody) == sizeof(FSolverBody));
-//#endif
-		}
-
 	}	// namespace Private
 }	// namespace Chaos
diff --git a/Engine/Source/Runtime/Experimental/Chaos/Public/Chaos/Collision/PBDCollisionContainerSolver.h b/Engine/Source/Runtime/Experimental/Chaos/Public/Chaos/Collision/PBDCollisionContainerSolver.h
--- a/Engine/Source/Runtime/Experimental/Chaos/Public/Chaos/Collision/PBDCollisionContainerSolver.h
+++ b/Engine/Source/Runtime/Experimental/Chaos/Public/Chaos/Collision/PBDCollisionContainerSolver.h
@@ -66,8 +66,8 @@ namespace Chaos
 		void UpdatePositionShockPropagation(const FReal Dt, const int32 It, const int32 NumIts, const int32 BeginIndex, const int32 EndIndex, const FPBDCollisionSolverSettings& SolverSettings);
 		void UpdateVelocityShockPropagation(const FReal Dt, const int32 It, const int32 NumIts, const int32 BeginIndex, const int32 EndIndex, const FPBDCollisionSolverSettings& SolverSettings);
 		void ApplyShockPropagation(const FSolverReal ShockPropagation);
-		bool SolvePositionImpl(const FReal Dt, const int32 It, const int32 NumIts, const int32 BeginIndex, const int32 EndIndex, const FPBDCollisionSolverSettings& SolverSettings);
-		bool SolveVelocityImpl(const FReal Dt, const int32 It, const int32 NumIts, const int32 BeginIndex, const int32 EndIndex, const FPBDCollisionSolverSettings& SolverSettings);
+		void SolvePositionImpl(const FReal Dt, const int32 It, const int32 NumIts, const int32 BeginIndex, const int32 EndIndex, const FPBDCollisionSolverSettings& SolverSettings);
+		void SolveVelocityImpl(const FReal Dt, const int32 It, const int32 NumIts, const int32 BeginIndex, const int32 EndIndex, const FPBDCollisionSolverSettings& SolverSettings);
 		void UpdateCollisions(const FReal InDt, const int32 BeginIndex, const int32 EndIndex);
 
 		// The constraints we are solving and the container to which they belong
diff --git a/Engine/Source/Runtime/Experimental/Chaos/Public/Chaos/Collision/PBDCollisionSolver.h b/Engine/Source/Runtime/Experimental/Chaos/Public/Chaos/Collision/PBDCollisionSolver.h
--- a/Engine/Source/Runtime/Experimental/Chaos/Public/Chaos/Collision/PBDCollisionSolver.h
+++ b/Engine/Source/Runtime/Experimental/Chaos/Public/Chaos/Collision/PBDCollisionSolver.h
@@ -32,36 +32,37 @@ namespace Chaos
 		class FPBDCollisionSolverManifoldPoint
 		{
 		public:
-			// World-space contact point relative to each particle's center of mass
-			FVec3f RelativeContactPoints[2];
-
-			// World-space contact normal and tangents
-			FVec3f ContactNormal;
-			FVec3f ContactTangentU;
-			FVec3f ContactTangentV;
-
-			// Errors to correct along each of the contact axes
-			FRealSingle ContactDeltaNormal;
-			FRealSingle ContactDeltaTangentU;
-			FRealSingle ContactDeltaTangentV;
-
-			// Target velocity along the normal direction
-			FRealSingle ContactTargetVelocityNormal;
-
-			// I^-1.(R x A) for each body where A is each axis (Normal, TangentU, TangentV)
-			FSolverVec3 ContactNormalAngular0;
-			FSolverVec3 ContactTangentUAngular0;
-			FSolverVec3 ContactTangentVAngular0;
-			FSolverVec3 ContactNormalAngular1;
-			FSolverVec3 ContactTangentUAngular1;
-			FSolverVec3 ContactTangentVAngular1;
-
-			// Contact mass (for non-friction)
-			FSolverReal ContactMassNormal;
-			FSolverReal ContactMassTangentU;
-			FSolverReal ContactMassTangentV;
-
-			// Solver outputs
+			// @todo(chaos): the contact point is only needed for SolveVelocityAverage - try to remove it
+			FSolverVec3 RelativeContactPoints[2];		// R (World-space) Contact point relative to each particle's center of mass
+
+			FSolverVec3 ContactNormal;					// N (World Space)
+			FSolverVec3 ContactTangentU;				// U (World Space)
+			FSolverVec3 ContactTangentV;				// V (World Space)
+
+			FSolverReal ContactDeltaNormal;				// Initial separation along N
+			FSolverReal ContactDeltaTangentU;			// Initial separation along U
+			FSolverReal ContactDeltaTangentV;			// Initial separation along V
+
+			FSolverReal ContactTargetVelocityNormal;	// Normal velocity target
+
+			FSolverVec3 ContactRxNormal0;				// R0xN
+			FSolverVec3 ContactRxNormal1;				// R1xN
+			FSolverVec3 ContactRxTangentU0;				// R0xU
+			FSolverVec3 ContactRxTangentV0;				// R0xV
+			FSolverVec3 ContactRxTangentU1;				// R1xU
+			FSolverVec3 ContactRxTangentV1;				// R1xV
+
+			FSolverVec3 ContactNormalAngular0;			// InvI0.R0xN
+			FSolverVec3 ContactTangentUAngular0;		// InvI0.R0xU
+			FSolverVec3 ContactTangentVAngular0;		// InvI0.R0xV
+			FSolverVec3 ContactNormalAngular1;			// InvI1.R1xN
+			FSolverVec3 ContactTangentUAngular1;		// InvI1.R1xU
+			FSolverVec3 ContactTangentVAngular1;		// InvI1.R1xV
+
+			FSolverReal ContactMassNormal;				// 1/[NxR0.InvI0.R0xN + NxR1.InvI0.R1xN + InvM0 + InvM1]
+			FSolverReal ContactMassTangentU;			// 1/[UxR0.InvI0.R0xU + UxR1.InvI0.R1xU + InvM0 + InvM1]
+			FSolverReal ContactMassTangentV;			// 1/[VxR0.InvI0.R0xV + VxR1.InvI0.R1xV + InvM0 + InvM1]
+
 			FSolverReal NetPushOutNormal;
 			FSolverReal NetPushOutTangentU;
 			FSolverReal NetPushOutTangentV;
@@ -695,9 +696,9 @@ namespace Chaos
 
 				if (IsDynamic(0))
 				{
-					const FSolverVec3 R0xN = FSolverVec3::CrossProduct(ManifoldPoint.RelativeContactPoints[0], ManifoldPoint.ContactNormal);
-					const FSolverVec3 R0xU = FSolverVec3::CrossProduct(ManifoldPoint.RelativeContactPoints[0], ManifoldPoint.ContactTangentU);
-					const FSolverVec3 R0xV = FSolverVec3::CrossProduct(ManifoldPoint.RelativeContactPoints[0], ManifoldPoint.ContactTangentV);
+					const FSolverVec3& R0xN = ManifoldPoint.ContactRxNormal0;
+					const FSolverVec3& R0xU = ManifoldPoint.ContactRxTangentU0;
+					const FSolverVec3& R0xV = ManifoldPoint.ContactRxTangentV0;
 
 					ManifoldPoint.ContactNormalAngular0 = InvI0 * R0xN;
 					ManifoldPoint.ContactTangentUAngular0 = InvI0 * R0xU;
@@ -709,9 +710,9 @@ namespace Chaos
 				}
 				if (IsDynamic(1))
 				{
-					const FSolverVec3 R1xN = FSolverVec3::CrossProduct(ManifoldPoint.RelativeContactPoints[1], ManifoldPoint.ContactNormal);
-					const FSolverVec3 R1xU = FSolverVec3::CrossProduct(ManifoldPoint.RelativeContactPoints[1], ManifoldPoint.ContactTangentU);
-					const FSolverVec3 R1xV = FSolverVec3::CrossProduct(ManifoldPoint.RelativeContactPoints[1], ManifoldPoint.ContactTangentV);
+					const FSolverVec3& R1xN = ManifoldPoint.ContactRxNormal1;
+					const FSolverVec3& R1xU = ManifoldPoint.ContactRxTangentU1;
+					const FSolverVec3& R1xV = ManifoldPoint.ContactRxTangentV1;
 
 					ManifoldPoint.ContactNormalAngular1 = InvI1 * R1xN;
 					ManifoldPoint.ContactTangentUAngular1 = InvI1 * R1xU;
@@ -747,13 +748,13 @@ namespace Chaos
 				FSolverReal ContactMassInvNormal = FSolverReal(0);
 				if (IsDynamic(0))
 				{
-					const FSolverVec3 R0xN = FSolverVec3::CrossProduct(ManifoldPoint.RelativeContactPoints[0], ManifoldPoint.ContactNormal);
+					const FSolverVec3& R0xN = ManifoldPoint.ContactRxNormal0;
 					ManifoldPoint.ContactNormalAngular0 = InvI0 * R0xN;
 					ContactMassInvNormal += FSolverVec3::DotProduct(R0xN, ManifoldPoint.ContactNormalAngular0) + InvM0;
 				}
 				if (IsDynamic(1))
 				{
-					const FSolverVec3 R1xN = FSolverVec3::CrossProduct(ManifoldPoint.RelativeContactPoints[1], ManifoldPoint.ContactNormal);
+					const FSolverVec3& R1xN = ManifoldPoint.ContactRxNormal1;
 					ManifoldPoint.ContactNormalAngular1 = InvI1 * R1xN;
 					ContactMassInvNormal += FSolverVec3::DotProduct(R1xN, ManifoldPoint.ContactNormalAngular1) + InvM1;
 				}
@@ -770,17 +771,32 @@ namespace Chaos
 			const FConstraintSolverBody& Body1 = SolverBody1();
 
 			// Linear version: calculate the contact delta assuming linear motion after applying a positional impulse at the contact point. There will be an error that depends on the size of the rotation.
-			const FSolverVec3 ContactDelta0 = Body0.DP() + FSolverVec3::CrossProduct(Body0.DQ(), ManifoldPoint.RelativeContactPoints[0]);
-			const FSolverVec3 ContactDelta1 = Body1.DP() + FSolverVec3::CrossProduct(Body1.DQ(), ManifoldPoint.RelativeContactPoints[1]);
-			const FSolverVec3 ContactDelta = ContactDelta0 - ContactDelta1;
-			OutContactDeltaNormal = ManifoldPoint.ContactDeltaNormal + FSolverVec3::DotProduct(ContactDelta, ManifoldPoint.ContactNormal);
+			FSolverReal ContactDelta = ManifoldPoint.ContactDeltaNormal;
+			if (IsDynamic(0) && IsDynamic(1))
+			{
+				ContactDelta += FSolverVec3::DotProduct(Body0.DP() - Body1.DP(), ManifoldPoint.ContactNormal);
+				ContactDelta += FSolverVec3::DotProduct(Body0.DQ(), ManifoldPoint.ContactRxNormal0);
+				ContactDelta -= FSolverVec3::DotProduct(Body1.DQ(), ManifoldPoint.ContactRxNormal1);
+			}
+			else if (IsDynamic(0))
+			{
+				ContactDelta += FSolverVec3::DotProduct(Body0.DQ(), ManifoldPoint.ContactRxNormal0);
+				ContactDelta += FSolverVec3::DotProduct(Body0.DP(), ManifoldPoint.ContactNormal);
+			}
+			else if (IsDynamic(1))
+			{
+				ContactDelta -= FSolverVec3::DotProduct(Body1.DQ(), ManifoldPoint.ContactRxNormal1);
+				ContactDelta -= FSolverVec3::DotProduct(Body1.DP(), ManifoldPoint.ContactNormal);
+			}
 
-			// NOTE: OutContactDeltaNormal is negative for penetration
+			// NOTE: ContactDelta is negative for penetration
 			// NOTE: MaxPushOut == 0 disables the pushout limits
-			if ((MaxPushOut > 0) && (OutContactDeltaNormal < -MaxPushOut))
+			if ((MaxPushOut > 0) && (ContactDelta < -MaxPushOut))
 			{
-				OutContactDeltaNormal = -MaxPushOut;
+				ContactDelta = -MaxPushOut;
 			}
+
+			OutContactDeltaNormal = ContactDelta;
 		}
 
 		FORCEINLINE_DEBUGGABLE void FPBDCollisionSolver::CalculateContactPositionErrorTangential(
@@ -792,11 +808,37 @@ namespace Chaos
 			const FConstraintSolverBody& Body1 = SolverBody1();
 
 			// Linear version: calculate the contact delta assuming linear motion after applying a positional impulse at the contact point. There will be an error that depends on the size of the rotation.
-			const FSolverVec3 ContactDelta0 = Body0.DP() + FSolverVec3::CrossProduct(Body0.DQ(), ManifoldPoint.RelativeContactPoints[0]);
-			const FSolverVec3 ContactDelta1 = Body1.DP() + FSolverVec3::CrossProduct(Body1.DQ(), ManifoldPoint.RelativeContactPoints[1]);
-			const FSolverVec3 ContactDelta = ContactDelta0 - ContactDelta1;
-			OutContactDeltaTangentU = ManifoldPoint.ContactDeltaTangentU + FSolverVec3::DotProduct(ContactDelta, ManifoldPoint.ContactTangentU);
-			OutContactDeltaTangentV = ManifoldPoint.ContactDeltaTangentV + FSolverVec3::DotProduct(ContactDelta, ManifoldPoint.ContactTangentV);
+			FSolverReal ContactDeltaU = ManifoldPoint.ContactDeltaTangentU;
+			FSolverReal ContactDeltaV = ManifoldPoint.ContactDeltaTangentV;
+			if (IsDynamic(0) && IsDynamic(1))
+			{
+				ContactDeltaU += FSolverVec3::DotProduct(Body0.DP() - Body1.DP(), ManifoldPoint.ContactTangentU);
+				ContactDeltaU += FSolverVec3::DotProduct(Body0.DQ(), ManifoldPoint.ContactRxTangentU0);
+				ContactDeltaU -= FSolverVec3::DotProduct(Body1.DQ(), ManifoldPoint.ContactRxTangentU1);
+				
+				ContactDeltaV += FSolverVec3::DotProduct(Body0.DP() - Body1.DP(), ManifoldPoint.ContactTangentV);
+				ContactDeltaV += FSolverVec3::DotProduct(Body0.DQ(), ManifoldPoint.ContactRxTangentV0);
+				ContactDeltaV -= FSolverVec3::DotProduct(Body1.DQ(), ManifoldPoint.ContactRxTangentV1);
+			}
+			else if (IsDynamic(0))
+			{
+				ContactDeltaU += FSolverVec3::DotProduct(Body0.DP(), ManifoldPoint.ContactTangentU);
+				ContactDeltaU += FSolverVec3::DotProduct(Body0.DQ(), ManifoldPoint.ContactRxTangentU0);
+
+				ContactDeltaV += FSolverVec3::DotProduct(Body0.DP(), ManifoldPoint.ContactTangentV);
+				ContactDeltaV += FSolverVec3::DotProduct(Body0.DQ(), ManifoldPoint.ContactRxTangentV0);
+			}
+			else if (IsDynamic(1))
+			{
+				ContactDeltaU -= FSolverVec3::DotProduct(Body1.DP(), ManifoldPoint.ContactTangentU);
+				ContactDeltaU -= FSolverVec3::DotProduct(Body1.DQ(), ManifoldPoint.ContactRxTangentU1);
+
+				ContactDeltaV -= FSolverVec3::DotProduct(Body1.DP(), ManifoldPoint.ContactTangentV);
+				ContactDeltaV -= FSolverVec3::DotProduct(Body1.DQ(), ManifoldPoint.ContactRxTangentV1);
+			}
+
+			OutContactDeltaTangentU = ContactDeltaU;
+			OutContactDeltaTangentV = ContactDeltaV;
 		}
 
 		FORCEINLINE_DEBUGGABLE void FPBDCollisionSolver::CalculateContactVelocityError(
@@ -804,22 +846,32 @@ namespace Chaos
 			const FSolverReal DynamicFriction,
 			const FSolverReal Dt,
 			FSolverReal& OutContactVelocityDeltaNormal,
-			FSolverReal& OutContactVelocityDeltaTangent0,
-			FSolverReal& OutContactVelocityDeltaTangent1) const
+			FSolverReal& OutContactVelocityDeltaTangentU,
+			FSolverReal& OutContactVelocityDeltaTangentV) const
 		{
 			const FConstraintSolverBody& Body0 = SolverBody0();
 			const FConstraintSolverBody& Body1 = SolverBody1();
 
-			const FSolverVec3 ContactVelocity0 = Body0.V() + FSolverVec3::CrossProduct(Body0.W(), ManifoldPoint.RelativeContactPoints[0]);
-			const FSolverVec3 ContactVelocity1 = Body1.V() + FSolverVec3::CrossProduct(Body1.W(), ManifoldPoint.RelativeContactPoints[1]);
-			const FSolverVec3 ContactVelocity = ContactVelocity0 - ContactVelocity1;
-			const FSolverReal ContactVelocityNormal = FSolverVec3::DotProduct(ContactVelocity, ManifoldPoint.ContactNormal);
-			const FSolverReal ContactVelocityTangent0 = FSolverVec3::DotProduct(ContactVelocity, ManifoldPoint.ContactTangentU);
-			const FSolverReal ContactVelocityTangent1 = FSolverVec3::DotProduct(ContactVelocity, ManifoldPoint.ContactTangentV);
+			FSolverReal ContactVelN = -ManifoldPoint.ContactTargetVelocityNormal;
+			FSolverReal ContactVelU = FSolverReal(0);
+			FSolverReal ContactVelV = FSolverReal(0);
+
+			// @todo(chaos): check for static or stationary kinematics?
+			ContactVelN += FSolverVec3::DotProduct(Body0.V() - Body1.V(), ManifoldPoint.ContactNormal);
+			ContactVelN += FSolverVec3::DotProduct(Body0.W(), ManifoldPoint.ContactRxNormal0);
+			ContactVelN -= FSolverVec3::DotProduct(Body1.W(), ManifoldPoint.ContactRxNormal1);
+
+			ContactVelU += FSolverVec3::DotProduct(Body0.V() - Body1.V(), ManifoldPoint.ContactTangentU);
+			ContactVelU += FSolverVec3::DotProduct(Body0.W(), ManifoldPoint.ContactRxTangentU0);
+			ContactVelU -= FSolverVec3::DotProduct(Body1.W(), ManifoldPoint.ContactRxTangentU1);
 
-			OutContactVelocityDeltaNormal = (ContactVelocityNormal - ManifoldPoint.ContactTargetVelocityNormal);
-			OutContactVelocityDeltaTangent0 = ContactVelocityTangent0;
-			OutContactVelocityDeltaTangent1 = ContactVelocityTangent1;
+			ContactVelV += FSolverVec3::DotProduct(Body0.V() - Body1.V(), ManifoldPoint.ContactTangentV);
+			ContactVelV += FSolverVec3::DotProduct(Body0.W(), ManifoldPoint.ContactRxTangentV0);
+			ContactVelV -= FSolverVec3::DotProduct(Body1.W(), ManifoldPoint.ContactRxTangentV1);
+
+			OutContactVelocityDeltaNormal = ContactVelN;
+			OutContactVelocityDeltaTangentU = ContactVelU;
+			OutContactVelocityDeltaTangentV = ContactVelV;
 		}
 
 		FORCEINLINE_DEBUGGABLE void FPBDCollisionSolver::CalculateContactVelocityErrorNormal(
@@ -829,13 +881,12 @@ namespace Chaos
 			const FConstraintSolverBody& Body0 = SolverBody0();
 			const FConstraintSolverBody& Body1 = SolverBody1();
 
-			const FSolverVec3 ContactVelocity0 = Body0.V() + FSolverVec3::CrossProduct(Body0.W(), ManifoldPoint.RelativeContactPoints[0]);
-			const FSolverVec3 ContactVelocity1 = Body1.V() + FSolverVec3::CrossProduct(Body1.W(), ManifoldPoint.RelativeContactPoints[1]);
-			const FSolverVec3 ContactVelocity = ContactVelocity0 - ContactVelocity1;
-			const FSolverReal ContactVelocityNormal = FSolverVec3::DotProduct(ContactVelocity, ManifoldPoint.ContactNormal);
+			FSolverReal ContactVelN = -ManifoldPoint.ContactTargetVelocityNormal;
+			ContactVelN += FSolverVec3::DotProduct(Body0.V() - Body1.V(), ManifoldPoint.ContactNormal);
+			ContactVelN += FSolverVec3::DotProduct(Body0.W(), ManifoldPoint.ContactRxNormal0);
+			ContactVelN -= FSolverVec3::DotProduct(Body1.W(), ManifoldPoint.ContactRxNormal1);
 
-			// Add up the errors in the velocity (current velocity - desired velocity)
-			OutContactVelocityDeltaNormal = (ContactVelocityNormal - ManifoldPoint.ContactTargetVelocityNormal);
+			OutContactVelocityDeltaNormal = ContactVelN;
 		}
 
 		FORCEINLINE_DEBUGGABLE bool FPBDCollisionSolver::ShouldSolveVelocity(
@@ -862,6 +913,7 @@ namespace Chaos
 
 			ManifoldPoint.RelativeContactPoints[0] = InRelativeContactPosition0;
 			ManifoldPoint.RelativeContactPoints[1] = InRelativeContactPosition1;
+
 			ManifoldPoint.ContactNormal = InWorldContactNormal;
 			ManifoldPoint.ContactTangentU = InWorldContactTangentU;
 			ManifoldPoint.ContactTangentV = InWorldContactTangentV;
@@ -870,6 +922,13 @@ namespace Chaos
 			ManifoldPoint.ContactDeltaTangentV = InWorldContactDeltaTangentV;
 			ManifoldPoint.ContactTargetVelocityNormal = InWorldContactVelocityTargetNormal;
 
+			ManifoldPoint.ContactRxNormal0 = FSolverVec3::CrossProduct(InRelativeContactPosition0, InWorldContactNormal);
+			ManifoldPoint.ContactRxTangentU0 = FSolverVec3::CrossProduct(InRelativeContactPosition0, InWorldContactTangentU);
+			ManifoldPoint.ContactRxTangentV0 = FSolverVec3::CrossProduct(InRelativeContactPosition0, InWorldContactTangentV);
+			ManifoldPoint.ContactRxNormal1 = FSolverVec3::CrossProduct(InRelativeContactPosition1, InWorldContactNormal);
+			ManifoldPoint.ContactRxTangentU1 = FSolverVec3::CrossProduct(InRelativeContactPosition1, InWorldContactTangentU);
+			ManifoldPoint.ContactRxTangentV1 = FSolverVec3::CrossProduct(InRelativeContactPosition1, InWorldContactTangentV);
+
 			ManifoldPoint.NetPushOutNormal = FSolverReal(0);
 			ManifoldPoint.NetPushOutTangentU = FSolverReal(0);
 			ManifoldPoint.NetPushOutTangentV = FSolverReal(0);
@@ -1039,24 +1098,6 @@ namespace Chaos
 			}
 		}
 
-		//////////////////////////////////////////////////////////////////////////////////////////////////
-		//////////////////////////////////////////////////////////////////////////////////////////////////
-		//////////////////////////////////////////////////////////////////////////////////////////////////
-		//////////////////////////////////////////////////////////////////////////////////////////////////
-
-		/**
-		 * A helper for solving arrays of constraints
-		 */
-		class FPBDCollisionSolverHelper
-		{
-		public:
-			static void SolvePositionNoFriction(const TArrayView<FPBDCollisionSolver>& CollisionSolvers, const FSolverReal Dt, const FSolverReal MaxPushOut);
-			static void SolvePositionWithFriction(const TArrayView<FPBDCollisionSolver>& CollisionSolvers, const FSolverReal Dt, const FSolverReal MaxPushOut);
-			static void SolveVelocity(const TArrayView<FPBDCollisionSolver>& CollisionSolvers, const FSolverReal Dt, const bool bApplyDynamicFriction);
-
-			static void CheckISPC();
-		};
-
 	}	// namespace Private
 }	// namespace Chaos
 
