dataflow : make sure proximity is passed to the convex decompositioon algorithm when using direct children for generating convex hulls

diff --git a/Engine/Source/Runtime/Experimental/Chaos/Private/GeometryCollection/GeometryCollectionConvexUtility.cpp b/Engine/Source/Runtime/Experimental/Chaos/Private/GeometryCollection/GeometryCollectionConvexUtility.cpp
--- a/Engine/Source/Runtime/Experimental/Chaos/Private/GeometryCollection/GeometryCollectionConvexUtility.cpp
+++ b/Engine/Source/Runtime/Experimental/Chaos/Private/GeometryCollection/GeometryCollectionConvexUtility.cpp
@@ -1583,22 +1583,52 @@ void FGeometryCollectionConvexUtility::GenerateClusterConvexHullsFromLeafOrChild
 					}
 				}
 				// get hull proximity out of geometry proximity.  (Note: Only works if merging leaf convexes ... need a more general source of proximity!)
-				for (int32 SourceTransformIndex : SourceTransformIndices)
+				if (bUseDirectChildren)
+				{
+					// when using direct childre, proximity data may not be available ( if children are not leaves ) 
+					// so we use declare everything connected to everything for now 
+					// @todo this is suboptimal in some cases and we could precompute proximity and use it instead 
+					for (int32 IndexA = 0; IndexA < SourceTransformIndices.Num(); IndexA++)
+					{
+						const int32 TransformIndexA = SourceTransformIndices[IndexA];
+						if (const TArray<int32>* HullIndicesA = TransformToHullIdx.Find(TransformIndexA))
+						{
+							for (int32 IndexB = IndexA + 1; IndexB < SourceTransformIndices.Num(); IndexB++)
+							{
+								const int32 TransformIndexB = SourceTransformIndices[IndexB];
+								if (const TArray<int32>* HullIndicesB = TransformToHullIdx.Find(TransformIndexB))
+								{
+									for (int32 SourceHullIdxA : (*HullIndicesA))
+									{
+										for (int32 SourceHullIdxB : (*HullIndicesB))
+										{
+											HullProximity.Emplace(SourceHullIdxA, SourceHullIdxB);
+										}
+									}
+								}
+							}
+						}
+					}
+				}
+				else
 				{
-					const TArray<int32, TInlineAllocator<8>> SourceHullIndices{ TransformToHullIdx.FindOrAdd(SourceTransformIndex) };
-					const int32 GeoIdx = Collection.TransformToGeometryIndex[SourceTransformIndex];
-					if (GeoIdx > INDEX_NONE)
+					for (int32 SourceTransformIndex : SourceTransformIndices)
 					{
-						for (int32 NeighborGeoIndex : Proximity[GeoIdx])
+						const TArray<int32, TInlineAllocator<8>> SourceHullIndices{ TransformToHullIdx.FindOrAdd(SourceTransformIndex) };
+						const int32 GeoIdx = Collection.TransformToGeometryIndex[SourceTransformIndex];
+						if (GeoIdx > INDEX_NONE)
 						{
-							const int32 NeightborTransformIdx = Collection.TransformIndex[NeighborGeoIndex];
-							if (const TArray<int32>* NeighborHullIndices = TransformToHullIdx.Find(NeightborTransformIdx))
+							for (int32 NeighborGeoIndex : Proximity[GeoIdx])
 							{
-								for (int32 SourceHullIdx : SourceHullIndices)
+								const int32 NeightborTransformIdx = Collection.TransformIndex[NeighborGeoIndex];
+								if (const TArray<int32>* NeighborHullIndices = TransformToHullIdx.Find(NeightborTransformIdx))
 								{
-									for (int32 NeighborHullIdx : (*NeighborHullIndices))
+									for (int32 SourceHullIdx : SourceHullIndices)
 									{
-										HullProximity.Emplace(SourceHullIdx, NeighborHullIdx);
+										for (int32 NeighborHullIdx : (*NeighborHullIndices))
+										{
+											HullProximity.Emplace(SourceHullIdx, NeighborHullIdx);
+										}
 									}
 								}
 							}
