Better parameterization of ISM pool components.
Still work to do for recycling components that are empty.

diff --git a/Engine/Source/Runtime/Experimental/GeometryCollectionEngine/Private/GeometryCollection/GeometryCollectionComponent.cpp b/Engine/Source/Runtime/Experimental/GeometryCollectionEngine/Private/GeometryCollection/GeometryCollectionComponent.cpp
--- a/Engine/Source/Runtime/Experimental/GeometryCollectionEngine/Private/GeometryCollection/GeometryCollectionComponent.cpp
+++ b/Engine/Source/Runtime/Experimental/GeometryCollectionEngine/Private/GeometryCollection/GeometryCollectionComponent.cpp
@@ -2406,11 +2406,7 @@ void UGeometryCollectionComponent::OnRegister()
 	if (Type && Type->ImplementsInterface(UGeometryCollectionExternalRenderInterface::StaticClass()))
 	{
 		CustomRenderer = (UGeometryCollectionExternalRenderInterface*)NewObject<UObject>(this, Type);
-
-		if (IGeometryCollectionExternalRenderInterface* RendererInterface = Cast<IGeometryCollectionExternalRenderInterface>(CustomRenderer))
-		{
-			RendererInterface->OnRegisterGeometryCollection(*this);
-		}
+		RegisterCustomRenderer();
 	}
 
 	Super::OnRegister();
@@ -2423,14 +2419,40 @@ void UGeometryCollectionComponent::OnUnregister()
 	// Remove any custom renderer.
 	if (CustomRenderer)
 	{
-		if (IGeometryCollectionExternalRenderInterface* RendererInterface = Cast<IGeometryCollectionExternalRenderInterface>(CustomRenderer))
-		{
-			RendererInterface->OnUnregisterGeometryCollection();
-		}
+		UnregisterCustomRenderer();
 		CustomRenderer = nullptr;
 	}
 }
 
+void UGeometryCollectionComponent::RegisterCustomRenderer()
+{
+	if (IGeometryCollectionExternalRenderInterface* RendererInterface = Cast<IGeometryCollectionExternalRenderInterface>(CustomRenderer))
+	{
+		RendererInterface->OnRegisterGeometryCollection(*this);
+	}
+}
+
+void UGeometryCollectionComponent::UnregisterCustomRenderer()
+{
+	if (IGeometryCollectionExternalRenderInterface* RendererInterface = Cast<IGeometryCollectionExternalRenderInterface>(CustomRenderer))
+	{
+		RendererInterface->OnUnregisterGeometryCollection();
+	}
+}
+
+void UGeometryCollectionComponent::ReregisterAllCustomRenderers()
+{
+	for (TObjectIterator<UGeometryCollectionComponent> It; It; ++It)
+	{
+		It->UnregisterCustomRenderer();
+	}
+	for (TObjectIterator<UGeometryCollectionComponent> It; It; ++It)
+	{
+		It->RegisterCustomRenderer();
+		It->RefreshCustomRenderer();
+	}
+}
+
 void UGeometryCollectionComponent::ResetDynamicCollection()
 {
 	bool bCreateDynamicCollection = true;
diff --git a/Engine/Source/Runtime/Experimental/GeometryCollectionEngine/Private/GeometryCollection/GeometryCollectionISMPoolComponent.cpp b/Engine/Source/Runtime/Experimental/GeometryCollectionEngine/Private/GeometryCollection/GeometryCollectionISMPoolComponent.cpp
--- a/Engine/Source/Runtime/Experimental/GeometryCollectionEngine/Private/GeometryCollection/GeometryCollectionISMPoolComponent.cpp
+++ b/Engine/Source/Runtime/Experimental/GeometryCollectionEngine/Private/GeometryCollection/GeometryCollectionISMPoolComponent.cpp
@@ -43,35 +43,53 @@ void FGeometryCollectionMeshGroup::RemoveAllMeshes(FGeometryCollectionISMPool& I
 	Meshes.Empty();
 }
 
-FGeometryCollectionISM::FGeometryCollectionISM(AActor* OwmingActor, const FGeometryCollectionStaticMeshInstance& MeshInstance, bool bPreferHISM)
+FGeometryCollectionISM::FGeometryCollectionISM(AActor* OwmingActor, const FGeometryCollectionStaticMeshInstance& MeshInstance)
 {
 	check(MeshInstance.StaticMesh);
 	check(OwmingActor);
 
-	const FName ISMName = MakeUniqueObjectName(OwmingActor, bPreferHISM? UHierarchicalInstancedStaticMeshComponent::StaticClass(): UInstancedStaticMeshComponent::StaticClass(), MeshInstance.StaticMesh->GetFName());
-	UInstancedStaticMeshComponent* ISMC = bPreferHISM
-		? NewObject<UHierarchicalInstancedStaticMeshComponent>(OwmingActor, ISMName, RF_Transient | RF_DuplicateTransient)
-		: NewObject<UInstancedStaticMeshComponent>(OwmingActor, ISMName, RF_Transient | RF_DuplicateTransient)
-		;
-	if (ISMC)
+	UHierarchicalInstancedStaticMeshComponent* HISMC = nullptr;
+	UInstancedStaticMeshComponent* ISMC = nullptr;
+	
+	if (MeshInstance.Desc.bUseHISM)
 	{
-		ISMC->SetStaticMesh(MeshInstance.StaticMesh);
-		// material overrides
-		for (int32 MaterialIndex = 0; MaterialIndex < MeshInstance.MaterialsOverrides.Num(); MaterialIndex++)
-		{
-			ISMC->SetMaterial(MaterialIndex, MeshInstance.MaterialsOverrides[MaterialIndex]);
-		}
-		ISMC->NumCustomDataFloats = MeshInstance.NumCustomDataFloats;
-		ISMC->SetCullDistances(0, 0);
-		ISMC->SetCanEverAffectNavigation(false);
-		ISMC->SetCollisionEnabled(ECollisionEnabled::NoCollision);
-		ISMC->SetCastShadow(true);
-		ISMC->SetMobility(EComponentMobility::Stationary);
-		OwmingActor->AddInstanceComponent(ISMC);
-		ISMC->RegisterComponent();
-		
-		ISMComponent = ISMC;
+		const FName ISMName = MakeUniqueObjectName(OwmingActor, UHierarchicalInstancedStaticMeshComponent::StaticClass(), MeshInstance.StaticMesh->GetFName());
+		ISMC = HISMC = NewObject<UHierarchicalInstancedStaticMeshComponent>(OwmingActor, ISMName, RF_Transient | RF_DuplicateTransient);
+	}
+	else
+	{
+		const FName ISMName = MakeUniqueObjectName(OwmingActor, UInstancedStaticMeshComponent::StaticClass(), MeshInstance.StaticMesh->GetFName());
+		ISMC = NewObject<UInstancedStaticMeshComponent>(OwmingActor, ISMName, RF_Transient | RF_DuplicateTransient);
+	}
+
+	if (!ensure(ISMC != nullptr))
+	{
+		return;
+	}
+
+	ISMC->SetStaticMesh(MeshInstance.StaticMesh);
+	for (int32 MaterialIndex = 0; MaterialIndex < MeshInstance.MaterialsOverrides.Num(); MaterialIndex++)
+	{
+		ISMC->SetMaterial(MaterialIndex, MeshInstance.MaterialsOverrides[MaterialIndex]);
+	}
+
+	ISMC->NumCustomDataFloats = MeshInstance.Desc.NumCustomDataFloats;
+	ISMC->SetMobility(MeshInstance.Desc.bIsStaticMobility ? EComponentMobility::Static : EComponentMobility::Stationary);
+	ISMC->SetCullDistances(MeshInstance.Desc.StartCullDistance, MeshInstance.Desc.EndCullDistance);
+	ISMC->SetCastShadow(MeshInstance.Desc.bAffectShadow);
+	ISMC->bAffectDynamicIndirectLighting = MeshInstance.Desc.bAffectDynamicIndirectLighting;
+	ISMC->bAffectDistanceFieldLighting = MeshInstance.Desc.bAffectDistanceFieldLighting;
+	ISMC->SetCanEverAffectNavigation(false);
+	ISMC->SetCollisionEnabled(ECollisionEnabled::NoCollision);
+	
+	if (HISMC)
+	{
+		HISMC->SetLODDistanceScale(MeshInstance.Desc.LodScale);
 	}
+	
+	OwmingActor->AddInstanceComponent(ISMC);
+	ISMC->RegisterComponent();
+	ISMComponent = ISMC;
 }
 
 int32 FGeometryCollectionISM::AddInstanceGroup(int32 InstanceCount, TArrayView<const float> CustomDataFloats)
@@ -105,14 +123,14 @@ int32 FGeometryCollectionISM::AddInstanceGroup(int32 InstanceCount, TArrayView<c
 	return InstanceGroupIndex;
 }
 
-FGeometryCollectionMeshInfo FGeometryCollectionISMPool::AddISM(UGeometryCollectionISMPoolComponent* OwningComponent, const FGeometryCollectionStaticMeshInstance& MeshInstance, int32 InstanceCount, TArrayView<const float> CustomDataFloats, bool bPreferHISM)
+FGeometryCollectionMeshInfo FGeometryCollectionISMPool::AddISM(UGeometryCollectionISMPoolComponent* OwningComponent, const FGeometryCollectionStaticMeshInstance& MeshInstance, int32 InstanceCount, TArrayView<const float> CustomDataFloats)
 {
 	FGeometryCollectionMeshInfo Info;
 
 	FISMIndex* ISMIndex = MeshToISMIndex.Find(MeshInstance);
 	if (!ISMIndex)
 	{
-		Info.ISMIndex = ISMs.Emplace(OwningComponent->GetOwner(), MeshInstance, bPreferHISM);
+		Info.ISMIndex = ISMs.Emplace(OwningComponent->GetOwner(), MeshInstance);
 		MeshToISMIndex.Add(MeshInstance, Info.ISMIndex);
 		ISMComponentToISMIndex.Add(ISMs[Info.ISMIndex].ISMComponent, Info.ISMIndex);
 	}
@@ -246,11 +264,11 @@ void UGeometryCollectionISMPoolComponent::DestroyMeshGroup(FMeshGroupId MeshGrou
 	}
 }
 
-UGeometryCollectionISMPoolComponent::FMeshId UGeometryCollectionISMPoolComponent::AddMeshToGroup(FMeshGroupId MeshGroupId, const FGeometryCollectionStaticMeshInstance& MeshInstance, int32 InstanceCount, TArrayView<const float> CustomDataFloats, bool bPreferHISM)
+UGeometryCollectionISMPoolComponent::FMeshId UGeometryCollectionISMPoolComponent::AddMeshToGroup(FMeshGroupId MeshGroupId, const FGeometryCollectionStaticMeshInstance& MeshInstance, int32 InstanceCount, TArrayView<const float> CustomDataFloats)
 {
 	if (FGeometryCollectionMeshGroup* MeshGroup = MeshGroups.Find(MeshGroupId))
 	{
-		const FGeometryCollectionMeshInfo ISMInstanceInfo = Pool.AddISM(this, MeshInstance, InstanceCount, CustomDataFloats, bPreferHISM);
+		const FGeometryCollectionMeshInfo ISMInstanceInfo = Pool.AddISM(this, MeshInstance, InstanceCount, CustomDataFloats);
 		return MeshGroup->AddMesh(MeshInstance, InstanceCount, ISMInstanceInfo);
 	}
 	UE_LOG(LogChaos, Warning, TEXT("UGeometryCollectionISMPoolComponent : Trying to add a mesh to a mesh group (%d) that does not exists"), MeshGroupId);
diff --git a/Engine/Source/Runtime/Experimental/GeometryCollectionEngine/Private/GeometryCollection/GeometryCollectionISMPoolRenderer.cpp b/Engine/Source/Runtime/Experimental/GeometryCollectionEngine/Private/GeometryCollection/GeometryCollectionISMPoolRenderer.cpp
--- a/Engine/Source/Runtime/Experimental/GeometryCollectionEngine/Private/GeometryCollection/GeometryCollectionISMPoolRenderer.cpp
+++ b/Engine/Source/Runtime/Experimental/GeometryCollectionEngine/Private/GeometryCollection/GeometryCollectionISMPoolRenderer.cpp
@@ -84,9 +84,10 @@ void UGeometryCollectionISMPoolRenderer::InitMergedMeshFromGeometryCollection(UG
 
 		FGeometryCollectionStaticMeshInstance StaticMeshInstance;
 		StaticMeshInstance.StaticMesh = StaticMesh;
+		StaticMeshInstance.Desc.bUseHISM = true;
 
 		TArray<float> DummyCustomData;
-		MergedMeshGroup.MeshIds.Add(ISMPoolComponent->AddMeshToGroup(MergedMeshGroup.GroupIndex, StaticMeshInstance, 1, DummyCustomData, true));
+		MergedMeshGroup.MeshIds.Add(ISMPoolComponent->AddMeshToGroup(MergedMeshGroup.GroupIndex, StaticMeshInstance, 1, DummyCustomData));
 	}
 }
 
@@ -128,9 +129,10 @@ void UGeometryCollectionISMPoolRenderer::InitInstancesFromGeometryCollection(UGe
 			{
 				StaticMeshInstance.MaterialsOverrides = AutoInstanceMesh.Materials;
 			}
+			StaticMeshInstance.Desc.bUseHISM = true;
 
 			TArray<float> DummyCustomData;
-			InstancesGroup.MeshIds.Add(ISMPoolComponent->AddMeshToGroup(InstancesGroup.GroupIndex, StaticMeshInstance, AutoInstanceMesh.NumInstances, DummyCustomData, true));
+			InstancesGroup.MeshIds.Add(ISMPoolComponent->AddMeshToGroup(InstancesGroup.GroupIndex, StaticMeshInstance, AutoInstanceMesh.NumInstances, DummyCustomData));
 		}
 	}
 }
diff --git a/Engine/Source/Runtime/Experimental/GeometryCollectionEngine/Public/GeometryCollection/GeometryCollectionComponent.h b/Engine/Source/Runtime/Experimental/GeometryCollectionEngine/Public/GeometryCollection/GeometryCollectionComponent.h
--- a/Engine/Source/Runtime/Experimental/GeometryCollectionEngine/Public/GeometryCollection/GeometryCollectionComponent.h
+++ b/Engine/Source/Runtime/Experimental/GeometryCollectionEngine/Public/GeometryCollection/GeometryCollectionComponent.h
@@ -1015,6 +1015,9 @@ public:
 
 	TArray<UStaticMeshComponent*> CreateProxyComponents() const;
 
+	/** Force all GC components to reregister their custom renderer objects. */
+	static void ReregisterAllCustomRenderers();
+
 public:
 	UPROPERTY(BlueprintAssignable, Category = "Collision")
 	FOnChaosPhysicsCollision OnChaosPhysicsCollision;
@@ -1256,7 +1259,9 @@ private:
 
 	/** return true if a a custom renderer has been set and the feature is enabled */
 	bool CanUseCustomRenderer() const;
-	/** update the custom renderer */
+
+	void RegisterCustomRenderer();
+	void UnregisterCustomRenderer();
 	void RefreshCustomRenderer();
 
 	/** return true if the root cluster is not longer active at runtime */
diff --git a/Engine/Source/Runtime/Experimental/GeometryCollectionEngine/Public/GeometryCollection/GeometryCollectionISMPoolComponent.h b/Engine/Source/Runtime/Experimental/GeometryCollectionEngine/Public/GeometryCollection/GeometryCollectionISMPoolComponent.h
--- a/Engine/Source/Runtime/Experimental/GeometryCollectionEngine/Public/GeometryCollection/GeometryCollectionISMPoolComponent.h
+++ b/Engine/Source/Runtime/Experimental/GeometryCollectionEngine/Public/GeometryCollection/GeometryCollectionISMPoolComponent.h
@@ -2,12 +2,11 @@
 
 #pragma once
 
-#if 1
-
 #include "Components/SceneComponent.h"
 #include "Containers/Map.h"
-#include "Materials/MaterialInterface.h"
 #include "InstancedStaticMeshDelegates.h"
+#include "Materials/MaterialInterface.h"
+
 #include "GeometryCollectionISMPoolComponent.generated.h"
 
 class AActor;
@@ -15,7 +14,6 @@ class UGeometryComponent;
 class UInstancedStaticMeshComponent;
 class UGeometryCollectionISMPoolComponent;
 
-
 struct FGeometryCollectionISMInstance
 {
 	int32 StartIndex;
@@ -55,32 +53,6 @@ public:
 		TArray<int32> InstanceIdToIndex;
 	};
 
-	//FInstanceGroupId AddGroup(int32 Count)
-	//{
-	//	int32 StartIndex = 0;
-	//	if (GroupRanges.Num())
-	//	{
-	//		const FInstanceGroupRange& LastGroupRange = GroupRanges.Last();
-	//		StartIndex = LastGroupRange.Start + LastGroupRange.Count;
-	//	}
-	//	return GroupRanges.Emplace(FInstanceGroupRange(StartIndex, Count));
-	//}
-
-	//void RemoveGroup(FInstanceGroupId GroupId)
-	//{
-	//	check(GroupRanges.IsValidIndex(BlockId));
-	//	const int32 StartOffset = GroupRanges[GroupId].Count;
-	//	FreeList.Add(GroupId);
-	//	// 
-	//	for (int32 Index = (GroupId + 1); Index < GroupRanges.Num(); Index++)
-	//	{
-	//		const FInstanceGroupRange& Curr = GroupRanges[Index];
-	//		FInstanceGroupRange& Prev = GroupRanges[Index - 1];
-	//		Prev.Start = Curr.Start - StartOffset;
-	//		Prev.Count = Curr.Count;
-	//	}
-	//}
-
 	FInstanceGroupId AddGroup(int32 Count)
 	{
 		const int32 StartIndex = InstancesCount;
@@ -132,6 +104,42 @@ private:
 	TMap<FInstanceGroupId, FInstanceGroupRange> GroupRanges;
 };
 
+/** */
+struct FISMComponentDescription
+{
+	bool bUseHISM = false;
+	bool bIsStaticMobility = false;
+	bool bAffectShadow = true;
+	bool bAffectDistanceFieldLighting = false;
+	bool bAffectDynamicIndirectLighting = false;
+	int32 NumCustomDataFloats = 0;
+	int32 StartCullDistance = 0;
+	int32 EndCullDistance = 0;
+	float LodScale = 1.f;
+
+	bool operator==(const FISMComponentDescription& Other) const
+	{
+		return bUseHISM == Other.bUseHISM &&
+			bIsStaticMobility == Other.bIsStaticMobility &&
+			bAffectShadow == Other.bAffectShadow &&
+			bAffectDistanceFieldLighting == Other.bAffectDistanceFieldLighting &&
+			bAffectDynamicIndirectLighting == Other.bAffectDistanceFieldLighting &&
+			NumCustomDataFloats == Other.NumCustomDataFloats &&
+			StartCullDistance == Other.StartCullDistance && 
+			EndCullDistance == Other.EndCullDistance &&
+			LodScale == Other.LodScale;
+	}
+};
+
+FORCEINLINE uint32 GetTypeHash(const FISMComponentDescription& Desc)
+{
+	const uint32 PackedBools = (Desc.bUseHISM ? 1 : 0) | (Desc.bIsStaticMobility ? 2 : 0) | (Desc.bAffectShadow ? 4 : 0) | (Desc.bAffectDistanceFieldLighting ? 8 : 0) | (Desc.bAffectDynamicIndirectLighting ? 16 : 0);
+	uint32 Hash = HashCombine(GetTypeHash(PackedBools), GetTypeHash(Desc.NumCustomDataFloats));
+	Hash = HashCombine(Hash, GetTypeHash(Desc.StartCullDistance));
+	Hash = HashCombine(Hash, GetTypeHash(Desc.EndCullDistance));
+	return HashCombine(Hash, GetTypeHash(Desc.LodScale));
+}
+
 /**
 * This represent a unique mesh with potentially overriden materials
 * if the array is empty , there's no overrides
@@ -140,11 +148,11 @@ struct FGeometryCollectionStaticMeshInstance
 {
 	UStaticMesh* StaticMesh = nullptr;
 	TArray<UMaterialInterface*> MaterialsOverrides;
-	int32 NumCustomDataFloats = 0;
+	FISMComponentDescription Desc;
 
 	bool operator==(const FGeometryCollectionStaticMeshInstance& Other) const 
 	{
-		if (StaticMesh == Other.StaticMesh && NumCustomDataFloats == Other.NumCustomDataFloats)
+		if (StaticMesh == Other.StaticMesh && Desc == Other.Desc)
 		{
 			if (MaterialsOverrides.Num() == Other.MaterialsOverrides.Num())
 			{
@@ -164,7 +172,6 @@ struct FGeometryCollectionStaticMeshInstance
 	}
 };
 
-
 FORCEINLINE uint32 GetTypeHash(const FGeometryCollectionStaticMeshInstance& MeshInstance)
 {
 	uint32 CombinedHash = GetTypeHash(MeshInstance.StaticMesh);
@@ -173,6 +180,7 @@ FORCEINLINE uint32 GetTypeHash(const FGeometryCollectionStaticMeshInstance& Mesh
 	{
 		CombinedHash = HashCombine(CombinedHash, GetTypeHash(Material));
 	}
+	CombinedHash = HashCombine(CombinedHash, GetTypeHash(MeshInstance.Desc));
 	return CombinedHash;
 }
 
@@ -202,7 +210,7 @@ struct FGeometryCollectionMeshGroup
 
 struct FGeometryCollectionISM
 {
-	FGeometryCollectionISM(AActor* OwmingActor, const FGeometryCollectionStaticMeshInstance& MeshInstance, bool bPreferHISM);
+	FGeometryCollectionISM(AActor* OwmingActor, const FGeometryCollectionStaticMeshInstance& MeshInstance);
 
 	int32 AddInstanceGroup(int32 InstanceCount, TArrayView<const float> CustomDataFloats);
 
@@ -215,7 +223,7 @@ struct FGeometryCollectionISMPool
 {
 	using FISMIndex = int32;
 
-	FGeometryCollectionMeshInfo AddISM(UGeometryCollectionISMPoolComponent* OwningComponent, const FGeometryCollectionStaticMeshInstance& MeshInstance, int32 InstanceCount, TArrayView<const float> CustomDataFloats, bool bPreferHISM);
+	FGeometryCollectionMeshInfo AddISM(UGeometryCollectionISMPoolComponent* OwningComponent, const FGeometryCollectionStaticMeshInstance& MeshInstance, int32 InstanceCount, TArrayView<const float> CustomDataFloats);
 	bool BatchUpdateInstancesTransforms(FGeometryCollectionMeshInfo& MeshInfo, int32 StartInstanceIndex, const TArray<FTransform>& NewInstancesTransforms, bool bWorldSpace, bool bMarkRenderStateDirty, bool bTeleport);
 	void RemoveISM(const FGeometryCollectionMeshInfo& MeshInfo);
 	
@@ -261,7 +269,7 @@ public:
 	void DestroyMeshGroup(FMeshGroupId MeshGroupId);
 
 	/** Add a static mesh for a nmesh group */
-	FMeshId AddMeshToGroup(FMeshGroupId MeshGroupId, const FGeometryCollectionStaticMeshInstance& MeshInstance, int32 InstanceCount, TArrayView<const float> CustomDataFloats, bool bPreferHISM = true);
+	FMeshId AddMeshToGroup(FMeshGroupId MeshGroupId, const FGeometryCollectionStaticMeshInstance& MeshInstance, int32 InstanceCount, TArrayView<const float> CustomDataFloats);
 
 	/** Add a static mesh for a nmesh group */
 	bool BatchUpdateInstancesTransforms(FMeshGroupId MeshGroupId, FMeshId MeshId, int32 StartInstanceIndex, const TArray<FTransform>& NewInstancesTransforms, bool bWorldSpace = false, bool bMarkRenderStateDirty = false, bool bTeleport = false);
@@ -286,4 +294,3 @@ private:
 	// Expose internals for debug draw support.
 	friend class UGeometryCollectionISMPoolDebugDrawComponent;
 };
-#endif
\ No newline at end of file
