Apply strain from collision to GC child instead of GC root for GCs inside of ClusteredUnions.
Also, change bRestoreBreakingMomentum to RestoreBreakingMomentumPercent so that the percent momentum restored is tweakable via cvar.

diff --git a/Engine/Source/Runtime/Experimental/Chaos/Private/Chaos/ClusterUnionManager.cpp b/Engine/Source/Runtime/Experimental/Chaos/Private/Chaos/ClusterUnionManager.cpp
--- a/Engine/Source/Runtime/Experimental/Chaos/Private/Chaos/ClusterUnionManager.cpp
+++ b/Engine/Source/Runtime/Experimental/Chaos/Private/Chaos/ClusterUnionManager.cpp
@@ -206,6 +206,12 @@ namespace Chaos
 		return ClusterUnions.Find(Index);
 	}
 
+	FClusterUnion* FClusterUnionManager::FindClusterUnionFromParticle(FPBDRigidParticleHandle* Particle)
+	{
+		FClusterUnionIndex ClusterUnionIndex = FindClusterUnionIndexFromParticle(Particle);
+		return FindClusterUnion(ClusterUnionIndex);
+	}
+
 	DECLARE_CYCLE_STAT(TEXT("FClusterUnionManager::HandleAddOperation"), STAT_HandleAddOperation, STATGROUP_Chaos);
 	void FClusterUnionManager::HandleAddOperation(FClusterUnionIndex ClusterIndex, const TArray<FPBDRigidParticleHandle*>& Particles, bool bReleaseClustersFirst)
 	{
@@ -440,7 +446,6 @@ namespace Chaos
 		UpdateClusterMassProperties(ClusterUnion.InternalCluster, FullChildrenSet, ClusterInertia, (bRecomputeMassOrientation && ClusterUnion.bNeedsXRInitialization) ? nullptr : &ForceMassOrientation);
 		UpdateKinematicProperties(ClusterUnion.InternalCluster, MClustering.GetChildrenMap(), MEvolution);
 
-		MEvolution.InvalidateParticle(ClusterUnion.InternalCluster);
 		// The recreation of the geometry must happen after the call to UpdateClusterMassProperties.
 		// Creating the geometry requires knowing the relative frame between the parent cluster and the child clusters. The
 		// parent transform is not set properly for a new empty cluster until UpdateClusterMassProperties is called for the first time.
diff --git a/Engine/Source/Runtime/Experimental/Chaos/Private/Chaos/PBDRigidClustering.cpp b/Engine/Source/Runtime/Experimental/Chaos/Private/Chaos/PBDRigidClustering.cpp
--- a/Engine/Source/Runtime/Experimental/Chaos/Private/Chaos/PBDRigidClustering.cpp
+++ b/Engine/Source/Runtime/Experimental/Chaos/Private/Chaos/PBDRigidClustering.cpp
@@ -53,8 +53,8 @@ namespace Chaos
 	float GraphPropagationBasedCollisionFactor = 1;
 	FAutoConsoleVariableRef CVarGraphPropagationBasedCollisionFactor(TEXT("p.GraphPropagationBasedCollisionFactor"), GraphPropagationBasedCollisionFactor, TEXT("when p.GraphPropagationBasedCollisionImpulseProcessing is on, the percentage [0-1] of remaining damage that is distributed to the connected pieces"));
 
-	bool bRestoreBreakingMomentum = 1;
-	FAutoConsoleVariableRef CVarRestoreBreakingMomentum(TEXT("p.RestoreBreakingMomentum"), bRestoreBreakingMomentum, TEXT("When a rigid cluster is broken, objects that its in contact with will receive an impulse to restore their momentum prior to the break."));
+	float RestoreBreakingMomentumPercent = .5;
+	FAutoConsoleVariableRef CVarRestoreBreakingMomentumPercent(TEXT("p.RestoreBreakingMomentumPercent"), RestoreBreakingMomentumPercent, TEXT("When a rigid cluster is broken, objects that its in contact with will receive an impulse to restore this percent of their momentum prior to the break."));
 
 	static FGeometryCollectionPhysicsProxy* GetConcreteProxy(FPBDRigidClusteredParticleHandle* ClusteredParticle)
 	{
@@ -498,10 +498,8 @@ namespace Chaos
 			const FMatrix33 OtherInvI = Utilities::ComputeWorldSpaceInertia(Generic->QCom(), Generic->ConditionedInvI());
 			const FVec3 AngularImpulseVelocity = OtherInvI * AngularImpulse;
 
-			// Temporary: Always restore only 50% of the impulse
-			const float RestorationPercent = .5f;
-
 			// Update linear and angular impulses for the body, to be integrated next solve
+			const float RestorationPercent = RestoreBreakingMomentumPercent;
 			Rigid.V() += ImpulseVelocity * RestorationPercent;
 			Rigid.W() += AngularImpulseVelocity * RestorationPercent;
 		}
@@ -810,6 +808,10 @@ namespace Chaos
 		// since it'll be modifying the children array.
 		TArray<FPBDRigidParticleHandle*> DeferredRemoveFromClusterUnion;
 
+		// Grab cluster union parent if there is one
+		FPBDRigidParticleHandle* ParentRigid = ClusteredParticle->ClusterIds().Id;
+		FPBDRigidClusteredParticleHandle* Parent = ParentRigid ? ParentRigid->CastToClustered() : nullptr;
+
 		for (int32 ChildIdx = Children.Num() - 1; ChildIdx >= 0; --ChildIdx)
 		{
 			FPBDRigidClusteredParticleHandle* Child = Children[ChildIdx]->CastToClustered();
@@ -854,13 +856,6 @@ namespace Chaos
 				
 				ActivatedChildren.Add(Child);
 				SendBreakingEvent(Child, bParentCrumbled);
-
-				// Restore some of the momentum of whatever collided with the parent
-				if (bRestoreBreakingMomentum)
-				{
-					TrackBreakingCollision(Child);
-					TrackBreakingCollision(ClusteredParticle);
-				}
 			}
 			if (bUseDamagePropagation)
 			{
@@ -874,6 +869,19 @@ namespace Chaos
 			ClusterUnionManager.HandleRemoveOperationWithClusterLookup(DeferredRemoveFromClusterUnion, true);
 		}
 
+		// Restore some of the momentum of whatever collided with the parent
+		if (bFoundFirstRelease && RestoreBreakingMomentumPercent > 0.f)
+		{
+			if (Parent)
+			{
+				TrackBreakingCollision(Parent);
+			}
+			else
+			{
+				TrackBreakingCollision(ClusteredParticle);
+			}
+		}
+
 		// if necessary propagate strain through the graph
 		// IMPORTANT: this assumes that the connectivity graph has not yet been updated from pieces that broke off
 		if (bUseDamagePropagation)
@@ -1116,19 +1124,52 @@ namespace Chaos
 		TArray<FPBDRigidClusteredParticleHandle*> ClusteredParticlesToProcess;
 		for (auto& Particle : MEvolution.GetNonDisabledClusteredView())
 		{
-			ClusteredParticlesToProcess.Add(Particle.Handle()->CastToClustered());
+			if (FPBDRigidClusteredParticleHandle* Clustered = Particle.Handle()->CastToClustered())
+			{
+				if (Clustered->ClusterIds().NumChildren > 0)
+				{
+					if (FClusterUnion* ClusterUnion = ClusterUnionManager.FindClusterUnionFromParticle(Clustered))
+					{
+						if (ClusterUnion->InternalCluster == Clustered)
+						{
+							// Clustered is itself a cluster union, so loop over its children and add those
+							// to process for breaking.
+							for (FPBDRigidParticleHandle* ChildParticle : ClusterUnion->ChildParticles)
+							{
+								if (ChildParticle)
+								{
+									if (FPBDRigidClusteredParticleHandle* ClusteredChild = ChildParticle->CastToClustered())
+									{
+										if (ClusteredChild->ClusterIds().NumChildren > 0)
+										{
+											ClusteredParticlesToProcess.Add(ClusteredChild);
+										}
+									}
+								}
+							}
+						}
+						else
+						{
+							// Clustered is inside a clustered union, but not a clustered union itself
+							ClusteredParticlesToProcess.Add(Clustered);
+						}
+					}
+					else
+					{
+						// Clustered is not a clustered union, and not _in_ a clustered union
+						ClusteredParticlesToProcess.Add(Clustered);
+					}
+				}
+			}
 		}
 
 		for (FPBDRigidClusteredParticleHandle* ClusteredParticle : ClusteredParticlesToProcess)
 		{
-			if (ClusteredParticle->ClusterIds().NumChildren)
-			{
-				ReleaseClusterParticles(ClusteredParticle);
-			}
+			ReleaseClusterParticles(ClusteredParticle);
 		}
 
 		// Restore some of the momentum of objects that were touching rigid clusters that broke
-		if (bRestoreBreakingMomentum)
+		if (RestoreBreakingMomentumPercent > 0.f)
 		{
 			RestoreBreakingMomentum();
 		}
@@ -1434,6 +1475,20 @@ namespace Chaos
 							{
 								if (TPBDRigidClusteredParticleHandle<FReal, 3>*ClusteredChild = Child->CastToClustered())
 								{
+									// If we got a child of a Clustered Union then this might just be the root particle of a geometry collection,
+									// in which case it might not make sense to strain it (?)
+									if (const IPhysicsProxyBase* Proxy = Cluster->PhysicsProxy())
+									{
+										if (Proxy->GetType() == EPhysicsProxyType::ClusterUnionProxy)
+										{
+											TArray<FRigidHandle>& ClusteredParentToChildren = MChildren[ClusteredChild];
+											if (FPBDRigidParticleHandle* ClusteredChildRigid = FindClosestParticle(ClusteredParentToChildren, ContactWorldLocation))
+											{
+												ClusteredChild = ClusteredChildRigid->CastToClustered();
+											}
+										}
+									}
+
 									ClusteredChild->CollisionImpulses() += AccumulatedImpulse;
 									UpdateTopLevelParticle(ClusteredChild);
 								}
diff --git a/Engine/Source/Runtime/Experimental/Chaos/Public/Chaos/ClusterUnionManager.h b/Engine/Source/Runtime/Experimental/Chaos/Public/Chaos/ClusterUnionManager.h
--- a/Engine/Source/Runtime/Experimental/Chaos/Public/Chaos/ClusterUnionManager.h
+++ b/Engine/Source/Runtime/Experimental/Chaos/Public/Chaos/ClusterUnionManager.h
@@ -93,7 +93,7 @@ namespace Chaos
 		// Access the cluster union externally.
 		FClusterUnion* FindClusterUnionFromExplicitIndex(FClusterUnionExplicitIndex Index);
 		FClusterUnion* FindClusterUnion(FClusterUnionIndex Index);
-
+		FClusterUnion* FindClusterUnionFromParticle(FPBDRigidParticleHandle* Particle);
 		FClusterUnionIndex FindClusterUnionIndexFromParticle(FPBDRigidParticleHandle* Particle);
 
 		// An extension to FindClusterUnionIndexFromParticle to check whether or not the given particle is the cluster union particle itself.
