Temporarily disable Chaos_Collision_MidPhase_MaxShapePairs to fix issue with collisions being ignored on larege GCs.

diff --git a/Engine/Source/Runtime/Experimental/Chaos/Private/Chaos/Collision/ParticlePairMidPhase.cpp b/Engine/Source/Runtime/Experimental/Chaos/Private/Chaos/Collision/ParticlePairMidPhase.cpp
--- a/Engine/Source/Runtime/Experimental/Chaos/Private/Chaos/Collision/ParticlePairMidPhase.cpp
+++ b/Engine/Source/Runtime/Experimental/Chaos/Private/Chaos/Collision/ParticlePairMidPhase.cpp
@@ -35,7 +35,7 @@ namespace Chaos
 		FAutoConsoleVariableRef CVarChaos_Collision_CullDistanceReferenceSize(TEXT("p.Chaos.Collision.CullDistanceReferenceSize"), Chaos_Collision_CullDistanceScaleInverseSize, TEXT(""));
 		FAutoConsoleVariableRef CVarChaos_Collision_MinCullDistanecScale(TEXT("p.Chaos.Collision.MinCullDistanceScale"), Chaos_Collision_MinCullDistanceScale, TEXT(""));
 
-		int32 Chaos_Collision_MidPhase_MaxShapePairs = 100;
+		int32 Chaos_Collision_MidPhase_MaxShapePairs = 0;
 		FAutoConsoleVariableRef CVarChaos_Collision_MaxShapePairs(TEXT("p.Chaos.Collision.MaxShapePairs"), Chaos_Collision_MidPhase_MaxShapePairs, TEXT(""));
 	}
 	using namespace CVars;
