Chaos Cloth: Add XPBDStretchBiasElementConstraints

This new constraint has warp/weft/bias control. Warp and Weft are just springs but transformed into Warp and Weft components. Bias is trying to enforce that the Warp and Weft directions stay orthogonal. Optionally can use 2d or 3d positions for rest lengths (default is 3d). Optionally can specify different warp and weft scales.

diff --git a/Engine/Plugins/Experimental/ChaosCloth/Source/Chaos/Private/ChaosCloth/ChaosClothConstraints.cpp b/Engine/Plugins/Experimental/ChaosCloth/Source/Chaos/Private/ChaosCloth/ChaosClothConstraints.cpp
--- a/Engine/Plugins/Experimental/ChaosCloth/Source/Chaos/Private/ChaosCloth/ChaosClothConstraints.cpp
+++ b/Engine/Plugins/Experimental/ChaosCloth/Source/Chaos/Private/ChaosCloth/ChaosClothConstraints.cpp
@@ -4,6 +4,7 @@
 #include "ChaosCloth/ChaosClothingPatternData.h"
 #include "Chaos/PBDSpringConstraints.h"
 #include "Chaos/XPBDSpringConstraints.h"
+#include "Chaos/XPBDStretchBiasElementConstraints.h"
 #include "Chaos/PBDBendingConstraints.h"
 #include "Chaos/XPBDBendingConstraints.h"
 #include "Chaos/XPBDAnisotropicBendingConstraints.h"
@@ -110,7 +111,7 @@ void FClothConstraints::AddRules(
 	CreateSelfCollisionConstraints(ConfigProperties, TriangleMesh);
 
 	// Edge constraints
-	CreateEdgeConstraints(ConfigProperties, WeightMaps, TriangleMesh);
+	CreateStretchConstraints(ConfigProperties, WeightMaps, TriangleMesh, PatternData);
 
 	// Bending constraints
 	CreateBendingConstraints(ConfigProperties, WeightMaps, TriangleMesh, PatternData);
@@ -187,12 +188,38 @@ void FClothConstraints::CreateSelfCollisionConstraints(const Softs::FCollectionP
 	}
 }
 
-void FClothConstraints::CreateEdgeConstraints(
+void FClothConstraints::CreateStretchConstraints(
 	const Softs::FCollectionPropertyConstFacade& ConfigProperties,
 	const TArray<TConstArrayView<FRealSingle>>& WeightMaps,
-	const FTriangleMesh& TriangleMesh)
+	const FTriangleMesh& TriangleMesh,
+	const FClothingPatternData* PatternData)
 {
-	if (Softs::FXPBDEdgeSpringConstraints::IsEnabled(ConfigProperties))
+	if (PatternData && PatternData->PatternPositions.Num() && Softs::FXPBDStretchBiasElementConstraints::IsEnabled(ConfigProperties))
+	{
+		// TODO: separate warp, weft, Bias multipliers
+		const TConstArrayView<FRealSingle>& EdgeStiffnessMultipliers = WeightMaps[(int32)EChaosWeightMapTarget::EdgeStiffness];
+		const TConstArrayView<FRealSingle> DampingMultipliers;  // TODO: Damping multiplier
+		const TConstArrayView<FRealSingle> WarpWeftScaleMultipliers; // TODO: separate multipliers for warp and weft scale
+
+		XStretchBiasConstraints = MakeShared<Softs::FXPBDStretchBiasElementConstraints>(
+			Evolution->Particles(),
+			ParticleOffset,
+			NumParticles,
+			TriangleMesh,
+			PatternData->WeldedFaceVertexPatternPositions,
+			EdgeStiffnessMultipliers,
+			EdgeStiffnessMultipliers,
+			EdgeStiffnessMultipliers,
+			DampingMultipliers,
+			WarpWeftScaleMultipliers,
+			WarpWeftScaleMultipliers,
+			ConfigProperties,
+			/*bTrimKinematicConstraints =*/ true);
+
+		++NumConstraintInits;  // Uses init to update the property tables
+		++NumConstraintRules;
+	}
+	else if (Softs::FXPBDEdgeSpringConstraints::IsEnabled(ConfigProperties))
 	{
 		const TConstArrayView<FRealSingle>& EdgeStiffnessMultipliers = WeightMaps[(int32)EChaosWeightMapTarget::EdgeStiffness];
 		const TConstArrayView<FRealSingle> DampingMultipliers;  // TODO: Damping multiplier
@@ -493,6 +520,21 @@ void FClothConstraints::CreateRules()
 	int32 ConstraintRuleIndex = 0;
 	int32 PostCollisionConstraintRuleIndex = 0;
 
+	if (XStretchBiasConstraints)
+	{
+		ConstraintInits[ConstraintInitIndex++] =
+			[this](Softs::FSolverParticles& /*Particles*/, const Softs::FSolverReal Dt)
+		{
+			XStretchBiasConstraints->Init();
+			XStretchBiasConstraints->ApplyProperties(Dt, Evolution->GetIterations());
+		};
+
+		ConstraintRules[ConstraintRuleIndex++] =
+			[this](Softs::FSolverParticles& Particles, const Softs::FSolverReal Dt)
+		{
+			XStretchBiasConstraints->Apply(Particles, Dt);
+		};
+	}
 	if (XEdgeConstraints)
 	{
 		ConstraintInits[ConstraintInitIndex++] =
diff --git a/Engine/Plugins/Experimental/ChaosCloth/Source/Chaos/Public/ChaosCloth/ChaosClothConstraints.h b/Engine/Plugins/Experimental/ChaosCloth/Source/Chaos/Public/ChaosCloth/ChaosClothConstraints.h
--- a/Engine/Plugins/Experimental/ChaosCloth/Source/Chaos/Public/ChaosCloth/ChaosClothConstraints.h
+++ b/Engine/Plugins/Experimental/ChaosCloth/Source/Chaos/Public/ChaosCloth/ChaosClothConstraints.h
@@ -158,10 +158,12 @@ namespace Chaos
 
 		const TSharedPtr<Softs::FPBDEdgeSpringConstraints>& GetEdgeSpringConstraints() const { return EdgeConstraints; }
 		const TSharedPtr<Softs::FXPBDEdgeSpringConstraints>& GetXEdgeSpringConstraints() const { return XEdgeConstraints; }
+		const TSharedPtr<Softs::FXPBDStretchBiasElementConstraints>& GetXStretchBiasConstraints() const { return XStretchBiasConstraints; }
 		const TSharedPtr<Softs::FPBDBendingSpringConstraints>& GetBendingSpringConstraints() const { return BendingConstraints; }
 		const TSharedPtr<Softs::FXPBDBendingSpringConstraints>& GetXBendingSpringConstraints() const { return XBendingConstraints; }
 		const TSharedPtr<Softs::FPBDBendingConstraints>& GetBendingElementConstraints() const { return BendingElementConstraints; }
 		const TSharedPtr<Softs::FXPBDBendingConstraints>& GetXBendingElementConstraints() const { return XBendingElementConstraints; }
+		const TSharedPtr<Softs::FXPBDAnisotropicBendingConstraints>& GetXAnisoBendingElementConstraints() const { return XAnisoBendingElementConstraints; }
 		const TSharedPtr<Softs::FPBDAreaSpringConstraints>& GetAreaSpringConstraints() const { return AreaConstraints; }
 		const TSharedPtr<Softs::FXPBDAreaSpringConstraints>& GetXAreaSpringConstraints() const { return XAreaConstraints; }
 		const TSharedPtr<Softs::FPBDLongRangeConstraints>& GetLongRangeConstraints() const { return LongRangeConstraints; }
@@ -178,10 +180,11 @@ namespace Chaos
 		void CreateSelfCollisionConstraints(
 			const Softs::FCollectionPropertyConstFacade& ConfigProperties,
 			const FTriangleMesh& TriangleMesh);
-		void CreateEdgeConstraints(
+		void CreateStretchConstraints(
 			const Softs::FCollectionPropertyConstFacade& ConfigProperties,
 			const TArray<TConstArrayView<FRealSingle>>& WeightMaps,
-			const FTriangleMesh& TriangleMesh);
+			const FTriangleMesh& TriangleMesh,
+			const FClothingPatternData* PatternData);
 		void CreateBendingConstraints(
 			const Softs::FCollectionPropertyConstFacade& ConfigProperties,
 			const TArray<TConstArrayView<FRealSingle>>& WeightMaps,
@@ -222,6 +225,7 @@ namespace Chaos
 
 		TSharedPtr<Softs::FPBDEdgeSpringConstraints> EdgeConstraints;
 		TSharedPtr<Softs::FXPBDEdgeSpringConstraints> XEdgeConstraints;
+		TSharedPtr<Softs::FXPBDStretchBiasElementConstraints> XStretchBiasConstraints;
 		TSharedPtr<Softs::FPBDBendingSpringConstraints> BendingConstraints;
 		TSharedPtr<Softs::FXPBDBendingSpringConstraints> XBendingConstraints;
 		TSharedPtr<Softs::FPBDBendingConstraints> BendingElementConstraints;
diff --git a/Engine/Plugins/Experimental/ChaosClothAssetEditor/Source/ChaosClothAssetDataflowNodes/Private/ChaosClothAsset/SimulationXPBDConfigNode.cpp b/Engine/Plugins/Experimental/ChaosClothAssetEditor/Source/ChaosClothAssetDataflowNodes/Private/ChaosClothAsset/SimulationXPBDConfigNode.cpp
--- a/Engine/Plugins/Experimental/ChaosClothAssetEditor/Source/ChaosClothAssetDataflowNodes/Private/ChaosClothAsset/SimulationXPBDConfigNode.cpp
+++ b/Engine/Plugins/Experimental/ChaosClothAssetEditor/Source/ChaosClothAssetDataflowNodes/Private/ChaosClothAsset/SimulationXPBDConfigNode.cpp
@@ -17,8 +17,14 @@ FChaosClothAssetSimulationXPBDConfigNode::FChaosClothAssetSimulationXPBDConfigNo
 	RegisterInputConnection(&Collection);
 
 	RegisterInputConnection(&bEnableXPBDStretchConstraints);
+	RegisterInputConnection(&bEnableStretchBias);
+	RegisterInputConnection(&bUse3dElementRestLengths);
 	RegisterInputConnection(&StretchStiffness);
+	RegisterInputConnection(&StretchStiffnessWeft);
+	RegisterInputConnection(&StretchStiffnessBias);
 	RegisterInputConnection(&StretchDampingRatio);
+	RegisterInputConnection(&WarpStretchScale);
+	RegisterInputConnection(&WeftStretchScale);
 
 	RegisterInputConnection(&bEnableXPBDBendConstraints);
 	RegisterInputConnection(&bEnableBendAnisotropy);
@@ -94,17 +100,55 @@ void FChaosClothAssetSimulationXPBDConfigNode::SetStretchProperties(Chaos::Softs
 {
 	if (bEnableXPBDStretchConstraints)
 	{
-		const int32 StretchStiffnessIndex = Properties.AddProperty(TEXT("XPBDEdgeSpringStiffness"));
-		Properties.SetWeightedValue(StretchStiffnessIndex, StretchStiffness.Low, StretchStiffness.High);
-		Properties.SetStringValue(StretchStiffnessIndex, TEXT("EdgeStiffness"));
+		if (bEnableStretchBias)
+		{
+			constexpr bool bEnable = true;
+			constexpr bool bNonAnimatable = false;
+			Properties.AddValue(TEXT("XPBDStretchBiasElementUse3dRestLengths"), bUse3dElementRestLengths, bEnable, bNonAnimatable);
+
+			const int32 StretchStiffnessWarpIndex = Properties.AddProperty(TEXT("XPBDStretchBiasElementStiffnessWarp"));
+			Properties.SetWeightedValue(StretchStiffnessWarpIndex, StretchStiffness.Low, StretchStiffness.High);
+			Properties.SetStringValue(StretchStiffnessWarpIndex, TEXT("StretchStiffnessWarp"));
+
+			const int32 StretchStiffnessWeftIndex = Properties.AddProperty(TEXT("XPBDStretchBiasElementStiffnessWeft"));
+			Properties.SetWeightedValue(StretchStiffnessWeftIndex, StretchStiffnessWeft.Low, StretchStiffnessWeft.High);
+			Properties.SetStringValue(StretchStiffnessWeftIndex, TEXT("StretchStiffnessWeft"));
+
+			const int32 BiasStiffnessIndex = Properties.AddProperty(TEXT("XPBDStretchBiasElementStiffnessBias"));
+			Properties.SetWeightedValue(BiasStiffnessIndex, StretchStiffnessBias.Low, StretchStiffnessBias.High);
+			Properties.SetStringValue(BiasStiffnessIndex, TEXT("StretchStiffnessBias"));
+
+			const int32 StretchDampingRatioIndex = Properties.AddProperty(TEXT("XPBDStretchBiasElementDamping"));
+			Properties.SetWeightedValue(StretchDampingRatioIndex, StretchDampingRatio.Low, StretchDampingRatio.High);
+			Properties.SetStringValue(StretchDampingRatioIndex, TEXT("StretchBiasDamping"));
+
+			const int32 WarpScaleIndex = Properties.AddProperty(TEXT("XPBDStretchBiasElementWarpScale"));
+			Properties.SetWeightedValue(WarpScaleIndex, WarpStretchScale.Low, WarpStretchScale.High);
+			Properties.SetStringValue(WarpScaleIndex, TEXT("WarpStretchScale"));
 
-		const int32 StretchDampingRatioIndex = Properties.AddProperty(TEXT("XPBDEdgeSpringDamping"));
-		Properties.SetWeightedValue(StretchDampingRatioIndex, StretchDampingRatio.Low, StretchDampingRatio.High);
-		Properties.SetStringValue(StretchDampingRatioIndex, TEXT("EdgeDamping"));
+			const int32 WeftScaleIndex = Properties.AddProperty(TEXT("XPBDStretchBiasElementWeftScale"));
+			Properties.SetWeightedValue(WeftScaleIndex, WeftStretchScale.Low, WeftStretchScale.High);
+			Properties.SetStringValue(WeftScaleIndex, TEXT("WeftStretchScale"));
+
+			Properties.SetEnabled(TEXT("XPBDEdgeSpringStiffness"), false);
+		}
+		else
+		{
+			const int32 StretchStiffnessIndex = Properties.AddProperty(TEXT("XPBDEdgeSpringStiffness"));
+			Properties.SetWeightedValue(StretchStiffnessIndex, StretchStiffness.Low, StretchStiffness.High);
+			Properties.SetStringValue(StretchStiffnessIndex, TEXT("EdgeStiffness"));
+
+			const int32 StretchDampingRatioIndex = Properties.AddProperty(TEXT("XPBDEdgeSpringDamping"));
+			Properties.SetWeightedValue(StretchDampingRatioIndex, StretchDampingRatio.Low, StretchDampingRatio.High);
+			Properties.SetStringValue(StretchDampingRatioIndex, TEXT("EdgeDamping"));
+
+			Properties.SetEnabled(TEXT("XPBDStretchBiasElementStiffnessWarp"), false);
+		}
 	}
 	else
 	{
 		Properties.SetEnabled(TEXT("XPBDEdgeSpringStiffness"), false);
+		Properties.SetEnabled(TEXT("XPBDStretchBiasElementStiffnessWarp"), false);
 	}
 }
 
diff --git a/Engine/Plugins/Experimental/ChaosClothAssetEditor/Source/ChaosClothAssetDataflowNodes/Public/ChaosClothAsset/SimulationXPBDConfigNode.h b/Engine/Plugins/Experimental/ChaosClothAssetEditor/Source/ChaosClothAssetDataflowNodes/Public/ChaosClothAsset/SimulationXPBDConfigNode.h
--- a/Engine/Plugins/Experimental/ChaosClothAssetEditor/Source/ChaosClothAssetDataflowNodes/Public/ChaosClothAsset/SimulationXPBDConfigNode.h
+++ b/Engine/Plugins/Experimental/ChaosClothAssetEditor/Source/ChaosClothAssetDataflowNodes/Public/ChaosClothAsset/SimulationXPBDConfigNode.h
@@ -29,24 +29,36 @@ public:
 	/** Cloth Simulation Properties. */
 	UPROPERTY(EditAnywhere, Category = "Stretch")
 	bool bEnableXPBDStretchConstraints = true;
+	UPROPERTY(EditAnywhere, Category = "Stretch", meta = (EditCondition = "bEnableXPBDStretchConstraints"))
+	bool bEnableStretchBias = false;
+	UPROPERTY(EditAnywhere, Category = "Stretch", meta = (EditCondition = "bEnableXPBDStretchConstraints && bEnableStretchBias"))
+	bool bUse3dElementRestLengths = true;
 
-	UPROPERTY(EditAnywhere, Category = "Stretch", DisplayName = "Stretch Stiffness [kg/s^2]", meta = (UIMin = "0", UIMax = "10000", ClampMin = "0", ClampMax = "10000000", EditCondition = "bEnableXPBDStretchConstraints"))
+	UPROPERTY(EditAnywhere, Category = "Stretch", DisplayName = "Stretch Stiffness (Warp) [kg cm /s^2]", meta = (UIMin = "0", UIMax = "10000", ClampMin = "0", ClampMax = "10000000", EditCondition = "bEnableXPBDStretchConstraints"))
 	FChaosClothWeightedValue StretchStiffness = { 100.f, 100.f };
+	UPROPERTY(EditAnywhere, Category = "Stretch", DisplayName = "Stretch Stiffness (Weft) [kg cm /s^2]", meta = (UIMin = "0", UIMax = "10000", ClampMin = "0", ClampMax = "10000000", EditCondition = "bEnableXPBDStretchConstraints && bEnableStretchBias"))
+	FChaosClothWeightedValue StretchStiffnessWeft = { 100.f, 100.f };
+	UPROPERTY(EditAnywhere, Category = "Stretch", DisplayName = "Stretch Stiffness (Bias) [kg cm^2 /s^2]", meta = (UIMin = "0", UIMax = "10000", ClampMin = "0", ClampMax = "10000000", EditCondition = "bEnableXPBDStretchConstraints && bEnableStretchBias"))
+	FChaosClothWeightedValue StretchStiffnessBias = { 100.f, 100.f };
 	UPROPERTY(EditAnywhere, Category = "Stretch", DisplayName = "Stretch Damping Ratio", meta = (UIMin = "0", UIMax = "10", ClampMin = "0", ClampMax = "1000", EditCondition = "bEnableXPBDStretchConstraints"))
 	FChaosClothWeightedValue StretchDampingRatio = { 1.f, 1.f };
+	UPROPERTY(EditAnywhere, Category = "Stretch", DisplayName = "Warp Scale", meta = (UIMin = "0", UIMax = "10", ClampMin = "0", ClampMax = "10000000", EditCondition = "bEnableXPBDStretchConstraints && bEnableStretchBias"))
+	FChaosClothWeightedValue WarpStretchScale = { 1.f, 1.f };
+	UPROPERTY(EditAnywhere, Category = "Stretch", DisplayName = "Weft Scale", meta = (UIMin = "0", UIMax = "10", ClampMin = "0", ClampMax = "10000000", EditCondition = "bEnableXPBDStretchConstraints && bEnableStretchBias"))
+	FChaosClothWeightedValue WeftStretchScale = { 1.f, 1.f };
 
 	UPROPERTY(EditAnywhere, Category = "Bending")
 	bool bEnableXPBDBendConstraints = true;
-	UPROPERTY(EditAnywhere, Category = "Bending")
+	UPROPERTY(EditAnywhere, Category = "Bending", meta = (EditCondition = "bEnableXPBDBendConstraints"))
 	bool bEnableBendAnisotropy = false;
 
-	UPROPERTY(EditAnywhere, Category = "Bending", DisplayName = "Bending Stiffness (Warp) [kg cm/s^2 rad]", meta = (UIMin = "0", UIMax = "10000", ClampMin = "0", ClampMax = "10000000", EditCondition = "bEnableXPBDBendConstraints"))
+	UPROPERTY(EditAnywhere, Category = "Bending", DisplayName = "Bending Stiffness (Warp) [kg cm^2/s^2 rad^2]", meta = (UIMin = "0", UIMax = "10000", ClampMin = "0", ClampMax = "10000000", EditCondition = "bEnableXPBDBendConstraints"))
 	FChaosClothWeightedValue BendingStiffness = { 100.f, 100.f };
 
-	UPROPERTY(EditAnywhere, Category = "Bending", DisplayName = "Bending Stiffness (Weft) [kg cm/s^2 rad]", meta = (UIMin = "0", UIMax = "10000", ClampMin = "0", ClampMax = "10000000", EditCondition = "bEnableXPBDBendConstraints && bEnableBendAnisotropy"))
+	UPROPERTY(EditAnywhere, Category = "Bending", DisplayName = "Bending Stiffness (Weft) [kg cm^2/s^2 rad^2]", meta = (UIMin = "0", UIMax = "10000", ClampMin = "0", ClampMax = "10000000", EditCondition = "bEnableXPBDBendConstraints && bEnableBendAnisotropy"))
 	FChaosClothWeightedValue BendingStiffnessWeft = { 100.f, 100.f };
 
-	UPROPERTY(EditAnywhere, Category = "Bending", DisplayName = "Bending Stiffness (Bias) [kg cm/s^2 rad]", meta = (UIMin = "0", UIMax = "10000", ClampMin = "0", ClampMax = "10000000", EditCondition = "bEnableXPBDBendConstraints && bEnableBendAnisotropy"))
+	UPROPERTY(EditAnywhere, Category = "Bending", DisplayName = "Bending Stiffness (Bias) [kg cm^2/s^2 rad^2]", meta = (UIMin = "0", UIMax = "10000", ClampMin = "0", ClampMax = "10000000", EditCondition = "bEnableXPBDBendConstraints && bEnableBendAnisotropy"))
 	FChaosClothWeightedValue BendingStiffnessBias = { 100.f, 100.f };
 
 	UPROPERTY(EditAnywhere, Category = "Bending", DisplayName = "Bending Damping Ratio", meta = (UIMin = "0", UIMax = "10", ClampMin = "0", ClampMax = "1000", EditCondition = "bEnableXPBDBendConstraints"))
diff --git a/Engine/Source/Runtime/Core/Public/Math/Vector.isph b/Engine/Source/Runtime/Core/Public/Math/Vector.isph
--- a/Engine/Source/Runtime/Core/Public/Math/Vector.isph
+++ b/Engine/Source/Runtime/Core/Public/Math/Vector.isph
@@ -391,6 +391,24 @@ inline void VectorStore(uniform FIntVector4 *uniform DstPtr, const varying FIntV
 }
 
 
+inline varying FVector2f VectorLoad(const uniform FVector2f *uniform SrcPtr)
+{
+	varying FVector2f Result;
+
+	if(((1<<TARGET_WIDTH)-1 ^ lanemask()) == 0)
+	{
+		aos_to_soa2_ispc((uniform float * uniform)SrcPtr, &Result.V[0], &Result.V[1]);
+	}
+	else
+	{
+		#pragma ignore warning(perf)
+		Result = SrcPtr[programIndex];
+	}
+
+	return Result;
+}
+
+
 inline varying FVector2 VectorLoad(const uniform FVector2 *uniform SrcPtr)
 {
 	varying FVector2 Result;
diff --git a/Engine/Source/Runtime/Experimental/Chaos/Private/Chaos/XPBDStretchBiasElementConstraints.cpp b/Engine/Source/Runtime/Experimental/Chaos/Private/Chaos/XPBDStretchBiasElementConstraints.cpp
new file mode 100644
--- /dev/null
+++ b/Engine/Source/Runtime/Experimental/Chaos/Private/Chaos/XPBDStretchBiasElementConstraints.cpp
@@ -0,0 +1,609 @@
+// Copyright Epic Games, Inc. All Rights Reserved.
+#include "Chaos/XPBDStretchBiasElementConstraints.h"
+#include "Chaos/GraphColoring.h"
+#include "Chaos/Framework/Parallel.h"
+#include "Chaos/TriangleMesh.h"
+#include "Chaos/Matrix.h"
+#include "ChaosStats.h"
+
+#if INTEL_ISPC
+#include "XPBDStretchBiasElementConstraints.ispc.generated.h"
+#endif
+
+DECLARE_CYCLE_STAT(TEXT("Chaos XPBD Stretch Bias Constraint"), STAT_XPBD_StretchBias, STATGROUP_Chaos);
+
+#if INTEL_ISPC && !UE_BUILD_SHIPPING
+bool bChaos_XPBDStretchBiasElement_ISPC_Enabled = true;
+FAutoConsoleVariableRef CVarChaosXPBDStretchBiasISPCEnabled(TEXT("p.Chaos.XPBDStretchBias.ISPC"), bChaos_XPBDStretchBiasElement_ISPC_Enabled, TEXT("Whether to use ISPC optimizations in XPBD Stretch Bias constraints"));
+static_assert(sizeof(ispc::FVector4f) == sizeof(Chaos::Softs::FPAndInvM), "sizeof(ispc::FVector4f) != sizeof(Chaos::Softs::FPAndInvM");
+static_assert(sizeof(ispc::FIntVector) == sizeof(Chaos::TVec3<int32>), "sizeof(ispc::FIntVector) != sizeof(Chaos::TVec3<int32>");
+static_assert(sizeof(ispc::FVector2f) == sizeof(Chaos::Softs::FSolverVec2), "sizeof(ispc::FVector2f) != sizeof(Chaos::Softs::FSolverVec2");
+static_assert(sizeof(ispc::FVector4f) == sizeof(Chaos::Softs::FSolverMatrix22), "sizeof(ispc::FVector2f) != sizeof(Chaos::Softs::FSolverMatrix22");
+static_assert(sizeof(ispc::FVector3f) == sizeof(Chaos::Softs::FSolverVec3), "sizeof(ispc::FVector3f) != sizeof(Chaos::Softs::FSolverVec3");
+#endif
+
+namespace Chaos::Softs {
+
+// @todo(chaos): the parallel threshold (or decision to run parallel) should probably be owned by the solver and passed to the constraint container
+static int32 Chaos_XPBDStretchBias_ParallelConstraintCount = 100;
+
+FXPBDStretchBiasElementConstraints::FXPBDStretchBiasElementConstraints(const FSolverParticles& InParticles,
+	int32 ParticleOffset,
+	int32 ParticleCount,
+	const FTriangleMesh& TriangleMesh,
+	const TArray<TVec3<FVec2f>>& FaceVertexUVs,
+	const TConstArrayView<FRealSingle>& StiffnessWarpMultipliers,
+	const TConstArrayView<FRealSingle>& StiffnessWeftMultipliers,
+	const TConstArrayView<FRealSingle>& StiffnessBiasMultipliers,
+	const TConstArrayView<FRealSingle>& DampingMultipliers,
+	const TConstArrayView<FRealSingle>& WarpScaleMultipliers,
+	const TConstArrayView<FRealSingle>& WeftScaleMultipliers,
+	const FCollectionPropertyConstFacade& PropertyCollection,
+	bool bTrimKinematicConstraints)
+	: StiffnessWarp(
+		FSolverVec2(GetWeightedFloatXPBDStretchBiasElementStiffnessWarp(PropertyCollection, MaxStiffness)),
+		StiffnessWarpMultipliers,
+		TConstArrayView<TVec3<int32>>(Constraints),
+		ParticleOffset,
+		ParticleCount,
+		FPBDStiffness::DefaultTableSize,
+		FPBDStiffness::DefaultParameterFitBase,
+		MaxStiffness)
+	, StiffnessWeft(
+		FSolverVec2(GetWeightedFloatXPBDStretchBiasElementStiffnessWeft(PropertyCollection, MaxStiffness)),
+		StiffnessWeftMultipliers,
+		TConstArrayView<TVec3<int32>>(Constraints),
+		ParticleOffset,
+		ParticleCount,
+		FPBDStiffness::DefaultTableSize,
+		FPBDStiffness::DefaultParameterFitBase,
+		MaxStiffness)
+	, StiffnessBias(
+		FSolverVec2(GetWeightedFloatXPBDStretchBiasElementStiffnessBias(PropertyCollection, MaxStiffness)),
+		StiffnessBiasMultipliers,
+		TConstArrayView<TVec3<int32>>(Constraints),
+		ParticleOffset,
+		ParticleCount,
+		FPBDStiffness::DefaultTableSize,
+		FPBDStiffness::DefaultParameterFitBase,
+		MaxStiffness)
+	, DampingRatio(
+		FSolverVec2(GetWeightedFloatXPBDStretchBiasElementDamping(PropertyCollection, MinDamping)).ClampAxes(MinDamping, MaxDamping),
+		DampingMultipliers,
+		TConstArrayView<TVec3<int32>>(Constraints),
+		ParticleOffset,
+		ParticleCount)
+	, WarpScale(
+		FSolverVec2(GetWeightedFloatXPBDStretchBiasElementWarpScale(PropertyCollection, DefaultWarpWeftScale)).ClampAxes(MinWarpWeftScale, MaxWarpWeftScale),
+		WarpScaleMultipliers,
+		TConstArrayView<TVec3<int32>>(Constraints),
+		ParticleOffset,
+		ParticleCount)
+	, WeftScale(
+		FSolverVec2(GetWeightedFloatXPBDStretchBiasElementWeftScale(PropertyCollection, DefaultWarpWeftScale)).ClampAxes(MinWarpWeftScale, MaxWarpWeftScale),
+		WeftScaleMultipliers,
+		TConstArrayView<TVec3<int32>>(Constraints),
+		ParticleOffset,
+		ParticleCount)
+{
+	Lambdas.Init(FSolverVec3(0.), Constraints.Num());
+	InitConstraintsAndRestData(InParticles, TriangleMesh, FaceVertexUVs, GetXPBDStretchBiasElementUse3dRestLengths(PropertyCollection, bDefaultUse3dRestLengths), bTrimKinematicConstraints);
+	InitColor(InParticles, ParticleOffset, ParticleCount);
+}
+
+FXPBDStretchBiasElementConstraints::FXPBDStretchBiasElementConstraints(const FSolverParticles& InParticles,
+	int32 ParticleOffset,
+	int32 ParticleCount,
+	const FTriangleMesh& TriangleMesh,
+	const TArray<TVec3<FVec2f>>& FaceVertexUVs,
+	const TConstArrayView<FRealSingle>& StiffnessWarpMultipliers,
+	const TConstArrayView<FRealSingle>& StiffnessWeftMultipliers,
+	const TConstArrayView<FRealSingle>& StiffnessBiasMultipliers,
+	const TConstArrayView<FRealSingle>& DampingMultipliers,
+	const TConstArrayView<FRealSingle>& WarpScaleMultipliers,
+	const TConstArrayView<FRealSingle>& WeftScaleMultipliers,
+	const FSolverVec2& InStiffnessWarp,
+	const FSolverVec2& InStiffnessWeft,
+	const FSolverVec2& InStiffnessBias,
+	const FSolverVec2& InDampingRatio,
+	const FSolverVec2& InWarpScale,
+	const FSolverVec2& InWeftScale,
+	bool bUse3dRestLengths,
+	bool bTrimKinematicConstraints)
+	: StiffnessWarp(
+		InStiffnessWarp,
+		StiffnessWarpMultipliers,
+		TConstArrayView<TVec3<int32>>(Constraints),
+		ParticleOffset,
+		ParticleCount,
+		FPBDStiffness::DefaultTableSize,
+		FPBDStiffness::DefaultParameterFitBase,
+		MaxStiffness)
+	, StiffnessWeft(
+		InStiffnessWeft,
+		StiffnessWeftMultipliers,
+		TConstArrayView<TVec3<int32>>(Constraints),
+		ParticleOffset,
+		ParticleCount,
+		FPBDStiffness::DefaultTableSize,
+		FPBDStiffness::DefaultParameterFitBase,
+		MaxStiffness)
+	, StiffnessBias(
+		InStiffnessBias,
+		StiffnessBiasMultipliers,
+		TConstArrayView<TVec3<int32>>(Constraints),
+		ParticleOffset,
+		ParticleCount,
+		FPBDStiffness::DefaultTableSize,
+		FPBDStiffness::DefaultParameterFitBase,
+		MaxStiffness)
+	, DampingRatio(
+		InDampingRatio,
+		DampingMultipliers,
+		TConstArrayView<TVec3<int32>>(Constraints),
+		ParticleOffset,
+		ParticleCount)
+	, WarpScale(
+		InWarpScale,
+		WarpScaleMultipliers,
+		TConstArrayView<TVec3<int32>>(Constraints),
+		ParticleOffset,
+		ParticleCount)
+	, WeftScale(
+		InWeftScale,
+		WeftScaleMultipliers,
+		TConstArrayView<TVec3<int32>>(Constraints),
+		ParticleOffset,
+		ParticleCount)
+{
+	Lambdas.Init(FSolverVec3(0.), Constraints.Num());
+	InitConstraintsAndRestData(InParticles, TriangleMesh, FaceVertexUVs, bUse3dRestLengths, bTrimKinematicConstraints);
+	InitColor(InParticles, ParticleOffset, ParticleCount);
+}
+
+void FXPBDStretchBiasElementConstraints::InitConstraintsAndRestData(const FSolverParticles& InParticles, const FTriangleMesh& TriangleMesh,
+	const TArray<TVec3<FSolverVec2>>& FaceVertexUVs, const bool bUse3dRestLengths, const bool bTrimKinematicConstraints)
+{
+	const TArray<TVec3<int32>>& Elements = TriangleMesh.GetElements();
+	check(Elements.Num() == FaceVertexUVs.Num());
+
+	Constraints.Reserve(Elements.Num());
+	RestStretchLengths.Reserve(Elements.Num());
+	DeltaUVInverse.Reserve(Elements.Num());
+	StiffnessScales.Reserve(Elements.Num());
+
+	for (int32 ElemIdx = 0; ElemIdx < Elements.Num(); ++ElemIdx)
+	{
+		const TVec3<int32>& Constraint = Elements[ElemIdx];
+
+		if (bTrimKinematicConstraints && InParticles.InvM(Constraint[0]) == (FSolverReal)0. && InParticles.InvM(Constraint[1]) == (FSolverReal)0. && InParticles.InvM(Constraint[2]) == (FSolverReal)0.)
+		{
+			continue;
+		}
+
+		Constraints.Add(Constraint);
+
+		const FSolverVec2& UV0 = FaceVertexUVs[ElemIdx][0];
+		const FSolverVec2& UV1 = FaceVertexUVs[ElemIdx][1];
+		const FSolverVec2& UV2 = FaceVertexUVs[ElemIdx][2];
+
+		const FSolverVec2 UV01 = UV1 - UV0;
+		const FSolverVec2 UV02 = UV2 - UV0;
+
+
+		if (bUse3dRestLengths)
+		{
+			const FSolverVec2 DeltaUNormalized = FSolverVec2(UV01[0], UV02[0]).GetSafeNormal();
+			const FSolverVec2 DeltaVNormalized = FSolverVec2(UV01[1], UV02[1]).GetSafeNormal();
+			const FSolverReal DeltaUVDet = DeltaUNormalized[0] * DeltaVNormalized[1] - DeltaVNormalized[0] * DeltaUNormalized[1];
+
+			const FSolverMatrix22& DeltaUVInv = FMath::Abs(DeltaUVDet) < UE_SMALL_NUMBER ? DeltaUVInverse.Add_GetRef(FSolverMatrix22(1.f, 0.f, 0.f, 1.f)) :
+				DeltaUVInverse.Add_GetRef(FSolverMatrix22(DeltaVNormalized[1] / DeltaUVDet, -DeltaVNormalized[0] / DeltaUVDet, -DeltaUNormalized[1] / DeltaUVDet, DeltaUNormalized[0] / DeltaUVDet));
+
+			const FSolverVec3& X0 = InParticles.X(Constraint[0]);
+			const FSolverVec3& X1 = InParticles.X(Constraint[1]);
+			const FSolverVec3& X2 = InParticles.X(Constraint[2]);
+			const FSolverVec3 X01 = X1 - X0;
+			const FSolverVec3 X02 = X2 - X0;
+
+			const FSolverVec3 dX_dU = X01 * DeltaUVInv.M[0] + X02 * DeltaUVInv.M[1];
+			const FSolverVec3 dX_dV = X01 * DeltaUVInv.M[2] + X02 * DeltaUVInv.M[3];
+			const FSolverVec2& RestStretch = RestStretchLengths.Add_GetRef(FSolverVec2(dX_dU.Length(), dX_dV.Length()));
+
+			const FSolverReal StiffnessScaleU = RestStretch[0] < UE_SMALL_NUMBER ? 1.f : 1.f / RestStretch[0];
+			const FSolverReal StiffnessScaleV = RestStretch[1] < UE_SMALL_NUMBER ? 1.f : 1.f / RestStretch[1];
+			const FSolverReal BiasScale = FMath::Sqrt((FSolverReal).5 * FSolverVec3::CrossProduct(X01, X02).Length()); // Sqrt of triangle area
+
+			StiffnessScales.Add(FSolverVec3(StiffnessScaleU, StiffnessScaleV, BiasScale));
+		}
+		else
+		{
+			const FSolverVec2 DeltaUUnnormalized = FSolverVec2(UV01[0], UV02[0]);
+			const FSolverVec2 DeltaVUnnormalized = FSolverVec2(UV01[1], UV02[1]);
+			const FSolverReal DeltaUVDet = DeltaUUnnormalized[0] * DeltaVUnnormalized[1] - DeltaVUnnormalized[0] * DeltaUUnnormalized[1];
+			const FSolverMatrix22& DeltaUVInv = FMath::Abs(DeltaUVDet) < UE_SMALL_NUMBER ? DeltaUVInverse.Add_GetRef(FSolverMatrix22(1.f, 0.f, 0.f, 1.f)) :
+				DeltaUVInverse.Add_GetRef(FSolverMatrix22(DeltaVUnnormalized[1] / DeltaUVDet, -DeltaVUnnormalized[0] / DeltaUVDet, -DeltaUUnnormalized[1] / DeltaUVDet, DeltaUUnnormalized[0] / DeltaUVDet));
+			RestStretchLengths.Add(FSolverVec2(1.f));
+
+			const FSolverReal DeltaULen = DeltaUUnnormalized.Length();
+			const FSolverReal DeltaVLen = DeltaVUnnormalized.Length();
+
+			const FSolverReal StiffnessScaleU = DeltaULen < UE_SMALL_NUMBER ? 1.f : 1.f / DeltaULen;
+			const FSolverReal StiffnessScaleV = DeltaVLen < UE_SMALL_NUMBER ? 1.f : 1.f / DeltaVLen;
+			const FSolverReal BiasScale = FMath::Sqrt((FSolverReal).5 * FMath::Abs(UV01[0] * UV02[1] - UV01[1] * UV02[0]));
+			StiffnessScales.Add(FSolverVec3(StiffnessScaleU, StiffnessScaleV, BiasScale));
+		}
+	}
+}
+
+void FXPBDStretchBiasElementConstraints::InitColor(const FSolverParticles& InParticles, const int32 ParticleOffset, const int32 ParticleCount)
+{
+	// In dev builds we always color so we can tune the system without restarting. See Apply()
+#if UE_BUILD_SHIPPING || UE_BUILD_TEST
+	if (Constraints.Num() > Chaos_XPBDStretchBias_ParallelConstraintCount)
+#endif
+	{
+		const TArray<TArray<int32>> ConstraintsPerColor = FGraphColoring::ComputeGraphColoring(Constraints, InParticles, ParticleOffset, ParticleOffset + ParticleCount);
+
+		// Reorder constraints based on color so each array in ConstraintsPerColor contains contiguous elements.
+		TArray<TVec3<int32>> ReorderedConstraints;
+		TArray<FSolverVec2> ReorderedRestStretchLengths;
+		TArray<FSolverMatrix22> ReorderedDeltaUVInverse;
+		TArray<FSolverVec3> ReorderedStiffnessScales;
+		TArray<int32> OrigToReorderedIndices; // used to reorder stiffness indices
+		ReorderedConstraints.SetNumUninitialized(Constraints.Num());
+		ReorderedRestStretchLengths.SetNumUninitialized(Constraints.Num());
+		ReorderedDeltaUVInverse.SetNumUninitialized(Constraints.Num());
+		ReorderedStiffnessScales.SetNumUninitialized(Constraints.Num());
+		OrigToReorderedIndices.SetNumUninitialized(Constraints.Num());
+
+		ConstraintsPerColorStartIndex.Reset(ConstraintsPerColor.Num() + 1);
+
+		int32 ReorderedIndex = 0;
+		for (const TArray<int32>& ConstraintsBatch : ConstraintsPerColor)
+		{
+			ConstraintsPerColorStartIndex.Add(ReorderedIndex);
+			for (const int32& BatchConstraint : ConstraintsBatch)
+			{
+				const int32 OrigIndex = BatchConstraint;
+				ReorderedConstraints[ReorderedIndex] = Constraints[OrigIndex];
+				ReorderedRestStretchLengths[ReorderedIndex] = RestStretchLengths[OrigIndex];
+				ReorderedDeltaUVInverse[ReorderedIndex] = DeltaUVInverse[OrigIndex];
+				ReorderedStiffnessScales[ReorderedIndex] = StiffnessScales[OrigIndex];
+				OrigToReorderedIndices[OrigIndex] = ReorderedIndex;
+
+				++ReorderedIndex;
+			}
+		}
+		ConstraintsPerColorStartIndex.Add(ReorderedIndex);
+
+		Constraints = MoveTemp(ReorderedConstraints);
+		RestStretchLengths = MoveTemp(ReorderedRestStretchLengths);
+		DeltaUVInverse = MoveTemp(ReorderedDeltaUVInverse);
+		StiffnessScales = MoveTemp(ReorderedStiffnessScales);
+		StiffnessWarp.ReorderIndices(OrigToReorderedIndices);
+		StiffnessWeft.ReorderIndices(OrigToReorderedIndices);
+		StiffnessBias.ReorderIndices(OrigToReorderedIndices);
+		DampingRatio.ReorderIndices(OrigToReorderedIndices);
+	}
+}
+
+void FXPBDStretchBiasElementConstraints::Apply(FSolverParticles& Particles, const FSolverReal Dt) const
+{
+	TRACE_CPUPROFILER_EVENT_SCOPE(FXPBDStretchBiasElementConstraints_Apply);
+	SCOPE_CYCLE_COUNTER(STAT_XPBD_StretchBias);
+
+	const bool StiffnessWarpHasWeightMap = StiffnessWarp.HasWeightMap();
+	const bool StiffnessWeftHasWeightMap = StiffnessWeft.HasWeightMap();
+	const bool StiffnessBiasHasWeightMap = StiffnessBias.HasWeightMap();
+	const bool DampingHasWeightMap = DampingRatio.HasWeightMap();
+	const bool WarpScaleHasWeightMap = WarpScale.HasWeightMap();
+	const bool WeftScaleHasWeightMap = WeftScale.HasWeightMap();
+
+	const FSolverReal StiffnessWarpNoMap = (FSolverReal)StiffnessWarp;
+	const FSolverReal StiffnessWeftNoMap = (FSolverReal)StiffnessWeft;
+	const FSolverReal StiffnessBiasNoMap = (FSolverReal)StiffnessBias;
+	const FSolverReal DampingNoMap = (FSolverReal)DampingRatio;
+	const FSolverReal WarpScaleNoMap = (FSolverReal)WarpScale;
+	const FSolverReal WeftScaleNoMap = (FSolverReal)WeftScale;
+
+	if (ConstraintsPerColorStartIndex.Num() > 0 && Constraints.Num() > Chaos_XPBDStretchBias_ParallelConstraintCount)
+	{
+		const int32 ConstraintColorNum = ConstraintsPerColorStartIndex.Num() - 1;
+		if (!StiffnessWarpHasWeightMap && !StiffnessWeftHasWeightMap && !StiffnessBiasHasWeightMap && !DampingHasWeightMap && !WarpScaleHasWeightMap && !WeftScaleHasWeightMap)
+		{
+			const FSolverVec3 ExpStiffnessValue(StiffnessWarpNoMap, StiffnessWeftNoMap, StiffnessBiasNoMap);
+			if (ExpStiffnessValue.Max() < MinStiffness)
+			{
+				return;
+			}
+
+#if INTEL_ISPC
+			if (bRealTypeCompatibleWithISPC && bChaos_XPBDStretchBiasElement_ISPC_Enabled)
+			{
+				if (DampingNoMap > 0)
+				{
+					for (int32 ConstraintColorIndex = 0; ConstraintColorIndex < ConstraintColorNum; ++ConstraintColorIndex)
+					{
+						const int32 ColorStart = ConstraintsPerColorStartIndex[ConstraintColorIndex];
+						const int32 ColorSize = ConstraintsPerColorStartIndex[ConstraintColorIndex + 1] - ColorStart;
+						ispc::ApplyXPBDStretchBiasConstraintsWithDamping(
+							(ispc::FVector4f*)Particles.GetPAndInvM().GetData(),
+							(const ispc::FVector3f*)Particles.X().GetData(),
+							(ispc::FIntVector*)&Constraints.GetData()[ColorStart],
+							(ispc::FVector2f*)&RestStretchLengths.GetData()[ColorStart],
+							(ispc::FVector4f*)&DeltaUVInverse.GetData()[ColorStart],
+							(ispc::FVector3f*)&StiffnessScales.GetData()[ColorStart],
+							(ispc::FVector3f*)&Lambdas.GetData()[ColorStart],
+							Dt,
+							MinStiffness,
+							reinterpret_cast<const ispc::FVector3f&>(ExpStiffnessValue),
+							DampingNoMap,
+							WarpScaleNoMap,
+							WeftScaleNoMap,
+							ColorSize);
+					}
+
+				}
+				else
+				{
+					for (int32 ConstraintColorIndex = 0; ConstraintColorIndex < ConstraintColorNum; ++ConstraintColorIndex)
+					{
+						const int32 ColorStart = ConstraintsPerColorStartIndex[ConstraintColorIndex];
+						const int32 ColorSize = ConstraintsPerColorStartIndex[ConstraintColorIndex + 1] - ColorStart;
+						ispc::ApplyXPBDStretchBiasConstraints(
+							(ispc::FVector4f*)Particles.GetPAndInvM().GetData(),
+							(ispc::FIntVector*)&Constraints.GetData()[ColorStart],
+							(ispc::FVector2f*)&RestStretchLengths.GetData()[ColorStart],
+							(ispc::FVector4f*)&DeltaUVInverse.GetData()[ColorStart],
+							(ispc::FVector3f*)&StiffnessScales.GetData()[ColorStart],
+							(ispc::FVector3f*)&Lambdas.GetData()[ColorStart],
+							Dt,
+							MinStiffness,
+							reinterpret_cast<const ispc::FVector3f&>(ExpStiffnessValue),
+							WarpScaleNoMap,
+							WeftScaleNoMap,
+							ColorSize);
+					}
+				}
+			}
+			else
+#endif
+			{
+				for (int32 ConstraintColorIndex = 0; ConstraintColorIndex < ConstraintColorNum; ++ConstraintColorIndex)
+				{
+					const int32 ColorStart = ConstraintsPerColorStartIndex[ConstraintColorIndex];
+					const int32 ColorSize = ConstraintsPerColorStartIndex[ConstraintColorIndex + 1] - ColorStart;
+					PhysicsParallelFor(ColorSize, [&](const int32 Index)
+					{
+						const int32 ConstraintIndex = ColorStart + Index;
+						ApplyHelper(Particles, Dt, ConstraintIndex, ExpStiffnessValue, DampingNoMap, WarpScaleNoMap, WeftScaleNoMap);
+					});
+				}
+			}			
+		}
+		else // has weight maps
+		{
+
+#if INTEL_ISPC
+			if (bRealTypeCompatibleWithISPC && bChaos_XPBDStretchBiasElement_ISPC_Enabled)
+			{
+				if (DampingHasWeightMap || DampingNoMap > 0)
+				{
+					for (int32 ConstraintColorIndex = 0; ConstraintColorIndex < ConstraintColorNum; ++ConstraintColorIndex)
+					{
+						const int32 ColorStart = ConstraintsPerColorStartIndex[ConstraintColorIndex];
+						const int32 ColorSize = ConstraintsPerColorStartIndex[ConstraintColorIndex + 1] - ColorStart;
+						ispc::ApplyXPBDStretchBiasConstraintsWithDampingAndMaps(
+							(ispc::FVector4f*)Particles.GetPAndInvM().GetData(),
+							(const ispc::FVector3f*)Particles.X().GetData(),
+							(ispc::FIntVector*)&Constraints.GetData()[ColorStart],
+							(ispc::FVector2f*)&RestStretchLengths.GetData()[ColorStart],
+							(ispc::FVector4f*)&DeltaUVInverse.GetData()[ColorStart],
+							(ispc::FVector3f*)&StiffnessScales.GetData()[ColorStart],
+							(ispc::FVector3f*)&Lambdas.GetData()[ColorStart],
+							Dt,
+							MinStiffness,
+							StiffnessWarpHasWeightMap,
+							&StiffnessWarp.GetIndices().GetData()[ColorStart],
+							&StiffnessWarp.GetTable().GetData()[0],
+							StiffnessWeftHasWeightMap,
+							&StiffnessWeft.GetIndices().GetData()[ColorStart],
+							&StiffnessWeft.GetTable().GetData()[0],
+							StiffnessBiasHasWeightMap,
+							&StiffnessBias.GetIndices().GetData()[ColorStart],
+							&StiffnessBias.GetTable().GetData()[0],
+							DampingHasWeightMap,
+							&DampingRatio.GetIndices().GetData()[ColorStart],
+							&DampingRatio.GetTable().GetData()[0],
+							WarpScaleHasWeightMap,
+							&WarpScale.GetIndices().GetData()[ColorStart],
+							&WarpScale.GetTable().GetData()[0],
+							WeftScaleHasWeightMap,
+							&WeftScale.GetIndices().GetData()[ColorStart],
+							&WeftScale.GetTable().GetData()[0],
+							ColorSize);
+					}
+
+				}
+				else
+				{
+					for (int32 ConstraintColorIndex = 0; ConstraintColorIndex < ConstraintColorNum; ++ConstraintColorIndex)
+					{
+						const int32 ColorStart = ConstraintsPerColorStartIndex[ConstraintColorIndex];
+						const int32 ColorSize = ConstraintsPerColorStartIndex[ConstraintColorIndex + 1] - ColorStart;
+						ispc::ApplyXPBDStretchBiasConstraintsWithMaps(
+							(ispc::FVector4f*)Particles.GetPAndInvM().GetData(),
+							(ispc::FIntVector*)&Constraints.GetData()[ColorStart],
+							(ispc::FVector2f*)&RestStretchLengths.GetData()[ColorStart],
+							(ispc::FVector4f*)&DeltaUVInverse.GetData()[ColorStart],
+							(ispc::FVector3f*)&StiffnessScales.GetData()[ColorStart],
+							(ispc::FVector3f*)&Lambdas.GetData()[ColorStart],
+							Dt,
+							MinStiffness,
+							StiffnessWarpHasWeightMap,
+							&StiffnessWarp.GetIndices().GetData()[ColorStart],
+							&StiffnessWarp.GetTable().GetData()[0],
+							StiffnessWeftHasWeightMap,
+							&StiffnessWeft.GetIndices().GetData()[ColorStart],
+							&StiffnessWeft.GetTable().GetData()[0],
+							StiffnessBiasHasWeightMap,
+							&StiffnessBias.GetIndices().GetData()[ColorStart],
+							&StiffnessBias.GetTable().GetData()[0],
+							WarpScaleHasWeightMap,
+							&WarpScale.GetIndices().GetData()[ColorStart],
+							&WarpScale.GetTable().GetData()[0],
+							WeftScaleHasWeightMap,
+							&WeftScale.GetIndices().GetData()[ColorStart],
+							&WeftScale.GetTable().GetData()[0],
+							ColorSize);
+					}
+				}
+
+			}
+			else
+#endif
+			{
+				for (int32 ConstraintColorIndex = 0; ConstraintColorIndex < ConstraintColorNum; ++ConstraintColorIndex)
+				{
+					const int32 ColorStart = ConstraintsPerColorStartIndex[ConstraintColorIndex];
+					const int32 ColorSize = ConstraintsPerColorStartIndex[ConstraintColorIndex + 1] - ColorStart;
+					PhysicsParallelFor(ColorSize, [&](const int32 Index)
+					{
+						const int32 ConstraintIndex = ColorStart + Index;
+						const FSolverReal ExpStiffnessWarpValue = StiffnessWarpHasWeightMap ? StiffnessWarp[ConstraintIndex] : StiffnessWarpNoMap;
+						const FSolverReal ExpStiffnessWeftValue = StiffnessWeftHasWeightMap ? StiffnessWeft[ConstraintIndex] : StiffnessWeftNoMap;
+						const FSolverReal ExpStiffnessBiasValue = StiffnessBiasHasWeightMap ? StiffnessBias[ConstraintIndex] : StiffnessBiasNoMap;
+						const FSolverReal DampingRatioValue = DampingHasWeightMap ? DampingRatio[ConstraintIndex] : DampingNoMap;
+						const FSolverReal WarpScaleValue = WarpScaleHasWeightMap ? WarpScale[ConstraintIndex] : WarpScaleNoMap;
+						const FSolverReal WeftScaleValue = WeftScaleHasWeightMap ? WeftScale[ConstraintIndex] : WeftScaleNoMap;
+						ApplyHelper(Particles, Dt, ConstraintIndex, FSolverVec3(ExpStiffnessWarpValue, ExpStiffnessWeftValue, ExpStiffnessBiasValue),DampingRatioValue, WarpScaleValue, WeftScaleValue);
+					});
+				}
+			}
+		}
+	}
+	else
+	{
+		if (!StiffnessWarpHasWeightMap && !StiffnessWeftHasWeightMap && !StiffnessBiasHasWeightMap && !DampingHasWeightMap)
+		{
+			const FSolverVec3 ExpStiffnessValue(StiffnessWarpNoMap, StiffnessWeftNoMap, StiffnessBiasNoMap);
+			if (ExpStiffnessValue.Max() < MinStiffness)
+			{
+				return;
+			}
+			for (int32 ConstraintIndex = 0; ConstraintIndex < Constraints.Num(); ++ConstraintIndex)
+			{
+				ApplyHelper(Particles, Dt, ConstraintIndex, ExpStiffnessValue, DampingNoMap, WarpScaleNoMap, WeftScaleNoMap);
+			}
+		}
+		else
+		{
+			for (int32 ConstraintIndex = 0; ConstraintIndex < Constraints.Num(); ++ConstraintIndex)
+			{
+				const FSolverReal ExpStiffnessWarpValue = StiffnessWarpHasWeightMap ? StiffnessWarp[ConstraintIndex] : StiffnessWarpNoMap;
+				const FSolverReal ExpStiffnessWeftValue = StiffnessWeftHasWeightMap ? StiffnessWeft[ConstraintIndex] : StiffnessWeftNoMap;
+				const FSolverReal ExpStiffnessBiasValue = StiffnessBiasHasWeightMap ? StiffnessBias[ConstraintIndex] : StiffnessBiasNoMap;
+				const FSolverReal DampingRatioValue = DampingHasWeightMap ? DampingRatio[ConstraintIndex] : DampingNoMap;
+				const FSolverReal WarpScaleValue = WarpScaleHasWeightMap ? WarpScale[ConstraintIndex] : WarpScaleNoMap;
+				const FSolverReal WeftScaleValue = WeftScaleHasWeightMap ? WeftScale[ConstraintIndex] : WeftScaleNoMap;
+				ApplyHelper(Particles, Dt, ConstraintIndex, FSolverVec3(ExpStiffnessWarpValue, ExpStiffnessWeftValue, ExpStiffnessBiasValue), DampingRatioValue, WarpScaleValue, WeftScaleValue);
+			}
+		}
+	}
+}
+
+static FSolverReal CalcDLambda(const FSolverReal StiffnessValue, const FSolverReal Dt, const FSolverReal Lambda, const FSolverReal Damping,
+	const FSolverReal InvM1, const FSolverReal InvM2, const FSolverReal InvM3,
+	const FSolverVec3& V1TimesDt, const FSolverVec3& V2TimesDt, const FSolverVec3& V3TimesDt,
+	const FSolverReal C, const FSolverVec3& dC_dX1, const FSolverVec3& dC_dX2, const FSolverVec3& dC_dX3)
+{
+	if (StiffnessValue < FXPBDStretchBiasElementConstraints::MinStiffness)
+	{
+		return (FSolverReal)0.f;
+	}
+	const FSolverReal Alpha = (FSolverReal)1.f / (StiffnessValue * Dt * Dt);
+	const FSolverReal Gamma = Alpha * Damping * Dt;
+	const FSolverReal DampingTerm = Gamma * (FSolverVec3::DotProduct(V1TimesDt, dC_dX1) + FSolverVec3::DotProduct(V2TimesDt, dC_dX2) + FSolverVec3::DotProduct(V3TimesDt, dC_dX3));
+	const FSolverReal Denom = ((FSolverReal)1.f + Gamma) * (InvM1 * dC_dX1.SizeSquared() + InvM2 * dC_dX2.SizeSquared() + InvM3 * dC_dX3.SizeSquared()) + Alpha;
+	return (C - Alpha * Lambda + DampingTerm) / Denom;
+}
+
+void FXPBDStretchBiasElementConstraints::ApplyHelper(FSolverParticles& Particles, const FSolverReal Dt, const int32 ConstraintIndex, const FSolverVec3& ExpStiffnessValue, const FSolverReal DampingRatioValue, const FSolverReal WarpScaleValue, const FSolverReal WeftScaleValue) const
+{
+	const TVec3<int32>& Constraint = Constraints[ConstraintIndex];
+	const int32 i0 = Constraint[0];
+	const int32 i1 = Constraint[1];
+	const int32 i2 = Constraint[2];
+	if (Particles.InvM(i0) == (FSolverReal)0. && Particles.InvM(i1) == (FSolverReal)0. && Particles.InvM(i2) == (FSolverReal)0.)
+	{
+		return;
+	}
+
+	const FSolverReal CombinedInvMass = Particles.InvM(i0) + Particles.InvM(i1) + Particles.InvM(i2);
+
+	const FSolverMatrix22& DeltaUVInv = DeltaUVInverse[ConstraintIndex];
+
+	const FSolverVec3 P01 = Particles.P(i1) - Particles.P(i0);
+	const FSolverVec3 P02 = Particles.P(i2) - Particles.P(i0);
+	const FSolverVec3 dX_dU = P01 * DeltaUVInv.M[0] + P02 * DeltaUVInv.M[1];
+	const FSolverVec3 dX_dV = P01 * DeltaUVInv.M[2] + P02 * DeltaUVInv.M[3];
+
+	const FSolverReal dX_dU_length = dX_dU.Length();
+	const FSolverReal dX_dV_length = dX_dV.Length();
+	auto SafeRecip = [](const FSolverReal Len, const FSolverReal Fallback)
+	{
+		if (Len > UE_SMALL_NUMBER)
+		{
+			return 1.f / Len;
+		}
+		return Fallback;
+	};
+	const FSolverReal OneOverDxDuLen = SafeRecip(dX_dU_length, 0.f);
+	const FSolverReal OneOverDxDvLen = SafeRecip(dX_dV_length, 0.f);
+
+	const FSolverVec3 dX_dU_normalized = dX_dU * OneOverDxDuLen;
+	const FSolverVec3 dX_dV_normalized = dX_dV * OneOverDxDvLen;
+
+	// constraints
+	const FSolverReal Cu = (dX_dU_length - RestStretchLengths[ConstraintIndex][0]*WarpScaleValue); // stretch in warp direction
+	const FSolverReal Cv = (dX_dV_length - RestStretchLengths[ConstraintIndex][1]*WeftScaleValue); // stretch in weft direction
+	const FSolverReal Cs = FSolverVec3::DotProduct(dX_dU_normalized, dX_dV_normalized);
+
+	const FSolverVec3 dCu_dX1 = dX_dU_normalized * DeltaUVInv.M[0];
+	const FSolverVec3 dCu_dX2 = dX_dU_normalized * DeltaUVInv.M[1];
+	const FSolverVec3 dCu_dX0 = -dCu_dX1 - dCu_dX2;
+	const FSolverVec3 dCv_dX1 = dX_dV_normalized * DeltaUVInv.M[2];
+	const FSolverVec3 dCv_dX2 = dX_dV_normalized * DeltaUVInv.M[3];
+	const FSolverVec3 dCv_dX0 = -dCv_dX1 - dCv_dX2;
+	const FSolverVec3 dCs_dX1 = (dX_dV_normalized * OneOverDxDuLen * DeltaUVInv.M[0] + dX_dU_normalized * OneOverDxDvLen * DeltaUVInv.M[2]);
+	const FSolverVec3 dCs_dX2 = (dX_dV_normalized * OneOverDxDuLen * DeltaUVInv.M[1] + dX_dU_normalized * OneOverDxDvLen * DeltaUVInv.M[3]);
+	const FSolverVec3 dCs_dX0 = -dCs_dX1 - dCs_dX2;
+
+	const FSolverVec3 V0TimesDt = Particles.P(i0) - Particles.X(i0);
+	const FSolverVec3 V1TimesDt = Particles.P(i1) - Particles.X(i1);
+	const FSolverVec3 V2TimesDt = Particles.P(i2) - Particles.X(i2);
+
+	// These scale factors make everything resolution independent
+	const FSolverVec3 FinalStiffnesses = StiffnessScales[ConstraintIndex] * ExpStiffnessValue;
+
+	const FSolverVec3 Damping = (FSolverReal)2.f * DampingRatioValue * FMath::InvSqrt(CombinedInvMass) *
+		FSolverVec3(FMath::Sqrt(FinalStiffnesses[0]), FMath::Sqrt(FinalStiffnesses[1]), FMath::Sqrt(FinalStiffnesses[2]));
+
+	FSolverVec3& Lambda = Lambdas[ConstraintIndex];
+	FSolverVec3 DLambda;
+	DLambda[0] = CalcDLambda(FinalStiffnesses[0], Dt, Lambda[0], Damping[0], Particles.InvM(i0), Particles.InvM(i1), Particles.InvM(i2),
+		V0TimesDt, V1TimesDt, V2TimesDt, Cu, dCu_dX0, dCu_dX1, dCu_dX2);
+	DLambda[1] = CalcDLambda(FinalStiffnesses[1], Dt, Lambda[1], Damping[1], Particles.InvM(i0), Particles.InvM(i1), Particles.InvM(i2),
+		V0TimesDt, V1TimesDt, V2TimesDt, Cv, dCv_dX0, dCv_dX1, dCv_dX2);
+	DLambda[2] = CalcDLambda(FinalStiffnesses[2], Dt, Lambda[2], Damping[2], Particles.InvM(i0), Particles.InvM(i1), Particles.InvM(i2),
+		V0TimesDt, V1TimesDt, V2TimesDt, Cs, dCs_dX0, dCs_dX1, dCs_dX2);
+
+	Particles.P(i0) -= Particles.InvM(i0) * (DLambda[0] * dCu_dX0 + DLambda[1] * dCv_dX0 + DLambda[2] * dCs_dX0);
+	Particles.P(i1) -= Particles.InvM(i1) * (DLambda[0] * dCu_dX1 + DLambda[1] * dCv_dX1 + DLambda[2] * dCs_dX1);
+	Particles.P(i2) -= Particles.InvM(i2) * (DLambda[0] * dCu_dX2 + DLambda[1] * dCv_dX2 + DLambda[2] * dCs_dX2);
+
+	Lambda += DLambda;
+}
+
+}  // End namespace Chaos::Softs
diff --git a/Engine/Source/Runtime/Experimental/Chaos/Private/Chaos/XPBDStretchBiasElementConstraints.ispc b/Engine/Source/Runtime/Experimental/Chaos/Private/Chaos/XPBDStretchBiasElementConstraints.ispc
new file mode 100644
--- /dev/null
+++ b/Engine/Source/Runtime/Experimental/Chaos/Private/Chaos/XPBDStretchBiasElementConstraints.ispc
@@ -0,0 +1,450 @@
+// Copyright Epic Games, Inc. All Rights Reserved.
+
+#define EXPLICIT_VECTOR4 1
+
+#include "Math/Vector.isph"
+#include "Chaos/PBDSofts.isph"
+
+static inline varying float SafeDivide(const varying float Numerator, const varying float Denominator)
+{
+	return select(Denominator <= FLOAT_SMALL_NUMBER, (uniform float)0, Numerator / Denominator);
+}
+
+static void CalcConstraintsAndGrads(const FVector3f& P0, const FVector3f& P1, const FVector3f& P2, const FVector4f& DeltaUVInv,
+	const FVector2f& RestStretchLengths, const float WarpScale, const float WeftScale,
+	float& Cu, float& Cv, float& Cs, 
+	FVector3f& dCu_dX1, FVector3f& dCu_dX2, FVector3f& dCv_dX1, FVector3f& dCv_dX2, FVector3f& dCs_dX1, FVector3f& dCs_dX2)
+{
+	const FVector3f P01 = P1 - P0;
+	const FVector3f P02 = P2 - P0;
+	const FVector3f dX_dU = DeltaUVInv.V[0] * P01 + DeltaUVInv.V[1] * P02;
+	const FVector3f dX_dV = DeltaUVInv.V[2] * P01 + DeltaUVInv.V[3] * P02;
+
+	const float dX_dU_length = VectorSize(dX_dU);
+	const float dX_dV_length = VectorSize(dX_dV);
+
+	const float OneOverDxDuLen = SafeDivide(1.f, dX_dU_length);
+	const float OneOverDxDvLen = SafeDivide(1.f, dX_dV_length);
+
+	const FVector3f dX_dU_normalized = OneOverDxDuLen * dX_dU;
+	const FVector3f dX_dV_normalized = OneOverDxDvLen * dX_dV;
+
+	Cu = dX_dU_length - RestStretchLengths.V[0] * WarpScale;
+	Cv = dX_dV_length - RestStretchLengths.V[1] * WeftScale;
+	Cs = VectorDot(dX_dU_normalized, dX_dV_normalized);
+
+	dCu_dX1 = dX_dU_normalized * DeltaUVInv.V[0];
+	dCu_dX2 = dX_dU_normalized * DeltaUVInv.V[1];
+	dCv_dX1 = dX_dV_normalized * DeltaUVInv.V[2];
+	dCv_dX2 = dX_dV_normalized * DeltaUVInv.V[3];
+	dCs_dX1 = (dX_dV_normalized * OneOverDxDuLen * DeltaUVInv.V[0] + dX_dU_normalized * OneOverDxDvLen * DeltaUVInv.V[2]);
+	dCs_dX2 = (dX_dV_normalized * OneOverDxDuLen * DeltaUVInv.V[1] + dX_dU_normalized * OneOverDxDvLen * DeltaUVInv.V[3]);
+}
+
+static inline float CalcDLambda(const float StiffnessValue, const float Dt, const float MinStiffness, const float Lambda,
+	const float InvM0, const float InvM1, const float InvM2,
+	const float C, const FVector3f& dC_dX1, const FVector3f& dC_dX2)
+{
+	const float Alpha = 1.f / (StiffnessValue  * Dt * Dt);
+	const float Denom = InvM0 * VectorSizeSquared(dC_dX1 + dC_dX2) + InvM1 * VectorSizeSquared(dC_dX1) + InvM2 * VectorSizeSquared(dC_dX2) + Alpha;
+	const float DLambda = (C - Alpha * Lambda) / Denom;
+	return select(StiffnessValue < MinStiffness, (uniform float)0, DLambda);
+}
+
+static inline float CalcDamping(const float StiffnessValue, const float CombinedInvMass, const float DampingRatio)
+{
+	return 2.f * DampingRatio * sqrt(StiffnessValue / CombinedInvMass);
+}
+
+static inline float CalcDLambda(const float StiffnessValue, const float Dt, const float MinStiffness, const float Lambda, const float Damping,
+	const float InvM0, const float InvM1, const float InvM2,
+	const FVector3f& V0TimesDt, const FVector3f& V1TimesDt, const FVector3f& V2TimesDt, 
+	const float C, const FVector3f& dC_dX1, const FVector3f& dC_dX2)
+{
+	const float Alpha = 1.f / (StiffnessValue  * Dt * Dt);
+	const float Gamma = Alpha * Damping * Dt;
+	const float DampingTerm = Gamma * (VectorDot(V1TimesDt, dC_dX1) + VectorDot(V2TimesDt, dC_dX2) - VectorDot(V0TimesDt, dC_dX1 + dC_dX2));
+	const float Denom = (1. + Gamma) * (InvM0 * VectorSizeSquared(dC_dX1 + dC_dX2) + InvM1 * VectorSizeSquared(dC_dX1) + InvM2 * VectorSizeSquared(dC_dX2)) + Alpha;
+	const float DLambda = (C - Alpha * Lambda + DampingTerm) / Denom;
+	return select(StiffnessValue < MinStiffness, (uniform float)0, DLambda);
+}
+
+export void ApplyXPBDStretchBiasConstraints(uniform FVector4f PandInvM[],
+									const uniform FIntVector Constraints[],
+									const uniform FVector2f RestStretchLengths[],
+									const uniform FVector4f DeltaUVInverses[],
+									const uniform FVector3f StiffnessScales[],
+									uniform FVector3f Lambdas[],
+									const uniform float Dt,
+									const uniform float MinStiffness,
+									const uniform FVector3f& Stiffness,
+									const uniform float WarpScale,
+									const uniform float WeftScale,
+									const uniform int32 NumConstraints)
+{
+
+	foreach(i = 0 ... NumConstraints)
+	{
+		const varying FIntVector Constraint = VectorLoad(&Constraints[extract(i,0)]);
+
+		const varying int32 i0 = Constraint.V[0];
+		const varying int32 i1 = Constraint.V[1];
+		const varying int32 i2 = Constraint.V[2];
+		
+		const varying FVector4f PandInvM0 = VectorGather(&PandInvM[i0]);
+		const varying FVector4f PandInvM1 = VectorGather(&PandInvM[i1]);
+		const varying FVector4f PandInvM2 = VectorGather(&PandInvM[i2]);
+
+		varying FVector3f P0, P1, P2;
+		varying float M0, M1, M2;
+		
+		UnzipPandInvM(PandInvM0, P0, M0);
+		UnzipPandInvM(PandInvM1, P1, M1);
+		UnzipPandInvM(PandInvM2, P2, M2);
+
+		const varying FVector4f DeltaUVInverse = VectorLoad(&DeltaUVInverses[extract(i,0)]);
+		const varying FVector2f RestStretchLength = VectorLoad(&RestStretchLengths[extract(i,0)]);
+
+		varying float Cu, Cv, Cs;
+		varying FVector3f dCu_dX1, dCu_dX2, dCv_dX1, dCv_dX2, dCs_dX1, dCs_dX2;
+		CalcConstraintsAndGrads(P0, P1, P2, DeltaUVInverse, RestStretchLength, WarpScale, WeftScale, Cu, Cv, Cs,
+			dCu_dX1, dCu_dX2, dCv_dX1, dCv_dX2, dCs_dX1, dCs_dX2);
+		
+		FVector3f Lambda = VectorLoad(&Lambdas[extract(i,0)]);
+		const varying FVector3f StiffnessScale = VectorLoad(&StiffnessScales[extract(i,0)]);
+		const varying FVector3f FinalStiffness = Stiffness * StiffnessScale;
+
+		const varying float DLambdaU = CalcDLambda(FinalStiffness.V[0], Dt, MinStiffness, Lambda.V[0], M0, M1, M2, Cu, dCu_dX1, dCu_dX2);
+		const varying float DLambdaV = CalcDLambda(FinalStiffness.V[1], Dt, MinStiffness, Lambda.V[1], M0, M1, M2, Cv, dCv_dX1, dCv_dX2);
+		const varying float DLambdaS = CalcDLambda(FinalStiffness.V[2], Dt, MinStiffness, Lambda.V[2], M0, M1, M2, Cs, dCs_dX1, dCs_dX2);
+		const varying FVector3f DLambda = {{DLambdaU, DLambdaV, DLambdaS}};
+
+		Lambda = Lambda + DLambda;
+		VectorStore(&Lambdas[extract(i,0)], Lambda);
+		
+		const varying FVector3f DP1 = DLambdaU * dCu_dX1 + DLambdaV * dCv_dX1 + DLambdaS * dCs_dX1;
+		const varying FVector3f DP2 = DLambdaU * dCu_dX2 + DLambdaV * dCv_dX2 + DLambdaS * dCs_dX2;
+
+		if (M0 > 0)
+		{
+			VectorScatter(&PandInvM[i0], SetVector4(P0 + M0 * (DP1 + DP2), M0));
+		}
+		if (M1 > 0)
+		{
+			VectorScatter(&PandInvM[i1], SetVector4(P1 - M1 * DP1, M1));
+		}
+		if (M2 > 0)
+		{
+			VectorScatter(&PandInvM[i2], SetVector4(P2 - M2 * DP2, M2));
+		}
+	}
+}
+
+export void ApplyXPBDStretchBiasConstraintsWithDamping(uniform FVector4f PandInvM[],
+									const uniform FVector3f X[],
+									const uniform FIntVector Constraints[],
+									const uniform FVector2f RestStretchLengths[],
+									const uniform FVector4f DeltaUVInverses[],
+									const uniform FVector3f StiffnessScales[],
+									uniform FVector3f Lambdas[],
+									const uniform float Dt,
+									const uniform float MinStiffness,
+									const uniform FVector3f& Stiffness,
+									const uniform float DampingRatio,
+									const uniform float WarpScale,
+									const uniform float WeftScale,
+									const uniform int32 NumConstraints)
+{
+
+	foreach(i = 0 ... NumConstraints)
+	{
+		const varying FIntVector Constraint = VectorLoad(&Constraints[extract(i,0)]);
+
+		const varying int32 i0 = Constraint.V[0];
+		const varying int32 i1 = Constraint.V[1];
+		const varying int32 i2 = Constraint.V[2];
+		
+		const varying FVector4f PandInvM0 = VectorGather(&PandInvM[i0]);
+		const varying FVector4f PandInvM1 = VectorGather(&PandInvM[i1]);
+		const varying FVector4f PandInvM2 = VectorGather(&PandInvM[i2]);
+
+		varying FVector3f P0, P1, P2;
+		varying float M0, M1, M2;
+		
+		UnzipPandInvM(PandInvM0, P0, M0);
+		UnzipPandInvM(PandInvM1, P1, M1);
+		UnzipPandInvM(PandInvM2, P2, M2);
+
+		const varying FVector4f DeltaUVInverse = VectorLoad(&DeltaUVInverses[extract(i,0)]);
+		const varying FVector2f RestStretchLength = VectorLoad(&RestStretchLengths[extract(i,0)]);
+
+		varying float Cu, Cv, Cs;
+		varying FVector3f dCu_dX1, dCu_dX2, dCv_dX1, dCv_dX2, dCs_dX1, dCs_dX2;
+		CalcConstraintsAndGrads(P0, P1, P2, DeltaUVInverse, RestStretchLength, WarpScale, WeftScale, Cu, Cv, Cs,
+			dCu_dX1, dCu_dX2, dCv_dX1, dCv_dX2, dCs_dX1, dCs_dX2);
+		
+		const varying FVector3f X0 = VectorGather(&X[i0]);
+		const varying FVector3f X1 = VectorGather(&X[i1]);
+		const varying FVector3f X2 = VectorGather(&X[i2]);
+		const varying FVector3f V0TimesDt = P0 - X0;
+		const varying FVector3f V1TimesDt = P1 - X1;
+		const varying FVector3f V2TimesDt = P2 - X2;
+		const varying float CombinedInvMass = M0 + M1 + M2;
+		
+		FVector3f Lambda = VectorLoad(&Lambdas[extract(i,0)]);
+		const varying FVector3f StiffnessScale = VectorLoad(&StiffnessScales[extract(i,0)]);
+		const varying FVector3f FinalStiffness = Stiffness * StiffnessScale;
+		const varying float DLambdaU = CalcDLambda(FinalStiffness.V[0], Dt, MinStiffness, Lambda.V[0], 
+			CalcDamping(FinalStiffness.V[0], CombinedInvMass, DampingRatio),
+			M0, M1, M2, V0TimesDt, V1TimesDt, V2TimesDt, Cu, dCu_dX1, dCu_dX2);
+		const varying float DLambdaV = CalcDLambda(FinalStiffness.V[1], Dt, MinStiffness, Lambda.V[1], 
+			CalcDamping(FinalStiffness.V[1], CombinedInvMass, DampingRatio),
+			M0, M1, M2, V0TimesDt, V1TimesDt, V2TimesDt, Cv, dCv_dX1, dCv_dX2);
+		const varying float DLambdaS = CalcDLambda(FinalStiffness.V[2], Dt, MinStiffness, Lambda.V[2], 
+			CalcDamping(FinalStiffness.V[2], CombinedInvMass, DampingRatio),
+			M0, M1, M2, V0TimesDt, V1TimesDt, V2TimesDt, Cs, dCs_dX1, dCs_dX2);
+		const varying FVector3f DLambda = {{DLambdaU, DLambdaV, DLambdaS}};
+
+		Lambda = Lambda + DLambda;
+		VectorStore(&Lambdas[extract(i,0)], Lambda);
+		
+		const varying FVector3f DP1 = DLambdaU * dCu_dX1 + DLambdaV * dCv_dX1 + DLambdaS * dCs_dX1;
+		const varying FVector3f DP2 = DLambdaU * dCu_dX2 + DLambdaV * dCv_dX2 + DLambdaS * dCs_dX2;
+
+		if (M0 > 0)
+		{
+			VectorScatter(&PandInvM[i0], SetVector4(P0 + M0 * (DP1 + DP2), M0));
+		}
+		if (M1 > 0)
+		{
+			VectorScatter(&PandInvM[i1], SetVector4(P1 - M1 * DP1, M1));
+		}
+		if (M2 > 0)
+		{
+			VectorScatter(&PandInvM[i2], SetVector4(P2 - M2 * DP2, M2));
+		}
+	}
+}
+
+export void ApplyXPBDStretchBiasConstraintsWithMaps(uniform FVector4f PandInvM[],
+											const uniform FIntVector Constraints[],
+											const uniform FVector2f RestStretchLengths[],
+											const uniform FVector4f DeltaUVInverses[],
+											const uniform FVector3f StiffnessScales[],
+											uniform FVector3f Lambdas[],
+											const uniform float Dt,
+											const uniform float MinStiffness,
+											const uniform bool StiffnessWarpHasMap,
+											const uniform uint8 StiffnessWarpIndices[],
+											const uniform float StiffnessWarpTable[],
+											const uniform bool StiffnessWeftHasMap,
+											const uniform uint8 StiffnessWeftIndices[],
+											const uniform float StiffnessWeftTable[],
+											const uniform bool StiffnessBiasHasMap,
+											const uniform uint8 StiffnessBiasIndices[],
+											const uniform float StiffnessBiasTable[],
+											const uniform bool WarpScaleHasMap,
+											const uniform uint8 WarpScaleIndices[],
+											const uniform float WarpScaleTable[],
+											const uniform bool WeftScaleHasMap,
+											const uniform uint8 WeftScaleIndices[],
+											const uniform float WeftScaleTable[],
+											const uniform int32 NumConstraints)
+{
+
+	foreach(i = 0 ... NumConstraints)
+	{
+		const varying FIntVector Constraint = VectorLoad(&Constraints[extract(i,0)]);
+
+		const varying int32 i0 = Constraint.V[0];
+		const varying int32 i1 = Constraint.V[1];
+		const varying int32 i2 = Constraint.V[2];
+		const varying int8 StiffnessWarpIndex = select(StiffnessWarpHasMap, StiffnessWarpIndices[i], 0);
+		const varying int8 StiffnessWeftIndex = select(StiffnessWeftHasMap, StiffnessWeftIndices[i], 0);
+		const varying int8 StiffnessBiasIndex = select(StiffnessBiasHasMap, StiffnessBiasIndices[i], 0);
+		const varying int8 WarpScaleIndex = select(WarpScaleHasMap, WarpScaleIndices[i], 0);
+		const varying int8 WeftScaleIndex = select(WeftScaleHasMap, WeftScaleIndices[i], 0);
+		
+		const varying FVector4f PandInvM0 = VectorGather(&PandInvM[i0]);
+		const varying FVector4f PandInvM1 = VectorGather(&PandInvM[i1]);
+		const varying FVector4f PandInvM2 = VectorGather(&PandInvM[i2]);
+
+		varying FVector3f P0, P1, P2;
+		varying float M0, M1, M2;
+		
+		UnzipPandInvM(PandInvM0, P0, M0);
+		UnzipPandInvM(PandInvM1, P1, M1);
+		UnzipPandInvM(PandInvM2, P2, M2);
+
+		const varying FVector4f DeltaUVInverse = VectorLoad(&DeltaUVInverses[extract(i,0)]);
+		const varying FVector2f RestStretchLength = VectorLoad(&RestStretchLengths[extract(i,0)]);
+		
+		#pragma ignore warning(perf)
+		const varying float WarpScale = WarpScaleTable[WarpScaleIndex];
+		#pragma ignore warning(perf)
+		const varying float WeftScale = WeftScaleTable[WeftScaleIndex];
+
+		varying float Cu, Cv, Cs;
+		varying FVector3f dCu_dX1, dCu_dX2, dCv_dX1, dCv_dX2, dCs_dX1, dCs_dX2;
+		CalcConstraintsAndGrads(P0, P1, P2, DeltaUVInverse, RestStretchLength, WarpScale, WeftScale, Cu, Cv, Cs,
+			dCu_dX1, dCu_dX2, dCv_dX1, dCv_dX2, dCs_dX1, dCs_dX2);		
+		
+		#pragma ignore warning(perf)
+		const varying float StiffnessWarp = StiffnessWarpTable[StiffnessWarpIndex];
+		#pragma ignore warning(perf)
+		const varying float StiffnessWeft = StiffnessWeftTable[StiffnessWeftIndex];
+		#pragma ignore warning(perf)
+		const varying float StiffnessBias = StiffnessBiasTable[StiffnessBiasIndex];
+		const varying FVector3f Stiffness = {{StiffnessWarp, StiffnessWeft, StiffnessBias}};
+
+		FVector3f Lambda = VectorLoad(&Lambdas[extract(i,0)]);
+		const varying FVector3f StiffnessScale = VectorLoad(&StiffnessScales[extract(i,0)]);
+		const varying FVector3f FinalStiffness = Stiffness * StiffnessScale;
+
+		const varying float DLambdaU = CalcDLambda(FinalStiffness.V[0], Dt, MinStiffness, Lambda.V[0], M0, M1, M2, Cu, dCu_dX1, dCu_dX2);
+		const varying float DLambdaV = CalcDLambda(FinalStiffness.V[1], Dt, MinStiffness, Lambda.V[1], M0, M1, M2, Cv, dCv_dX1, dCv_dX2);
+		const varying float DLambdaS = CalcDLambda(FinalStiffness.V[2], Dt, MinStiffness, Lambda.V[2], M0, M1, M2, Cs, dCs_dX1, dCs_dX2);
+		const varying FVector3f DLambda = {{DLambdaU, DLambdaV, DLambdaS}};
+
+		Lambda = Lambda + DLambda;
+		VectorStore(&Lambdas[extract(i,0)], Lambda);
+		
+		const varying FVector3f DP1 = DLambdaU * dCu_dX1 + DLambdaV * dCv_dX1 + DLambdaS * dCs_dX1;
+		const varying FVector3f DP2 = DLambdaU * dCu_dX2 + DLambdaV * dCv_dX2 + DLambdaS * dCs_dX2;
+
+		if (M0 > 0)
+		{
+			VectorScatter(&PandInvM[i0], SetVector4(P0 + M0 * (DP1 + DP2), M0));
+		}
+		if (M1 > 0)
+		{
+			VectorScatter(&PandInvM[i1], SetVector4(P1 - M1 * DP1, M1));
+		}
+		if (M2 > 0)
+		{
+			VectorScatter(&PandInvM[i2], SetVector4(P2 - M2 * DP2, M2));
+		}
+	}
+}
+
+export void ApplyXPBDStretchBiasConstraintsWithDampingAndMaps(uniform FVector4f PandInvM[],
+											const uniform FVector3f X[],
+											const uniform FIntVector Constraints[],
+											const uniform FVector2f RestStretchLengths[],
+											const uniform FVector4f DeltaUVInverses[],
+											const uniform FVector3f StiffnessScales[],
+											uniform FVector3f Lambdas[],
+											const uniform float Dt,
+											const uniform float MinStiffness,
+											const uniform bool StiffnessWarpHasMap,
+											const uniform uint8 StiffnessWarpIndices[],
+											const uniform float StiffnessWarpTable[],
+											const uniform bool StiffnessWeftHasMap,
+											const uniform uint8 StiffnessWeftIndices[],
+											const uniform float StiffnessWeftTable[],
+											const uniform bool StiffnessBiasHasMap,
+											const uniform uint8 StiffnessBiasIndices[],
+											const uniform float StiffnessBiasTable[],
+											const uniform bool DampingHasMap,
+											const uniform uint8 DampingIndices[],
+											const uniform float DampingTable[],
+											const uniform bool WarpScaleHasMap,
+											const uniform uint8 WarpScaleIndices[],
+											const uniform float WarpScaleTable[],
+											const uniform bool WeftScaleHasMap,
+											const uniform uint8 WeftScaleIndices[],
+											const uniform float WeftScaleTable[],
+											const uniform int32 NumConstraints)
+{
+
+	foreach(i = 0 ... NumConstraints)
+	{
+		const varying FIntVector Constraint = VectorLoad(&Constraints[extract(i,0)]);
+
+		const varying int32 i0 = Constraint.V[0];
+		const varying int32 i1 = Constraint.V[1];
+		const varying int32 i2 = Constraint.V[2];
+		const varying int8 StiffnessWarpIndex = select(StiffnessWarpHasMap, StiffnessWarpIndices[i], 0);
+		const varying int8 StiffnessWeftIndex = select(StiffnessWeftHasMap, StiffnessWeftIndices[i], 0);
+		const varying int8 StiffnessBiasIndex = select(StiffnessBiasHasMap, StiffnessBiasIndices[i], 0);
+		const varying int8 WarpScaleIndex = select(WarpScaleHasMap, WarpScaleIndices[i], 0);
+		const varying int8 WeftScaleIndex = select(WeftScaleHasMap, WeftScaleIndices[i], 0);
+		const varying int8 DampingIndex = select(DampingHasMap, DampingIndices[i], 0);
+		
+		const varying FVector4f PandInvM0 = VectorGather(&PandInvM[i0]);
+		const varying FVector4f PandInvM1 = VectorGather(&PandInvM[i1]);
+		const varying FVector4f PandInvM2 = VectorGather(&PandInvM[i2]);
+
+		varying FVector3f P0, P1, P2;
+		varying float M0, M1, M2;
+		
+		UnzipPandInvM(PandInvM0, P0, M0);
+		UnzipPandInvM(PandInvM1, P1, M1);
+		UnzipPandInvM(PandInvM2, P2, M2);
+
+		const varying FVector4f DeltaUVInverse = VectorLoad(&DeltaUVInverses[extract(i,0)]);
+		const varying FVector2f RestStretchLength = VectorLoad(&RestStretchLengths[extract(i,0)]);		
+		
+		#pragma ignore warning(perf)
+		const varying float WarpScale = WarpScaleTable[WarpScaleIndex];
+		#pragma ignore warning(perf)
+		const varying float WeftScale = WeftScaleTable[WeftScaleIndex];
+
+		varying float Cu, Cv, Cs;
+		varying FVector3f dCu_dX1, dCu_dX2, dCv_dX1, dCv_dX2, dCs_dX1, dCs_dX2;
+		CalcConstraintsAndGrads(P0, P1, P2, DeltaUVInverse, RestStretchLength, WarpScale, WeftScale, Cu, Cv, Cs,
+			dCu_dX1, dCu_dX2, dCv_dX1, dCv_dX2, dCs_dX1, dCs_dX2);
+		
+		const varying FVector3f X0 = VectorGather(&X[i0]);
+		const varying FVector3f X1 = VectorGather(&X[i1]);
+		const varying FVector3f X2 = VectorGather(&X[i2]);
+		const varying FVector3f V0TimesDt = P0 - X0;
+		const varying FVector3f V1TimesDt = P1 - X1;
+		const varying FVector3f V2TimesDt = P2 - X2;
+		const varying float CombinedInvMass = M0 + M1 + M2;
+		
+		#pragma ignore warning(perf)
+		const varying float StiffnessWarp = StiffnessWarpTable[StiffnessWarpIndex];
+		#pragma ignore warning(perf)
+		const varying float StiffnessWeft = StiffnessWeftTable[StiffnessWeftIndex];
+		#pragma ignore warning(perf)
+		const varying float StiffnessBias = StiffnessBiasTable[StiffnessBiasIndex];
+		const varying FVector3f Stiffness = {{StiffnessWarp, StiffnessWeft, StiffnessBias}};
+
+		#pragma ignore warning(perf)
+		const varying float DampingRatio = DampingTable[DampingIndex];
+		
+		FVector3f Lambda = VectorLoad(&Lambdas[extract(i,0)]);
+		const varying FVector3f StiffnessScale = VectorLoad(&StiffnessScales[extract(i,0)]);
+		const varying FVector3f FinalStiffness = Stiffness * StiffnessScale;
+		const varying float DLambdaU = CalcDLambda(FinalStiffness.V[0], Dt, MinStiffness, Lambda.V[0], 
+			CalcDamping(FinalStiffness.V[0], CombinedInvMass, DampingRatio),
+			M0, M1, M2, V0TimesDt, V1TimesDt, V2TimesDt, Cu, dCu_dX1, dCu_dX2);
+		const varying float DLambdaV = CalcDLambda(FinalStiffness.V[1], Dt, MinStiffness, Lambda.V[1], 
+			CalcDamping(FinalStiffness.V[1], CombinedInvMass, DampingRatio),
+			M0, M1, M2, V0TimesDt, V1TimesDt, V2TimesDt, Cv, dCv_dX1, dCv_dX2);
+		const varying float DLambdaS = CalcDLambda(FinalStiffness.V[2], Dt, MinStiffness, Lambda.V[2], 
+			CalcDamping(FinalStiffness.V[2], CombinedInvMass, DampingRatio),
+			M0, M1, M2, V0TimesDt, V1TimesDt, V2TimesDt, Cs, dCs_dX1, dCs_dX2);
+		const varying FVector3f DLambda = {{DLambdaU, DLambdaV, DLambdaS}};
+
+		Lambda = Lambda + DLambda;
+		VectorStore(&Lambdas[extract(i,0)], Lambda);
+		
+		const varying FVector3f DP1 = DLambdaU * dCu_dX1 + DLambdaV * dCv_dX1 + DLambdaS * dCs_dX1;
+		const varying FVector3f DP2 = DLambdaU * dCu_dX2 + DLambdaV * dCv_dX2 + DLambdaS * dCs_dX2;
+
+		if (M0 > 0)
+		{
+			VectorScatter(&PandInvM[i0], SetVector4(P0 + M0 * (DP1 + DP2), M0));
+		}
+		if (M1 > 0)
+		{
+			VectorScatter(&PandInvM[i1], SetVector4(P1 - M1 * DP1, M1));
+		}
+		if (M2 > 0)
+		{
+			VectorScatter(&PandInvM[i2], SetVector4(P2 - M2 * DP2, M2));
+		}
+	}
+}
\ No newline at end of file
diff --git a/Engine/Source/Runtime/Experimental/Chaos/Public/Chaos/PBDSoftsEvolutionFwd.h b/Engine/Source/Runtime/Experimental/Chaos/Public/Chaos/PBDSoftsEvolutionFwd.h
--- a/Engine/Source/Runtime/Experimental/Chaos/Public/Chaos/PBDSoftsEvolutionFwd.h
+++ b/Engine/Source/Runtime/Experimental/Chaos/Public/Chaos/PBDSoftsEvolutionFwd.h
@@ -30,6 +30,7 @@ namespace Chaos
 		using FSolverReal = FRealSingle;
 		using FSolverVec2 = TVector<FSolverReal, 2>;
 		using FSolverVec3 = TVector<FSolverReal, 3>;
+		using FSolverMatrix22 = PMatrix<FSolverReal, 2, 2>;
 		using FSolverMatrix33 = PMatrix<FSolverReal, 3, 3>;
 		using FSolverMatrix44 = PMatrix<FSolverReal, 4, 4>;
 		using FSolverRigidParticles = TKinematicGeometryParticlesImp<FSolverReal, 3, EGeometryParticlesSimType::Other>;
@@ -46,6 +47,7 @@ namespace Chaos
 		// Softs solver constraint classes
 		class FPBDSpringConstraints;
 		class FPBDEdgeSpringConstraints;
+		class FXPBDStretchBiasElementConstraints;
 		class FPBDBendingSpringConstraints;
 		class FXPBDSpringConstraints;
 		class FXPBDEdgeSpringConstraints;
diff --git a/Engine/Source/Runtime/Experimental/Chaos/Public/Chaos/XPBDStretchBiasElementConstraints.h b/Engine/Source/Runtime/Experimental/Chaos/Public/Chaos/XPBDStretchBiasElementConstraints.h
new file mode 100644
--- /dev/null
+++ b/Engine/Source/Runtime/Experimental/Chaos/Public/Chaos/XPBDStretchBiasElementConstraints.h
@@ -0,0 +1,169 @@
+// Copyright Epic Games, Inc. All Rights Reserved.
+#pragma once
+
+#include "Chaos/CollectionPropertyFacade.h"
+#include "Chaos/PBDSoftsEvolutionFwd.h"
+#include "Chaos/PBDSoftsSolverParticles.h"
+#include "Chaos/PBDStiffness.h"
+#include "Chaos/ParticleRule.h"
+#include "ChaosStats.h"
+
+namespace Chaos
+{
+class FTriangleMesh;
+}
+
+namespace Chaos::Softs
+{
+
+class CHAOS_API FXPBDStretchBiasElementConstraints
+{
+public:
+	// Stiffness is in kg cm / s^2 for stretch, kg cm^2 / s^2 for Bias
+	static constexpr FSolverReal MinStiffness = (FSolverReal)1e-4; // Stiffness below this will be considered 0 since all of our calculations are actually based on 1 / stiffness.
+	static constexpr FSolverReal MaxStiffness = (FSolverReal)1e7;
+	static constexpr FSolverReal MinDamping = (FSolverReal)0.;
+	static constexpr FSolverReal MaxDamping = (FSolverReal)1000.;
+	static constexpr bool bDefaultUse3dRestLengths = true;
+	static constexpr FSolverReal MinWarpWeftScale = (FSolverReal)0.; 
+	static constexpr FSolverReal MaxWarpWeftScale = (FSolverReal)1e7; // No particular reason for this number. Just can't imagine wanting something bigger?
+	static constexpr FSolverReal DefaultWarpWeftScale = (FSolverReal)1.;
+
+	static bool IsEnabled(const FCollectionPropertyConstFacade& PropertyCollection)
+	{
+		return IsXPBDStretchBiasElementStiffnessWarpEnabled(PropertyCollection, false);
+	}
+
+	FXPBDStretchBiasElementConstraints(const FSolverParticles& InParticles,
+		int32 ParticleOffset,
+		int32 ParticleCount,
+		const FTriangleMesh& TriangleMesh,
+		const TArray<TVec3<FVec2f>>& FaceVertexUVs,
+		const TConstArrayView<FRealSingle>& StiffnessWarpMultipliers,
+		const TConstArrayView<FRealSingle>& StiffnessWeftMultipliers,
+		const TConstArrayView<FRealSingle>& StiffnessBiasMultipliers,
+		const TConstArrayView<FRealSingle>& DampingMultipliers,
+		const TConstArrayView<FRealSingle>& WarpScaleMultipliers,
+		const TConstArrayView<FRealSingle>& WeftScaleMultipliers,
+		const FCollectionPropertyConstFacade& PropertyCollection,
+		bool bTrimKinematicConstraints = false);
+
+	FXPBDStretchBiasElementConstraints(const FSolverParticles& InParticles,
+		int32 ParticleOffset,
+		int32 ParticleCount,
+		const FTriangleMesh& TriangleMesh,
+		const TArray<TVec3<FVec2f>>& FaceVertexUVs,
+		const TConstArrayView<FRealSingle>& StiffnessWarpMultipliers,
+		const TConstArrayView<FRealSingle>& StiffnessWeftMultipliers,
+		const TConstArrayView<FRealSingle>& StiffnessBiasMultipliers,
+		const TConstArrayView<FRealSingle>& DampingMultipliers,
+		const TConstArrayView<FRealSingle>& WarpScaleMultipliers,
+		const TConstArrayView<FRealSingle>& WeftScaleMultipliers,
+		const FSolverVec2& InStiffnessWarp,
+		const FSolverVec2& InStiffnessWeft,
+		const FSolverVec2& InStiffnessBias,
+		const FSolverVec2& InDampingRatio,
+		const FSolverVec2& InWarpScale,
+		const FSolverVec2& InWeftScale,
+		bool bUse3dRestLengths,
+		bool bTrimKinematicConstraints = false);
+
+	virtual ~FXPBDStretchBiasElementConstraints() {}
+
+	void Init()
+	{ 
+		Lambdas.Reset();
+		Lambdas.AddZeroed(Constraints.Num());
+	}
+
+	void SetProperties(const FCollectionPropertyConstFacade& PropertyCollection)
+	{
+		if (IsXPBDStretchBiasElementStiffnessWarpMutable(PropertyCollection))
+		{
+			StiffnessWarp.SetWeightedValue(FSolverVec2(GetWeightedFloatXPBDStretchBiasElementStiffnessWarp(PropertyCollection)), MaxStiffness);
+		}
+		if (IsXPBDStretchBiasElementStiffnessWeftMutable(PropertyCollection))
+		{
+			StiffnessWeft.SetWeightedValue(FSolverVec2(GetWeightedFloatXPBDStretchBiasElementStiffnessWeft(PropertyCollection)), MaxStiffness);
+		}
+		if (IsXPBDStretchBiasElementStiffnessBiasMutable(PropertyCollection))
+		{
+			StiffnessBias.SetWeightedValue(FSolverVec2(GetWeightedFloatXPBDStretchBiasElementStiffnessBias(PropertyCollection)), MaxStiffness);
+		}
+		if (IsXPBDStretchBiasElementDampingMutable(PropertyCollection))
+		{
+			DampingRatio.SetWeightedValue(FSolverVec2(GetWeightedFloatXPBDStretchBiasElementDamping(PropertyCollection)).ClampAxes(MinDamping, MaxDamping));
+		}
+		if (IsXPBDStretchBiasElementWarpScaleMutable(PropertyCollection))
+		{
+			WarpScale.SetWeightedValue(FSolverVec2(GetWeightedFloatXPBDStretchBiasElementWarpScale(PropertyCollection)).ClampAxes(MinWarpWeftScale, MaxWarpWeftScale));
+		}
+		if (IsXPBDStretchBiasElementWeftScaleMutable(PropertyCollection))
+		{
+			WeftScale.SetWeightedValue(FSolverVec2(GetWeightedFloatXPBDStretchBiasElementWeftScale(PropertyCollection)).ClampAxes(MinWarpWeftScale, MaxWarpWeftScale));
+		}
+	}
+
+	void SetProperties(const FSolverVec2& InStiffnessWarp, const FSolverVec2& InStiffnessWeft, const FSolverVec2& InStiffnessBias, const FSolverVec2& InDampingRatio, const FSolverVec2& InWarpScale, const FSolverVec2& InWeftScale)
+	{
+		StiffnessWarp.SetWeightedValue(InStiffnessWarp, MaxStiffness);
+		StiffnessWeft.SetWeightedValue(InStiffnessWeft, MaxStiffness);
+		StiffnessBias.SetWeightedValue(InStiffnessBias, MaxStiffness);
+		DampingRatio.SetWeightedValue(InDampingRatio.ClampAxes(MinDamping, MaxDamping));
+		WarpScale.SetWeightedValue(InWarpScale.ClampAxes(MinWarpWeftScale, MaxWarpWeftScale));
+		WeftScale.SetWeightedValue(InWeftScale.ClampAxes(MinWarpWeftScale, MaxWarpWeftScale));
+	}
+
+	// Update stiffness table, as well as the simulation stiffness exponent
+	void ApplyProperties(const FSolverReal /*Dt*/, const int32 /*NumIterations*/)
+	{
+		StiffnessWarp.ApplyXPBDValues(MaxStiffness);
+		StiffnessWeft.ApplyXPBDValues(MaxStiffness);
+		StiffnessBias.ApplyXPBDValues(MaxStiffness);
+		DampingRatio.ApplyValues();
+		WarpScale.ApplyValues();
+		WeftScale.ApplyValues();
+	}
+
+	void Apply(FSolverParticles& Particles, const FSolverReal Dt) const;
+
+private:
+	void InitConstraintsAndRestData(const FSolverParticles& InParticles, const FTriangleMesh& TriangleMesh,
+		const TArray<TVec3<FVec2f>>& FaceVertexUVs, const bool bUse3dRestLengths, const bool bTrimKinematicConstraints);
+	void InitColor(const FSolverParticles& InParticles, const int32 ParticleOffset, const int32 ParticleCount);
+	void ApplyHelper(FSolverParticles& Particles, const FSolverReal Dt, const int32 ConstraintIndex, const FSolverVec3& ExpStiffnessValue, const FSolverReal DampingRatioValue, const FSolverReal WarpScaleValue, const FSolverReal WeftScaleValue) const;
+
+	TArray<TVec3<int32>> Constraints;
+
+	FPBDStiffness StiffnessWarp;
+	FPBDStiffness StiffnessWeft;
+	FPBDStiffness StiffnessBias;
+	FPBDWeightMap DampingRatio;
+	FPBDWeightMap WarpScale;
+	FPBDWeightMap WeftScale;
+	mutable TArray<FSolverVec3> Lambdas; // separate for stretchU, stretchV, Bias
+	TArray<int32> ConstraintsPerColorStartIndex; // Constraints are ordered so each batch is contiguous. This is ColorNum + 1 length so it can be used as start and end.
+
+	TArray<FSolverMatrix22> DeltaUVInverse; // Used to convert from DeltaX to dX/dU and dX/dV
+	TArray<FSolverVec2> RestStretchLengths;
+	TArray<FSolverVec3> StiffnessScales; // Used to make everything resolution independent.
+
+	UE_CHAOS_DECLARE_PROPERTYCOLLECTION_NAME(XPBDStretchBiasElementUse3dRestLengths, bool);
+	UE_CHAOS_DECLARE_PROPERTYCOLLECTION_NAME(XPBDStretchBiasElementStiffnessWarp, float);
+	UE_CHAOS_DECLARE_PROPERTYCOLLECTION_NAME(XPBDStretchBiasElementStiffnessWeft, float);
+	UE_CHAOS_DECLARE_PROPERTYCOLLECTION_NAME(XPBDStretchBiasElementStiffnessBias, float);
+	UE_CHAOS_DECLARE_PROPERTYCOLLECTION_NAME(XPBDStretchBiasElementDamping, float);
+	UE_CHAOS_DECLARE_PROPERTYCOLLECTION_NAME(XPBDStretchBiasElementWarpScale, float);
+	UE_CHAOS_DECLARE_PROPERTYCOLLECTION_NAME(XPBDStretchBiasElementWeftScale, float);
+};
+
+}  // End namespace Chaos::Softs
+
+// Support ISPC enable/disable in non-shipping builds
+#if !INTEL_ISPC
+const bool bChaos_XPBDStretchBiasElement_ISPC_Enabled = false;
+#elif UE_BUILD_SHIPPING
+const bool bChaos_XPBDStretchBiasElement_ISPC_Enabled = true;
+#else
+extern CHAOS_API bool bChaos_XPBDStretchBiasElement_ISPC_Enabled;
+#endif
