Iris - Fixed bug in non-vectorized NetObjectPollFrequencyLimiter update.

#rnx

diff --git a/Engine/Source/Runtime/Experimental/Iris/Core/Private/Iris/ReplicationSystem/ObjectPollFrequencyLimiter.cpp b/Engine/Source/Runtime/Experimental/Iris/Core/Private/Iris/ReplicationSystem/ObjectPollFrequencyLimiter.cpp
--- a/Engine/Source/Runtime/Experimental/Iris/Core/Private/Iris/ReplicationSystem/ObjectPollFrequencyLimiter.cpp
+++ b/Engine/Source/Runtime/Experimental/Iris/Core/Private/Iris/ReplicationSystem/ObjectPollFrequencyLimiter.cpp
@@ -5,11 +5,7 @@
 #include "HAL/IConsoleManager.h"
 #include "Math/VectorRegister.h"
 
-namespace UE
-{
-namespace Net
-{
-namespace Private
+namespace UE::Net::Private
 {
 
 // Storing poll frame period as uint8
@@ -125,9 +121,10 @@ void FObjectPollFrequencyLimiter::Update(const FNetBitArrayView& ScopableObjects
 			}
 
 			WordType ObjectsToPoll = 0;
-			for (SIZE_T It = 0, IndexOffset = 0; It < 8; ++It, IndexOffset += 4U, ObjectsInScopeWord >>= 4U)
+			WordType ObjectsInScopeMask = ObjectsInScopeWord;
+			for (SIZE_T It = 0, IndexOffset = 0; It < 8; ++It, IndexOffset += 4U, ObjectsInScopeMask >>= 4U)
 			{
-				if (!(ObjectsInScopeWord & 15U))
+				if (!(ObjectsInScopeMask & 15U))
 				{
 					continue;
 				}
@@ -148,16 +145,16 @@ void FObjectPollFrequencyLimiter::Update(const FNetBitArrayView& ScopableObjects
 				--Counter3;
 
 				const uint32 Mask0 = Counter0 == 255 ? 1 : 0;
-				const uint32 Mask1 = Counter1 == 255 ? 1 : 0;
-				const uint32 Mask2 = Counter2 == 255 ? 1 : 0;
-				const uint32 Mask3 = Counter3 == 255 ? 1 : 0;
+				const uint32 Mask1 = Counter1 == 255 ? 2 : 0;
+				const uint32 Mask2 = Counter2 == 255 ? 4 : 0;
+				const uint32 Mask3 = Counter3 == 255 ? 8 : 0;
 
 				Counter0 = Mask0 ? FramesBetweenUpdates0 : Counter0;
 				Counter1 = Mask1 ? FramesBetweenUpdates1 : Counter1;
 				Counter2 = Mask2 ? FramesBetweenUpdates2 : Counter2;
 				Counter3 = Mask3 ? FramesBetweenUpdates3 : Counter3;
 
-				const uint32 Mask = (Mask3 << 3U) | (Mask2 << 2U) | (Mask1 << 1U) | Mask0;
+				const uint32 Mask = Mask3 | Mask2 | Mask1 | Mask0;
 				ObjectsToPoll |= (Mask << IndexOffset);
 
 				CountersData[IndexOffset + 0] = Counter0;
@@ -174,5 +171,3 @@ void FObjectPollFrequencyLimiter::Update(const FNetBitArrayView& ScopableObjects
 }
 
 }
-}
-}
