IrisDebugHelper
* Added ability to set and break on a internal index (BreakOnNetInternalIndex)
* DebugNetObject now prints the UObject associated with the handle/index/...

diff --git a/Engine/Source/Runtime/Experimental/Iris/Core/Private/Iris/Core/IrisDebugging.cpp b/Engine/Source/Runtime/Experimental/Iris/Core/Private/Iris/Core/IrisDebugging.cpp
--- a/Engine/Source/Runtime/Experimental/Iris/Core/Private/Iris/Core/IrisDebugging.cpp
+++ b/Engine/Source/Runtime/Experimental/Iris/Core/Private/Iris/Core/IrisDebugging.cpp
@@ -56,6 +56,24 @@ namespace IrisDebugHelperInternal
 				GIrisDebugNetRefHandle = FNetRefHandle();
 			}
 		}));
+
+	static UE::Net::Private::FInternalNetRefIndex GIrisDebugInternalIndex;
+	static FAutoConsoleCommand NetIrisDebugNetInternalIndex(
+		TEXT("Net.Iris.DebugNetInternalIndex"),
+		TEXT("Specify an internal index that we will break on (or none to turn off)."),
+		FConsoleCommandWithArgsDelegate::CreateLambda([](const TArray<FString>& Args)
+			{
+				if (Args.Num() > 0)
+				{
+					uint32 InternalIndex = 0;
+					LexFromString(InternalIndex, *Args[0]);
+					GIrisDebugInternalIndex = InternalIndex;
+				}
+				else
+				{
+					GIrisDebugInternalIndex = 0;
+				}
+			}));
 }; // namespace IrisDebugHelperInternal
 
 bool BreakOnObjectName(UObject* Object)
@@ -101,6 +119,28 @@ bool BreakOnRPCName(FName RPCName)
 	return false;
 }
 
+
+bool BreakOnNetInternalIndex(UE::Net::Private::FInternalNetRefIndex InternalIndex)
+{
+	if (IrisDebugHelperInternal::GIrisDebugInternalIndex != 0 && IrisDebugHelperInternal::GIrisDebugInternalIndex == InternalIndex)
+	{
+		UE_DEBUG_BREAK();
+		return true;
+	}
+
+	return false;
+}
+
+UE::Net::Private::FInternalNetRefIndex GetDebugNetInternalIndex()
+{
+	return IrisDebugHelperInternal::GIrisDebugInternalIndex;
+}
+
+FNetRefHandle GetDebugNetRefHandle()
+{
+	return IrisDebugHelperInternal::GIrisDebugNetRefHandle;
+}
+
 void SetIrisDebugObjectName(const ANSICHAR* NameBuffer)
 {
 	if (NameBuffer)
@@ -130,6 +170,12 @@ void SetIrisDebugRPCName(const ANSICHAR* NameBuffer)
 	}
 }
 
+void SetIrisDebugInternalIndex(UE::Net::Private::FInternalNetRefIndex InternalIndex)
+{
+	IrisDebugHelperInternal::GIrisDebugInternalIndex = InternalIndex;
+}
+
+
 UReplicationSystem* GetReplicationSystemForDebug(uint32 Id)
 {
 	return GetReplicationSystem(Id);
@@ -149,9 +195,11 @@ uint64 Init()
 	UE_NET_FORCE_REFERENCE_DEBUGFUNCTION(DebugNetObjectById);
 	UE_NET_FORCE_REFERENCE_DEBUGFUNCTION(DebugNetRefHandle);
 	UE_NET_FORCE_REFERENCE_DEBUGFUNCTION(DebugNetRefHandleById);
+	UE_NET_FORCE_REFERENCE_DEBUGFUNCTION(DebugNetInternalIndex);
 	UE_NET_FORCE_REFERENCE_DEBUGFUNCTION(DebugNetObjectProtocolReferencesToString);
 	UE_NET_FORCE_REFERENCE_DEBUGFUNCTION(SetIrisDebugObjectName);
 	UE_NET_FORCE_REFERENCE_DEBUGFUNCTION(SetIrisDebugNetRefHandle);
+	UE_NET_FORCE_REFERENCE_DEBUGFUNCTION(SetIrisDebugInternalIndex);
 	UE_NET_FORCE_REFERENCE_DEBUGFUNCTION(SetIrisDebugRPCName);
 	
 	return FunctionReferenceAccumulator;
@@ -264,6 +312,7 @@ FNetReplicatedObjectDebugInfo DebugNetObject(UObject* Instance)
 					Info.Protocol = ObjectData.Protocol;
 					Info.InstanceProtocol = ObjectData.InstanceProtocol;
 					Info.ReplicationSystem = ReplicationSystem;
+					Info.Object = NetRefHandleManager.GetReplicatedInstances()[Info.InternalNetRefIndex];
 
 					return Info;
 				}
@@ -296,6 +345,7 @@ FNetReplicatedObjectDebugInfo DebugNetObjectById(UObject* Instance, uint32 Repli
 				Info.Protocol = ObjectData.Protocol;
 				Info.InstanceProtocol = ObjectData.InstanceProtocol;
 				Info.ReplicationSystem = ReplicationSystem;
+				Info.Object = NetRefHandleManager.GetReplicatedInstances()[Info.InternalNetRefIndex];
 
 				return Info;
 			}
@@ -333,6 +383,40 @@ FNetReplicatedObjectDebugInfo DebugNetRefHandleById(uint32 NetRefHandleId, uint3
 					Info.Protocol = ObjectData.Protocol;
 					Info.InstanceProtocol = ObjectData.InstanceProtocol;
 					Info.ReplicationSystem = ReplicationSystem;
+					Info.Object = NetRefHandleManager.GetReplicatedInstances()[InternalNetRefIndex];
+
+					return Info;
+				}
+			}
+		}
+	}
+
+	return Info;
+}
+
+FNetReplicatedObjectDebugInfo DebugNetInternalIndex(uint32 InternalIndex, uint32 ReplicationSystemId)
+{
+	using namespace UE::Net::Private;
+	FNetReplicatedObjectDebugInfo Info = {};
+
+	// See if we can find the instance in any replication system
+	{
+		const UReplicationSystem* ReplicationSystem = GetReplicationSystemForDebug(ReplicationSystemId);
+		if (ReplicationSystem)
+		{
+			const FNetRefHandleManager& NetRefHandleManager = ReplicationSystem->GetReplicationSystemInternal()->GetNetRefHandleManager();
+			if (InternalIndex != FNetRefHandleManager::InvalidInternalIndex)
+			{
+				FNetRefHandle ObjectHandle = NetRefHandleManager.GetNetRefHandleFromInternalIndex(InternalIndex);
+				const FNetRefHandleManager::FReplicatedObjectData& ObjectData = NetRefHandleManager.GetReplicatedObjectDataNoCheck(InternalIndex);
+				if (ObjectData.RefHandle == ObjectHandle)
+				{
+					Info.RefHandle = &ObjectData.RefHandle;
+					Info.InternalNetRefIndex = InternalIndex;
+					Info.Protocol = ObjectData.Protocol;
+					Info.InstanceProtocol = ObjectData.InstanceProtocol;
+					Info.ReplicationSystem = ReplicationSystem;
+					Info.Object = NetRefHandleManager.GetReplicatedInstances()[InternalIndex];
 
 					return Info;
 				}
diff --git a/Engine/Source/Runtime/Experimental/Iris/Core/Public/Iris/Core/IrisDebugging.h b/Engine/Source/Runtime/Experimental/Iris/Core/Public/Iris/Core/IrisDebugging.h
--- a/Engine/Source/Runtime/Experimental/Iris/Core/Public/Iris/Core/IrisDebugging.h
+++ b/Engine/Source/Runtime/Experimental/Iris/Core/Public/Iris/Core/IrisDebugging.h
@@ -13,6 +13,11 @@ struct FReplicationProtocol;
 struct FReplicationInstanceProtocol;
 }
 
+namespace UE::Net::Private
+{
+	typedef uint32 FInternalNetRefIndex;
+}
+
 #if !UE_BUILD_SHIPPING
 
 namespace UE::Net::IrisDebugHelper
@@ -30,9 +35,18 @@ IRISCORE_API bool BreakOnNetRefHandle(FNetRefHandle NetRefHandle);
 /** Trigger a breakpoint and return true if the name contains the debug RPC string */
 IRISCORE_API bool BreakOnRPCName(FName RPCName);
 
+/** Trigger a breakpoint and return true if the index is the current debug index */
+IRISCORE_API bool BreakOnNetInternalIndex(UE::Net::Private::FInternalNetRefIndex InternalIndex);
+
 /** Returns true if the object name contains the current debug name, will return true if no debug name is set */
 IRISCORE_API bool FilterDebuggedObject(UObject* Object);
 
+/** Returns the internal index set via SetIrisDebugInternalIndex */
+IRISCORE_API UE::Net::Private::FInternalNetRefIndex GetDebugNetInternalIndex();
+
+/** Returns the NetRefHandle set via SetIrisDebugNetRefHandle */
+IRISCORE_API FNetRefHandle GetDebugNetRefHandle();
+
 /** Output state data to StringBuilder for the specified Handle */
 void NetObjectStateToString(FStringBuilderBase& StringBuilder, FNetRefHandle RefHandle);
 
@@ -60,10 +74,11 @@ extern "C" IRISCORE_API void DebugOutputNetObjectProtocolReferences(uint64 Proto
 struct FNetReplicatedObjectDebugInfo
 {
 	const FNetRefHandle* RefHandle;
-	uint32 InternalNetRefIndex;
+	UE::Net::Private::FInternalNetRefIndex InternalNetRefIndex;
 	const UReplicationSystem* ReplicationSystem;
 	const FReplicationProtocol* Protocol;
 	const FReplicationInstanceProtocol* InstanceProtocol;
+	const UObject* Object;
 };
 
 /** Look up replicated handle from Instance pointer and return debug information. This variant searches all active replication systems and returns information for the first one replicating the Instance. */
@@ -78,6 +93,9 @@ extern "C" IRISCORE_API FNetReplicatedObjectDebugInfo DebugNetRefHandle(FNetRefH
 /** Look up replicated handle specified by handle id and replicationsystem id and return debug information */
 extern "C" IRISCORE_API FNetReplicatedObjectDebugInfo DebugNetRefHandleById(uint32 NetRefHandleId, uint32 ReplicationSystemId);
 
+/** Look up replicated handle specified by an internal index and replicationsystem id and return debug information */
+extern "C" IRISCORE_API FNetReplicatedObjectDebugInfo DebugNetInternalIndex(UE::Net::Private::FInternalNetRefIndex InternalIndex, uint32 ReplicationSystemId);
+
 /** Variant of DebugOutputNetObjectProtocolReferences that can be used from breakpoints or in watch window to find all handles references registered for a protocol and output to DebugOutput in debugger
  * NOTE: Use only for debugging as this variant uses a static buffer which is not thread safe	
 */
@@ -87,7 +105,10 @@ extern "C" IRISCORE_API const TCHAR* DebugNetObjectProtocolReferencesToString(ui
 extern "C" IRISCORE_API void SetIrisDebugObjectName(const ANSICHAR* NameBuffer);
 
 /** Set the NetHandle to break on */
-extern "C" IRISCORE_API void SetIrisDebugNetHandle(uint32 NetHandleId);
+extern "C" IRISCORE_API void SetIrisDebugNetRefHandle(uint32 NetHandleId);
+
+/** Set the InternalIndex to break on */
+extern "C" IRISCORE_API void SetIrisDebugInternalIndex(UE::Net::Private::FInternalNetRefIndex InternalIndex);
 
 /** Set the RPC Name to break on */
 extern "C" IRISCORE_API void SetIrisDebugRPCName(const ANSICHAR* NameBuffer);
