fix crash when trying to access a key in MChildren that doesn't exist when computing strain from collisions

diff --git a/Engine/Source/Runtime/Experimental/Chaos/Private/Chaos/PBDRigidClustering.cpp b/Engine/Source/Runtime/Experimental/Chaos/Private/Chaos/PBDRigidClustering.cpp
--- a/Engine/Source/Runtime/Experimental/Chaos/Private/Chaos/PBDRigidClustering.cpp
+++ b/Engine/Source/Runtime/Experimental/Chaos/Private/Chaos/PBDRigidClustering.cpp
@@ -1481,10 +1481,12 @@ namespace Chaos
 									{
 										if (Proxy->GetType() == EPhysicsProxyType::ClusterUnionProxy)
 										{
-											TArray<FRigidHandle>& ClusteredParentToChildren = MChildren[ClusteredChild];
-											if (FPBDRigidParticleHandle* ClusteredChildRigid = FindClosestParticle(ClusteredParentToChildren, ContactWorldLocation))
+											if (TArray<FRigidHandle>* ClusteredParentToChildren = MChildren.Find(ClusteredChild))
 											{
-												ClusteredChild = ClusteredChildRigid->CastToClustered();
+												if (FPBDRigidParticleHandle* ClusteredChildRigid = FindClosestParticle(*ClusteredParentToChildren, ContactWorldLocation))
+												{
+													ClusteredChild = ClusteredChildRigid->CastToClustered();
+												}
 											}
 										}
 									}
