Move creation of ISMPoolActor to outside of place where it may be called in a construction script.

diff --git a/Engine/Source/Runtime/Experimental/GeometryCollectionEngine/Private/GeometryCollection/GeometryCollectionISMPoolRenderer.cpp b/Engine/Source/Runtime/Experimental/GeometryCollectionEngine/Private/GeometryCollection/GeometryCollectionISMPoolRenderer.cpp
--- a/Engine/Source/Runtime/Experimental/GeometryCollectionEngine/Private/GeometryCollection/GeometryCollectionISMPoolRenderer.cpp
+++ b/Engine/Source/Runtime/Experimental/GeometryCollectionEngine/Private/GeometryCollection/GeometryCollectionISMPoolRenderer.cpp
@@ -15,10 +15,6 @@
 
 void UGeometryCollectionISMPoolRenderer::OnRegisterGeometryCollection(UGeometryCollectionComponent const& InComponent)
 {
-	if (UGeometryCollectionISMPoolSubSystem* ISMPoolSubSystem = UWorld::GetSubsystem<UGeometryCollectionISMPoolSubSystem>(GetWorld()))
-	{
-		ISMPoolActor = ISMPoolSubSystem->FindISMPoolActor(InComponent);
-	}
 }
 
 void UGeometryCollectionISMPoolRenderer::OnUnregisterGeometryCollection()
@@ -60,6 +56,18 @@ void UGeometryCollectionISMPoolRenderer::UpdateTransforms(UGeometryCollection co
 	UpdateInstanceTransforms(InGeometryCollection, InBaseTransform, InMatrices);
 }
 
+UGeometryCollectionISMPoolComponent* UGeometryCollectionISMPoolRenderer::GetOrCreateISMPoolComponent()
+{
+	if (ISMPoolActor == nullptr)
+	{
+		if (UGeometryCollectionISMPoolSubSystem* ISMPoolSubSystem = UWorld::GetSubsystem<UGeometryCollectionISMPoolSubSystem>(GetWorld()))
+		{
+			ISMPoolActor = ISMPoolSubSystem->FindISMPoolActor();
+		}
+	}
+	return ISMPoolActor != nullptr ? ISMPoolActor->GetISMPoolComp() : nullptr;
+}
+
 void UGeometryCollectionISMPoolRenderer::InitMergedMeshFromGeometryCollection(UGeometryCollection const& InGeometryCollection)
 {
 	if (InGeometryCollection.RootProxyData.ProxyMeshes.Num() == 0)
@@ -67,7 +75,7 @@ void UGeometryCollectionISMPoolRenderer::InitMergedMeshFromGeometryCollection(UG
 		return;
 	}
 
-	UGeometryCollectionISMPoolComponent* ISMPoolComponent = ISMPoolActor != nullptr ? ISMPoolActor->GetISMPoolComp() : nullptr;
+	UGeometryCollectionISMPoolComponent* ISMPoolComponent = GetOrCreateISMPoolComponent();
 	MergedMeshGroup.GroupIndex = ISMPoolComponent != nullptr ? ISMPoolComponent->CreateMeshGroup() : INDEX_NONE;
 
 	if (MergedMeshGroup.GroupIndex == INDEX_NONE)
@@ -99,7 +107,7 @@ void UGeometryCollectionISMPoolRenderer::InitInstancesFromGeometryCollection(UGe
 		return;
 	}
 
-	UGeometryCollectionISMPoolComponent* ISMPoolComponent = ISMPoolActor != nullptr ? ISMPoolActor->GetISMPoolComp() : nullptr;
+	UGeometryCollectionISMPoolComponent* ISMPoolComponent = GetOrCreateISMPoolComponent();
 	InstancesGroup.GroupIndex = ISMPoolComponent != nullptr ? ISMPoolComponent->CreateMeshGroup() : INDEX_NONE;
 
 	if (InstancesGroup.GroupIndex == INDEX_NONE)
@@ -144,7 +152,7 @@ void UGeometryCollectionISMPoolRenderer::UpdateMergedMeshTransforms(FTransform c
 		return;
 	}
 
-	UGeometryCollectionISMPoolComponent* ISMPoolComponent = ISMPoolActor != nullptr ? ISMPoolActor->GetISMPoolComp() : nullptr;
+	UGeometryCollectionISMPoolComponent* ISMPoolComponent = GetOrCreateISMPoolComponent();
 	if (ISMPoolComponent == nullptr)
 	{
 		return;
@@ -166,7 +174,7 @@ void UGeometryCollectionISMPoolRenderer::UpdateInstanceTransforms(UGeometryColle
 		return;
 	}
 
-	UGeometryCollectionISMPoolComponent* ISMPoolComponent = ISMPoolActor != nullptr ? ISMPoolActor->GetISMPoolComp() : nullptr;
+	UGeometryCollectionISMPoolComponent* ISMPoolComponent = GetOrCreateISMPoolComponent();
 	if (ISMPoolComponent == nullptr)
 	{
 		return;
@@ -200,7 +208,7 @@ void UGeometryCollectionISMPoolRenderer::UpdateInstanceTransforms(UGeometryColle
 
 void UGeometryCollectionISMPoolRenderer::ReleaseGroup(FISMPoolGroup& InOutGroup)
 {
-	UGeometryCollectionISMPoolComponent* ISMPoolComponent = ISMPoolActor != nullptr ? ISMPoolActor->GetISMPoolComp() : nullptr;
+	UGeometryCollectionISMPoolComponent* ISMPoolComponent = GetOrCreateISMPoolComponent();
 	if (ISMPoolComponent != nullptr)
 	{
 		ISMPoolComponent->DestroyMeshGroup(InOutGroup.GroupIndex);
diff --git a/Engine/Source/Runtime/Experimental/GeometryCollectionEngine/Private/GeometryCollection/GeometryCollectionISMPoolRenderer.h b/Engine/Source/Runtime/Experimental/GeometryCollectionEngine/Private/GeometryCollection/GeometryCollectionISMPoolRenderer.h
--- a/Engine/Source/Runtime/Experimental/GeometryCollectionEngine/Private/GeometryCollection/GeometryCollectionISMPoolRenderer.h
+++ b/Engine/Source/Runtime/Experimental/GeometryCollectionEngine/Private/GeometryCollection/GeometryCollectionISMPoolRenderer.h
@@ -7,6 +7,7 @@
 #include "GeometryCollectionISMPoolRenderer.generated.h"
 
 class AGeometryCollectionISMPoolActor;
+class UGeometryCollectionISMPoolComponent;
 
 /** Implementation of a geometry collection custom renderer that pushes AutoInstanceMeshes to an ISMPool. */
 UCLASS()
@@ -23,10 +24,6 @@ public:
 	virtual void UpdateTransforms(UGeometryCollection const& InGeometryCollection, FTransform const& InBaseTransform, TArrayView<const FMatrix> InMatrices) override;
 	//~ End IGeometryCollectionExternalRenderInterface Interface.
 
-	/** Instanced Static Mesh Pool actor that is used to render our meshes. */
-	UPROPERTY(Transient)
-	TObjectPtr<AGeometryCollectionISMPoolActor> ISMPoolActor = nullptr;
-
 	/** Description for a group of meshes that are added/updated together. */
 	struct FISMPoolGroup
 	{
@@ -35,11 +32,16 @@ public:
 	};
 
 protected:
+	/** Instanced Static Mesh Pool actor that is used to render our meshes. */
+	UPROPERTY(Transient)
+	TObjectPtr<AGeometryCollectionISMPoolActor> ISMPoolActor;
+
 	/** ISM pool groups per rendering element type. */
 	FISMPoolGroup MergedMeshGroup;
 	FISMPoolGroup InstancesGroup;
 
 private:
+	UGeometryCollectionISMPoolComponent* GetOrCreateISMPoolComponent();
 	void InitMergedMeshFromGeometryCollection(UGeometryCollection const& InGeometryCollection);
 	void InitInstancesFromGeometryCollection(UGeometryCollection const& InGeometryCollection);
 	void UpdateMergedMeshTransforms(FTransform const& InBaseTransform);
diff --git a/Engine/Source/Runtime/Experimental/GeometryCollectionEngine/Private/GeometryCollection/GeometryCollectionISMPoolSubSystem.cpp b/Engine/Source/Runtime/Experimental/GeometryCollectionEngine/Private/GeometryCollection/GeometryCollectionISMPoolSubSystem.cpp
--- a/Engine/Source/Runtime/Experimental/GeometryCollectionEngine/Private/GeometryCollection/GeometryCollectionISMPoolSubSystem.cpp
+++ b/Engine/Source/Runtime/Experimental/GeometryCollectionEngine/Private/GeometryCollection/GeometryCollectionISMPoolSubSystem.cpp
@@ -30,7 +30,7 @@ void UGeometryCollectionISMPoolSubSystem::Deinitialize()
 	Super::Deinitialize();
 }
 
-AGeometryCollectionISMPoolActor* UGeometryCollectionISMPoolSubSystem::FindISMPoolActor(const USceneComponent& Requester)
+AGeometryCollectionISMPoolActor* UGeometryCollectionISMPoolSubSystem::FindISMPoolActor()
 {
 	// on demand creation of the actor 
 	// very simple logic for now, we can extend that in the future to return a specific actor based on the criteria
@@ -43,5 +43,3 @@ AGeometryCollectionISMPoolActor* UGeometryCollectionISMPoolSubSystem::FindISMPoo
 	}
 	return ISMPoolActor;
 }
-
-
diff --git a/Engine/Source/Runtime/Experimental/GeometryCollectionEngine/Public/GeometryCollection/GeometryCollectionISMPoolSubSystem.h b/Engine/Source/Runtime/Experimental/GeometryCollectionEngine/Public/GeometryCollection/GeometryCollectionISMPoolSubSystem.h
--- a/Engine/Source/Runtime/Experimental/GeometryCollectionEngine/Public/GeometryCollection/GeometryCollectionISMPoolSubSystem.h
+++ b/Engine/Source/Runtime/Experimental/GeometryCollectionEngine/Public/GeometryCollection/GeometryCollectionISMPoolSubSystem.h
@@ -24,7 +24,7 @@ public:
 	virtual void Deinitialize() override;
 	// USubsystem END
 
-	AGeometryCollectionISMPoolActor* FindISMPoolActor(const USceneComponent& Requester);
+	AGeometryCollectionISMPoolActor* FindISMPoolActor();
 
 protected:
 
