Chaos - invalidating cluster particle in FClusterUnionManager::UpdateAllClusterUnionProperties

This is required because we are changing the state and geometry on the particle so must destroy collisions, remove from the graph, etc.

diff --git a/Engine/Source/Runtime/Experimental/Chaos/Private/Chaos/ClusterUnionManager.cpp b/Engine/Source/Runtime/Experimental/Chaos/Private/Chaos/ClusterUnionManager.cpp
--- a/Engine/Source/Runtime/Experimental/Chaos/Private/Chaos/ClusterUnionManager.cpp
+++ b/Engine/Source/Runtime/Experimental/Chaos/Private/Chaos/ClusterUnionManager.cpp
@@ -446,6 +446,9 @@ namespace Chaos
 		UpdateClusterMassProperties(ClusterUnion.InternalCluster, FullChildrenSet, ClusterInertia, (bRecomputeMassOrientation && ClusterUnion.bNeedsXRInitialization) ? nullptr : &ForceMassOrientation);
 		UpdateKinematicProperties(ClusterUnion.InternalCluster, MClustering.GetChildrenMap(), MEvolution);
 
+		// We must reset collisions etc on this particle since geometry will be changed
+		MEvolution.InvalidateParticle(ClusterUnion.InternalCluster);
+
 		// The recreation of the geometry must happen after the call to UpdateClusterMassProperties.
 		// Creating the geometry requires knowing the relative frame between the parent cluster and the child clusters. The
 		// parent transform is not set properly for a new empty cluster until UpdateClusterMassProperties is called for the first time.
