Geometry collection : add blueprint functions to get initial and current transform of the root node of the collection

diff --git a/Engine/Source/Runtime/Experimental/GeometryCollectionEngine/Private/GeometryCollection/GeometryCollectionComponent.cpp b/Engine/Source/Runtime/Experimental/GeometryCollectionEngine/Private/GeometryCollection/GeometryCollectionComponent.cpp
--- a/Engine/Source/Runtime/Experimental/GeometryCollectionEngine/Private/GeometryCollection/GeometryCollectionComponent.cpp
+++ b/Engine/Source/Runtime/Experimental/GeometryCollectionEngine/Private/GeometryCollection/GeometryCollectionComponent.cpp
@@ -4749,6 +4749,30 @@ int32 UGeometryCollectionComponent::GetRootIndex() const
 	return INDEX_NONE;
 }
 
+
+
+FTransform UGeometryCollectionComponent::GetRootInitialTransform() const
+{
+	FTransform RootInitialTransform{ FTransform::Identity };
+	const int32 RootIndex = RestCollection->GetRootIndex();
+	if (RestTransforms.IsValidIndex(RootIndex))
+	{
+		RootInitialTransform = RestTransforms[RootIndex] * GetComponentTransform();
+	}
+	return RootInitialTransform;
+}
+
+FTransform UGeometryCollectionComponent::GetRootCurrentTransform() const
+{
+	FTransform RootInitialTransform{ FTransform::Identity };
+	const int32 RootIndex = RestCollection->GetRootIndex();
+	if (GlobalMatrices.IsValidIndex(RootIndex))
+	{
+		RootInitialTransform = FTransform(GlobalMatrices[RootIndex]) * GetComponentTransform();
+	}
+	return RootInitialTransform;
+}
+
 void UGeometryCollectionComponent::GetMassAndExtents(int32 ItemIndex, float& OutMass, FBox& OutExtents)
 {
 	using FGeometryCollectionPtr = const TSharedPtr<FGeometryCollection, ESPMode::ThreadSafe>;
diff --git a/Engine/Source/Runtime/Experimental/GeometryCollectionEngine/Public/GeometryCollection/GeometryCollectionComponent.h b/Engine/Source/Runtime/Experimental/GeometryCollectionEngine/Public/GeometryCollection/GeometryCollectionComponent.h
--- a/Engine/Source/Runtime/Experimental/GeometryCollectionEngine/Public/GeometryCollection/GeometryCollectionComponent.h
+++ b/Engine/Source/Runtime/Experimental/GeometryCollectionEngine/Public/GeometryCollection/GeometryCollectionComponent.h
@@ -629,6 +629,14 @@ public:
 	UFUNCTION(BlueprintCallable, Category = "ChaosPhysics")
 	int32 GetRootIndex() const;
 
+	/** Get the root item initial transform */
+	UFUNCTION(BlueprintCallable, Category = "ChaosPhysics")
+	FTransform GetRootInitialTransform() const;
+
+	/** Get the root item initial transform */
+	UFUNCTION(BlueprintCallable, Category = "ChaosPhysics")
+	FTransform GetRootCurrentTransform() const;
+
 	/**
 	* Get mass and extent of a specific piece
 	* @param ItemIndex item index ( from HitResult) of the cluster to get level from
