Client Crash UGeometryCollectionComponent

diff --git a/Engine/Source/Runtime/Experimental/GeometryCollectionEngine/Private/GeometryCollection/GeometryCollectionComponent.cpp b/Engine/Source/Runtime/Experimental/GeometryCollectionEngine/Private/GeometryCollection/GeometryCollectionComponent.cpp
--- a/Engine/Source/Runtime/Experimental/GeometryCollectionEngine/Private/GeometryCollection/GeometryCollectionComponent.cpp
+++ b/Engine/Source/Runtime/Experimental/GeometryCollectionEngine/Private/GeometryCollection/GeometryCollectionComponent.cpp
@@ -1908,30 +1908,32 @@ bool UGeometryCollectionComponent::ProcessRepData(const float DeltaTime, const f
 	{
 		const FGeometryCollectionActivatedCluster& ActivatedCluster = RepData.OneOffActivated[OneOffActivatedProcessed];
 		FPBDRigidParticleHandle* OneOff = PhysicsProxy->GetParticles()[ActivatedCluster.ActivatedIndex];
-		check(OneOff);
 
-		if (FPBDRigidClusteredParticleHandle* ClusterParticle = OneOff->CastToClustered())
+		if (ensure(OneOff))
 		{
-			// If there's a parent cluster particle we need to release them first.
-			// This is generally an indication that something desynced between the client and server though...maybe something needs to be done
-			// to ensure internal clusters stay in sync.
-			if (FPBDRigidClusteredParticleHandle* ParentParticle = ClusterParticle->Parent())
+			if (FPBDRigidClusteredParticleHandle* ClusterParticle = OneOff->CastToClustered())
 			{
-				// server authoritative particles are unbreakable, we need to set them breakable again 
-				ParentParticle->SetUnbreakable(false);
-				RigidClustering.ReleaseClusterParticles(TArray<FPBDRigidParticleHandle*>{ ParentParticle }, true);
-			}
+				// If there's a parent cluster particle we need to release them first.
+				// This is generally an indication that something desynced between the client and server though...maybe something needs to be done
+				// to ensure internal clusters stay in sync.
+				if (FPBDRigidClusteredParticleHandle* ParentParticle = ClusterParticle->Parent())
+				{
+					// server authoritative particles are unbreakable, we need to set them breakable again 
+					ParentParticle->SetUnbreakable(false);
+					RigidClustering.ReleaseClusterParticles(TArray<FPBDRigidParticleHandle*>{ ParentParticle }, true);
+				}
 
-			// Set initial velocities if not hard snapping
-			if (!bHardSnap)
-			{
-				// TODO: we should get an update cluster position first so that when particles break off they get the right position 
-				// TODO: should we invalidate?
-				OneOff->SetV(ActivatedCluster.InitialLinearVelocity);
-				OneOff->SetW(ActivatedCluster.InitialAngularVelocity);
-			}
+				// Set initial velocities if not hard snapping
+				if (!bHardSnap)
+				{
+					// TODO: we should get an update cluster position first so that when particles break off they get the right position 
+					// TODO: should we invalidate?
+					OneOff->SetV(ActivatedCluster.InitialLinearVelocity);
+					OneOff->SetW(ActivatedCluster.InitialAngularVelocity);
+				}
 
-			RigidClustering.ReleaseClusterParticles(TArray<FPBDRigidParticleHandle*>{ ClusterParticle }, true);
+				RigidClustering.ReleaseClusterParticles(TArray<FPBDRigidParticleHandle*>{ ClusterParticle }, true);
+			}
 		}
 	}
 
