Fix issues caused by removal of MEvolution.InvalidateParticle - move calls to TrackBreakingCollision before HandleRemoveOperationWithClusterLookup

diff --git a/Engine/Source/Runtime/Experimental/Chaos/Private/Chaos/PBDRigidClustering.cpp b/Engine/Source/Runtime/Experimental/Chaos/Private/Chaos/PBDRigidClustering.cpp
--- a/Engine/Source/Runtime/Experimental/Chaos/Private/Chaos/PBDRigidClustering.cpp
+++ b/Engine/Source/Runtime/Experimental/Chaos/Private/Chaos/PBDRigidClustering.cpp
@@ -827,6 +827,22 @@ namespace Chaos
 			{
 				if (!bFoundFirstRelease)
 				{
+					// Restore some of the momentum of whatever collided with the parent
+					// NOTE: This has to come before HandleRemoveOperationWithClusterLookup, because
+					// in FClusterUnionManager::UpdateAllClusterUnionProperties, the particle is
+					// invalidated with MEvolution.InvalidateParticle, which clears its contacts
+					if (RestoreBreakingMomentumPercent > 0.f)
+					{
+						if (Parent)
+						{
+							TrackBreakingCollision(Parent);
+						}
+						else
+						{
+							TrackBreakingCollision(ClusteredParticle);
+						}
+					}
+
 					ClusterUnionManager.HandleRemoveOperationWithClusterLookup({ ClusteredParticle }, true);
 					bFoundFirstRelease = true;
 				}
@@ -869,19 +885,6 @@ namespace Chaos
 			ClusterUnionManager.HandleRemoveOperationWithClusterLookup(DeferredRemoveFromClusterUnion, true);
 		}
 
-		// Restore some of the momentum of whatever collided with the parent
-		if (bFoundFirstRelease && RestoreBreakingMomentumPercent > 0.f)
-		{
-			if (Parent)
-			{
-				TrackBreakingCollision(Parent);
-			}
-			else
-			{
-				TrackBreakingCollision(ClusteredParticle);
-			}
-		}
-
 		// if necessary propagate strain through the graph
 		// IMPORTANT: this assumes that the connectivity graph has not yet been updated from pieces that broke off
 		if (bUseDamagePropagation)
