add ability to defer the update of cluster union properties - use this in the case the breaking model where we're popping off GCs from the cluster union one by one

diff --git a/Engine/Source/Runtime/Experimental/Chaos/Private/Chaos/ClusterUnionManager.cpp b/Engine/Source/Runtime/Experimental/Chaos/Private/Chaos/ClusterUnionManager.cpp
--- a/Engine/Source/Runtime/Experimental/Chaos/Private/Chaos/ClusterUnionManager.cpp
+++ b/Engine/Source/Runtime/Experimental/Chaos/Private/Chaos/ClusterUnionManager.cpp
@@ -70,7 +70,7 @@ namespace Chaos
 			// Note that we need to make a copy of the array here since the children list will be modified by the HandleRemoveOperation.
 			// However, the function does not expect that the input array will change.
 			TArray<FPBDRigidParticleHandle*> ChildrenCopy = ClusterUnion->ChildParticles;
-			HandleRemoveOperation(Index, ChildrenCopy, false);
+			HandleRemoveOperation(Index, ChildrenCopy, EClusterUnionOperationTiming::Never);
 			ClusterUnion->ChildParticles.Empty();
 			MClustering.DestroyClusterParticle(ClusterUnion->InternalCluster);
 
@@ -182,7 +182,7 @@ namespace Chaos
 					HandleAddOperation(OpMap.Key, Op.Value, Op.Key == EClusterUnionOperation::AddReleased);
 					break;
 				case EClusterUnionOperation::Remove:
-					HandleRemoveOperation(OpMap.Key, Op.Value, true);
+					HandleRemoveOperation(OpMap.Key, Op.Value, EClusterUnionOperationTiming::Immediate);
 					break;
 				}
 			}
@@ -223,8 +223,8 @@ namespace Chaos
 		}
 
 		// If we're adding particles to a cluster we need to first make sure they're not part of any other cluster.
-		// Book-keeping might get a bit odd if we try to add a particle to a new clutser and then only later remove the particle from its old cluster.
-		HandleRemoveOperationWithClusterLookup(Particles, true);
+		// Book-keeping might get a bit odd if we try to add a particle to a new cluster and then only later remove the particle from its old cluster.
+		HandleRemoveOperationWithClusterLookup(Particles, EClusterUnionOperationTiming::Immediate);
 
 		TGuardValue_Bitfield_Cleanup<TFunction<void()>> Cleanup(
 			[this, OldGenerateClusterBreaking=MClustering.GetDoGenerateBreakingData()]() {
@@ -345,8 +345,28 @@ namespace Chaos
 		MEvolution.GetParticles().MarkTransientDirtyParticle(Cluster->InternalCluster);
 	}
 
+	DECLARE_CYCLE_STAT(TEXT("FClusterUnionManager::HandleDeferredClusterUnionUpdateProperties"), STAT_HandleDeferredClusterUnionUpdateProperties, STATGROUP_Chaos);
+	void FClusterUnionManager::HandleDeferredClusterUnionUpdateProperties()
+	{
+		SCOPE_CYCLE_COUNTER(STAT_HandleDeferredClusterUnionUpdateProperties);
+		if (DeferredClusterUnionsForUpdateProperties.IsEmpty())
+		{
+			return;
+		}
+
+		for (FClusterUnionIndex Index : DeferredClusterUnionsForUpdateProperties)
+		{
+			if (FClusterUnion* Union = FindClusterUnion(Index))
+			{
+				UpdateAllClusterUnionProperties(*Union, false);
+			}
+		}
+
+		DeferredClusterUnionsForUpdateProperties.Reset();
+	}
+
 	DECLARE_CYCLE_STAT(TEXT("FClusterUnionManager::HandleRemoveOperation"), STAT_HandleRemoveOperation, STATGROUP_Chaos);
-	void FClusterUnionManager::HandleRemoveOperation(FClusterUnionIndex ClusterIndex, const TArray<FPBDRigidParticleHandle*>& Particles, bool bUpdateClusterProperties)
+	void FClusterUnionManager::HandleRemoveOperation(FClusterUnionIndex ClusterIndex, const TArray<FPBDRigidParticleHandle*>& Particles, EClusterUnionOperationTiming UpdateClusterPropertiesTiming)
 	{
 		SCOPE_CYCLE_COUNTER(STAT_HandleRemoveOperation);
 		FClusterUnion* Cluster = ClusterUnions.Find(ClusterIndex);
@@ -390,9 +410,14 @@ namespace Chaos
 
 		MClustering.RemoveParticlesFromCluster(Cluster->InternalCluster, Particles);
 
-		if (bUpdateClusterProperties)
+		switch (UpdateClusterPropertiesTiming)
 		{
+		case EClusterUnionOperationTiming::Immediate:
 			UpdateAllClusterUnionProperties(*Cluster, false);
+			break;
+		case EClusterUnionOperationTiming::Defer:
+			DeferredClusterUnionsForUpdateProperties.Add(ClusterIndex);
+			break;
 		}
 
 		// Removing a particle should have no bearing on the proxy of the cluster.
@@ -413,7 +438,7 @@ namespace Chaos
 	}
 
 	DECLARE_CYCLE_STAT(TEXT("FClusterUnionManager::HandleRemoveOperationWithClusterLookup"), STAT_HandleRemoveOperationWithClusterLookup, STATGROUP_Chaos);
-	void FClusterUnionManager::HandleRemoveOperationWithClusterLookup(const TArray<FPBDRigidParticleHandle*>& InParticles, bool bUpdateClusterProperties)
+	void FClusterUnionManager::HandleRemoveOperationWithClusterLookup(const TArray<FPBDRigidParticleHandle*>& InParticles, EClusterUnionOperationTiming UpdateClusterPropertiesTiming)
 	{
 		SCOPE_CYCLE_COUNTER(STAT_HandleRemoveOperationWithClusterLookup);
 		TMap<FClusterUnionIndex, TSet<FPBDRigidParticleHandle*>> ParticlesPerCluster;
@@ -427,7 +452,7 @@ namespace Chaos
 
 		for (const TPair<FClusterUnionIndex, TSet<FPBDRigidParticleHandle*>>& Kvp : ParticlesPerCluster)
 		{
-			HandleRemoveOperation(Kvp.Key, Kvp.Value.Array(), bUpdateClusterProperties);
+			HandleRemoveOperation(Kvp.Key, Kvp.Value.Array(), UpdateClusterPropertiesTiming);
 		}
 	}
 
diff --git a/Engine/Source/Runtime/Experimental/Chaos/Private/Chaos/PBDRigidClustering.cpp b/Engine/Source/Runtime/Experimental/Chaos/Private/Chaos/PBDRigidClustering.cpp
--- a/Engine/Source/Runtime/Experimental/Chaos/Private/Chaos/PBDRigidClustering.cpp
+++ b/Engine/Source/Runtime/Experimental/Chaos/Private/Chaos/PBDRigidClustering.cpp
@@ -1,4 +1,5 @@
 // Copyright Epic Games, Inc. All Rights Reserved.
+
 #include "Chaos/PBDRigidClustering.h"
 
 #include "Chaos/ErrorReporter.h"
@@ -843,7 +844,7 @@ namespace Chaos
 						}
 					}
 
-					ClusterUnionManager.HandleRemoveOperationWithClusterLookup({ ClusteredParticle }, true);
+					ClusterUnionManager.HandleRemoveOperationWithClusterLookup({ ClusteredParticle }, EClusterUnionOperationTiming::Defer);
 					bFoundFirstRelease = true;
 				}
 
@@ -882,7 +883,7 @@ namespace Chaos
 
 		if (!DeferredRemoveFromClusterUnion.IsEmpty())
 		{
-			ClusterUnionManager.HandleRemoveOperationWithClusterLookup(DeferredRemoveFromClusterUnion, true);
+			ClusterUnionManager.HandleRemoveOperationWithClusterLookup(DeferredRemoveFromClusterUnion, EClusterUnionOperationTiming::Defer);
 		}
 
 		// if necessary propagate strain through the graph
@@ -1171,6 +1172,10 @@ namespace Chaos
 			ReleaseClusterParticles(ClusteredParticle);
 		}
 
+		// This way if we break apart a large cluster union here (i.e. many of its children want to be released from ReleaseClusterParticles due to strain)
+		// we'll only update the cluster properties once here (connection graph, geometry, etc.).
+		ClusterUnionManager.HandleDeferredClusterUnionUpdateProperties();
+
 		// Restore some of the momentum of objects that were touching rigid clusters that broke
 		if (RestoreBreakingMomentumPercent > 0.f)
 		{
@@ -1640,7 +1645,7 @@ namespace Chaos
 		}
 
 		// Pre-emptively remove the particle from cluster unions it might be in.
-		ClusterUnionManager.HandleRemoveOperationWithClusterLookup({ ClusteredParticle }, true);
+		ClusterUnionManager.HandleRemoveOperationWithClusterLookup({ ClusteredParticle }, EClusterUnionOperationTiming::Immediate);
 
 		// max strain will allow to unconditionally release the children when strain is evaluated
 		constexpr FRealSingle MaxStrain = TNumericLimits<FRealSingle>::Max();
diff --git a/Engine/Source/Runtime/Experimental/Chaos/Private/PhysicsProxy/ClusterUnionPhysicsProxy.cpp b/Engine/Source/Runtime/Experimental/Chaos/Private/PhysicsProxy/ClusterUnionPhysicsProxy.cpp
--- a/Engine/Source/Runtime/Experimental/Chaos/Private/PhysicsProxy/ClusterUnionPhysicsProxy.cpp
+++ b/Engine/Source/Runtime/Experimental/Chaos/Private/PhysicsProxy/ClusterUnionPhysicsProxy.cpp
@@ -174,7 +174,7 @@ namespace Chaos
 				FReadPhysicsObjectInterface_Internal Interface = FPhysicsObjectInternalInterface::GetRead();
 				if (FPBDRigidsEvolutionGBF* Evolution = GetEvolution(this))
 				{
-					Evolution->GetRigidClustering().GetClusterUnionManager().HandleRemoveOperation(ClusterUnionIndex, Interface.GetAllRigidParticles(Objects.Array()), true);
+					Evolution->GetRigidClustering().GetClusterUnionManager().HandleRemoveOperation(ClusterUnionIndex, Interface.GetAllRigidParticles(Objects.Array()), EClusterUnionOperationTiming::Immediate);
 				}
 			}
 		);
diff --git a/Engine/Source/Runtime/Experimental/Chaos/Private/PhysicsProxy/GeometryCollectionPhysicsProxy.cpp b/Engine/Source/Runtime/Experimental/Chaos/Private/PhysicsProxy/GeometryCollectionPhysicsProxy.cpp
--- a/Engine/Source/Runtime/Experimental/Chaos/Private/PhysicsProxy/GeometryCollectionPhysicsProxy.cpp
+++ b/Engine/Source/Runtime/Experimental/Chaos/Private/PhysicsProxy/GeometryCollectionPhysicsProxy.cpp
@@ -1974,7 +1974,7 @@ void FGeometryCollectionPhysicsProxy::BreakClusters_External(TArray<FGeometryCol
 				if (Chaos::FPBDRigidClusteredParticleHandle* ClusteredHandle = FindClusteredParticleHandleByItemIndex_Internal(ItemIndex))
 				{
 					// Pre-emptively remove the particle from cluster unions it might be in.
-					Clustering.GetClusterUnionManager().HandleRemoveOperationWithClusterLookup({ ClusteredHandle }, true);
+					Clustering.GetClusterUnionManager().HandleRemoveOperationWithClusterLookup({ ClusteredHandle }, Chaos::EClusterUnionOperationTiming::Immediate);
 					Clustering.BreakCluster(ClusteredHandle);
 				}
 			}
diff --git a/Engine/Source/Runtime/Experimental/Chaos/Public/Chaos/ClusterUnionManager.h b/Engine/Source/Runtime/Experimental/Chaos/Public/Chaos/ClusterUnionManager.h
--- a/Engine/Source/Runtime/Experimental/Chaos/Public/Chaos/ClusterUnionManager.h
+++ b/Engine/Source/Runtime/Experimental/Chaos/Public/Chaos/ClusterUnionManager.h
@@ -24,6 +24,13 @@ namespace Chaos
 		Remove
 	};
 
+	enum class EClusterUnionOperationTiming
+	{
+		Never,
+		Defer,
+		Immediate
+	};
+
 	struct CHAOS_API FClusterUnionCreationParameters
 	{
 		FClusterUnionExplicitIndex ExplicitIndex = INDEX_NONE;
@@ -82,14 +89,17 @@ namespace Chaos
 		void HandleAddOperation(FClusterUnionIndex ClusterIndex, const TArray<FPBDRigidParticleHandle*>& Particles, bool bReleaseClustersFirst);
 
 		// Removes the specified particles from the specified cluster.
-		void HandleRemoveOperation(FClusterUnionIndex ClusterIndex, const TArray<FPBDRigidParticleHandle*>& Particles, bool bUpdateClusterProperties);
+		void HandleRemoveOperation(FClusterUnionIndex ClusterIndex, const TArray<FPBDRigidParticleHandle*>& Particles, EClusterUnionOperationTiming UpdateClusterPropertiesTiming);
 
 		// Helper function to remove particles given only the particle handle. This will consult the lookup table to find which cluster the particle is in before doing a normal remove operation.
-		void HandleRemoveOperationWithClusterLookup(const TArray<FPBDRigidParticleHandle*>& InParticles, bool bUpdateClusterProperties);
+		void HandleRemoveOperationWithClusterLookup(const TArray<FPBDRigidParticleHandle*>& InParticles, EClusterUnionOperationTiming UpdateClusterPropertiesTiming);
 
 		// Will be called at the beginning of every time step to ensure that all the expected cluster unions have been modified.
 		void FlushPendingOperations();
 
+		// Update cluster union properties here if they were deferred. This can be called manually but will otherwise also be handled in FlushPendingOperations.
+		void HandleDeferredClusterUnionUpdateProperties();
+
 		// Access the cluster union externally.
 		FClusterUnion* FindClusterUnionFromExplicitIndex(FClusterUnionExplicitIndex Index);
 		FClusterUnion* FindClusterUnion(FClusterUnionIndex Index);
@@ -129,6 +139,9 @@ namespace Chaos
 		// All of our actively managed cluster unions. We need to keep track of these
 		// so a user could use the index to request modifications to a specific cluster union.
 		TMap<FClusterUnionIndex, FClusterUnion> ClusterUnions;
+
+		// The set of cluster unions that we want to defer updating their cluster union properties for.
+		TSet<FClusterUnionIndex> DeferredClusterUnionsForUpdateProperties;
 		
 		//
 		// There are two ways we can pick a new union index:
