In FRigidClustering::ComputeStrainFromCollision, make sure GCs are unpacked from clustered unions for collision strain processing, and respect propagation settings (FindAllIntersectingChildren does not work properly for ClusteredUnion). Also templatize GetConcreteProxy so we can get gc or cu proxy

diff --git a/Engine/Source/Runtime/Experimental/Chaos/Private/Chaos/PBDRigidClustering.cpp b/Engine/Source/Runtime/Experimental/Chaos/Private/Chaos/PBDRigidClustering.cpp
--- a/Engine/Source/Runtime/Experimental/Chaos/Private/Chaos/PBDRigidClustering.cpp
+++ b/Engine/Source/Runtime/Experimental/Chaos/Private/Chaos/PBDRigidClustering.cpp
@@ -23,6 +23,7 @@
 #include "Chaos/PerParticlePBDEulerStep.h"
 #include "Chaos/StrainModification.h"
 #include "PhysicsProxy/GeometryCollectionPhysicsProxy.h"
+#include "PhysicsProxy/ClusterUnionPhysicsProxy.h"
 #include "CoreMinimal.h"
 
 namespace Chaos
@@ -57,30 +58,32 @@ namespace Chaos
 	float RestoreBreakingMomentumPercent = .5;
 	FAutoConsoleVariableRef CVarRestoreBreakingMomentumPercent(TEXT("p.RestoreBreakingMomentumPercent"), RestoreBreakingMomentumPercent, TEXT("When a rigid cluster is broken, objects that its in contact with will receive an impulse to restore this percent of their momentum prior to the break."));
 
-	static FGeometryCollectionPhysicsProxy* GetConcreteProxy(FPBDRigidClusteredParticleHandle* ClusteredParticle)
+	template <typename TProxy=FGeometryCollectionPhysicsProxy>
+	TProxy* GetConcreteProxy(FPBDRigidClusteredParticleHandle* ClusteredParticle)
 	{
 		if (ClusteredParticle)
 		{
 			if (IPhysicsProxyBase* Proxy = ClusteredParticle->PhysicsProxy())
 			{
-				if (Proxy->GetType() == EPhysicsProxyType::GeometryCollectionType)
+				if (Proxy->GetType() == TProxy::ConcreteType())
 				{
-					return static_cast<FGeometryCollectionPhysicsProxy*>(Proxy);
+					return static_cast<TProxy*>(Proxy);
 				}
 			}
 		}
 		return nullptr;
 	}
 
-	static const FGeometryCollectionPhysicsProxy* GetConcreteProxy(const FPBDRigidClusteredParticleHandle* ClusteredParticle)
+	template <typename TProxy=FGeometryCollectionPhysicsProxy>
+	const TProxy* GetConcreteProxy(const FPBDRigidClusteredParticleHandle* ClusteredParticle)
 	{
 		if (ClusteredParticle)
 		{
 			if (const IPhysicsProxyBase* Proxy = ClusteredParticle->PhysicsProxy())
 			{
-				if (Proxy->GetType() == EPhysicsProxyType::GeometryCollectionType)
+				if (Proxy->GetType() == TProxy::ConcreteType())
 				{
-					return static_cast<const FGeometryCollectionPhysicsProxy*>(Proxy);
+					return static_cast<const TProxy*>(Proxy);
 				}
 			}
 		}
@@ -1435,19 +1438,37 @@ namespace Chaos
 			}
 
 			auto ComputeStrainLambda = [&](
-				const FPBDRigidClusteredParticleHandle* Cluster,
-				const TArray<FPBDRigidParticleHandle*>& ParentToChildren)
+				const FPBDRigidClusteredParticleHandle* Cluster)
 			{
 				const FVec3 ContactWorldLocation = ContactHandle->GetContact().CalculateWorldContactLocation();
 				
 				const FRealSingle AccumulatedImpulse = static_cast<FRealSingle>(ContactHandle->GetAccumulatedImpulse().Size());
 				if (AccumulatedImpulse > UE_SMALL_NUMBER && FMath::IsFinite(AccumulatedImpulse))
 				{
+					if (const FClusterUnionPhysicsProxy* ClusterUnionProxy = GetConcreteProxy<FClusterUnionPhysicsProxy>(Cluster))
+					{
+						// At the moment, we don't want to apply strains to children of ClusterUnions, we want instead
+						// to apply the strains to GRANDchildren of ClusterUnions.
+						FPBDRigidParticleHandle* ClosestChild = FindClosestChild(Cluster, ContactWorldLocation);
+
+						// If Closest child is not a clustered, then there is no substructure to apply strain to,
+						// so null Cluster
+						Cluster
+							= ClosestChild
+							? ClosestChild->CastToClustered()
+							: nullptr;
+					}
+
+					if (Cluster == nullptr)
+					{
+						return;
+					}
+
 					// gather propagation information from the parent proxy
 					bool bUseDamagePropagation = false;
-					if (const FGeometryCollectionPhysicsProxy* ConcreteProxy = GetConcreteProxy(Cluster))
+					if (const FGeometryCollectionPhysicsProxy* GeometryCollectionProxy = GetConcreteProxy<FGeometryCollectionPhysicsProxy>(Cluster))
 					{
-						const FSimulationParameters& SimParams = ConcreteProxy->GetSimParameters();
+						const FSimulationParameters& SimParams = GeometryCollectionProxy->GetSimParameters();
 						bUseDamagePropagation = SimParams.bUseDamagePropagation;
 						if (!SimParams.bEnableStrainOnCollision)
 						{
@@ -1459,8 +1480,7 @@ namespace Chaos
 					{
 						// propagation based breaking model start from the closest particle and propagate through the connection graph
 						// propagation logic is dealt when evaluating the strain
-						FPBDRigidParticleHandle* ClosestChild = FindClosestParticle(ParentToChildren, ContactWorldLocation);
-						if (ClosestChild)
+						if (FPBDRigidParticleHandle* ClosestChild = FindClosestChild(Cluster, ContactWorldLocation))
 						{
 							if (TPBDRigidClusteredParticleHandle<FReal, 3>* ClusteredChild = ClosestChild->CastToClustered())
 							{
@@ -1483,22 +1503,6 @@ namespace Chaos
 							{
 								if (TPBDRigidClusteredParticleHandle<FReal, 3>*ClusteredChild = Child->CastToClustered())
 								{
-									// If we got a child of a Clustered Union then this might just be the root particle of a geometry collection,
-									// in which case it might not make sense to strain it (?)
-									if (const IPhysicsProxyBase* Proxy = Cluster->PhysicsProxy())
-									{
-										if (Proxy->GetType() == EPhysicsProxyType::ClusterUnionProxy)
-										{
-											if (TArray<FRigidHandle>* ClusteredParentToChildren = MChildren.Find(ClusteredChild))
-											{
-												if (FPBDRigidParticleHandle* ClusteredChildRigid = FindClosestParticle(*ClusteredParentToChildren, ContactWorldLocation))
-												{
-													ClusteredChild = ClusteredChildRigid->CastToClustered();
-												}
-											}
-										}
-									}
-
 									ClusteredChild->CollisionImpulses() += AccumulatedImpulse;
 									UpdateTopLevelParticle(ClusteredChild);
 								}
@@ -1510,18 +1514,12 @@ namespace Chaos
 
 			if (ClusteredConstrainedParticles0)
 			{
-				if (const TArray<FPBDRigidParticleHandle*>* ChildrenPtr = MParentToChildren.Find(ClusteredConstrainedParticles0))
-				{
-					ComputeStrainLambda(ClusteredConstrainedParticles0, *ChildrenPtr);
-				}
+				ComputeStrainLambda(ClusteredConstrainedParticles0);
 			}
 
 			if (ClusteredConstrainedParticles1)
 			{
-				if (const TArray<FPBDRigidParticleHandle*>* ChildrenPtr = MParentToChildren.Find(ClusteredConstrainedParticles1))
-				{
-					ComputeStrainLambda(ClusteredConstrainedParticles1, *ChildrenPtr);
-				}
+				ComputeStrainLambda(ClusteredConstrainedParticles1);
 			}
 
 			MCollisionImpulseArrayDirty = true;
diff --git a/Engine/Source/Runtime/Experimental/Chaos/Public/PhysicsProxy/ClusterUnionPhysicsProxy.h b/Engine/Source/Runtime/Experimental/Chaos/Public/PhysicsProxy/ClusterUnionPhysicsProxy.h
--- a/Engine/Source/Runtime/Experimental/Chaos/Public/PhysicsProxy/ClusterUnionPhysicsProxy.h
+++ b/Engine/Source/Runtime/Experimental/Chaos/Public/PhysicsProxy/ClusterUnionPhysicsProxy.h
@@ -136,7 +136,7 @@ namespace Chaos
 		void DisableCollisionsCallback(TSet<TTuple<int32, int32>>& InPairs) {}
 		void AddForceCallback(FParticlesType& InParticles, const float InDt, const int32 InIndex) {}
 		void BindParticleCallbackMapping(Chaos::TArrayCollectionArray<PhysicsProxyWrapper>& PhysicsProxyReverseMap, Chaos::TArrayCollectionArray<int32>& ParticleIDReverseMap) {}
-		EPhysicsProxyType ConcreteType() { return EPhysicsProxyType::ClusterUnionProxy; }
+		static EPhysicsProxyType ConcreteType() { return EPhysicsProxyType::ClusterUnionProxy; }
 		void SyncBeforeDestroy() {}
 		void OnRemoveFromScene() {}
 		bool IsDirty() { return false; }
diff --git a/Engine/Source/Runtime/Experimental/Chaos/Public/PhysicsProxy/GeometryCollectionPhysicsProxy.h b/Engine/Source/Runtime/Experimental/Chaos/Public/PhysicsProxy/GeometryCollectionPhysicsProxy.h
--- a/Engine/Source/Runtime/Experimental/Chaos/Public/PhysicsProxy/GeometryCollectionPhysicsProxy.h
+++ b/Engine/Source/Runtime/Experimental/Chaos/Public/PhysicsProxy/GeometryCollectionPhysicsProxy.h
@@ -226,7 +226,7 @@ public:
 
 	bool IsDirty() { return false; }
 
-	EPhysicsProxyType ConcreteType() { return EPhysicsProxyType::GeometryCollectionType; }
+	static EPhysicsProxyType ConcreteType() { return EPhysicsProxyType::GeometryCollectionType; }
 
 	void SyncBeforeDestroy();
 	void OnRemoveFromSolver(Chaos::FPBDRigidsSolver *RBDSolver);
