SOA crash fix for Geometry Collection while resim is in progress

diff --git a/Engine/Source/Runtime/Experimental/Chaos/Public/Chaos/PBDRigidsSOAs.h b/Engine/Source/Runtime/Experimental/Chaos/Public/Chaos/PBDRigidsSOAs.h
--- a/Engine/Source/Runtime/Experimental/Chaos/Public/Chaos/PBDRigidsSOAs.h
+++ b/Engine/Source/Runtime/Experimental/Chaos/Public/Chaos/PBDRigidsSOAs.h
@@ -623,36 +623,57 @@ public:
 				MovingKinematicsMapArray.Remove(Particle);
 			}
 
-			// Move to appropriate dynamic SOA
-			switch (State)
+			if (Particle->Type == EParticleType::Clustered)
 			{
-			case EObjectStateType::Kinematic:
-				Particle->MoveToSOA(*DynamicKinematicParticles);
-				if(bResimulating)
-				{
-					ResimDynamicParticles.Remove(Particle);
-					ResimDynamicKinematicParticles.Insert(Particle);
-				}
-				break;
-
-			case EObjectStateType::Dynamic:
-				Particle->MoveToSOA(*DynamicParticles);
+				Particle->MoveToSOA(*ClusteredParticles);
 				if (bResimulating)
 				{
 					ResimDynamicKinematicParticles.Remove(Particle);
 					ResimDynamicParticles.Insert(Particle);
 				}
-				break;
-
-			default:
-				// TODO: Special SOAs for sleeping and static particles?
-				Particle->MoveToSOA(*DynamicParticles);
+			}
+			else if (Particle->Type == EParticleType::GeometryCollection)
+			{
+				Particle->MoveToSOA(*GeometryCollectionParticles);
 				if (bResimulating)
 				{
 					ResimDynamicKinematicParticles.Remove(Particle);
 					ResimDynamicParticles.Insert(Particle);
 				}
-				break;
+			}
+			else
+			{
+				// Move to appropriate dynamic SOA
+				switch (State)
+				{
+				case EObjectStateType::Kinematic:
+					Particle->MoveToSOA(*DynamicKinematicParticles);
+					if (bResimulating)
+					{
+						ResimDynamicParticles.Remove(Particle);
+						ResimDynamicKinematicParticles.Insert(Particle);
+					}
+					break;
+
+				case EObjectStateType::Dynamic:
+					Particle->MoveToSOA(*DynamicParticles);
+					if (bResimulating)
+					{
+						ResimDynamicKinematicParticles.Remove(Particle);
+						ResimDynamicParticles.Insert(Particle);
+					}
+					break;
+
+				default:
+					// TODO: Special SOAs for sleeping and static particles?
+					Particle->MoveToSOA(*DynamicParticles);
+					if (bResimulating)
+					{
+						ResimDynamicKinematicParticles.Remove(Particle);
+						ResimDynamicParticles.Insert(Particle);
+					}
+					break;
+				}
 			}
 		}
 		UpdateViews();
@@ -963,6 +984,7 @@ private:
 			auto TmpArray = bResimulating
 			? TArray<TSOAView<FGeometryParticles>>
 			{
+				// re-sim only works with a reduced set of particle types, i.e. no DynamicClusteredMapArray (DynamicGeometryCollectionArray??) BH
 				{&ResimStaticParticles},
 				{&ResimKinematicParticles},
 				{&ResimDynamicParticles.GetArray() },
@@ -1225,3 +1247,5 @@ template <typename T, int d>
 using TPBDRigidsSOAs UE_DEPRECATED(4.27, "Deprecated. this class is to be deleted, use FPBDRigidsSOAs instead") = FPBDRigidsSOAs;
 
 }
+
+PRAGMA_ENABLE_OPTIMIZATION
