CIS fix removing pragma that was accidently left in

diff --git a/Engine/Source/Runtime/Experimental/Chaos/Public/Chaos/PBDRigidsSOAs.h b/Engine/Source/Runtime/Experimental/Chaos/Public/Chaos/PBDRigidsSOAs.h
--- a/Engine/Source/Runtime/Experimental/Chaos/Public/Chaos/PBDRigidsSOAs.h
+++ b/Engine/Source/Runtime/Experimental/Chaos/Public/Chaos/PBDRigidsSOAs.h
@@ -1247,5 +1247,3 @@ template <typename T, int d>
 using TPBDRigidsSOAs UE_DEPRECATED(4.27, "Deprecated. this class is to be deleted, use FPBDRigidsSOAs instead") = FPBDRigidsSOAs;
 
 }
-
-PRAGMA_ENABLE_OPTIMIZATION
