Iris - Use ranged-based for when iterating over valid connections.

#rnx

diff --git a/Engine/Source/Runtime/Experimental/Iris/Core/Private/Iris/ReplicationSystem/Conditionals/ReplicationConditionals.cpp b/Engine/Source/Runtime/Experimental/Iris/Core/Private/Iris/ReplicationSystem/Conditionals/ReplicationConditionals.cpp
--- a/Engine/Source/Runtime/Experimental/Iris/Core/Private/Iris/ReplicationSystem/Conditionals/ReplicationConditionals.cpp
+++ b/Engine/Source/Runtime/Experimental/Iris/Core/Private/Iris/ReplicationSystem/Conditionals/ReplicationConditionals.cpp
@@ -771,15 +771,8 @@ void FReplicationConditionals::ClearConnectionInfosForObject(const FNetBitArray&
 {
 	IRIS_PROFILER_SCOPE(FReplicationConditionals_ClearConnectionInfosForObject);
 
-	// 0 is not a valid connection ID so we adjust it to one-based in the loop.
-	for (uint32 ConnIt = 0, ConnEndIt = MaxConnectionId; ConnIt != ConnEndIt; ++ConnIt)
+	for (uint32 ConnectionId : ValidConnections)
 	{
-		const uint32 ConnectionId = ConnIt + 1U;
-		if (!ValidConnections.GetBit(ConnectionId))
-		{
-			continue;
-		}
-
 		FPerConnectionInfo& ConnectionInfo = ConnectionInfos[ConnectionId];
 		ConnectionInfo.ObjectConditionals[ObjectIndex] = FConditionalsMask{};
 	}
diff --git a/Engine/Source/Runtime/Experimental/Iris/Core/Private/Iris/ReplicationSystem/DeltaCompression/DeltaCompressionBaselineManager.cpp b/Engine/Source/Runtime/Experimental/Iris/Core/Private/Iris/ReplicationSystem/DeltaCompression/DeltaCompressionBaselineManager.cpp
--- a/Engine/Source/Runtime/Experimental/Iris/Core/Private/Iris/ReplicationSystem/DeltaCompression/DeltaCompressionBaselineManager.cpp
+++ b/Engine/Source/Runtime/Experimental/Iris/Core/Private/Iris/ReplicationSystem/DeltaCompression/DeltaCompressionBaselineManager.cpp
@@ -258,15 +258,8 @@ void FDeltaCompressionBaselineManager::UpdateDirtyStateMasks(const FChangeMaskCa
 		ChangeMaskStorageType*const ConnectionChangeMaskStoragePtr = ObjectInfo->ChangeMasksForConnections;
 		ChangeMaskStorageType const*const UpdatedChanges = DirtyChangeMaskStoragePtr + Entry.StorageOffset;
 		const SIZE_T ChangeMaskStridePerConnection = ObjectInfo->ChangeMaskStride;
-		// 0 is not a valid connection ID so we adjust it to one-based in the loop.
-		for (uint32 ConnIt = 0, ConnEndIt = MaxConnectionId; ConnIt != ConnEndIt; ++ConnIt)
+		for (uint32 ConnectionId : ValidConnections)
 		{
-			const uint32 ConnectionId = ConnIt + 1U;
-			if (!ValidConnections.GetBit(ConnectionId))
-			{
-				continue;
-			}
-
 			const SIZE_T ChangeMaskOffset = ChangeMaskStridePerConnection*ConnectionId;
 			ChangeMaskStorageType* ConnectionChangeMask = ConnectionChangeMaskStoragePtr + ChangeMaskOffset;
 			for (uint32 WordIt = 0, WordEndIt = ChangeMaskStridePerConnection; WordIt != WordEndIt; ++WordIt)
@@ -640,15 +633,8 @@ void FDeltaCompressionBaselineManager::InvalidateBaselinesDueToModifiedCondition
 		FPerObjectInfo* ObjectInfo = GetPerObjectInfo(ObjectInfoIndex);
 		if (InvalidationInfo.ConnId == BaselineInvalidationTracker->InvalidateBaselineForAllConnections)
 		{
-			// 0 is not a valid connection ID so we adjust it to one-based in the loop.
-			for (uint32 ConnIt = 0, ConnEndIt = MaxConnectionId; ConnIt != ConnEndIt; ++ConnIt)
+			for (uint32 ConnectionId : ValidConnections)
 			{
-				const uint32 ConnectionId = ConnIt + 1U;
-				if (!ValidConnections.GetBit(ConnectionId))
-				{
-					continue;
-				}
-
 				FObjectBaselineInfo& BaselineInfo = ObjectInfo->BaselinesForConnections[ConnectionId];
 				for (const uint32 BaselineIndex : {0, 1})
 				{
diff --git a/Engine/Source/Runtime/Experimental/Iris/Core/Private/Iris/ReplicationSystem/Filtering/ReplicationFiltering.cpp b/Engine/Source/Runtime/Experimental/Iris/Core/Private/Iris/ReplicationSystem/Filtering/ReplicationFiltering.cpp
--- a/Engine/Source/Runtime/Experimental/Iris/Core/Private/Iris/ReplicationSystem/Filtering/ReplicationFiltering.cpp
+++ b/Engine/Source/Runtime/Experimental/Iris/Core/Private/Iris/ReplicationSystem/Filtering/ReplicationFiltering.cpp
@@ -808,13 +808,9 @@ void FReplicationFiltering::UpdateObjectsInScope()
 	}
 
 	// Update the scope for all valid connections.
-	for (FPerConnectionInfo& ConnectionInfo : ConnectionInfos)
+	for (uint32 ConnectionId : ValidConnections)
 	{
-		const SIZE_T Index = &ConnectionInfo - ConnectionInfos.GetData();
-		if (!ValidConnections.GetBit(Index))
-		{
-			continue;
-		}
+		FPerConnectionInfo& ConnectionInfo = ConnectionInfos[ConnectionId];
 
 		// Or in brand new objects and mask off deleted objects by anding with now existing objects.
 		uint32* FilteredObjectsStorage = ConnectionInfo.ConnectionFilteredObjects.GetData();
