fix jitter due to the deferred set child to parent for cluster unions

clients will now always initialize child to parent to some value; however, as soon as they receive a server authoritative value of ChildToParent, we will lock the value so that UpdateGeometry will no longer attempt to override that value.

diff --git a/Engine/Source/Runtime/Engine/Private/PhysicsEngine/ClusterUnionComponent.cpp b/Engine/Source/Runtime/Engine/Private/PhysicsEngine/ClusterUnionComponent.cpp
--- a/Engine/Source/Runtime/Engine/Private/PhysicsEngine/ClusterUnionComponent.cpp
+++ b/Engine/Source/Runtime/Engine/Private/PhysicsEngine/ClusterUnionComponent.cpp
@@ -200,7 +200,6 @@ void UClusterUnionComponent::OnCreatePhysicsState()
 	// TODO: Expose these parameters via the component.
 	Chaos::FClusterCreationParameters Parameters{ 0.3f, 100, false, false };
 	Parameters.ConnectionMethod = Chaos::FClusterCreationParameters::EConnectionMethod::DelaunayTriangulation;
-	Parameters.bUseExistingChildToParent = !GetOwner()->HasAuthority();
 
 	FChaosUserData::Set<UPrimitiveComponent>(&PhysicsUserData, this);
 
@@ -579,7 +578,8 @@ void UClusterUnionComponent::ForceSetChildToParent(UPrimitiveComponent* InCompon
 		Objects.Add(Handle);
 	}
 
-	PhysicsProxy->BulkSetChildToParent_External(Objects, ChildToParent);
+	// If we're on the client we want to lock the child to parent transform for this particle as soon as we get a server authoritative value.
+	PhysicsProxy->BulkSetChildToParent_External(Objects, ChildToParent, !IsAuthority());
 }
 
 void UClusterUnionComponent::SetSimulatePhysics(bool bSimulate)
diff --git a/Engine/Source/Runtime/Experimental/Chaos/Private/Chaos/ClusterUnionManager.cpp b/Engine/Source/Runtime/Experimental/Chaos/Private/Chaos/ClusterUnionManager.cpp
--- a/Engine/Source/Runtime/Experimental/Chaos/Private/Chaos/ClusterUnionManager.cpp
+++ b/Engine/Source/Runtime/Experimental/Chaos/Private/Chaos/ClusterUnionManager.cpp
@@ -103,7 +103,7 @@ namespace Chaos
 		{
 			FRigidTransform3 Frame = FRigidTransform3::Identity;
 			
-			if (FPBDRigidClusteredParticleHandle* ClusterChild = Child->CastToClustered(); ClusterChild && Union.Parameters.bUseExistingChildToParent)
+			if (FPBDRigidClusteredParticleHandle* ClusterChild = Child->CastToClustered(); ClusterChild && ClusterChild->IsChildToParentLocked())
 			{
 				Frame = ClusterChild->ChildToParent();
 			}
@@ -392,6 +392,13 @@ namespace Chaos
 			if (ParticleIndex != INDEX_NONE)
 			{
 				ParticleIndicesToRemove.Add(ParticleIndex);
+
+				// Remove the child to parent lock if it exists since it's no longer managed by the cluster union manager.
+				if (FPBDRigidClusteredParticleHandle* ClusterHandle = Handle->CastToClustered())
+				{
+					ClusterHandle->SetChildToParentLocked(false);
+				}
+
 				// Remove the parent proxy only if it's a cluster union proxy.
 				if (IPhysicsProxyBase* Proxy = Handle->PhysicsProxy(); Proxy && Proxy->GetParentProxy() && Proxy->GetParentProxy()->GetType() == EPhysicsProxyType::ClusterUnionProxy)
 				{
@@ -568,7 +575,7 @@ namespace Chaos
 	}
 
 	DECLARE_CYCLE_STAT(TEXT("FClusterUnionManager::UpdateClusterUnionParticlesChildToParent"), STAT_UpdateClusterUnionParticlesChildToParent, STATGROUP_Chaos);
-	void FClusterUnionManager::UpdateClusterUnionParticlesChildToParent(FClusterUnionIndex Index, const TArray<FPBDRigidParticleHandle*>& Particles, const TArray<FTransform>& ChildToParent)
+	void FClusterUnionManager::UpdateClusterUnionParticlesChildToParent(FClusterUnionIndex Index, const TArray<FPBDRigidParticleHandle*>& Particles, const TArray<FTransform>& ChildToParent, bool bLock)
 	{
 		SCOPE_CYCLE_COUNTER(STAT_UpdateClusterUnionParticlesChildToParent);
 
@@ -588,6 +595,11 @@ namespace Chaos
 					if (FPBDRigidClusteredParticleHandle* ChildHandle = ClusterUnion->ChildParticles[ChildIndex]->CastToClustered())
 					{
 						ChildHandle->SetChildToParent(ChildToParent[ParticleIndex]);
+
+						if (bLock)
+						{
+							ChildHandle->SetChildToParentLocked(bLock);
+						}						
 					}
 				}
 			}
diff --git a/Engine/Source/Runtime/Experimental/Chaos/Private/Chaos/PBDRigidClusteringAlgo.cpp b/Engine/Source/Runtime/Experimental/Chaos/Private/Chaos/PBDRigidClusteringAlgo.cpp
--- a/Engine/Source/Runtime/Experimental/Chaos/Private/Chaos/PBDRigidClusteringAlgo.cpp
+++ b/Engine/Source/Runtime/Experimental/Chaos/Private/Chaos/PBDRigidClusteringAlgo.cpp
@@ -335,7 +335,7 @@ namespace Chaos
 			{
 				FRigidTransform3 Frame = FRigidTransform3::Identity;
 
-				if (FPBDRigidClusteredParticleHandle* ClusterChild = Child->CastToClustered(); ClusterChild && Parameters.bUseExistingChildToParent)
+				if (FPBDRigidClusteredParticleHandle* ClusterChild = Child->CastToClustered(); ClusterChild && ClusterChild->IsChildToParentLocked())
 				{
 					Frame = ClusterChild->ChildToParent();
 				}
diff --git a/Engine/Source/Runtime/Experimental/Chaos/Private/PhysicsProxy/ClusterUnionPhysicsProxy.cpp b/Engine/Source/Runtime/Experimental/Chaos/Private/PhysicsProxy/ClusterUnionPhysicsProxy.cpp
--- a/Engine/Source/Runtime/Experimental/Chaos/Private/PhysicsProxy/ClusterUnionPhysicsProxy.cpp
+++ b/Engine/Source/Runtime/Experimental/Chaos/Private/PhysicsProxy/ClusterUnionPhysicsProxy.cpp
@@ -524,12 +524,12 @@ namespace Chaos
 		Solver->AddDirtyProxy(this);
 	}
 
-	void FClusterUnionPhysicsProxy::SetChildToParent_External(FPhysicsObjectHandle Child, const FTransform& RelativeTransform)
+	void FClusterUnionPhysicsProxy::SetChildToParent_External(FPhysicsObjectHandle Child, const FTransform& RelativeTransform, bool bLock)
 	{
-		BulkSetChildToParent_External({ Child }, { RelativeTransform });
+		BulkSetChildToParent_External({ Child }, { RelativeTransform }, bLock);
 	}
 
-	void FClusterUnionPhysicsProxy::BulkSetChildToParent_External(const TArray<FPhysicsObjectHandle>& Objects, const TArray<FTransform>& Transforms)
+	void FClusterUnionPhysicsProxy::BulkSetChildToParent_External(const TArray<FPhysicsObjectHandle>& Objects, const TArray<FTransform>& Transforms, bool bLock)
 	{
 		if (!Solver || !ensure(Particle_External) || !ensure(Objects.Num() == Transforms.Num()))
 		{
@@ -537,7 +537,7 @@ namespace Chaos
 		}
 
 		Solver->EnqueueCommandImmediate(
-			[this, Objects, Transforms]() mutable
+			[this, Objects, Transforms, bLock]() mutable
 			{
 				FReadPhysicsObjectInterface_Internal Interface = FPhysicsObjectInternalInterface::GetRead();
 				TArray<FPBDRigidParticleHandle*> Particles = Interface.GetAllRigidParticles(Objects);
@@ -545,7 +545,7 @@ namespace Chaos
 				{
 					FPBDRigidsEvolutionGBF& Evolution = *static_cast<FPBDRigidsSolver*>(Solver)->GetEvolution();
 					FClusterUnionManager& ClusterUnionManager = Evolution.GetRigidClustering().GetClusterUnionManager();
-					ClusterUnionManager.UpdateClusterUnionParticlesChildToParent(ClusterUnionIndex, Particles, Transforms);
+					ClusterUnionManager.UpdateClusterUnionParticlesChildToParent(ClusterUnionIndex, Particles, Transforms, bLock);
 				}
 			}
 		);
diff --git a/Engine/Source/Runtime/Experimental/Chaos/Private/PhysicsProxy/GeometryCollectionPhysicsProxy.cpp b/Engine/Source/Runtime/Experimental/Chaos/Private/PhysicsProxy/GeometryCollectionPhysicsProxy.cpp
--- a/Engine/Source/Runtime/Experimental/Chaos/Private/PhysicsProxy/GeometryCollectionPhysicsProxy.cpp
+++ b/Engine/Source/Runtime/Experimental/Chaos/Private/PhysicsProxy/GeometryCollectionPhysicsProxy.cpp
@@ -2708,7 +2708,7 @@ void FGeometryCollectionPhysicsProxy::PushToPhysicsState()
 
 			if (ClusterUnionIndex != INDEX_NONE && !DeferredClusterUnionParticleUpdates.IsEmpty() && !DeferredClusterUnionChildToParentUpdates.IsEmpty())
 			{
-				ClusterUnionManager.UpdateClusterUnionParticlesChildToParent(ClusterUnionIndex, DeferredClusterUnionParticleUpdates, DeferredClusterUnionChildToParentUpdates);
+				ClusterUnionManager.UpdateClusterUnionParticlesChildToParent(ClusterUnionIndex, DeferredClusterUnionParticleUpdates, DeferredClusterUnionChildToParentUpdates, false);
 			}
 		}
 	}
diff --git a/Engine/Source/Runtime/Experimental/Chaos/Public/Chaos/ClusterCreationParameters.h b/Engine/Source/Runtime/Experimental/Chaos/Public/Chaos/ClusterCreationParameters.h
--- a/Engine/Source/Runtime/Experimental/Chaos/Public/Chaos/ClusterCreationParameters.h
+++ b/Engine/Source/Runtime/Experimental/Chaos/Public/Chaos/ClusterCreationParameters.h
@@ -34,7 +34,6 @@ namespace Chaos
 			, Chaos::TPBDRigidClusteredParticleHandle<Chaos::FReal,3>* ClusterParticleHandleIn = nullptr
 			, const FVec3& ScaleIn = FVec3::OneVector
 			, bool bIsAnchoredIn = false
-			, bool bInUseExistingChildToParent = false
 		)
 			: CoillisionThicknessPercent(CoillisionThicknessPercentIn)
 			, MaxNumConnections(MaxNumConnectionsIn)
@@ -47,7 +46,6 @@ namespace Chaos
 			, ClusterParticleHandle(ClusterParticleHandleIn)
 			, Scale(ScaleIn)
 			, bIsAnchored(bIsAnchoredIn)
-			, bUseExistingChildToParent(bInUseExistingChildToParent)
 		{}
 
 		FReal CoillisionThicknessPercent;
@@ -61,7 +59,5 @@ namespace Chaos
 		Chaos::FPBDRigidClusteredParticleHandle* ClusterParticleHandle;
 		FVec3 Scale;
 		bool bIsAnchored;
-		// Whether or not to use the particle's existing child to parent transform instead of trying to recompute it.
-		bool bUseExistingChildToParent;
 	};
 }
diff --git a/Engine/Source/Runtime/Experimental/Chaos/Public/Chaos/ClusterUnionManager.h b/Engine/Source/Runtime/Experimental/Chaos/Public/Chaos/ClusterUnionManager.h
--- a/Engine/Source/Runtime/Experimental/Chaos/Public/Chaos/ClusterUnionManager.h
+++ b/Engine/Source/Runtime/Experimental/Chaos/Public/Chaos/ClusterUnionManager.h
@@ -109,8 +109,8 @@ namespace Chaos
 		// An extension to FindClusterUnionIndexFromParticle to check whether or not the given particle is the cluster union particle itself.
 		bool IsClusterUnionParticle(FPBDRigidClusteredParticleHandle* Particle);
 
-		// Changes the ChildToParent of a number of particles in a cluster union.
-		void UpdateClusterUnionParticlesChildToParent(FClusterUnionIndex Index, const TArray<FPBDRigidParticleHandle*>& Particles, const TArray<FTransform>& ChildToParent);
+		// Changes the ChildToParent of a number of particles in a cluster union. bLock will prevent other functions from overriding the ChildToParent (aside from another call to UpdateClusterUnionParticlesChildToParent).
+		void UpdateClusterUnionParticlesChildToParent(FClusterUnionIndex Index, const TArray<FPBDRigidParticleHandle*>& Particles, const TArray<FTransform>& ChildToParent, bool bLock);
 
 		// Update the cluster union's properties after its set of particle changes.
 		void UpdateAllClusterUnionProperties(FClusterUnion& ClusterUnion, bool bRecomputeMassOrientation);
diff --git a/Engine/Source/Runtime/Experimental/Chaos/Public/Chaos/PBDRigidClusteredParticles.h b/Engine/Source/Runtime/Experimental/Chaos/Public/Chaos/PBDRigidClusteredParticles.h
--- a/Engine/Source/Runtime/Experimental/Chaos/Public/Chaos/PBDRigidClusteredParticles.h
+++ b/Engine/Source/Runtime/Experimental/Chaos/Public/Chaos/PBDRigidClusteredParticles.h
@@ -88,12 +88,17 @@ public:
 	bool GetUnbreakable() const { return Flags.bUnbreakable; }
 	void SetUnbreakable(bool bSet) { Flags.bUnbreakable = bSet; }
 
+	// Set to true if we want to soft lock the particle's ChildToParent transform. This prevents its from being updated in Chaos::UpdateGeometry.
+	bool GetChildToParentLocked() const { return Flags.bChildToParentLocked; }
+	void SetChildToParentLocked(bool bSet) { Flags.bChildToParentLocked = bSet; }
+
 private:
 	struct FFlags
 	{
 		FStorage bInternalCluster : 1;
 		FStorage bAnchored : 1;
 		FStorage bUnbreakable : 1;
+		FStorage bChildToParentLocked : 1;
 
 		// Add new properties above this line
 		// Change FStorage typedef if we exceed the max bits
diff --git a/Engine/Source/Runtime/Experimental/Chaos/Public/Chaos/ParticleHandle.h b/Engine/Source/Runtime/Experimental/Chaos/Public/Chaos/ParticleHandle.h
--- a/Engine/Source/Runtime/Experimental/Chaos/Public/Chaos/ParticleHandle.h
+++ b/Engine/Source/Runtime/Experimental/Chaos/Public/Chaos/ParticleHandle.h
@@ -1323,6 +1323,9 @@ public:
 	bool Unbreakable() const { return PBDRigidClusteredParticles->RigidClusteredFlags(ParticleIdx).GetUnbreakable(); }
 	void SetUnbreakable(bool bValue) { PBDRigidClusteredParticles->RigidClusteredFlags(ParticleIdx).SetUnbreakable(bValue); }
 
+	bool IsChildToParentLocked() const { return PBDRigidClusteredParticles->RigidClusteredFlags(ParticleIdx).GetChildToParentLocked(); }
+	void SetChildToParentLocked(bool bValue) { PBDRigidClusteredParticles->RigidClusteredFlags(ParticleIdx).SetChildToParentLocked(bValue); }
+
 	const TPBDRigidClusteredParticleHandleImp<T, d, true>* Handle() const { return PBDRigidClusteredParticles->Handle(ParticleIdx); }
 	TPBDRigidClusteredParticleHandleImp<T, d, true>* Handle() { return PBDRigidClusteredParticles->Handle(ParticleIdx); }
 
diff --git a/Engine/Source/Runtime/Experimental/Chaos/Public/PhysicsProxy/ClusterUnionPhysicsProxy.h b/Engine/Source/Runtime/Experimental/Chaos/Public/PhysicsProxy/ClusterUnionPhysicsProxy.h
--- a/Engine/Source/Runtime/Experimental/Chaos/Public/PhysicsProxy/ClusterUnionPhysicsProxy.h
+++ b/Engine/Source/Runtime/Experimental/Chaos/Public/PhysicsProxy/ClusterUnionPhysicsProxy.h
@@ -91,8 +91,8 @@ namespace Chaos
 		void SetXR_External(const FVector& X, const FQuat& R);
 		void SetLinearVelocity_External(const FVector& V);
 		void SetAngularVelocity_External(const FVector& W);
-		void SetChildToParent_External(FPhysicsObjectHandle Child, const FTransform& RelativeTransform);
-		void BulkSetChildToParent_External(const TArray<FPhysicsObjectHandle>& Objects, const TArray<FTransform>& Transforms);
+		void SetChildToParent_External(FPhysicsObjectHandle Child, const FTransform& RelativeTransform, bool bLock);
+		void BulkSetChildToParent_External(const TArray<FPhysicsObjectHandle>& Objects, const TArray<FTransform>& Transforms, bool bLock);
 
 		//
 		// These functions take care of marshaling data back and forth between the game thread
