Chaos - remove collisions for destroyed particles from the Active Collision List by nulling it out and handle all the knock-ons

diff --git a/Engine/Source/Runtime/Experimental/Chaos/Private/Chaos/ChaosDebugDraw.cpp b/Engine/Source/Runtime/Experimental/Chaos/Private/Chaos/ChaosDebugDraw.cpp
--- a/Engine/Source/Runtime/Experimental/Chaos/Private/Chaos/ChaosDebugDraw.cpp
+++ b/Engine/Source/Runtime/Experimental/Chaos/Private/Chaos/ChaosDebugDraw.cpp
@@ -1693,6 +1693,11 @@ namespace Chaos
 			{
 				for (const Chaos::FPBDCollisionConstraintHandle * ConstraintHandle : Collisions.GetConstConstraintHandles())
 				{
+					if (ConstraintHandle == nullptr)
+					{
+						continue;
+					}
+
 					TVec2<const FGeometryParticleHandle*> ConstrainedParticles = ConstraintHandle->GetConstrainedParticles();
 					if ((ConstrainedParticles[0] == Particle) || (ConstrainedParticles[1] == Particle))
 					{
diff --git a/Engine/Source/Runtime/Experimental/Chaos/Private/Chaos/Collision/CollisionConstraintAllocator.cpp b/Engine/Source/Runtime/Experimental/Chaos/Private/Chaos/Collision/CollisionConstraintAllocator.cpp
--- a/Engine/Source/Runtime/Experimental/Chaos/Private/Chaos/Collision/CollisionConstraintAllocator.cpp
+++ b/Engine/Source/Runtime/Experimental/Chaos/Private/Chaos/Collision/CollisionConstraintAllocator.cpp
@@ -101,5 +101,55 @@ namespace Chaos
 			}
 		}
 
+		void FCollisionConstraintAllocator::RemoveActiveConstraint(FPBDCollisionConstraint& Constraint)
+		{
+			FPBDCollisionConstraintContainerCookie& Cookie = Constraint.GetContainerCookie();
+
+			// ConstraintIndex is only valid for one frame, so make sure we
+			// were actually activated during the most recent tick
+			if (Cookie.LastUsedEpoch != CurrentEpoch)
+			{
+				return;
+			}
+
+			// Remove from the active list
+			if (Cookie.ConstraintIndex != INDEX_NONE)
+			{
+				check(ActiveConstraints[Cookie.ConstraintIndex] == &Constraint);
+				ActiveConstraints[Cookie.ConstraintIndex] = nullptr;
+				Cookie.ConstraintIndex = INDEX_NONE;
+			}
+
+			// Remove from the active CCD list
+			if (Cookie.CCDConstraintIndex != INDEX_NONE)
+			{
+				check(ActiveCCDConstraints[Cookie.CCDConstraintIndex] == &Constraint);
+				ActiveCCDConstraints[Cookie.CCDConstraintIndex] = nullptr;
+				Cookie.CCDConstraintIndex = INDEX_NONE;
+			}
+		}
+
+		void FCollisionConstraintAllocator::RemoveParticle(FGeometryParticleHandle* Particle)
+		{
+			// Removal not supported during the (parallel) collision detection phase
+			check(!bInCollisionDetectionPhase);
+
+			// Loop over all particle pairs involving this particle and tell each MidPhase 
+			// that one of its particles is gone. It will get pruned at the next collision detection phase.
+			// Also remove its collisions from the Active lists.
+			Particle->ParticleCollisions().VisitMidPhases([this, Particle](FParticlePairMidPhase& MidPhase)
+				{
+					MidPhase.VisitCollisions([this](FPBDCollisionConstraint& Constraint)
+						{
+							RemoveActiveConstraint(Constraint);
+							return ECollisionVisitorResult::Continue;
+						});
+
+					MidPhase.DetachParticle(Particle);
+
+					return ECollisionVisitorResult::Continue;
+				});
+		}
+
 	}	// namespace Private
 }	// namespace Chaos
\ No newline at end of file
diff --git a/Engine/Source/Runtime/Experimental/Chaos/Private/Chaos/Collision/PBDCollisionContainerSolver.cpp b/Engine/Source/Runtime/Experimental/Chaos/Private/Chaos/Collision/PBDCollisionContainerSolver.cpp
--- a/Engine/Source/Runtime/Experimental/Chaos/Private/Chaos/Collision/PBDCollisionContainerSolver.cpp
+++ b/Engine/Source/Runtime/Experimental/Chaos/Private/Chaos/Collision/PBDCollisionContainerSolver.cpp
@@ -335,6 +335,7 @@ namespace Chaos
 
 		for (FPBDCollisionConstraintHandle* ConstraintHandle : ConstraintContainer.GetConstraintHandles())
 		{
+			check(ConstraintHandle != nullptr);
 			FPBDCollisionConstraint& Constraint = ConstraintHandle->GetContact();
 
 			AddConstraint(Constraint);
diff --git a/Engine/Source/Runtime/Experimental/Chaos/Private/Chaos/Collision/PBDCollisionContainerSolverJacobi.cpp b/Engine/Source/Runtime/Experimental/Chaos/Private/Chaos/Collision/PBDCollisionContainerSolverJacobi.cpp
--- a/Engine/Source/Runtime/Experimental/Chaos/Private/Chaos/Collision/PBDCollisionContainerSolverJacobi.cpp
+++ b/Engine/Source/Runtime/Experimental/Chaos/Private/Chaos/Collision/PBDCollisionContainerSolverJacobi.cpp
@@ -297,6 +297,7 @@ namespace Chaos
 
 			for (FPBDCollisionConstraintHandle* ConstraintHandle : ConstraintContainer.GetConstraintHandles())
 			{
+				check(ConstraintHandle != nullptr);
 				FPBDCollisionConstraint& Constraint = ConstraintHandle->GetContact();
 
 				AddConstraint(Constraint);
diff --git a/Engine/Source/Runtime/Experimental/Chaos/Private/Chaos/Collision/ParticlePairMidPhase.cpp b/Engine/Source/Runtime/Experimental/Chaos/Private/Chaos/Collision/ParticlePairMidPhase.cpp
--- a/Engine/Source/Runtime/Experimental/Chaos/Private/Chaos/Collision/ParticlePairMidPhase.cpp
+++ b/Engine/Source/Runtime/Experimental/Chaos/Private/Chaos/Collision/ParticlePairMidPhase.cpp
@@ -533,8 +533,11 @@ namespace Chaos
 #endif
 
 			// We need to refresh the epoch so that the constraint state will be used as the previous
-			// state if the pair is still colliding in the next tick
+			// state if the pair is still colliding in the next tick. However, we are not in the active 
+			// collisions array so reset the indices.
 			Constraint->GetContainerCookie().LastUsedEpoch = CurrentEpoch;
+			Constraint->GetContainerCookie().ConstraintIndex = INDEX_NONE;
+			Constraint->GetContainerCookie().CCDConstraintIndex = INDEX_NONE;
 			LastUsedEpoch = CurrentEpoch;
 
 			// We have skipped collision detection for this particle because it was asleep, so we need to update the transforms...
@@ -771,6 +774,8 @@ namespace Chaos
 			if (Constraint->GetContainerCookie().LastUsedEpoch >= SleepEpoch)
 			{
 				Constraint->GetContainerCookie().LastUsedEpoch = CurrentEpoch;
+				Constraint->GetContainerCookie().ConstraintIndex = INDEX_NONE;
+				Constraint->GetContainerCookie().CCDConstraintIndex = INDEX_NONE;
 			}
 		}
 	}
@@ -1205,6 +1210,8 @@ namespace Chaos
 			if (Constraint->GetContainerCookie().LastUsedEpoch >= SleepEpoch)
 			{
 				Constraint->GetContainerCookie().LastUsedEpoch = CurrentEpoch;
+				Constraint->GetContainerCookie().ConstraintIndex = INDEX_NONE;
+				Constraint->GetContainerCookie().CCDConstraintIndex = INDEX_NONE;
 			}
 		}
 	}
diff --git a/Engine/Source/Runtime/Experimental/Chaos/Private/Chaos/PBDCollisionConstraints.cpp b/Engine/Source/Runtime/Experimental/Chaos/Private/Chaos/PBDCollisionConstraints.cpp
--- a/Engine/Source/Runtime/Experimental/Chaos/Private/Chaos/PBDCollisionConstraints.cpp
+++ b/Engine/Source/Runtime/Experimental/Chaos/Private/Chaos/PBDCollisionConstraints.cpp
@@ -529,7 +529,7 @@ namespace Chaos
 		// for a contact which no longer is occurring due to resolution of another constraint.
 		for (FPBDCollisionConstraint* Contact : GetConstraints())
 		{
-			if (Contact->IsProbe())
+			if ((Contact != nullptr) && Contact->IsProbe())
 			{
 				const FGeometryParticleHandle* Particle0 = Contact->GetParticle0();
 				const FGeometryParticleHandle* Particle1 = Contact->GetParticle1();
@@ -573,6 +573,7 @@ namespace Chaos
 	{
 		if (GetConstraints().Num() > 0)
 		{
+			// NOTE: at this point, we will not have any nullptrs in the constraint handles
 			TArrayView<FPBDCollisionConstraint* const> ConstraintHandles = GetConstraintHandles();
 			FCollisionContactModifier Modifier(ConstraintHandles, Dt);
 
diff --git a/Engine/Source/Runtime/Experimental/Chaos/Private/Chaos/PBDRigidClustering.cpp b/Engine/Source/Runtime/Experimental/Chaos/Private/Chaos/PBDRigidClustering.cpp
--- a/Engine/Source/Runtime/Experimental/Chaos/Private/Chaos/PBDRigidClustering.cpp
+++ b/Engine/Source/Runtime/Experimental/Chaos/Private/Chaos/PBDRigidClustering.cpp
@@ -1405,6 +1405,11 @@ namespace Chaos
 
 		for (const Chaos::FPBDCollisionConstraintHandle* ContactHandle : CollisionRule.GetConstConstraintHandles())
 		{
+			if (ContactHandle == nullptr)
+			{
+				continue;
+			}
+
 			TVector<const FGeometryParticleHandle*, 2> ConstrainedParticles = ContactHandle->GetConstrainedParticles();
 			
 			// make sure we only compute things if one of the two particle is clustered
diff --git a/Engine/Source/Runtime/Experimental/Chaos/Private/Chaos/PBDRigidsEvolutionGBF.cpp b/Engine/Source/Runtime/Experimental/Chaos/Private/Chaos/PBDRigidsEvolutionGBF.cpp
--- a/Engine/Source/Runtime/Experimental/Chaos/Private/Chaos/PBDRigidsEvolutionGBF.cpp
+++ b/Engine/Source/Runtime/Experimental/Chaos/Private/Chaos/PBDRigidsEvolutionGBF.cpp
@@ -700,9 +700,12 @@ void FPBDRigidsEvolutionGBF::TestModeResetCollisions()
 {
 	for (FPBDCollisionConstraintHandle* Collision : CollisionConstraints.GetConstraintHandles())
 	{
-		Collision->GetContact().ResetManifold();
-		Collision->GetContact().ResetModifications();
-		Collision->GetContact().GetGJKWarmStartData().Reset();
+		if (Collision != nullptr)
+		{
+			Collision->GetContact().ResetManifold();
+			Collision->GetContact().ResetModifications();
+			Collision->GetContact().GetGJKWarmStartData().Reset();
+		}
 	}
 }
 
@@ -710,7 +713,10 @@ void FPBDRigidsEvolutionGBF::ResetCollisions()
 {
 	for (FPBDCollisionConstraintHandle* Collision : CollisionConstraints.GetConstraintHandles())
 	{
-		Collision->GetContact().GetGJKWarmStartData().Reset();
+		if (Collision != nullptr)
+		{
+			Collision->GetContact().GetGJKWarmStartData().Reset();
+		}
 	}
 }
 
diff --git a/Engine/Source/Runtime/Experimental/Chaos/Private/EventDefaults.cpp b/Engine/Source/Runtime/Experimental/Chaos/Private/EventDefaults.cpp
--- a/Engine/Source/Runtime/Experimental/Chaos/Private/EventDefaults.cpp
+++ b/Engine/Source/Runtime/Experimental/Chaos/Private/EventDefaults.cpp
@@ -122,15 +122,20 @@ namespace Chaos
 					{
 						for (int32 Index = StartRangeIndex; Index < EndRangeIndex; ++Index)
 						{
-							const Chaos::FPBDCollisionConstraintHandle* ContactHandle = CollisionHandles[Index];
 							ValidArray[Index] = false;
+
+							const Chaos::FPBDCollisionConstraintHandle* ContactHandle = CollisionHandles[Index];
+							if (ContactHandle == nullptr)
+							{
+								continue;
+							}
+
 							if (NumValidCollisions >= CollisionRule.NumConstraints())
 							{
 								break;
 							}
 
 							const FPBDCollisionConstraint& Constraint = ContactHandle->GetContact();
-
 							if (ensure(!Constraint.AccumulatedImpulse.ContainsNaN() && FMath::IsFinite(Constraint.GetPhi())))
 							{
 	
diff --git a/Engine/Source/Runtime/Experimental/Chaos/Private/PBDRigidsSolver.cpp b/Engine/Source/Runtime/Experimental/Chaos/Private/PBDRigidsSolver.cpp
--- a/Engine/Source/Runtime/Experimental/Chaos/Private/PBDRigidsSolver.cpp
+++ b/Engine/Source/Runtime/Experimental/Chaos/Private/PBDRigidsSolver.cpp
@@ -1935,6 +1935,11 @@ CSV_CUSTOM_STAT(PhysicsCounters, Name, Value, ECsvCustomStatOp::Set);
 		int32 NumUpdatedManifoldPoints = 0;
 		for (const FPBDCollisionConstraintHandle* Collision : GetEvolution()->GetCollisionConstraints().GetConstraints())
 		{
+			if (Collision == nullptr)
+			{
+				continue;
+			}
+
 			const FPBDCollisionConstraint& Contact = Collision->GetContact();
 			if (Contact.IsEnabled())
 			{
diff --git a/Engine/Source/Runtime/Experimental/Chaos/Public/Chaos/Collision/CollisionConstraintAllocator.h b/Engine/Source/Runtime/Experimental/Chaos/Public/Chaos/Collision/CollisionConstraintAllocator.h
--- a/Engine/Source/Runtime/Experimental/Chaos/Public/Chaos/Collision/CollisionConstraintAllocator.h
+++ b/Engine/Source/Runtime/Experimental/Chaos/Public/Chaos/Collision/CollisionConstraintAllocator.h
@@ -88,6 +88,8 @@ namespace Chaos
 				if (Cookie.LastUsedEpoch != CurrentEpoch)
 				{
 					Cookie.LastUsedEpoch = CurrentEpoch;
+					Cookie.ConstraintIndex = INDEX_NONE;
+					Cookie.CCDConstraintIndex = INDEX_NONE;
 
 					NewActiveConstraints.Push(Constraint);
 
@@ -467,20 +469,7 @@ namespace Chaos
 			 * @brief Destroy all collision and caches involving the particle
 			 * Called when a particle is destroyed or disabled (not sleeping).
 			*/
-			void RemoveParticle(FGeometryParticleHandle* Particle)
-			{
-				// We will be removing collisions, and don't want to have to prune the queues
-				check(!bInCollisionDetectionPhase);
-
-				// Loop over all particle pairs involving this particle.
-				// Tell each Particle Pair MidPhase that one of its particles is gone. 
-				// It will get pruned at the next collision detection phase.
-				Particle->ParticleCollisions().VisitMidPhases([Particle](FParticlePairMidPhase& MidPhase)
-					{
-						MidPhase.DetachParticle(Particle);
-						return ECollisionVisitorResult::Continue;
-					});
-			}
+			void RemoveParticle(FGeometryParticleHandle* Particle);
 
 			/**
 			 * @brief Iterate over all collisions (read-only), including sleeping ones
@@ -593,6 +582,8 @@ namespace Chaos
 				Cookie.LastUsedEpoch = CurrentEpoch;
 			}
 
+			void RemoveActiveConstraint(FPBDCollisionConstraint& Constraint);
+
 			// The container that owns the allocator (only needed because new constraints need to know)
 			FPBDCollisionConstraints* CollisionContainer;
 
diff --git a/Engine/Source/Runtime/Experimental/Chaos/Public/Chaos/Island/IslandManager.h b/Engine/Source/Runtime/Experimental/Chaos/Public/Chaos/Island/IslandManager.h
--- a/Engine/Source/Runtime/Experimental/Chaos/Public/Chaos/Island/IslandManager.h
+++ b/Engine/Source/Runtime/Experimental/Chaos/Public/Chaos/Island/IslandManager.h
@@ -109,7 +109,7 @@ public:
 		const int32 ContainerId = ConstraintContainer.GetContainerId();
 		for (FConstraintHandle* ConstraintHandle : ConstraintContainer.GetConstraintHandles())
 		{
-			if (ConstraintHandle->IsEnabled())
+			if ((ConstraintHandle != nullptr) && ConstraintHandle->IsEnabled())
 			{
 				AddConstraint(ContainerId, ConstraintHandle, ConstraintHandle->GetConstrainedParticles());
 			}
