ensure that the cluster union particle is dirtied/disabled as necessary when we handle the deferred removal updates

diff --git a/Engine/Source/Runtime/Experimental/Chaos/Private/Chaos/ClusterUnionManager.cpp b/Engine/Source/Runtime/Experimental/Chaos/Private/Chaos/ClusterUnionManager.cpp
--- a/Engine/Source/Runtime/Experimental/Chaos/Private/Chaos/ClusterUnionManager.cpp
+++ b/Engine/Source/Runtime/Experimental/Chaos/Private/Chaos/ClusterUnionManager.cpp
@@ -345,6 +345,30 @@ namespace Chaos
 		MEvolution.GetParticles().MarkTransientDirtyParticle(Cluster->InternalCluster);
 	}
 
+	DECLARE_CYCLE_STAT(TEXT("FClusterUnionManager::PostRemovalClusterUnionUpdate"), STAT_PostRemovalClusterUnionUpdate, STATGROUP_Chaos);
+	void FClusterUnionManager::PostRemovalClusterUnionUpdate(FClusterUnion& Union)
+	{
+		SCOPE_CYCLE_COUNTER(STAT_PostRemovalClusterUnionUpdate);
+		IPhysicsProxyBase* OldProxy = Union.InternalCluster->PhysicsProxy();
+		UpdateAllClusterUnionProperties(Union, false);
+
+		// Removing a particle should have no bearing on the proxy of the cluster.
+		// This gets changed because we go through an internal initialization route when we update the cluster union particle's properties.
+		Union.InternalCluster->SetPhysicsProxy(OldProxy);
+
+		if (!Union.ChildParticles.IsEmpty())
+		{
+			MEvolution.DirtyParticle(*Union.InternalCluster);
+		}
+		else
+		{
+			// Note that if we have 0 child particles, our implicit object union will have an invalid bounding box.
+			// We must eject from the acceleration structure otherwise we risk cashes.
+			MEvolution.DisableParticle(Union.InternalCluster);
+		}
+		MEvolution.GetParticles().MarkTransientDirtyParticle(Union.InternalCluster);
+	}
+
 	DECLARE_CYCLE_STAT(TEXT("FClusterUnionManager::HandleDeferredClusterUnionUpdateProperties"), STAT_HandleDeferredClusterUnionUpdateProperties, STATGROUP_Chaos);
 	void FClusterUnionManager::HandleDeferredClusterUnionUpdateProperties()
 	{
@@ -358,7 +382,7 @@ namespace Chaos
 		{
 			if (FClusterUnion* Union = FindClusterUnion(Index))
 			{
-				UpdateAllClusterUnionProperties(*Union, false);
+				PostRemovalClusterUnionUpdate(*Union);
 			}
 		}
 
@@ -420,28 +444,12 @@ namespace Chaos
 		switch (UpdateClusterPropertiesTiming)
 		{
 		case EClusterUnionOperationTiming::Immediate:
-			UpdateAllClusterUnionProperties(*Cluster, false);
+			PostRemovalClusterUnionUpdate(*Cluster);
 			break;
 		case EClusterUnionOperationTiming::Defer:
 			DeferredClusterUnionsForUpdateProperties.Add(ClusterIndex);
 			break;
-		}
-
-		// Removing a particle should have no bearing on the proxy of the cluster.
-		// This gets changed because we go through an internal initialization route when we update the cluster union particle's properties.
-		Cluster->InternalCluster->SetPhysicsProxy(OldProxy);
-
-		if (!Cluster->ChildParticles.IsEmpty())
-		{
-			MEvolution.DirtyParticle(*Cluster->InternalCluster);
-		}
-		else
-		{
-			// Note that if we have 0 child particles, our implicit object union will have an invalid bounding box.
-			// We must eject from the acceleration structure otherwise we risk cashes.
-			MEvolution.DisableParticle(Cluster->InternalCluster);
-		}
-		MEvolution.GetParticles().MarkTransientDirtyParticle(Cluster->InternalCluster);
+		}		
 	}
 
 	DECLARE_CYCLE_STAT(TEXT("FClusterUnionManager::HandleRemoveOperationWithClusterLookup"), STAT_HandleRemoveOperationWithClusterLookup, STATGROUP_Chaos);
diff --git a/Engine/Source/Runtime/Experimental/Chaos/Public/Chaos/ClusterUnionManager.h b/Engine/Source/Runtime/Experimental/Chaos/Public/Chaos/ClusterUnionManager.h
--- a/Engine/Source/Runtime/Experimental/Chaos/Public/Chaos/ClusterUnionManager.h
+++ b/Engine/Source/Runtime/Experimental/Chaos/Public/Chaos/ClusterUnionManager.h
@@ -167,6 +167,9 @@ namespace Chaos
 
 		// Forcefully recreate the shared geometry on a cluster. Potentially expensive so ideally should be used rarely.
 		TSharedPtr<FImplicitObject, ESPMode::ThreadSafe> ForceRecreateClusterUnionSharedGeometry(const FClusterUnion& Union);
+
+		// Handles updating the cluster union after it has particles removed from it. Could be called from either immediate or deferred removal.
+		void PostRemovalClusterUnionUpdate(FClusterUnion& Union);
 	};
 
 }
\ No newline at end of file
