Add support for ISMPool reverse culling flag. This allows us to potentially create two ISM components for left and right handedness if needed.
Add free list for ISMPool ISM component slots so that we can recycle them when the number of instances for a component becomes zero.

diff --git a/Engine/Source/Runtime/Experimental/GeometryCollectionEngine/Private/GeometryCollection/GeometryCollectionComponent.cpp b/Engine/Source/Runtime/Experimental/GeometryCollectionEngine/Private/GeometryCollection/GeometryCollectionComponent.cpp
--- a/Engine/Source/Runtime/Experimental/GeometryCollectionEngine/Private/GeometryCollection/GeometryCollectionComponent.cpp
+++ b/Engine/Source/Runtime/Experimental/GeometryCollectionEngine/Private/GeometryCollection/GeometryCollectionComponent.cpp
@@ -4341,16 +4341,16 @@ void UGeometryCollectionComponent::RefreshCustomRenderer()
 				const int32 RootIndex = GetRootIndex();
 				const bool bIsBroken = DynamicCollection ? !DynamicCollection->Active[RootIndex] : false;
 
-				RendererInterface->UpdateState(*RestCollection, bIsBroken);
+				RendererInterface->UpdateState(*RestCollection, ComponentTransform, bIsBroken);
 
 				if (bIsBroken)
 				{
-					RendererInterface->UpdateTransforms(*RestCollection, ComponentTransform, GlobalMatrices);
+					RendererInterface->UpdateTransforms(*RestCollection, GlobalMatrices);
 				}
 				else
 				{
 					const FTransform RootTransform = GlobalMatrices.IsValidIndex(RootIndex) ? FTransform(GlobalMatrices[RootIndex]) : FTransform::Identity;
-					RendererInterface->UpdateRootTransform(*RestCollection, ComponentTransform, RootTransform);
+					RendererInterface->UpdateRootTransform(*RestCollection, RootTransform);
 				}
 			}
 		}
diff --git a/Engine/Source/Runtime/Experimental/GeometryCollectionEngine/Private/GeometryCollection/GeometryCollectionISMPoolComponent.cpp b/Engine/Source/Runtime/Experimental/GeometryCollectionEngine/Private/GeometryCollection/GeometryCollectionISMPoolComponent.cpp
--- a/Engine/Source/Runtime/Experimental/GeometryCollectionEngine/Private/GeometryCollection/GeometryCollectionISMPoolComponent.cpp
+++ b/Engine/Source/Runtime/Experimental/GeometryCollectionEngine/Private/GeometryCollection/GeometryCollectionISMPoolComponent.cpp
@@ -43,8 +43,10 @@ void FGeometryCollectionMeshGroup::RemoveAllMeshes(FGeometryCollectionISMPool& I
 	Meshes.Empty();
 }
 
-FGeometryCollectionISM::FGeometryCollectionISM(AActor* OwmingActor, const FGeometryCollectionStaticMeshInstance& MeshInstance)
+FGeometryCollectionISM::FGeometryCollectionISM(AActor* OwmingActor, const FGeometryCollectionStaticMeshInstance& InMeshInstance)
 {
+	MeshInstance = InMeshInstance;
+
 	check(MeshInstance.StaticMesh);
 	check(OwmingActor);
 
@@ -74,6 +76,7 @@ FGeometryCollectionISM::FGeometryCollectionISM(AActor* OwmingActor, const FGeome
 	}
 
 	ISMC->NumCustomDataFloats = MeshInstance.Desc.NumCustomDataFloats;
+	ISMC->SetReverseCulling(MeshInstance.Desc.bReverseCulling);
 	ISMC->SetMobility(MeshInstance.Desc.bIsStaticMobility ? EComponentMobility::Static : EComponentMobility::Stationary);
 	ISMC->SetCullDistances(MeshInstance.Desc.StartCullDistance, MeshInstance.Desc.EndCullDistance);
 	ISMC->SetCastShadow(MeshInstance.Desc.bAffectShadow);
@@ -130,7 +133,16 @@ FGeometryCollectionMeshInfo FGeometryCollectionISMPool::AddISM(UGeometryCollecti
 	FISMIndex* ISMIndex = MeshToISMIndex.Find(MeshInstance);
 	if (!ISMIndex)
 	{
-		Info.ISMIndex = ISMs.Emplace(OwningComponent->GetOwner(), MeshInstance);
+		if (FreeList.Num())
+		{
+			Info.ISMIndex = FreeList.Last();
+			FreeList.RemoveAt(FreeList.Num() - 1);
+			ISMs[Info.ISMIndex] = FGeometryCollectionISM(OwningComponent->GetOwner(), MeshInstance);
+		}
+		else
+		{
+			Info.ISMIndex = ISMs.Emplace(OwningComponent->GetOwner(), MeshInstance);
+		}
 		MeshToISMIndex.Add(MeshInstance, Info.ISMIndex);
 		ISMComponentToISMIndex.Add(ISMs[Info.ISMIndex].ISMComponent, Info.ISMIndex);
 	}
@@ -186,6 +198,21 @@ void FGeometryCollectionISMPool::RemoveISM(const FGeometryCollectionMeshInfo& Me
 		const FInstanceGroups::FInstanceGroupRange& InstanceGroup = ISM.InstanceGroups.GetGroup(MeshInfo.InstanceGroupIndex);
 		ISM.ISMComponent->RemoveInstances(InstanceGroup.InstanceIdToIndex);
 		ISM.InstanceGroups.RemoveGroup(MeshInfo.InstanceGroupIndex);
+	
+		if (ISM.InstanceGroups.IsEmpty() && ISM.ISMComponent->PerInstanceSMData.Num() == 0)
+		{
+			// Remove component and push this ISM slot to the free list.
+			// todo: profile if it is better to push component into a free pool and recycle it.
+			ISM.ISMComponent->GetOwner()->RemoveInstanceComponent(ISM.ISMComponent);
+			ISM.ISMComponent->UnregisterComponent();
+			ISM.ISMComponent->DestroyComponent();
+			
+			MeshToISMIndex.Remove(ISM.MeshInstance);
+			ISMComponentToISMIndex.Remove(ISM.ISMComponent);
+			FreeList.Add(MeshInfo.ISMIndex);
+
+			ISM.ISMComponent = nullptr;
+		}
 	}
 }
 
@@ -217,6 +244,7 @@ void FGeometryCollectionISMPool::Clear()
 {
 	MeshToISMIndex.Reset();
 	ISMComponentToISMIndex.Reset();
+	FreeList.Reset();
 	if (ISMs.Num() > 0)
 	{
 		if (AActor* OwningActor = ISMs[0].ISMComponent->GetOwner())
@@ -284,19 +312,3 @@ bool UGeometryCollectionISMPoolComponent::BatchUpdateInstancesTransforms(FMeshGr
 	UE_LOG(LogChaos, Warning, TEXT("UGeometryCollectionISMPoolComponent : Trying to update instance with mesh group (%d) that not exists"), MeshGroupId);
 	return false;
 }
-
-UInstancedStaticMeshComponent* UGeometryCollectionISMPoolComponent::GetISMForMeshId(FMeshGroupId MeshGroupId, FMeshId MeshId) const
-{
-	if (FGeometryCollectionMeshGroup const* MeshGroup = MeshGroups.Find(MeshGroupId))
-	{
-		if (MeshGroup->MeshInfos.IsValidIndex(MeshId))
-		{
-			int32 ISMIndex = MeshGroup->MeshInfos[MeshId].ISMIndex;
-			if (Pool.ISMs.IsValidIndex(ISMIndex))
-			{
-				return Pool.ISMs[ISMIndex].ISMComponent;
-			}
-		}
-	}
-	return nullptr;
-}
diff --git a/Engine/Source/Runtime/Experimental/GeometryCollectionEngine/Private/GeometryCollection/GeometryCollectionISMPoolRenderer.cpp b/Engine/Source/Runtime/Experimental/GeometryCollectionEngine/Private/GeometryCollection/GeometryCollectionISMPoolRenderer.cpp
--- a/Engine/Source/Runtime/Experimental/GeometryCollectionEngine/Private/GeometryCollection/GeometryCollectionISMPoolRenderer.cpp
+++ b/Engine/Source/Runtime/Experimental/GeometryCollectionEngine/Private/GeometryCollection/GeometryCollectionISMPoolRenderer.cpp
@@ -25,8 +25,10 @@ void UGeometryCollectionISMPoolRenderer::OnUnregisterGeometryCollection()
 	ISMPoolActor = nullptr;
 }
 
-void UGeometryCollectionISMPoolRenderer::UpdateState(UGeometryCollection const& InGeometryCollection, bool bIsBroken)
+void UGeometryCollectionISMPoolRenderer::UpdateState(UGeometryCollection const& InGeometryCollection, FTransform const& InComponentTransform, bool bIsBroken)
 {
+	ComponentTransform = InComponentTransform;
+
 	if (!bIsBroken && MergedMeshGroup.GroupIndex == INDEX_NONE)
 	{
 		// Remove broken primitives.
@@ -46,14 +48,14 @@ void UGeometryCollectionISMPoolRenderer::UpdateState(UGeometryCollection const&
 	}
 }
 
-void UGeometryCollectionISMPoolRenderer::UpdateRootTransform(UGeometryCollection const& InGeometryCollection, FTransform const& InBaseTransform, FTransform const& InRootTransform)
+void UGeometryCollectionISMPoolRenderer::UpdateRootTransform(UGeometryCollection const& InGeometryCollection, FTransform const& InRootTransform)
 {
-	UpdateMergedMeshTransforms(InRootTransform * InBaseTransform);
+	UpdateMergedMeshTransforms(InRootTransform * ComponentTransform);
 }
 
-void UGeometryCollectionISMPoolRenderer::UpdateTransforms(UGeometryCollection const& InGeometryCollection, FTransform const& InBaseTransform, TArrayView<const FMatrix> InMatrices)
+void UGeometryCollectionISMPoolRenderer::UpdateTransforms(UGeometryCollection const& InGeometryCollection, TArrayView<const FMatrix> InMatrices)
 {
-	UpdateInstanceTransforms(InGeometryCollection, InBaseTransform, InMatrices);
+	UpdateInstanceTransforms(InGeometryCollection, ComponentTransform, InMatrices);
 }
 
 UGeometryCollectionISMPoolComponent* UGeometryCollectionISMPoolRenderer::GetOrCreateISMPoolComponent()
diff --git a/Engine/Source/Runtime/Experimental/GeometryCollectionEngine/Private/GeometryCollection/GeometryCollectionISMPoolRenderer.h b/Engine/Source/Runtime/Experimental/GeometryCollectionEngine/Private/GeometryCollection/GeometryCollectionISMPoolRenderer.h
--- a/Engine/Source/Runtime/Experimental/GeometryCollectionEngine/Private/GeometryCollection/GeometryCollectionISMPoolRenderer.h
+++ b/Engine/Source/Runtime/Experimental/GeometryCollectionEngine/Private/GeometryCollection/GeometryCollectionISMPoolRenderer.h
@@ -19,9 +19,9 @@ public:
 	//~ Begin IGeometryCollectionExternalRenderInterface Interface.
 	virtual void OnRegisterGeometryCollection(UGeometryCollectionComponent const& InComponent) override;
 	virtual void OnUnregisterGeometryCollection() override;
-	virtual void UpdateState(UGeometryCollection const& InGeometryCollection, bool bInIsBroken) override;
-	virtual void UpdateRootTransform(UGeometryCollection const& InGeometryCollection, FTransform const& InBaseTransform, FTransform const& InRootTransform) override;
-	virtual void UpdateTransforms(UGeometryCollection const& InGeometryCollection, FTransform const& InBaseTransform, TArrayView<const FMatrix> InMatrices) override;
+	virtual void UpdateState(UGeometryCollection const& InGeometryCollection, FTransform const& InComponentTransform, bool bInIsBroken) override;
+	virtual void UpdateRootTransform(UGeometryCollection const& InGeometryCollection, FTransform const& InRootTransform) override;
+	virtual void UpdateTransforms(UGeometryCollection const& InGeometryCollection, TArrayView<const FMatrix> InMatrices) override;
 	//~ End IGeometryCollectionExternalRenderInterface Interface.
 
 	/** Description for a group of meshes that are added/updated together. */
@@ -36,6 +36,9 @@ protected:
 	UPROPERTY(Transient)
 	TObjectPtr<AGeometryCollectionISMPoolActor> ISMPoolActor;
 
+	/** Cached component transform. */
+	FTransform ComponentTransform = FTransform::Identity;
+
 	/** ISM pool groups per rendering element type. */
 	FISMPoolGroup MergedMeshGroup;
 	FISMPoolGroup InstancesGroup;
diff --git a/Engine/Source/Runtime/Experimental/GeometryCollectionEngine/Public/GeometryCollection/GeometryCollectionExternalRenderInterface.h b/Engine/Source/Runtime/Experimental/GeometryCollectionEngine/Public/GeometryCollection/GeometryCollectionExternalRenderInterface.h
--- a/Engine/Source/Runtime/Experimental/GeometryCollectionEngine/Public/GeometryCollection/GeometryCollectionExternalRenderInterface.h
+++ b/Engine/Source/Runtime/Experimental/GeometryCollectionEngine/Public/GeometryCollection/GeometryCollectionExternalRenderInterface.h
@@ -22,7 +22,7 @@ class GEOMETRYCOLLECTIONENGINE_API IGeometryCollectionExternalRenderInterface
 public:
 	virtual void OnRegisterGeometryCollection(UGeometryCollectionComponent const& InComponent) = 0;
 	virtual void OnUnregisterGeometryCollection() = 0;
-	virtual void UpdateState(UGeometryCollection const& InGeometryCollection, bool bInIsBroken) = 0;
-	virtual void UpdateRootTransform(UGeometryCollection const& InGeometryCollection, FTransform const& InBaseTransform, FTransform const& InRootTransform) = 0;
-	virtual void UpdateTransforms(UGeometryCollection const& InGeometryCollection, FTransform const& InBaseTransform, TArrayView<const FMatrix> InMatrices) = 0;
+	virtual void UpdateState(UGeometryCollection const& InGeometryCollection, FTransform const& InComponentTransform, bool bInIsBroken) = 0;
+	virtual void UpdateRootTransform(UGeometryCollection const& InGeometryCollection, FTransform const& InRootTransform) = 0;
+	virtual void UpdateTransforms(UGeometryCollection const& InGeometryCollection, TArrayView<const FMatrix> InMatrices) = 0;
 };
diff --git a/Engine/Source/Runtime/Experimental/GeometryCollectionEngine/Public/GeometryCollection/GeometryCollectionISMPoolComponent.h b/Engine/Source/Runtime/Experimental/GeometryCollectionEngine/Public/GeometryCollection/GeometryCollectionISMPoolComponent.h
--- a/Engine/Source/Runtime/Experimental/GeometryCollectionEngine/Public/GeometryCollection/GeometryCollectionISMPoolComponent.h
+++ b/Engine/Source/Runtime/Experimental/GeometryCollectionEngine/Public/GeometryCollection/GeometryCollectionISMPoolComponent.h
@@ -53,6 +53,11 @@ public:
 		TArray<int32> InstanceIdToIndex;
 	};
 
+	bool IsEmpty() const
+	{
+		return InstancesCount == 0;
+	}
+
 	FInstanceGroupId AddGroup(int32 Count)
 	{
 		const int32 StartIndex = InstancesCount;
@@ -108,6 +113,7 @@ private:
 struct FISMComponentDescription
 {
 	bool bUseHISM = false;
+	bool bReverseCulling = false;
 	bool bIsStaticMobility = false;
 	bool bAffectShadow = true;
 	bool bAffectDistanceFieldLighting = false;
@@ -120,6 +126,7 @@ struct FISMComponentDescription
 	bool operator==(const FISMComponentDescription& Other) const
 	{
 		return bUseHISM == Other.bUseHISM &&
+			bReverseCulling == Other.bReverseCulling &&
 			bIsStaticMobility == Other.bIsStaticMobility &&
 			bAffectShadow == Other.bAffectShadow &&
 			bAffectDistanceFieldLighting == Other.bAffectDistanceFieldLighting &&
@@ -133,7 +140,7 @@ struct FISMComponentDescription
 
 FORCEINLINE uint32 GetTypeHash(const FISMComponentDescription& Desc)
 {
-	const uint32 PackedBools = (Desc.bUseHISM ? 1 : 0) | (Desc.bIsStaticMobility ? 2 : 0) | (Desc.bAffectShadow ? 4 : 0) | (Desc.bAffectDistanceFieldLighting ? 8 : 0) | (Desc.bAffectDynamicIndirectLighting ? 16 : 0);
+	const uint32 PackedBools = (Desc.bUseHISM ? 1 : 0) | (Desc.bReverseCulling ? 2 : 0) | (Desc.bIsStaticMobility ? 4 : 0) | (Desc.bAffectShadow ? 8 : 0) | (Desc.bAffectDistanceFieldLighting ? 16 : 0) | (Desc.bAffectDynamicIndirectLighting ? 32 : 0);
 	uint32 Hash = HashCombine(GetTypeHash(PackedBools), GetTypeHash(Desc.NumCustomDataFloats));
 	Hash = HashCombine(Hash, GetTypeHash(Desc.StartCullDistance));
 	Hash = HashCombine(Hash, GetTypeHash(Desc.EndCullDistance));
@@ -210,11 +217,12 @@ struct FGeometryCollectionMeshGroup
 
 struct FGeometryCollectionISM
 {
-	FGeometryCollectionISM(AActor* OwmingActor, const FGeometryCollectionStaticMeshInstance& MeshInstance);
+	FGeometryCollectionISM(AActor* OwmingActor, const FGeometryCollectionStaticMeshInstance& InMeshInstance);
 
 	int32 AddInstanceGroup(int32 InstanceCount, TArrayView<const float> CustomDataFloats);
 
 	TObjectPtr<UInstancedStaticMeshComponent> ISMComponent;
+	FGeometryCollectionStaticMeshInstance MeshInstance;
 	FInstanceGroups InstanceGroups;
 };
 
@@ -235,8 +243,7 @@ struct FGeometryCollectionISMPool
 	TMap<FGeometryCollectionStaticMeshInstance, FISMIndex> MeshToISMIndex;
 	TMap<UInstancedStaticMeshComponent*, FISMIndex> ISMComponentToISMIndex;
 	TArray<FGeometryCollectionISM> ISMs;
-	
-	// Todo : since ISMs index cannot change, we'll need a free list 
+	TArray<int32> FreeList;
 };
 
 
@@ -274,9 +281,6 @@ public:
 	/** Add a static mesh for a nmesh group */
 	bool BatchUpdateInstancesTransforms(FMeshGroupId MeshGroupId, FMeshId MeshId, int32 StartInstanceIndex, const TArray<FTransform>& NewInstancesTransforms, bool bWorldSpace = false, bool bMarkRenderStateDirty = false, bool bTeleport = false);
 
-	/** Get the ISM component related to a mesh id. */
-	UInstancedStaticMeshComponent* GetISMForMeshId(FMeshGroupId MeshGroupId, FMeshId MeshId) const;
-	
 	/** Instance Index updated on the InstancedStaticMeshComponent which might need to be handled by the pool instance groups */
 	void OnISMInstanceIndexUpdated(UInstancedStaticMeshComponent* InComponent, TArrayView<const FInstancedStaticMeshDelegates::FInstanceIndexUpdateData> InIndexUpdates)
 	{
