cluster unions will force the ChildToParent to match the world position of the particle on the client when replicating the cluster union's XR before ChildToParent is set

diff --git a/Engine/Source/Runtime/Experimental/Chaos/Private/PhysicsProxy/ClusterUnionPhysicsProxy.cpp b/Engine/Source/Runtime/Experimental/Chaos/Private/PhysicsProxy/ClusterUnionPhysicsProxy.cpp
--- a/Engine/Source/Runtime/Experimental/Chaos/Private/PhysicsProxy/ClusterUnionPhysicsProxy.cpp
+++ b/Engine/Source/Runtime/Experimental/Chaos/Private/PhysicsProxy/ClusterUnionPhysicsProxy.cpp
@@ -257,8 +257,40 @@ namespace Chaos
 
 		if (const FParticlePositionRotation* NewXR = ParticleData.FindClusterXR(Manager, DataIdx))
 		{
+			// If we change the cluster union's location and the child particles are not soft locked with a given child to parent,
+			// our goal is to maintain the particle's location rather than its child to parent.
+			FTransform NewTransform{ NewXR->R(), NewXR->X() };
+
+			TArray<FPBDRigidParticleHandle*> ParticlesToUpdate;
+			TArray<FTransform> NewChildToParent;
+
+			FClusterUnionManager& ClusterUnionManager = Evolution.GetRigidClustering().GetClusterUnionManager();
+			if (FClusterUnion* Union = ClusterUnionManager.FindClusterUnion(ClusterUnionIndex); Union && !Union->bNeedsXRInitialization)
+			{
+				for (FPBDRigidParticleHandle* Particle : Union->ChildParticles)
+				{
+					if (FPBDRigidClusteredParticleHandle* ClusterParticle = Particle->CastToClustered())
+					{
+						if (!ClusterParticle->IsChildToParentLocked())
+						{
+							ParticlesToUpdate.Add(Particle);
+
+							const FRigidTransform3 ChildWorldTM(Particle->X(), Particle->R());
+							NewChildToParent.Add(ChildWorldTM.GetRelativeTransform(NewTransform));
+						}
+					}
+				}
+			}
+
 			Evolution.SetParticleTransform(Particle_Internal, NewXR->X(), NewXR->R(), true);
 
+			if (!ParticlesToUpdate.IsEmpty() && !NewChildToParent.IsEmpty())
+			{
+				// Do not lock this child to parent as it's not coming from the server. It's just to smooth out the XR replication on the client
+				// while we still have not yet received the ChildToParent yet.
+				ClusterUnionManager.UpdateClusterUnionParticlesChildToParent(ClusterUnionIndex, ParticlesToUpdate, NewChildToParent, false);
+			}
+
 			// Particle needs to actually be marked dirty! Less so to pull state from the PT back to the GT for the cluster union
 			// but this is primarily for proxies within the cluster union. They (i.e. GCs) need to see this change to the particle's transform.
 			Evolution.GetParticles().MarkTransientDirtyParticle(Particle_Internal);
