updates to the child to parent on cluster unions should always happen on the server and client and should be overridden when computing the cluster union properties

diff --git a/Engine/Source/Runtime/Experimental/Chaos/Private/Chaos/ClusterUnionManager.cpp b/Engine/Source/Runtime/Experimental/Chaos/Private/Chaos/ClusterUnionManager.cpp
--- a/Engine/Source/Runtime/Experimental/Chaos/Private/Chaos/ClusterUnionManager.cpp
+++ b/Engine/Source/Runtime/Experimental/Chaos/Private/Chaos/ClusterUnionManager.cpp
@@ -589,6 +589,17 @@ namespace Chaos
 
 		if (FClusterUnion* ClusterUnion = FindClusterUnion(Index))
 		{
+			// We need to keep track of all the particles we touched here. We need this because UpdateClusterUnionParticlesChildToParent
+			// is an *authoritative* update on the ChildToParent of the particle. However, UpdateAllClusterUnionProperties in certain cases
+			// may try to recompute the ChildToParent using the position of the particle. To counteract this, we will force the ChildToParent
+			// to be *temporarily* locked for the duration of UpdateAllClusterUnionProperties. However, unless bLock is true, we will restore the
+			// lock state of the particle to what it was previously.
+			TArray<TPair<FPBDRigidClusteredParticleHandle*, bool>> DirtyParticleLockStates;
+			if (!bLock)
+			{
+				DirtyParticleLockStates.Reserve(Particles.Num());
+			}
+
 			for (int32 ParticleIndex = 0; ParticleIndex < Particles.Num() && ParticleIndex < ChildToParent.Num(); ++ParticleIndex)
 			{
 				FPBDRigidParticleHandle* Particle = Particles[ParticleIndex];
@@ -604,15 +615,25 @@ namespace Chaos
 					{
 						ChildHandle->SetChildToParent(ChildToParent[ParticleIndex]);
 
-						if (bLock)
+						if (!bLock)
 						{
-							ChildHandle->SetChildToParentLocked(bLock);
-						}						
+							DirtyParticleLockStates.Add({ ChildHandle, ChildHandle->IsChildToParentLocked() });
+						}
+						ChildHandle->SetChildToParentLocked(true);
 					}
 				}
 			}
 
 			UpdateAllClusterUnionProperties(*ClusterUnion, false);
+
+			if (!bLock)
+			{
+				for (TPair<FPBDRigidClusteredParticleHandle*, bool>& Pair : DirtyParticleLockStates)
+				{
+					Pair.Key->SetChildToParentLocked(Pair.Value);
+				}
+			}
+
 			MEvolution.GetParticles().MarkTransientDirtyParticle(ClusterUnion->InternalCluster);
 			MEvolution.DirtyParticle(*ClusterUnion->InternalCluster);
 		}
