Geometry collection : fix replication where client may try to remove a particle that still has a grand parent

diff --git a/Engine/Source/Runtime/Experimental/Chaos/Private/Chaos/PBDRigidClustering.cpp b/Engine/Source/Runtime/Experimental/Chaos/Private/Chaos/PBDRigidClustering.cpp
--- a/Engine/Source/Runtime/Experimental/Chaos/Private/Chaos/PBDRigidClustering.cpp
+++ b/Engine/Source/Runtime/Experimental/Chaos/Private/Chaos/PBDRigidClustering.cpp
@@ -1037,6 +1037,30 @@ namespace Chaos
 		return ActivatedBodies;
 	}
 
+	DECLARE_CYCLE_STAT(TEXT("TPBDRigidClustering<>::ReleaseChildrenParticleAndParents"), ReleaseChildrenParticleAndParents, STATGROUP_Chaos);
+	void FRigidClustering::ForceReleaseChildParticleAndParents(FPBDRigidClusteredParticleHandle* ChildClusteredParticle, bool bTriggerBreakEvents)
+	{
+		SCOPE_CYCLE_COUNTER(ReleaseChildrenParticleAndParents);
+
+		if (ChildClusteredParticle)
+		{
+			// make sure we set unbreakable to false so that the children can be released
+			ChildClusteredParticle->SetUnbreakable(false);
+			if (ChildClusteredParticle->Disabled())
+			{
+				// first release any parent if any
+				if (FPBDRigidClusteredParticleHandle* ParentCluster = ChildClusteredParticle->Parent())
+				{
+					// we need now to force parents to break
+					ForceReleaseChildParticleAndParents(ParentCluster, bTriggerBreakEvents);
+
+					SetExternalStrain(ChildClusteredParticle, TNumericLimits<FRealSingle>::Max());
+					ReleaseClusterParticles(ParentCluster, bTriggerBreakEvents);
+				}
+			}
+		}
+	}
+
 
 	DECLARE_CYCLE_STAT(TEXT("TPBDRigidClustering<>::AdvanceClustering"), STAT_AdvanceClustering, STATGROUP_Chaos);
 	DECLARE_CYCLE_STAT(TEXT("TPBDRigidClustering<>::Update Impulse from Strain"), STAT_UpdateImpulseStrain, STATGROUP_Chaos);
diff --git a/Engine/Source/Runtime/Experimental/Chaos/Public/Chaos/PBDRigidClustering.h b/Engine/Source/Runtime/Experimental/Chaos/Public/Chaos/PBDRigidClustering.h
--- a/Engine/Source/Runtime/Experimental/Chaos/Public/Chaos/PBDRigidClustering.h
+++ b/Engine/Source/Runtime/Experimental/Chaos/Public/Chaos/PBDRigidClustering.h
@@ -151,6 +151,12 @@ public:
 	TSet<FPBDRigidParticleHandle*> ReleaseClusterParticles(
 		TArray<FPBDRigidParticleHandle*> ChildrenParticles, bool bTriggerBreakEvents = false);
 
+	/** 
+	* Force Release a particle at any level by making sure their parent particles are also release if necessary 
+	* @Warning this will force all particles including the ones in the parent chain to be made breakable 
+	*/
+	void ForceReleaseChildParticleAndParents(FPBDRigidClusteredParticleHandle* ChildClusteredParticle, bool bTriggerBreakEvents);
+
 	/*
 	*  DestroyClusterParticle
 	*    Disable the cluster particle and remove from all internal clustering
diff --git a/Engine/Source/Runtime/Experimental/GeometryCollectionEngine/Private/GeometryCollection/GeometryCollectionComponent.cpp b/Engine/Source/Runtime/Experimental/GeometryCollectionEngine/Private/GeometryCollection/GeometryCollectionComponent.cpp
--- a/Engine/Source/Runtime/Experimental/GeometryCollectionEngine/Private/GeometryCollection/GeometryCollectionComponent.cpp
+++ b/Engine/Source/Runtime/Experimental/GeometryCollectionEngine/Private/GeometryCollection/GeometryCollectionComponent.cpp
@@ -1629,7 +1629,8 @@ void UGeometryCollectionComponent::UpdateRepData()
 						//a one off so record it
 						ensureMsgf(TransformGroupIdx >= 0, TEXT("Non-internal cluster should always have a group index"));
 						ensureMsgf(TransformGroupIdx < TNumericLimits<uint16>::Max(), TEXT("Trying to replicate GC with more than 65k pieces. We assumed uint16 would suffice"));
-	
+						ensureMsgf(PhysicsProxy->GetParticles().IsValidIndex(TransformGroupIdx) && PhysicsProxy->GetParticles()[TransformGroupIdx] != nullptr, TEXT("Invalid particle index being replicated - Contact physics team"));
+
 						// Because we cull ClustersToRep with abandoned level, we must make sure we don't add duplicates to one off activated.
 						// TODO: avoid search for entry for perf
 						// TODO: once we support deep fracture we should be able to remove one offs clusters that are now disabled, reducing the amount to be replicated
@@ -1913,16 +1914,6 @@ bool UGeometryCollectionComponent::ProcessRepData(const float DeltaTime, const f
 		{
 			if (FPBDRigidClusteredParticleHandle* ClusterParticle = OneOff->CastToClustered())
 			{
-				// If there's a parent cluster particle we need to release them first.
-				// This is generally an indication that something desynced between the client and server though...maybe something needs to be done
-				// to ensure internal clusters stay in sync.
-				if (FPBDRigidClusteredParticleHandle* ParentParticle = ClusterParticle->Parent())
-				{
-					// server authoritative particles are unbreakable, we need to set them breakable again 
-					ParentParticle->SetUnbreakable(false);
-					RigidClustering.ReleaseClusterParticles(TArray<FPBDRigidParticleHandle*>{ ParentParticle }, true);
-				}
-
 				// Set initial velocities if not hard snapping
 				if (!bHardSnap)
 				{
@@ -1932,7 +1923,9 @@ bool UGeometryCollectionComponent::ProcessRepData(const float DeltaTime, const f
 					OneOff->SetW(ActivatedCluster.InitialAngularVelocity);
 				}
 
-				RigidClustering.ReleaseClusterParticles(TArray<FPBDRigidParticleHandle*>{ ClusterParticle }, true);
+				// OneOffActivated is not garanteed to be sorted by level, so we may need to release particles before their parent get released
+				// so we make sure the full parent chain is properly released as well 
+				RigidClustering.ForceReleaseChildParticleAndParents(ClusterParticle, /* bTriggerBreakEvents */true);
 			}
 		}
 	}
