change FReadPhysicsObjectInterface::ShapeSweep to pass bComputeMTD=true to better match internal chaos SQVisitor behavior

diff --git a/Engine/Source/Runtime/Experimental/Chaos/Public/Chaos/PhysicsObjectInterface.h b/Engine/Source/Runtime/Experimental/Chaos/Public/Chaos/PhysicsObjectInterface.h
--- a/Engine/Source/Runtime/Experimental/Chaos/Public/Chaos/PhysicsObjectInterface.h
+++ b/Engine/Source/Runtime/Experimental/Chaos/Public/Chaos/PhysicsObjectInterface.h
@@ -289,7 +289,8 @@ namespace Chaos
 								StartTM,
 								[Shape, &WorldTM, &Dir, DeltaMag, &Distance, &WorldPosition, &WorldNormal, &FaceIdx, &FaceNormal](const auto& Downcast, const auto& FullTransformB)
 								{
-									return Chaos::SweepQuery(*Shape->GetGeometry(), WorldTM, Downcast, FullTransformB, Dir, DeltaMag, Distance, WorldPosition, WorldNormal, FaceIdx, FaceNormal, 0.f, false);
+									// Set bComputeMTD to true to better match internal SQ Visitor behavior.
+									return Chaos::SweepQuery(*Shape->GetGeometry(), WorldTM, Downcast, FullTransformB, Dir, DeltaMag, Distance, WorldPosition, WorldNormal, FaceIdx, FaceNormal, 0.f, true);
 								}
 							);
 
