don't rely on bFirstUpdate in geometry collection replication to determine whether or not to replicate a cluster - use the particle leve instead since that's what we're actually checking for

- using bFirstUpdate is unreliable when GC's are added into cluster unions since cluster unions will disable replication immediately

diff --git a/Engine/Source/Runtime/Experimental/GeometryCollectionEngine/Private/GeometryCollection/GeometryCollectionComponent.cpp b/Engine/Source/Runtime/Experimental/GeometryCollectionEngine/Private/GeometryCollection/GeometryCollectionComponent.cpp
--- a/Engine/Source/Runtime/Experimental/GeometryCollectionEngine/Private/GeometryCollection/GeometryCollectionComponent.cpp
+++ b/Engine/Source/Runtime/Experimental/GeometryCollectionEngine/Private/GeometryCollection/GeometryCollectionComponent.cpp
@@ -1624,7 +1624,7 @@ void UGeometryCollectionComponent::UpdateRepData()
 						bClustersChanged = true;
 					}
 
-					if(Root->InternalCluster() == false && bFirstUpdate == false)	//if bFirstUpdate it must be that these are the initial roots of the GC. These did not break off so no need to replicate
+					if(Root->InternalCluster() == false && Level > 0)
 					{
 						//a one off so record it
 						ensureMsgf(TransformGroupIdx >= 0, TEXT("Non-internal cluster should always have a group index"));
