cluster union manager must do the proxy restore when we remove particles not when we update the cluster union properties

diff --git a/Engine/Source/Runtime/Experimental/Chaos/Private/Chaos/ClusterUnionManager.cpp b/Engine/Source/Runtime/Experimental/Chaos/Private/Chaos/ClusterUnionManager.cpp
--- a/Engine/Source/Runtime/Experimental/Chaos/Private/Chaos/ClusterUnionManager.cpp
+++ b/Engine/Source/Runtime/Experimental/Chaos/Private/Chaos/ClusterUnionManager.cpp
@@ -349,13 +349,8 @@ namespace Chaos
 	void FClusterUnionManager::PostRemovalClusterUnionUpdate(FClusterUnion& Union)
 	{
 		SCOPE_CYCLE_COUNTER(STAT_PostRemovalClusterUnionUpdate);
-		IPhysicsProxyBase* OldProxy = Union.InternalCluster->PhysicsProxy();
 		UpdateAllClusterUnionProperties(Union, false);
 
-		// Removing a particle should have no bearing on the proxy of the cluster.
-		// This gets changed because we go through an internal initialization route when we update the cluster union particle's properties.
-		Union.InternalCluster->SetPhysicsProxy(OldProxy);
-
 		if (!Union.ChildParticles.IsEmpty())
 		{
 			MEvolution.DirtyParticle(*Union.InternalCluster);
@@ -441,6 +436,10 @@ namespace Chaos
 
 		MClustering.RemoveParticlesFromCluster(Cluster->InternalCluster, Particles);
 
+		// Removing a particle should have no bearing on the proxy of the cluster.
+		// This gets changed because we go through an internal initialization route when we update the cluster union particle's properties.
+		Cluster->InternalCluster->SetPhysicsProxy(OldProxy);
+
 		switch (UpdateClusterPropertiesTiming)
 		{
 		case EClusterUnionOperationTiming::Immediate:
