Remove p.GeometryCollectionAssetForceStripRenderDataOnCook which existed for temporary data migration.

diff --git a/Engine/Source/Runtime/Experimental/GeometryCollectionEngine/Private/GeometryCollection/GeometryCollectionRenderData.cpp b/Engine/Source/Runtime/Experimental/GeometryCollectionEngine/Private/GeometryCollection/GeometryCollectionRenderData.cpp
--- a/Engine/Source/Runtime/Experimental/GeometryCollectionEngine/Private/GeometryCollection/GeometryCollectionRenderData.cpp
+++ b/Engine/Source/Runtime/Experimental/GeometryCollectionEngine/Private/GeometryCollection/GeometryCollectionRenderData.cpp
@@ -720,20 +720,12 @@ FGeometryCollectionRenderData::~FGeometryCollectionRenderData()
 	Fence.Wait();
 }
 
-// Temporary cvar driven option to force strip render data on cook if we are also stripping other GC data.
-bool GeometryCollectionAssetForceStripRenderDataOnCook = false;
-FAutoConsoleVariableRef CVarGeometryCollectionStripRenderData(
-	TEXT("p.GeometryCollectionAssetForceStripRenderDataOnCook"),
-	GeometryCollectionAssetForceStripRenderDataOnCook,
-	TEXT("Bypass the cook of render data when other geom collection data is also being stripped."));
-
 void FGeometryCollectionRenderData::Serialize(FArchive& Ar, UGeometryCollection& Owner)
 {
-	const bool bStripRenderData = Owner.bStripRenderDataOnCook || (GeometryCollectionAssetForceStripRenderDataOnCook && Owner.bStripOnCook);
-	if (bStripRenderData && Ar.IsCooking())
+	if (Owner.bStripRenderDataOnCook && Ar.IsCooking())
 	{
 		// Don't cook rendering data.
-		// This is used if we expect to use a non-native GC render path such as the ISM pool.
+		// This is used if we expect to use a custom GC render path such as the ISM pool.
 		bool bDummy = false;
 		Ar << bDummy; // bHasMeshData
 		Ar << bDummy; // bHasNaniteData
