Fix nullptr crash in ISM pool debug component when ISM components have been put on free list.

diff --git a/Engine/Source/Runtime/Experimental/GeometryCollectionEngine/Public/GeometryCollection/GeometryCollectionISMPoolDebugDrawComponent.cpp b/Engine/Source/Runtime/Experimental/GeometryCollectionEngine/Public/GeometryCollection/GeometryCollectionISMPoolDebugDrawComponent.cpp
--- a/Engine/Source/Runtime/Experimental/GeometryCollectionEngine/Public/GeometryCollection/GeometryCollectionISMPoolDebugDrawComponent.cpp
+++ b/Engine/Source/Runtime/Experimental/GeometryCollectionEngine/Public/GeometryCollection/GeometryCollectionISMPoolDebugDrawComponent.cpp
@@ -293,13 +293,16 @@ void UGeometryCollectionISMPoolDebugDrawComponent::GetOnScreenMessages(TMultiMap
 
 			for (FGeometryCollectionISM const& ISM : ISMPoolComponent->Pool.ISMs)
 			{
-				const bool bIsHISM = ISM.ISMComponent->IsA(UHierarchicalInstancedStaticMeshComponent::StaticClass());
-				NumISMs += bIsHISM ? 0 : 1;
-				NumHISMs += bIsHISM ? 1 : 0;
+				if (ISM.ISMComponent != nullptr)
+				{
+					const bool bIsHISM = ISM.ISMComponent->IsA(UHierarchicalInstancedStaticMeshComponent::StaticClass());
+					NumISMs += bIsHISM ? 0 : 1;
+					NumHISMs += bIsHISM ? 1 : 0;
 
-				const int32 NumInstances = ISM.ISMComponent->GetNumRenderInstances();
-				NumInstancesTotal += NumInstances;
-				NumInstancesPerISM.Add(NumInstances);
+					const int32 NumInstances = ISM.ISMComponent->GetNumRenderInstances();
+					NumInstancesTotal += NumInstances;
+					NumInstancesPerISM.Add(NumInstances);
+				}
 			}
 
 			if (NumInstancesTotal > 0)
