refactor out the physics object funtions related to overlap/sweep into PhysicsObjectCollisionInterface

diff --git a/Engine/Source/Runtime/Engine/Private/Components/PrimitiveComponent.cpp b/Engine/Source/Runtime/Engine/Private/Components/PrimitiveComponent.cpp
--- a/Engine/Source/Runtime/Engine/Private/Components/PrimitiveComponent.cpp
+++ b/Engine/Source/Runtime/Engine/Private/Components/PrimitiveComponent.cpp
@@ -7,6 +7,7 @@
 #include "Components/PrimitiveComponent.h"
 
 #include "Chaos/ChaosEngineInterface.h"
+#include "Chaos/PhysicsObjectCollisionInterface.h"
 #include "ChaosInterfaceWrapperCore.h"
 #include "Collision/CollisionConversions.h"
 #include "Engine/Level.h"
@@ -2936,8 +2937,9 @@ bool UPrimitiveComponent::LineTraceComponent(struct FHitResult& OutHit, const FV
 			}
 		);
 
+		Chaos::FPhysicsObjectCollisionInterface_External CollisionInterface{ Interface.GetInterface() };
 		ChaosInterface::FRaycastHit BestHit;
-		if (Interface->LineTrace(Objects, Start, End, Params.bTraceComplex, BestHit))
+		if (CollisionInterface.LineTrace(Objects, Start, End, Params.bTraceComplex, BestHit))
 		{
 			bHaveHit = true;
 
@@ -2988,8 +2990,17 @@ bool UPrimitiveComponent::SweepComponent(struct FHitResult& OutHit, const FVecto
 		}
 	);
 
+	Chaos::FPhysicsObjectCollisionInterface_External CollisionInterface{ Interface.GetInterface() };
 	ChaosInterface::FSweepHit BestHit;
-	if (Interface->ShapeSweep(Objects, ShapeAdapter.GetGeometry(), ShapeAdapter.GetGeomPose(Start), End, bTraceComplex, BestHit))
+
+	Chaos::FSweepParameters Params;
+	Params.bSweepComplex = bTraceComplex;
+
+	// TODO: Expose this even further via parameters in the primitive component.
+	// For now, having this be always true guarantees us identical behavior to tracing via the Chaos SQ
+	// since TSQTraits::GetHitFlags() will always have the MTD flag on.
+	Params.bComputeMTD = true;
+	if (CollisionInterface.ShapeSweep(Objects, ShapeAdapter.GetGeometry(), ShapeAdapter.GetGeomPose(Start), End, Params, BestHit))
 	{
 		FCollisionFilterData QueryFilter;
 		QueryFilter.Word1 = 0xFFFFF;
@@ -3051,11 +3062,12 @@ bool UPrimitiveComponent::ComponentOverlapComponentImpl(class UPrimitiveComponen
 		}
 	);
 
+	Chaos::FPhysicsObjectCollisionInterface_External CollisionInterface{ Interface.GetInterface() };
 	for (Chaos::FPhysicsObjectHandle InObject : InObjects)
 	{
 		for (Chaos::FPhysicsObjectHandle ThisObject : ThisObjects)
 		{
-			if (Interface->PhysicsObjectOverlap(InObject, FTransform::Identity, ThisObject, FTransform::Identity, Params.bTraceComplex))
+			if (CollisionInterface.PhysicsObjectOverlap(InObject, FTransform::Identity, ThisObject, FTransform::Identity, Params.bTraceComplex))
 			{
 				return true;
 			}
@@ -3085,12 +3097,13 @@ bool UPrimitiveComponent::ComponentOverlapComponentWithResultImpl(const class UP
 		}
 	);
 
+	Chaos::FPhysicsObjectCollisionInterface_External CollisionInterface{ Interface.GetInterface() };
 	for (Chaos::FPhysicsObjectHandle InObject : InObjects)
 	{
 		for (Chaos::FPhysicsObjectHandle ThisObject : ThisObjects)
 		{
 			TArray<ChaosInterface::FOverlapHit> OverlapHits;
-			if (Interface->PhysicsObjectOverlap(ThisObject, FTransform::Identity, InObject, InTransform, Params.bTraceComplex, OverlapHits))
+			if (CollisionInterface.PhysicsObjectOverlap(ThisObject, FTransform::Identity, InObject, InTransform, Params.bTraceComplex, OverlapHits))
 			{
 				TArray<FOverlapResult> Overlaps;
 
@@ -3132,8 +3145,9 @@ bool UPrimitiveComponent::OverlapComponentWithResult(const FVector& Pos, const F
 		}
 	);
 
+	Chaos::FPhysicsObjectCollisionInterface_External CollisionInterface{ Interface.GetInterface() };
 	TArray<ChaosInterface::FOverlapHit> OverlapHits;
-	if (Interface->ShapeOverlap(Objects, ShapeAdapter.GetGeometry(), ShapeAdapter.GetGeomPose(Pos), OverlapHits))
+	if (CollisionInterface.ShapeOverlap(Objects, ShapeAdapter.GetGeometry(), ShapeAdapter.GetGeomPose(Pos), OverlapHits))
 	{
 		TArray<FOverlapResult> Overlaps;
 
diff --git a/Engine/Source/Runtime/Experimental/Chaos/Private/Chaos/PhysicsObjectCollisionInterface.cpp b/Engine/Source/Runtime/Experimental/Chaos/Private/Chaos/PhysicsObjectCollisionInterface.cpp
new file mode 100644
--- /dev/null
+++ b/Engine/Source/Runtime/Experimental/Chaos/Private/Chaos/PhysicsObjectCollisionInterface.cpp
@@ -0,0 +1,164 @@
+// Copyright Epic Games, Inc. All Rights Reserved.
+#include "Chaos/PhysicsObjectCollisionInterface.h"
+
+namespace Chaos
+{
+	template<EThreadContext Id>
+	bool FPhysicsObjectCollisionInterface<Id>::PhysicsObjectOverlap(const FConstPhysicsObjectHandle ObjectA, const FTransform& InTransformA, const FConstPhysicsObjectHandle ObjectB, const FTransform& InTransformB, bool bTraceComplex)
+	{
+		TRACE_CPUPROFILER_EVENT_SCOPE(FPhysicsObjectCollisionInterface<Id>::PhysicsObjectOverlap);
+		return PairwiseShapeOverlapHelper(ObjectA, InTransformA, ObjectB, InTransformB, bTraceComplex, false, FVector::Zero(), [](const FShapeOverlapData&, const FShapeOverlapData&, const FMTDInfo&) { return false; });
+	}
+
+	template<EThreadContext Id>
+	bool FPhysicsObjectCollisionInterface<Id>::PhysicsObjectOverlapWithMTD(const FConstPhysicsObjectHandle ObjectA, const FTransform& InTransformA, const FConstPhysicsObjectHandle ObjectB, const FTransform& InTransformB, bool bTraceComplex, FMTDInfo& OutMTD)
+	{
+		TRACE_CPUPROFILER_EVENT_SCOPE(FPhysicsObjectCollisionInterface<Id>::PhysicsObjectOverlapWithMTD);
+		OutMTD.Penetration = 0.0;
+		return PairwiseShapeOverlapHelper(
+			ObjectA,
+			InTransformA,
+			ObjectB,
+			InTransformB,
+			bTraceComplex,
+			true,
+			FVector::Zero(),
+			[&OutMTD](const FShapeOverlapData&, const FShapeOverlapData&, const FMTDInfo& MTDInfo)
+			{
+				OutMTD.Penetration = FMath::Max(OutMTD.Penetration, MTDInfo.Penetration);
+				return true;
+			}
+		);
+	}
+
+	template<EThreadContext Id>
+	bool FPhysicsObjectCollisionInterface<Id>::PhysicsObjectOverlapWithAABB(const FConstPhysicsObjectHandle ObjectA, const FTransform& InTransformA, const FConstPhysicsObjectHandle ObjectB, const FTransform& InTransformB, bool bTraceComplex, const FVector& Tolerance, FBox& OutOverlap)
+	{
+		TRACE_CPUPROFILER_EVENT_SCOPE(FPhysicsObjectCollisionInterface<Id>::PhysicsObjectOverlapWithAABB);
+		OutOverlap = FBox{ EForceInit::ForceInitToZero };
+		return PairwiseShapeOverlapHelper(
+			ObjectA,
+			InTransformA,
+			ObjectB,
+			InTransformB,
+			bTraceComplex,
+			false,
+			Tolerance,
+			[&OutOverlap, &Tolerance](const FShapeOverlapData& ShapeA, const FShapeOverlapData& ShapeB, const FMTDInfo&)
+			{
+				const FAABB3 Intersection = ShapeA.BoundingBox.GetIntersection(ShapeB.BoundingBox);
+				OutOverlap += FBox{ Intersection.Min(), Intersection.Max() };
+				return true;
+			}
+		);
+	}
+
+	template<EThreadContext Id>
+	bool FPhysicsObjectCollisionInterface<Id>::PhysicsObjectOverlapWithAABBSize(const FConstPhysicsObjectHandle ObjectA, const FTransform& InTransformA, const FConstPhysicsObjectHandle ObjectB, const FTransform& InTransformB, bool bTraceComplex, const FVector& Tolerance, FVector& OutOverlapSize)
+	{
+		TRACE_CPUPROFILER_EVENT_SCOPE(FPhysicsObjectCollisionInterface<Id>::PhysicsObjectOverlapWithAABBSize);
+		OutOverlapSize = FVector::Zero();
+		return PairwiseShapeOverlapHelper(
+			ObjectA,
+			InTransformA,
+			ObjectB,
+			InTransformB,
+			bTraceComplex,
+			false,
+			Tolerance,
+			[&OutOverlapSize](const FShapeOverlapData& ShapeA, const FShapeOverlapData& ShapeB, const FMTDInfo&)
+			{
+				const FAABB3 Intersection = ShapeA.BoundingBox.GetIntersection(ShapeB.BoundingBox);
+				OutOverlapSize += Intersection.Extents();
+				return true;
+			}
+		);
+	}
+
+	template<EThreadContext Id>
+	bool FPhysicsObjectCollisionInterface<Id>::PairwiseShapeOverlapHelper(const FConstPhysicsObjectHandle ObjectA, const FTransform& InTransformA, const FConstPhysicsObjectHandle ObjectB, const FTransform& InTransformB, bool bTraceComplex, bool bComputeMTD, const FVector& Tolerance, const TFunction<bool(const FShapeOverlapData&, const FShapeOverlapData&, const FMTDInfo&)>& Lambda)
+	{
+		TArray<TThreadShapeInstance<Id>*> ShapesA = Interface.GetAllThreadShapes({ &ObjectA, 1 });
+		const FTransform TransformA = FTransform{ Interface.GetR(ObjectA), Interface.GetX(ObjectA) } *InTransformA;
+		const FBox BoxA = Interface.GetWorldBounds({ &ObjectA, 1 }).TransformBy(InTransformA);
+
+		TArray<TThreadShapeInstance<Id>*> ShapesB = Interface.GetAllThreadShapes({ &ObjectB, 1 });
+		const FTransform TransformB = FTransform{ Interface.GetR(ObjectB), Interface.GetX(ObjectB) } *InTransformB;
+		const FBox BoxB = Interface.GetWorldBounds({ &ObjectB, 1 }).TransformBy(InTransformB);
+
+		if (!BoxA.Intersect(BoxB))
+		{
+			return false;
+		}
+
+		bool bFoundOverlap = false;
+		for (TThreadShapeInstance<Id>* B : ShapesB)
+		{
+			if (!B)
+			{
+				continue;
+			}
+
+			const TSerializablePtr<FImplicitObject> GeomB = B->GetGeometry();
+			if (!GeomB || !GeomB->IsConvex())
+			{
+				continue;
+			}
+
+			const FAABB3 BoxShapeB = GeomB->CalculateTransformedBounds(TransformB).ShrinkSymmetrically(Tolerance);
+			// At this point on, this function should be mirror the Overlap_GeomInternal function in PhysInterface_Chaos.cpp.
+			// ShapeA is equivalent to InInstance and GeomB is equivalent to InGeom.
+
+			for (TThreadShapeInstance<Id>* A : ShapesA)
+			{
+				if (!A)
+				{
+					continue;
+				}
+
+				FCollisionFilterData ShapeFilter = A->GetQueryData();
+				const bool bShapeIsComplex = (ShapeFilter.Word3 & static_cast<uint8>(EFilterFlags::ComplexCollision)) != 0;
+				const bool bShapeIsSimple = (ShapeFilter.Word3 & static_cast<uint8>(EFilterFlags::SimpleCollision)) != 0;
+				const bool bShouldTrace = (bTraceComplex && bShapeIsComplex) || (!bTraceComplex && bShapeIsSimple);
+				if (!bShouldTrace)
+				{
+					continue;
+				}
+
+				const FAABB3 BoxShapeA = A->GetGeometry()->CalculateTransformedBounds(TransformA).ShrinkSymmetrically(Tolerance);
+				if (!BoxShapeA.Intersects(BoxShapeB))
+				{
+					continue;
+				}
+
+				Chaos::FMTDInfo TmpMTDInfo;
+				const bool bOverlap = Chaos::Utilities::CastHelper(
+					*GeomB,
+					TransformB,
+					[A, &TransformA, bComputeMTD, &TmpMTDInfo](const auto& Downcast, const auto& FullTransformB)
+					{
+						return Chaos::OverlapQuery(*A->GetGeometry(), TransformA, Downcast, FullTransformB, 0, bComputeMTD ? &TmpMTDInfo : nullptr);
+					}
+				);
+
+				if (bOverlap)
+				{
+					bFoundOverlap = true;
+
+					FShapeOverlapData OverlapDataA = { A, BoxShapeA };
+					FShapeOverlapData OverlapDataB = { B, BoxShapeB };
+
+					if (!Lambda(OverlapDataA, OverlapDataB, TmpMTDInfo))
+					{
+						return true;
+					}
+				}
+			}
+		}
+
+		return bFoundOverlap;
+	}
+
+	template class FPhysicsObjectCollisionInterface<EThreadContext::External>;
+	template class FPhysicsObjectCollisionInterface<EThreadContext::Internal>;
+}
\ No newline at end of file
diff --git a/Engine/Source/Runtime/Experimental/Chaos/Private/Chaos/PhysicsObjectInterface.cpp b/Engine/Source/Runtime/Experimental/Chaos/Private/Chaos/PhysicsObjectInterface.cpp
--- a/Engine/Source/Runtime/Experimental/Chaos/Private/Chaos/PhysicsObjectInterface.cpp
+++ b/Engine/Source/Runtime/Experimental/Chaos/Private/Chaos/PhysicsObjectInterface.cpp
@@ -3,6 +3,7 @@
 
 #include "Chaos/CastingUtilities.h"
 #include "Chaos/CollisionFilterData.h"
+#include "Chaos/PhysicsObjectCollisionInterface.h"
 #include "Chaos/PhysicsObjectInternal.h"
 #include "PhysicsProxy/SingleParticlePhysicsProxy.h"
 #include "Math/Transform.h"
@@ -306,16 +307,17 @@ namespace Chaos
 	bool FReadPhysicsObjectInterface<Id>::GetPhysicsObjectOverlap(const FConstPhysicsObjectHandle ObjectA, const FConstPhysicsObjectHandle ObjectB, bool bTraceComplex, Chaos::FOverlapInfo& OutOverlap)
 	{
 		TRACE_CPUPROFILER_EVENT_SCOPE(FReadPhysicsObjectInterface<Id>::GetPhysicsObjectOverlap);
+		FPhysicsObjectCollisionInterface Interface{ *this };
 		// This is slow and inefficient and hence deprecated.
 		bool bRetOverlap = false;
 		if (OutOverlap.MTD)
 		{
-			bRetOverlap |= PhysicsObjectOverlapWithMTD(ObjectA, FTransform::Identity, ObjectB, FTransform::Identity, bTraceComplex, *OutOverlap.MTD);
+			bRetOverlap |= Interface.PhysicsObjectOverlapWithMTD(ObjectA, FTransform::Identity, ObjectB, FTransform::Identity, bTraceComplex, *OutOverlap.MTD);
 		}
 
 		if (OutOverlap.AxisOverlap)
 		{
-			bRetOverlap |= PhysicsObjectOverlapWithAABB(ObjectA, FTransform::Identity, ObjectB, FTransform::Identity, bTraceComplex, FVector::Zero(), *OutOverlap.AxisOverlap);
+			bRetOverlap |= Interface.PhysicsObjectOverlapWithAABB(ObjectA, FTransform::Identity, ObjectB, FTransform::Identity, bTraceComplex, FVector::Zero(), *OutOverlap.AxisOverlap);
 		}
 
 		return bRetOverlap;
@@ -325,177 +327,22 @@ namespace Chaos
 	bool FReadPhysicsObjectInterface<Id>::GetPhysicsObjectOverlapWithTransform(const FConstPhysicsObjectHandle ObjectA, const FTransform& InTransformA, const FConstPhysicsObjectHandle ObjectB, const FTransform& InTransformB, bool bTraceComplex, Chaos::FOverlapInfo& OutOverlap)
 	{
 		TRACE_CPUPROFILER_EVENT_SCOPE(FReadPhysicsObjectInterface<Id>::GetPhysicsObjectOverlapWithTransform);
+		FPhysicsObjectCollisionInterface Interface{ *this };
 		// This is slow and inefficient and hence deprecated.
 		bool bRetOverlap = false;
 		if (OutOverlap.MTD)
 		{
-			bRetOverlap |= PhysicsObjectOverlapWithMTD(ObjectA, InTransformA, ObjectB, InTransformB, bTraceComplex, *OutOverlap.MTD);
+			bRetOverlap |= Interface.PhysicsObjectOverlapWithMTD(ObjectA, InTransformA, ObjectB, InTransformB, bTraceComplex, *OutOverlap.MTD);
 		}
 
 		if (OutOverlap.AxisOverlap)
 		{
-			bRetOverlap |= PhysicsObjectOverlapWithAABB(ObjectA, InTransformA, ObjectB, InTransformB, bTraceComplex, FVector::Zero(), *OutOverlap.AxisOverlap);
+			bRetOverlap |= Interface.PhysicsObjectOverlapWithAABB(ObjectA, InTransformA, ObjectB, InTransformB, bTraceComplex, FVector::Zero(), *OutOverlap.AxisOverlap);
 		}
 
 		return bRetOverlap;
 	}
 
-	template<EThreadContext Id>
-	bool FReadPhysicsObjectInterface<Id>::PhysicsObjectOverlap(const FConstPhysicsObjectHandle ObjectA, const FTransform& InTransformA, const FConstPhysicsObjectHandle ObjectB, const FTransform& InTransformB, bool bTraceComplex)
-	{
-		TRACE_CPUPROFILER_EVENT_SCOPE(FReadPhysicsObjectInterface<Id>::PhysicsObjectOverlap);
-		return PairwiseShapeOverlapHelper(ObjectA, InTransformA, ObjectB, InTransformB, bTraceComplex, false, FVector::Zero(), [](const FShapeOverlapData&, const FShapeOverlapData&, const FMTDInfo&) { return false; });
-	}
-
-	template<EThreadContext Id>
-	bool FReadPhysicsObjectInterface<Id>::PhysicsObjectOverlapWithMTD(const FConstPhysicsObjectHandle ObjectA, const FTransform& InTransformA, const FConstPhysicsObjectHandle ObjectB, const FTransform& InTransformB, bool bTraceComplex, FMTDInfo& OutMTD)
-	{
-		TRACE_CPUPROFILER_EVENT_SCOPE(FReadPhysicsObjectInterface<Id>::PhysicsObjectOverlapWithMTD);
-		OutMTD.Penetration = 0.0;
-		return PairwiseShapeOverlapHelper(
-			ObjectA,
-			InTransformA,
-			ObjectB,
-			InTransformB,
-			bTraceComplex,
-			true,
-			FVector::Zero(),
-			[&OutMTD](const FShapeOverlapData&, const FShapeOverlapData&, const FMTDInfo& MTDInfo)
-			{
-				OutMTD.Penetration = FMath::Max(OutMTD.Penetration, MTDInfo.Penetration);
-				return true;
-			}
-		);
-	}
-
-	template<EThreadContext Id>
-	bool FReadPhysicsObjectInterface<Id>::PhysicsObjectOverlapWithAABB(const FConstPhysicsObjectHandle ObjectA, const FTransform& InTransformA, const FConstPhysicsObjectHandle ObjectB, const FTransform& InTransformB, bool bTraceComplex, const FVector& Tolerance, FBox& OutOverlap)
-	{
-		TRACE_CPUPROFILER_EVENT_SCOPE(FReadPhysicsObjectInterface<Id>::PhysicsObjectOverlapWithAABB);
-		OutOverlap = FBox{ EForceInit::ForceInitToZero };
-		return PairwiseShapeOverlapHelper(
-			ObjectA,
-			InTransformA,
-			ObjectB,
-			InTransformB,
-			bTraceComplex,
-			false,
-			Tolerance,
-			[&OutOverlap, &Tolerance](const FShapeOverlapData& ShapeA, const FShapeOverlapData& ShapeB, const FMTDInfo&)
-			{
-				const FAABB3 Intersection = ShapeA.BoundingBox.GetIntersection(ShapeB.BoundingBox);
-				OutOverlap += FBox{ Intersection.Min(), Intersection.Max() };
-				return true;
-			}
-		);
-	}
-
-	template<EThreadContext Id>
-	bool FReadPhysicsObjectInterface<Id>::PhysicsObjectOverlapWithAABBSize(const FConstPhysicsObjectHandle ObjectA, const FTransform& InTransformA, const FConstPhysicsObjectHandle ObjectB, const FTransform& InTransformB, bool bTraceComplex, const FVector& Tolerance, FVector& OutOverlapSize)
-	{
-		TRACE_CPUPROFILER_EVENT_SCOPE(FReadPhysicsObjectInterface<Id>::PhysicsObjectOverlapWithAABBSize);
-		OutOverlapSize = FVector::Zero();
-		return PairwiseShapeOverlapHelper(
-			ObjectA,
-			InTransformA,
-			ObjectB,
-			InTransformB,
-			bTraceComplex,
-			false,
-			Tolerance,
-			[&OutOverlapSize](const FShapeOverlapData& ShapeA, const FShapeOverlapData& ShapeB, const FMTDInfo&)
-			{
-				const FAABB3 Intersection = ShapeA.BoundingBox.GetIntersection(ShapeB.BoundingBox);
-				OutOverlapSize += Intersection.Extents();
-				return true;
-			}
-		);
-	}
-
-	template<EThreadContext Id>
-	bool FReadPhysicsObjectInterface<Id>::PairwiseShapeOverlapHelper(const FConstPhysicsObjectHandle ObjectA, const FTransform& InTransformA, const FConstPhysicsObjectHandle ObjectB, const FTransform& InTransformB, bool bTraceComplex, bool bComputeMTD, const FVector& Tolerance, const TFunction<bool(const FShapeOverlapData&, const FShapeOverlapData&, const FMTDInfo&)>& Lambda)
-	{
-		TArray<TThreadShapeInstance<Id>*> ShapesA = GetAllThreadShapes({ &ObjectA, 1 });
-		const FTransform TransformA = FTransform{ GetR(ObjectA), GetX(ObjectA) } *InTransformA;
-		const FBox BoxA = GetWorldBounds({ &ObjectA, 1 }).TransformBy(InTransformA);
-
-		TArray<TThreadShapeInstance<Id>*> ShapesB = GetAllThreadShapes({ &ObjectB, 1 });
-		const FTransform TransformB = FTransform{ GetR(ObjectB), GetX(ObjectB) } *InTransformB;
-		const FBox BoxB = GetWorldBounds({ &ObjectB, 1 }).TransformBy(InTransformB);
-
-		if (!BoxA.Intersect(BoxB))
-		{
-			return false;
-		}
-
-		bool bFoundOverlap = false;
-		for (TThreadShapeInstance<Id>* B : ShapesB)
-		{
-			if (!B)
-			{
-				continue;
-			}
-
-			const TSerializablePtr<FImplicitObject> GeomB = B->GetGeometry();
-			if (!GeomB || !GeomB->IsConvex())
-			{
-				continue;
-			}
-
-			const FAABB3 BoxShapeB = GeomB->CalculateTransformedBounds(TransformB).ShrinkSymmetrically(Tolerance);
-			// At this point on, this function should be mirror the Overlap_GeomInternal function in PhysInterface_Chaos.cpp.
-			// ShapeA is equivalent to InInstance and GeomB is equivalent to InGeom.
-
-			for (TThreadShapeInstance<Id>* A : ShapesA)
-			{
-				if (!A)
-				{
-					continue;
-				}
-
-				FCollisionFilterData ShapeFilter = A->GetQueryData();
-				const bool bShapeIsComplex = (ShapeFilter.Word3 & static_cast<uint8>(EFilterFlags::ComplexCollision)) != 0;
-				const bool bShapeIsSimple = (ShapeFilter.Word3 & static_cast<uint8>(EFilterFlags::SimpleCollision)) != 0;
-				const bool bShouldTrace = (bTraceComplex && bShapeIsComplex) || (!bTraceComplex && bShapeIsSimple);
-				if (!bShouldTrace)
-				{
-					continue;
-				}
-
-				const FAABB3 BoxShapeA = A->GetGeometry()->CalculateTransformedBounds(TransformA).ShrinkSymmetrically(Tolerance);
-				if (!BoxShapeA.Intersects(BoxShapeB))
-				{
-					continue;
-				}
-
-				Chaos::FMTDInfo TmpMTDInfo;
-				const bool bOverlap = Chaos::Utilities::CastHelper(
-					*GeomB,
-					TransformB,
-					[A, &TransformA, bComputeMTD, &TmpMTDInfo](const auto& Downcast, const auto& FullTransformB)
-					{
-						return Chaos::OverlapQuery(*A->GetGeometry(), TransformA, Downcast, FullTransformB, 0, bComputeMTD ? &TmpMTDInfo : nullptr);
-					}
-				);
-
-				if (bOverlap)
-				{
-					bFoundOverlap = true;
-
-					FShapeOverlapData OverlapDataA = { A, BoxShapeA };
-					FShapeOverlapData OverlapDataB = { B, BoxShapeB };
-
-					if (!Lambda(OverlapDataA, OverlapDataB, TmpMTDInfo))
-					{
-						return true;
-					}
-				}
-			}
-		}
-
-		return bFoundOverlap;
-	}
-
 	template<EThreadContext Id>
 	bool FReadPhysicsObjectInterface<Id>::AreAllValid(TArrayView<const FConstPhysicsObjectHandle> InObjects)
 	{
diff --git a/Engine/Source/Runtime/Experimental/Chaos/Public/Chaos/PhysicsObjectCollisionInterface.h b/Engine/Source/Runtime/Experimental/Chaos/Public/Chaos/PhysicsObjectCollisionInterface.h
new file mode 100644
--- /dev/null
+++ b/Engine/Source/Runtime/Experimental/Chaos/Public/Chaos/PhysicsObjectCollisionInterface.h
@@ -0,0 +1,267 @@
+// Copyright Epic Games, Inc. All Rights Reserved.
+
+#pragma once
+
+#include "Chaos/Interface/SQTypes.h"
+#include "Chaos/GeometryQueries.h"
+#include "Chaos/PhysicsObject.h"
+#include "Chaos/PhysicsObjectInterface.h"
+#include "Framework/Threading.h"
+
+namespace Chaos
+{
+	struct CHAOS_API FSweepParameters
+	{
+		bool bSweepComplex = false;
+		bool bComputeMTD = false;
+	};
+
+	template<EThreadContext Id>
+	class CHAOS_API FPhysicsObjectCollisionInterface
+	{
+	public:
+		explicit FPhysicsObjectCollisionInterface(FReadPhysicsObjectInterface<Id>& InInterface) :
+			Interface(InInterface)
+		{}
+
+		// This function will not compute any overlap heuristic.
+		bool PhysicsObjectOverlap(const FConstPhysicsObjectHandle ObjectA, const FTransform& InTransformA, const FConstPhysicsObjectHandle ObjectB, const FTransform& InTransformB, bool bTraceComplex);
+
+		// Returns all the overlaps within A given a shape B.
+		template<typename TOverlapHit>
+		bool PhysicsObjectOverlap(const FConstPhysicsObjectHandle ObjectA, const FTransform& InTransformA, const FConstPhysicsObjectHandle ObjectB, const FTransform& InTransformB, bool bTraceComplex, TArray<TOverlapHit>& OutOverlaps)
+		{
+			static_assert(std::is_same_v<TOverlapHit, ChaosInterface::TThreadOverlapHit<Id>>);
+			return PairwiseShapeOverlapHelper(
+				ObjectA,
+				InTransformA,
+				ObjectB,
+				InTransformB,
+				bTraceComplex,
+				false,
+				FVector::Zero(),
+				[this, ObjectA, &OutOverlaps](const FShapeOverlapData& A, const FShapeOverlapData& B, const FMTDInfo&)
+				{
+					ChaosInterface::FOverlapHit Overlap;
+					Overlap.Shape = A.Shape;
+					Overlap.Actor = Interface.GetParticle(ObjectA);
+					OutOverlaps.Add(Overlap);
+					return false;
+				}
+			);
+		}
+
+		// This function does the same as GetPhysicsObjectOverlap but also computes the MTD metric.
+		bool PhysicsObjectOverlapWithMTD(const FConstPhysicsObjectHandle ObjectA, const FTransform& InTransformA, const FConstPhysicsObjectHandle ObjectB, const FTransform& InTransformB, bool bTraceComplex, FMTDInfo& OutMTD);
+
+		// This function does the same as GetPhysicsObjectOverlap but also computes the AABB overlap metric.
+		bool PhysicsObjectOverlapWithAABB(const FConstPhysicsObjectHandle ObjectA, const FTransform& InTransformA, const FConstPhysicsObjectHandle ObjectB, const FTransform& InTransformB, bool bTraceComplex, const FVector& Tolerance, FBox& OutOverlap);
+		bool PhysicsObjectOverlapWithAABBSize(const FConstPhysicsObjectHandle ObjectA, const FTransform& InTransformA, const FConstPhysicsObjectHandle ObjectB, const FTransform& InTransformB, bool bTraceComplex, const FVector& Tolerance, FVector& OutOverlapSize);
+
+		template<typename TRaycastHit>
+		bool LineTrace(TArrayView<const FConstPhysicsObjectHandle> InObjects, const FVector& WorldStart, const FVector& WorldEnd, bool bTraceComplex, TRaycastHit& OutBestHit)
+		{
+			static_assert(std::is_same_v<TRaycastHit, ChaosInterface::TThreadRaycastHit<Id>>);
+			bool bHit = false;
+			OutBestHit.Distance = TNumericLimits<float>::Max();
+
+			const FVector Delta = WorldEnd - WorldStart;
+			const FReal DeltaMag = Delta.Size();
+			if (DeltaMag < UE_KINDA_SMALL_NUMBER)
+			{
+				return false;
+			}
+
+			FTransform BestWorldTM = FTransform::Identity;
+
+			for (const FConstPhysicsObjectHandle Object : InObjects)
+			{
+				const FTransform WorldTM = Interface.GetTransform(Object);
+				const FVector LocalStart = WorldTM.InverseTransformPositionNoScale(WorldStart);
+				const FVector LocalDelta = WorldTM.InverseTransformVectorNoScale(Delta);
+
+				Interface.VisitEveryShape(
+					{ &Object, 1 },
+					[this, &bHit, &WorldTM, &LocalStart, &LocalDelta, &Delta, DeltaMag, &BestWorldTM, bTraceComplex, &OutBestHit](const FConstPhysicsObjectHandle IterObject, TThreadShapeInstance<Id>* Shape)
+					{
+						check(Shape);
+
+						FCollisionFilterData ShapeFilter = Shape->GetQueryData();
+						const bool bShapeIsComplex = (ShapeFilter.Word3 & static_cast<uint8>(EFilterFlags::ComplexCollision)) != 0;
+						const bool bShapeIsSimple = (ShapeFilter.Word3 & static_cast<uint8>(EFilterFlags::SimpleCollision)) != 0;
+						if ((bTraceComplex && bShapeIsComplex) || (!bTraceComplex && bShapeIsSimple))
+						{
+							FReal Distance;
+							FVec3 LocalPosition;
+							FVec3 LocalNormal;
+							int32 FaceIndex;
+
+							const bool bRaycastHit = Shape->GetGeometry()->Raycast(
+								LocalStart,
+								LocalDelta / DeltaMag,
+								DeltaMag,
+								0,
+								Distance,
+								LocalPosition,
+								LocalNormal,
+								FaceIndex
+							);
+
+							if (bRaycastHit)
+							{
+								if (Distance < OutBestHit.Distance)
+								{
+									bHit = true;
+									BestWorldTM = WorldTM;
+									OutBestHit.Distance = static_cast<float>(Distance);
+									OutBestHit.WorldNormal = LocalNormal;
+									OutBestHit.WorldPosition = LocalPosition;
+									OutBestHit.Shape = Shape;
+									OutBestHit.Actor = Interface.GetParticle(IterObject);
+									OutBestHit.FaceIndex = FaceIndex;
+								}
+							}
+						}
+						return false;
+					}
+				);
+			}
+
+			if (bHit)
+			{
+				OutBestHit.WorldNormal = BestWorldTM.TransformVectorNoScale(OutBestHit.WorldNormal);
+				OutBestHit.WorldPosition = BestWorldTM.TransformPositionNoScale(OutBestHit.WorldPosition);
+			}
+
+			return bHit;
+		}
+
+		template<typename TOverlapHit>
+		bool ShapeOverlap(TArrayView<const FConstPhysicsObjectHandle> InObjects, const Chaos::FImplicitObject& InGeom, const FTransform& GeomTransform, TArray<TOverlapHit>& OutOverlaps)
+		{
+			static_assert(std::is_same_v<TOverlapHit, ChaosInterface::TThreadOverlapHit<Id>>);
+			bool bHasOverlap = false;
+			for (const FConstPhysicsObjectHandle Object : InObjects)
+			{
+				const FTransform WorldTM = Interface.GetTransform(Object);
+
+				Interface.VisitEveryShape(
+					{ &Object, 1 },
+					[this, &bHasOverlap, &WorldTM, &InGeom, &GeomTransform, &OutOverlaps](const FConstPhysicsObjectHandle IterObject, TThreadShapeInstance<Id>* Shape)
+					{
+						check(Shape);
+						const bool bOverlap = Chaos::Utilities::CastHelper(
+							InGeom,
+							GeomTransform,
+							[Shape, &WorldTM](const auto& Downcast, const auto& FullTransformB)
+							{
+								return Chaos::OverlapQuery(*Shape->GetGeometry(), WorldTM, Downcast, FullTransformB, 0, nullptr);
+							}
+						);
+
+						if (bOverlap)
+						{
+							bHasOverlap = true;
+
+							ChaosInterface::FOverlapHit Overlap;
+							Overlap.Shape = Shape;
+							Overlap.Actor = Interface.GetParticle(IterObject);
+							OutOverlaps.Add(Overlap);
+						}
+						return false;
+					}
+				);
+			}
+			return bHasOverlap;
+		}
+
+		template<typename TSweepHit>
+		bool ShapeSweep(TArrayView<const FConstPhysicsObjectHandle> InObjects, const Chaos::FImplicitObject& InGeom, const FTransform& StartTM, const FVector& EndPos, const FSweepParameters& Params, TSweepHit& OutBestHit)
+		{
+			static_assert(std::is_same_v<TSweepHit, ChaosInterface::TThreadSweepHit<Id>>);
+			bool bHit = false;
+			const FVector StartPos = StartTM.GetTranslation();
+			const FVector Delta = EndPos - StartPos;
+			const FReal DeltaMag = Delta.Size();
+			if (DeltaMag < UE_KINDA_SMALL_NUMBER)
+			{
+				return false;
+			}
+			const FVec3 Dir = Delta / DeltaMag;
+
+			for (const FConstPhysicsObjectHandle Object : InObjects)
+			{
+				const FTransform WorldTM = Interface.GetTransform(Object);
+
+				Interface.VisitEveryShape(
+					{ &Object, 1 },
+					[this, &WorldTM, &InGeom, &StartTM, &bHit, &Delta, DeltaMag, &Dir, &OutBestHit, &Params](const FConstPhysicsObjectHandle IterObject, TThreadShapeInstance<Id>* Shape)
+					{
+						check(Shape);
+
+						FCollisionFilterData ShapeFilter = Shape->GetQueryData();
+						const bool bShapeIsComplex = (ShapeFilter.Word3 & static_cast<uint8>(EFilterFlags::ComplexCollision)) != 0;
+						const bool bShapeIsSimple = (ShapeFilter.Word3 & static_cast<uint8>(EFilterFlags::SimpleCollision)) != 0;
+						if ((Params.bSweepComplex && bShapeIsComplex) || (!Params.bSweepComplex && bShapeIsSimple))
+						{
+							FVec3 WorldPosition;
+							FVec3 WorldNormal;
+							FReal Distance;
+							int32 FaceIdx;
+							FVec3 FaceNormal;
+
+							const bool bShapeHit = Chaos::Utilities::CastHelper(
+								InGeom,
+								StartTM,
+								[Shape, &WorldTM, &Dir, DeltaMag, &Distance, &WorldPosition, &WorldNormal, &FaceIdx, &FaceNormal, &Params](const auto& Downcast, const auto& FullTransformB)
+								{
+									// Set bComputeMTD to true to better match internal SQ Visitor behavior.
+									return Chaos::SweepQuery(*Shape->GetGeometry(), WorldTM, Downcast, FullTransformB, Dir, DeltaMag, Distance, WorldPosition, WorldNormal, FaceIdx, FaceNormal, 0.f, Params.bComputeMTD);
+								}
+							);
+
+							if (bShapeHit)
+							{
+								bHit = true;
+
+								OutBestHit.Shape = Shape;
+								OutBestHit.WorldPosition = WorldPosition;
+								OutBestHit.WorldNormal = WorldNormal;
+								OutBestHit.Distance = static_cast<float>(Distance);
+								OutBestHit.FaceIndex = FaceIdx;
+								if (OutBestHit.Distance > 0.f)
+								{
+									const FVector LocalPosition = WorldTM.InverseTransformPositionNoScale(OutBestHit.WorldPosition);
+									const FVector LocalUnitDir = WorldTM.InverseTransformVectorNoScale(Dir);
+									OutBestHit.FaceIndex = Shape->GetGeometry()->FindMostOpposingFace(LocalPosition, LocalUnitDir, OutBestHit.FaceIndex, 1);
+								}
+								OutBestHit.Actor = Interface.GetParticle(IterObject);
+							}
+						}
+						return false;
+					}
+				);
+
+			}
+			return bHit;
+		}
+
+
+	private:
+		FReadPhysicsObjectInterface<Id>& Interface;
+
+		struct FShapeOverlapData
+		{
+			TThreadShapeInstance<Id>* Shape;
+			FAABB3 BoundingBox;
+		};
+
+		/**
+		 * For every pair of shapes that overlap, allows the caller to perform some computation. If additional pairs of shapes need to be examined, the input TFunction should return true.
+		 */
+		bool PairwiseShapeOverlapHelper(const FConstPhysicsObjectHandle ObjectA, const FTransform& InTransformA, const FConstPhysicsObjectHandle ObjectB, const FTransform& InTransformB, bool bTraceComplex, bool bComputeMTD, const FVector& Tolerance, const TFunction<bool(const FShapeOverlapData&, const FShapeOverlapData&, const FMTDInfo&)>& Lambda);
+	};
+
+	using FPhysicsObjectCollisionInterface_External = FPhysicsObjectCollisionInterface<EThreadContext::External>;
+	using FPhysicsObjectCollisionInterface_Internal = FPhysicsObjectCollisionInterface<EThreadContext::Internal>;
+}
\ No newline at end of file
diff --git a/Engine/Source/Runtime/Experimental/Chaos/Public/Chaos/PhysicsObjectInterface.h b/Engine/Source/Runtime/Experimental/Chaos/Public/Chaos/PhysicsObjectInterface.h
--- a/Engine/Source/Runtime/Experimental/Chaos/Public/Chaos/PhysicsObjectInterface.h
+++ b/Engine/Source/Runtime/Experimental/Chaos/Public/Chaos/PhysicsObjectInterface.h
@@ -2,8 +2,6 @@
 
 #pragma once
 
-#include "Chaos/Interface/SQTypes.h"
-#include "Chaos/GeometryQueries.h"
 #include "Chaos/ParticleHandle.h"
 #include "Chaos/PhysicsObject.h"
 #include "Chaos/ShapeInstanceFwd.h"
@@ -79,46 +77,12 @@ namespace Chaos
 		// Returns true if a shape is found and we can stop iterating.
 		void VisitEveryShape(TArrayView<const FConstPhysicsObjectHandle> InObjects, TFunctionRef<bool(const FConstPhysicsObjectHandle, TThreadShapeInstance<Id>*)> Lambda);
 
-		UE_DEPRECATED(5.3, "GetPhysicsObjectOverlap has been deprecated. Please use the function for the specific overlap metric you wish to compute instead.")
+		UE_DEPRECATED(5.3, "GetPhysicsObjectOverlap has been deprecated. Please use the function for the specific overlap metric you wish to compute instead in the FPhysicsObjectCollisionInterface.")
 		bool GetPhysicsObjectOverlap(const FConstPhysicsObjectHandle ObjectA, const FConstPhysicsObjectHandle ObjectB, bool bTraceComplex, Chaos::FOverlapInfo& OutOverlap);
 
-		UE_DEPRECATED(5.3, "GetPhysicsObjectOverlapWithTransform has been deprecated. Please use the function for the specific overlap metric you wish to compute instead.")
+		UE_DEPRECATED(5.3, "GetPhysicsObjectOverlapWithTransform has been deprecated. Please use the function for the specific overlap metric you wish to compute instead in the FPhysicsObjectCollisionInterface.")
 		bool GetPhysicsObjectOverlapWithTransform(const FConstPhysicsObjectHandle ObjectA, const FTransform& InTransformA, const FConstPhysicsObjectHandle ObjectB, const FTransform& InTransformB, bool bTraceComplex, Chaos::FOverlapInfo& OutOverlap);
 
-		// This function will not compute any overlap heuristic.
-		bool PhysicsObjectOverlap(const FConstPhysicsObjectHandle ObjectA, const FTransform& InTransformA, const FConstPhysicsObjectHandle ObjectB, const FTransform& InTransformB, bool bTraceComplex);
-
-		// Returns all the overlaps within A given a shape B.
-		template<typename TOverlapHit>
-		bool PhysicsObjectOverlap(const FConstPhysicsObjectHandle ObjectA, const FTransform& InTransformA, const FConstPhysicsObjectHandle ObjectB, const FTransform& InTransformB, bool bTraceComplex, TArray<TOverlapHit>& OutOverlaps)
-		{
-			static_assert(std::is_same_v<TOverlapHit, ChaosInterface::TThreadOverlapHit<Id>>);
-			return PairwiseShapeOverlapHelper(
-				ObjectA,
-				InTransformA,
-				ObjectB,
-				InTransformB,
-				bTraceComplex,
-				false,
-				FVector::Zero(),
-				[this, ObjectA, &OutOverlaps](const FShapeOverlapData& A , const FShapeOverlapData& B, const FMTDInfo&)
-				{
-					ChaosInterface::FOverlapHit Overlap;
-					Overlap.Shape = A.Shape;
-					Overlap.Actor = GetParticle(ObjectA);
-					OutOverlaps.Add(Overlap);
-					return false;
-				}
-			);
-		}
-
-		// This function does the same as GetPhysicsObjectOverlap but also computes the MTD metric.
-		bool PhysicsObjectOverlapWithMTD(const FConstPhysicsObjectHandle ObjectA, const FTransform& InTransformA, const FConstPhysicsObjectHandle ObjectB, const FTransform& InTransformB, bool bTraceComplex, FMTDInfo& OutMTD);
-
-		// This function does the same as GetPhysicsObjectOverlap but also computes the AABB overlap metric.
-		bool PhysicsObjectOverlapWithAABB(const FConstPhysicsObjectHandle ObjectA, const FTransform& InTransformA, const FConstPhysicsObjectHandle ObjectB, const FTransform& InTransformB, bool bTraceComplex, const FVector& Tolerance, FBox& OutOverlap);
-		bool PhysicsObjectOverlapWithAABBSize(const FConstPhysicsObjectHandle ObjectA, const FTransform& InTransformA, const FConstPhysicsObjectHandle ObjectB, const FTransform& InTransformB, bool bTraceComplex, const FVector& Tolerance, FVector& OutOverlapSize);
-
 		bool AreAllValid(TArrayView<const FConstPhysicsObjectHandle> InObjects);
 		bool AreAllKinematic(TArrayView<const FConstPhysicsObjectHandle> InObjects);
 		bool AreAllSleeping(TArrayView<const FConstPhysicsObjectHandle> InObjects);
@@ -132,209 +96,9 @@ namespace Chaos
 		FClosestPhysicsObjectResult GetClosestPhysicsBodyFromLocation(TArrayView<const FConstPhysicsObjectHandle> InObjects, const FVector& WorldLocation);
 		FAccelerationStructureHandle CreateAccelerationStructureHandle(const FConstPhysicsObjectHandle Handle);
 
-		template<typename TRaycastHit>
-		bool LineTrace(TArrayView<const FConstPhysicsObjectHandle> InObjects, const FVector& WorldStart, const FVector& WorldEnd, bool bTraceComplex, TRaycastHit& OutBestHit)
-		{
-			static_assert(std::is_same_v<TRaycastHit, ChaosInterface::TThreadRaycastHit<Id>>);
-			bool bHit = false;
-			OutBestHit.Distance = TNumericLimits<float>::Max();
-
-			const FVector Delta = WorldEnd - WorldStart;
-			const FReal DeltaMag = Delta.Size();
-			if (DeltaMag < UE_KINDA_SMALL_NUMBER)
-			{
-				return false;
-			}
-
-			FTransform BestWorldTM = FTransform::Identity;
-
-			for (const FConstPhysicsObjectHandle Object : InObjects)
-			{
-				const FTransform WorldTM = GetTransform(Object);
-				const FVector LocalStart = WorldTM.InverseTransformPositionNoScale(WorldStart);
-				const FVector LocalDelta = WorldTM.InverseTransformVectorNoScale(Delta);
-
-				VisitEveryShape(
-					{ &Object, 1 },
-					[this, &bHit, &WorldTM, &LocalStart, &LocalDelta, &Delta, DeltaMag, &BestWorldTM, bTraceComplex, &OutBestHit](const FConstPhysicsObjectHandle IterObject, TThreadShapeInstance<Id>* Shape)
-					{
-						check(Shape);
-
-						FCollisionFilterData ShapeFilter = Shape->GetQueryData();
-						const bool bShapeIsComplex = (ShapeFilter.Word3 & static_cast<uint8>(EFilterFlags::ComplexCollision)) != 0;
-						const bool bShapeIsSimple = (ShapeFilter.Word3 & static_cast<uint8>(EFilterFlags::SimpleCollision)) != 0;
-						if ((bTraceComplex && bShapeIsComplex) || (!bTraceComplex && bShapeIsSimple))
-						{
-							FReal Distance;
-							FVec3 LocalPosition;
-							FVec3 LocalNormal;
-							int32 FaceIndex;
-
-							const bool bRaycastHit = Shape->GetGeometry()->Raycast(
-								LocalStart,
-								LocalDelta / DeltaMag,
-								DeltaMag,
-								0,
-								Distance,
-								LocalPosition,
-								LocalNormal,
-								FaceIndex
-							);
-
-							if (bRaycastHit)
-							{
-								if (Distance < OutBestHit.Distance)
-								{
-									bHit = true;
-									BestWorldTM = WorldTM;
-									OutBestHit.Distance = static_cast<float>(Distance);
-									OutBestHit.WorldNormal = LocalNormal;
-									OutBestHit.WorldPosition = LocalPosition;
-									OutBestHit.Shape = Shape;
-									OutBestHit.Actor = GetParticle(IterObject);
-									OutBestHit.FaceIndex = FaceIndex;
-								}
-							}
-						}
-						return false;
-					}
-				);
-			}
-
-			if (bHit)
-			{
-				OutBestHit.WorldNormal = BestWorldTM.TransformVectorNoScale(OutBestHit.WorldNormal);
-				OutBestHit.WorldPosition = BestWorldTM.TransformPositionNoScale(OutBestHit.WorldPosition);
-			}
-
-			return bHit;
-		}
-
-		template<typename TOverlapHit>
-		bool ShapeOverlap(TArrayView<const FConstPhysicsObjectHandle> InObjects, const Chaos::FImplicitObject& InGeom, const FTransform& GeomTransform, TArray<TOverlapHit>& OutOverlaps)
-		{
-			static_assert(std::is_same_v<TOverlapHit, ChaosInterface::TThreadOverlapHit<Id>>);
-			bool bHasOverlap = false;
-			for (const FConstPhysicsObjectHandle Object : InObjects)
-			{
-				const FTransform WorldTM = GetTransform(Object);
-
-				VisitEveryShape(
-					{ &Object, 1 },
-					[this, &bHasOverlap, &WorldTM, &InGeom, &GeomTransform, &OutOverlaps](const FConstPhysicsObjectHandle IterObject, TThreadShapeInstance<Id>* Shape)
-					{
-						check(Shape);
-						const bool bOverlap = Chaos::Utilities::CastHelper(
-							InGeom,
-							GeomTransform,
-							[Shape, &WorldTM](const auto& Downcast, const auto& FullTransformB)
-							{
-								return Chaos::OverlapQuery(*Shape->GetGeometry(), WorldTM, Downcast, FullTransformB, 0, nullptr);
-							}
-						);
-
-						if (bOverlap)
-						{
-							bHasOverlap = true;
-
-							ChaosInterface::FOverlapHit Overlap;
-							Overlap.Shape = Shape;
-							Overlap.Actor = GetParticle(IterObject);
-							OutOverlaps.Add(Overlap);
-						}
-						return false;
-					}
-				);
-			}
-			return bHasOverlap;
-		}
-
-		template<typename TSweepHit>
-		bool ShapeSweep(TArrayView<const FConstPhysicsObjectHandle> InObjects, const Chaos::FImplicitObject& InGeom, const FTransform& StartTM, const FVector& EndPos, bool bSweepComplex, TSweepHit& OutBestHit)
-		{
-			static_assert(std::is_same_v<TSweepHit, ChaosInterface::TThreadSweepHit<Id>>);
-			bool bHit = false;
-			const FVector StartPos = StartTM.GetTranslation();
-			const FVector Delta = EndPos - StartPos;
-			const FReal DeltaMag = Delta.Size();
-			if (DeltaMag < UE_KINDA_SMALL_NUMBER)
-			{
-				return false;
-			}
-			const FVec3 Dir = Delta / DeltaMag;
-
-			for (const FConstPhysicsObjectHandle Object : InObjects)
-			{
-				const FTransform WorldTM = GetTransform(Object);
-
-				VisitEveryShape(
-					{ &Object, 1 },
-					[this, &WorldTM, &InGeom, &StartTM, &bHit, &Delta, DeltaMag, &Dir, &OutBestHit, bSweepComplex](const FConstPhysicsObjectHandle IterObject, TThreadShapeInstance<Id>* Shape)
-					{
-						check(Shape);
-
-						FCollisionFilterData ShapeFilter = Shape->GetQueryData();
-						const bool bShapeIsComplex = (ShapeFilter.Word3 & static_cast<uint8>(EFilterFlags::ComplexCollision)) != 0;
-						const bool bShapeIsSimple = (ShapeFilter.Word3 & static_cast<uint8>(EFilterFlags::SimpleCollision)) != 0;
-						if ((bSweepComplex && bShapeIsComplex) || (!bSweepComplex && bShapeIsSimple))
-						{
-							FVec3 WorldPosition;
-							FVec3 WorldNormal;
-							FReal Distance;
-							int32 FaceIdx;
-							FVec3 FaceNormal;
-
-							const bool bShapeHit = Chaos::Utilities::CastHelper(
-								InGeom,
-								StartTM,
-								[Shape, &WorldTM, &Dir, DeltaMag, &Distance, &WorldPosition, &WorldNormal, &FaceIdx, &FaceNormal](const auto& Downcast, const auto& FullTransformB)
-								{
-									// Set bComputeMTD to true to better match internal SQ Visitor behavior.
-									return Chaos::SweepQuery(*Shape->GetGeometry(), WorldTM, Downcast, FullTransformB, Dir, DeltaMag, Distance, WorldPosition, WorldNormal, FaceIdx, FaceNormal, 0.f, true);
-								}
-							);
-
-							if (bShapeHit)
-							{
-								bHit = true;
-
-								OutBestHit.Shape = Shape;
-								OutBestHit.WorldPosition = WorldPosition;
-								OutBestHit.WorldNormal = WorldNormal;
-								OutBestHit.Distance = static_cast<float>(Distance);
-								OutBestHit.FaceIndex = FaceIdx;
-								if (OutBestHit.Distance > 0.f)
-								{
-									const FVector LocalPosition = WorldTM.InverseTransformPositionNoScale(OutBestHit.WorldPosition);
-									const FVector LocalUnitDir = WorldTM.InverseTransformVectorNoScale(Dir);
-									OutBestHit.FaceIndex = Shape->GetGeometry()->FindMostOpposingFace(LocalPosition, LocalUnitDir, OutBestHit.FaceIndex, 1);
-								}
-								OutBestHit.Actor = GetParticle(IterObject);
-							}
-						}
-						return false;
-					}
-				);
-
-			}
-			return bHit;
-		}
-
 		friend class FPhysicsObjectInterface;
 	protected:
 		FReadPhysicsObjectInterface() = default;
-
-	private:
-		struct FShapeOverlapData
-		{
-			TThreadShapeInstance<Id>* Shape;
-			FAABB3 BoundingBox;
-		};
-
-		/**
-		 * For every pair of shapes that overlap, allows the caller to perform some computation. If additional pairs of shapes need to be examined, the input TFunction should return true.
-		 */
-		bool PairwiseShapeOverlapHelper(const FConstPhysicsObjectHandle ObjectA, const FTransform& InTransformA, const FConstPhysicsObjectHandle ObjectB, const FTransform& InTransformB, bool bTraceComplex, bool bComputeMTD, const FVector& Tolerance, const TFunction<bool(const FShapeOverlapData&, const FShapeOverlapData&, const FMTDInfo&)>& Lambda);
 	};
 
 	using FReadPhysicsObjectInterface_External = FReadPhysicsObjectInterface<EThreadContext::External>;
