Add MinLod option to ISMPool components.

diff --git a/Engine/Source/Runtime/Experimental/GeometryCollectionEngine/Private/GeometryCollection/GeometryCollectionISMPoolComponent.cpp b/Engine/Source/Runtime/Experimental/GeometryCollectionEngine/Private/GeometryCollection/GeometryCollectionISMPoolComponent.cpp
--- a/Engine/Source/Runtime/Experimental/GeometryCollectionEngine/Private/GeometryCollection/GeometryCollectionISMPoolComponent.cpp
+++ b/Engine/Source/Runtime/Experimental/GeometryCollectionEngine/Private/GeometryCollection/GeometryCollectionISMPoolComponent.cpp
@@ -84,6 +84,8 @@ FGeometryCollectionISM::FGeometryCollectionISM(AActor* OwmingActor, const FGeome
 	ISMC->bAffectDistanceFieldLighting = MeshInstance.Desc.bAffectDistanceFieldLighting;
 	ISMC->SetCanEverAffectNavigation(false);
 	ISMC->SetCollisionEnabled(ECollisionEnabled::NoCollision);
+	ISMC->bOverrideMinLOD = MeshInstance.Desc.MinLod > 0;	
+	ISMC->MinLOD = MeshInstance.Desc.MinLod;
 	
 	if (HISMC)
 	{
diff --git a/Engine/Source/Runtime/Experimental/GeometryCollectionEngine/Public/GeometryCollection/GeometryCollectionISMPoolComponent.h b/Engine/Source/Runtime/Experimental/GeometryCollectionEngine/Public/GeometryCollection/GeometryCollectionISMPoolComponent.h
--- a/Engine/Source/Runtime/Experimental/GeometryCollectionEngine/Public/GeometryCollection/GeometryCollectionISMPoolComponent.h
+++ b/Engine/Source/Runtime/Experimental/GeometryCollectionEngine/Public/GeometryCollection/GeometryCollectionISMPoolComponent.h
@@ -121,6 +121,7 @@ struct FISMComponentDescription
 	int32 NumCustomDataFloats = 0;
 	int32 StartCullDistance = 0;
 	int32 EndCullDistance = 0;
+	int32 MinLod = 0;
 	float LodScale = 1.f;
 
 	bool operator==(const FISMComponentDescription& Other) const
@@ -134,6 +135,7 @@ struct FISMComponentDescription
 			NumCustomDataFloats == Other.NumCustomDataFloats &&
 			StartCullDistance == Other.StartCullDistance && 
 			EndCullDistance == Other.EndCullDistance &&
+			MinLod == Other.MinLod &&
 			LodScale == Other.LodScale;
 	}
 };
@@ -144,6 +146,7 @@ FORCEINLINE uint32 GetTypeHash(const FISMComponentDescription& Desc)
 	uint32 Hash = HashCombine(GetTypeHash(PackedBools), GetTypeHash(Desc.NumCustomDataFloats));
 	Hash = HashCombine(Hash, GetTypeHash(Desc.StartCullDistance));
 	Hash = HashCombine(Hash, GetTypeHash(Desc.EndCullDistance));
+	Hash = HashCombine(Hash, GetTypeHash(Desc.MinLod));
 	return HashCombine(Hash, GetTypeHash(Desc.LodScale));
 }
 
