Iris
- Restructured NetObject serialization in ReplicationWriter/ReplicationReader to always write root object and subobjects in a single batch including size and eports in order to allow us to skip or defer processing the data in the batch.
- Fixed reporting of data written in hugeobjects on server side.

diff --git a/Engine/Source/Runtime/Experimental/Iris/Core/Private/Iris/ReplicationSystem/AttachmentReplication.cpp b/Engine/Source/Runtime/Experimental/Iris/Core/Private/Iris/ReplicationSystem/AttachmentReplication.cpp
--- a/Engine/Source/Runtime/Experimental/Iris/Core/Private/Iris/ReplicationSystem/AttachmentReplication.cpp
+++ b/Engine/Source/Runtime/Experimental/Iris/Core/Private/Iris/ReplicationSystem/AttachmentReplication.cpp
@@ -366,8 +366,8 @@ uint32 FNetObjectAttachmentSendQueue::SerializeUnreliable(FNetSerializationConte
 
 		const TRefCountPtr<FNetBlob>& Attachment = UnreliableQueue.PeekAtOffsetNoCheck(AttachmentIt);
 
-		// Write exports, if attachment does not fit we will rollback exports as well.
-		ObjectReferenceCache->WriteExports(Context, Attachment->CallGetExports());
+		// If we have exports, append them, if attachment is rolled back we will roll back any appended exports as well.
+		ObjectReferenceCache->AddPendingExports(Context, Attachment->CallGetExports());
 
 		Attachment->SerializeCreationInfo(Context, Attachment->GetCreationInfo());
 		if (bSerializeWithObject)
@@ -872,8 +872,6 @@ uint32 FNetObjectAttachmentReceiveQueue::DeserializeUnreliable(FNetSerialization
 			break;
 		}
 
-		Context.GetInternalContext()->ObjectReferenceCache->ReadExports(Context);
-
 		FNetBlobCreationInfo CreationInfo;
 		FNetBlob::DeserializeCreationInfo(Context, CreationInfo);
 		const TRefCountPtr<FNetBlob>& Attachment = BlobReceiver->CreateNetBlob(CreationInfo);
diff --git a/Engine/Source/Runtime/Experimental/Iris/Core/Private/Iris/ReplicationSystem/NetBlob/ReliableNetBlobQueue.cpp b/Engine/Source/Runtime/Experimental/Iris/Core/Private/Iris/ReplicationSystem/NetBlob/ReliableNetBlobQueue.cpp
--- a/Engine/Source/Runtime/Experimental/Iris/Core/Private/Iris/ReplicationSystem/NetBlob/ReliableNetBlobQueue.cpp
+++ b/Engine/Source/Runtime/Experimental/Iris/Core/Private/Iris/ReplicationSystem/NetBlob/ReliableNetBlobQueue.cpp
@@ -107,8 +107,8 @@ uint32 FReliableNetBlobQueue::SerializeInternal(FNetSerializationContext& Contex
 			Writer->WriteBits(Index, IndexBitCount);
 		}
 
-		// If we have exports, write them now, if attachment is rolled back we will roll back any written exports as well.
-		ObjectReferenceCache->WriteExports(Context, Attachment->CallGetExports());
+		// If we have exports, append them, if attachment is rolled back we will roll back any appended exports as well.
+		ObjectReferenceCache->AddPendingExports(Context, Attachment->CallGetExports());
 
 		Attachment->SerializeCreationInfo(Context, Attachment->GetCreationInfo());
 		if (bSerializeWithObject)
@@ -224,9 +224,6 @@ uint32 FReliableNetBlobQueue::DeserializeInternal(FNetSerializationContext& Cont
 			return DeserializedCount;
 		}
 
-		// If we have exports, import them now.
-		ObjectReferenceCache->ReadExports(Context);
-
 		FNetBlobCreationInfo CreationInfo;
 		FNetBlob::DeserializeCreationInfo(Context, CreationInfo);
 		CreationInfo.Flags = CreationInfo.Flags | ENetBlobFlags::Reliable;
diff --git a/Engine/Source/Runtime/Experimental/Iris/Core/Private/Iris/ReplicationSystem/ObjectReferenceCache.cpp b/Engine/Source/Runtime/Experimental/Iris/Core/Private/Iris/ReplicationSystem/ObjectReferenceCache.cpp
--- a/Engine/Source/Runtime/Experimental/Iris/Core/Private/Iris/ReplicationSystem/ObjectReferenceCache.cpp
+++ b/Engine/Source/Runtime/Experimental/Iris/Core/Private/Iris/ReplicationSystem/ObjectReferenceCache.cpp
@@ -113,7 +113,7 @@ bool FObjectReferenceCache::IsAuthority() const
 bool FObjectReferenceCache::CanClientLoadObjectInternal(const UObject* Object, bool bIsDynamic) const
 {
 	// PackageMapClient can't load maps, we must wait for the client to load the map when ready
-	// These guids are special guids, where the guid and all child guids resolve once the map has been loaded
+	// These references are special references, where the reference and all child references resolve once the map has been loaded
 	if (bIsDynamic)
 	{
 		return false;
@@ -142,7 +142,7 @@ bool FObjectReferenceCache::ShouldIgnoreWhenMissing(FNetRefHandle RefHandle) con
 {
 	if (RefHandle.IsDynamic())
 	{
-		return true;		// Ignore missing dynamic guids (even on server because client may send RPC on/with object it doesn't know server destroyed)
+		return true;		// Ignore missing dynamic references (even on server because client may send RPC on/with object it doesn't know server destroyed)
 	}
 
 	if (IsAuthority())
@@ -153,7 +153,7 @@ bool FObjectReferenceCache::ShouldIgnoreWhenMissing(FNetRefHandle RefHandle) con
 	const FCachedNetObjectReference* CacheObject = ReferenceHandleToCachedReference.Find(RefHandle);
 	if (CacheObject == nullptr)
 	{
-		return false;		// If we haven't been told about this static guid before, we need to warn
+		return false;		// If we haven't been told about this static reference before, we need to warn
 	}
 
 	const FCachedNetObjectReference* OutermostCacheObject = CacheObject;
@@ -743,7 +743,7 @@ UObject* FObjectReferenceCache::ResolveObjectReferenceHandleInternal(FNetRefHand
 		{
 			// This isn't really a package but it should be
 			CacheObjectPtr->bIsBroken = true;
-			//UE_LOG( LogNetPackageMap, Error, TEXT( "GetObjectFromNetGUID: Object is not a package but should be! Path: %s, NetGUID: %s" ), *Path.ToString(), *NetGUID.ToString() );
+			UE_LOG(LogIris, Error, TEXT( "GetObjectFromRefHandle: Object is not a package but should be! Path: %s, NetRefHandle: %s" ), *ObjectPath, *RefHandle.ToString() );
 			return nullptr;
 		}
 
@@ -775,7 +775,7 @@ UObject* FObjectReferenceCache::ResolveObjectReferenceHandleInternal(FNetRefHand
 	//	{
 	//		if ( NetworkChecksumMode == ENetworkChecksumMode::SaveAndUse )
 	//		{
-	//			FString ErrorStr = FString::Printf(TEXT("GetObjectFromNetGUID: Network checksum mismatch. ``IDPath: %s, %u, %u"), *FullNetGUIDPath(NetGUID), CacheObjectPtr->NetworkChecksum, NetworkChecksum);
+	//			FString ErrorStr = FString::Printf(TEXT("GetObjectFromRefHandle: Network checksum mismatch. ``IDPath: %s, %u, %u"), *FullRefHandlePath(RefHandle), CacheObjectPtr->NetworkChecksum, NetworkChecksum);
 	//			UE_LOG( LogNetPackageMap, Warning, TEXT("%s"), *ErrorStr );
 
 	//			CacheObjectPtr->bIsBroken = true;
@@ -785,7 +785,7 @@ UObject* FObjectReferenceCache::ResolveObjectReferenceHandleInternal(FNetRefHand
 	//		}
 	//		else
 	//		{
-	//			UE_LOG( LogNetPackageMap, Verbose, TEXT( "GetObjectFromNetGUID: Network checksum mismatch. ``IDPath: %s, %u, %u" ), *FullNetGUIDPath( NetGUID ), CacheObjectPtr->NetworkChecksum, NetworkChecksum );
+	//			UE_LOG( LogNetPackageMap, Verbose, TEXT( "GetObjectFromRefHandle: Network checksum mismatch. ``IDPath: %s, %u, %u" ), *FullRefHandlePath( RefHandle ), CacheObjectPtr->NetworkChecksum, NetworkChecksum );
 	//		}
 	//	}
 	//}
@@ -797,7 +797,7 @@ UObject* FObjectReferenceCache::ResolveObjectReferenceHandleInternal(FNetRefHand
 		return nullptr;
 	}
 
-	// Assign the resolved object to this guid
+	// Assign the resolved object to this reference
 	if (!Object)
 	{
 		if (!CacheObjectPtr->bIgnoreWhenMissing)
@@ -1236,6 +1236,13 @@ void FObjectReferenceCache::ReadFullReferenceInternal(FNetSerializationContext&
 
 void FObjectReferenceCache::ReadFullReference(FNetSerializationContext& Context, FNetObjectReference& OutRef)
 {
+	// Normally all references are imported, unless we are doing this when reading exports
+	if (Context.GetInternalContext()->bInlineObjectReferenceExports == 0U)
+	{
+		ReadReference(Context, OutRef);
+		return;
+	}
+
 	FNetObjectReference ObjectRef;
 	FNetBitStreamReader* Reader = Context.GetBitStreamReader();
 
@@ -1267,6 +1274,17 @@ void FObjectReferenceCache::WriteFullReference(FNetSerializationContext& Context
 {
 	UE_NET_TRACE_SCOPE(FullNetObjectReference, *Context.GetBitStreamWriter(), Context.GetTraceCollector(), ENetTraceVerbosity::Trace);
 
+	// If we do not inline reference exports we just write the reference
+	if (Context.GetInternalContext()->bInlineObjectReferenceExports == 0U)
+	{
+		if (FNetExportContext* ExportContext = Context.GetExportContext())
+		{
+			AddPendingExport(*ExportContext, Ref);
+		}
+		WriteReference(Context, Ref);
+		return;
+	}		
+
 	// Handle client assigned reference 
 	FNetBitStreamWriter* Writer = Context.GetBitStreamWriter();
 	const bool bIsClientAssignedReference = Ref.PathToken.IsValid();
@@ -1274,9 +1292,7 @@ void FObjectReferenceCache::WriteFullReference(FNetSerializationContext& Context
 	{
 		WriteNetRefHandle(Context, Ref.GetRefHandle());
 		WriteNetToken(Writer, Ref.PathToken);
-		FNetExportContext* ExportContext = Context.GetExportContext();
-		ConditionalWriteNetTokenData(Context, ExportContext, Ref.PathToken);
-
+		ConditionalWriteNetTokenData(Context, Context.GetExportContext(), Ref.PathToken);
 		return;
 	}
 
@@ -1326,7 +1342,6 @@ void FObjectReferenceCache::ReadReference(FNetSerializationContext& Context, FNe
 {
 	UE_NET_TRACE_SCOPE(NetObjectReference, *Context.GetBitStreamReader(), Context.GetTraceCollector(), ENetTraceVerbosity::Trace);
 	UE_NET_TRACE_NAMED_OBJECT_SCOPE(ReferenceScope, FNetRefHandle(), *Context.GetBitStreamReader(), Context.GetTraceCollector(), ENetTraceVerbosity::Verbose);
-	LLM_SCOPE_BYTAG(Iris);
 
 	FNetBitStreamReader* Reader = Context.GetBitStreamReader();
 
@@ -1370,48 +1385,156 @@ void FObjectReferenceCache::ReadReference(FNetSerializationContext& Context, FNe
 	}
 }
 
-bool FObjectReferenceCache::WriteExports(FNetSerializationContext& Context, TArrayView<const FNetObjectReference> ExportsView) const
+void FObjectReferenceCache::AddPendingExport(FNetExportContext& ExportContext, const FNetObjectReference& Reference) const
 {
-	FNetBitStreamWriter& Writer = *Context.GetBitStreamWriter();
+	const bool bIsClientAssigned = Reference.PathToken.IsValid();
+	const FNetRefHandle Handle = Reference.GetRefHandle();
+
+	if (bIsClientAssigned || (IsAuthority() && Reference.CanBeExported()))
+	{
+		LLM_SCOPE_BYTAG(Iris);
+		ExportContext.AddPendingExport(Reference);
+	}
+}
 
+void FObjectReferenceCache::AddPendingExports(FNetSerializationContext& Context, TArrayView<const FNetObjectReference> ExportsView) const
+{
 	if (FNetExportContext* ExportContext = Context.GetExportContext())
 	{
-		UE_NET_TRACE_SCOPE(Exports, Writer, Context.GetTraceCollector(), ENetTraceVerbosity::Verbose);
+		for (const FNetObjectReference& Reference : ExportsView)
+		{
+			AddPendingExport(*ExportContext, Reference);
+		}
+	}
+}
+
+FObjectReferenceCache::EWriteExportsResult FObjectReferenceCache::WritePendingExports(FNetSerializationContext& Context)
+{	
+	FNetBitStreamWriter& Writer = *Context.GetBitStreamWriter();
+
+	FNetExportContext* ExportContext = Context.GetExportContext();
+	if (!ExportContext || ExportContext->GetBatchExports().ReferencesPendingExportInCurrentBatch.IsEmpty())
+	{
+		return EWriteExportsResult::NoExports;
+	}
+
+	UE_NET_TRACE_SCOPE(Exports, Writer, Context.GetTraceCollector(), ENetTraceVerbosity::Verbose);
+	FNetBitStreamRollbackScope Rollback(*Context.GetBitStreamWriter());
+	const uint32 ExportsStartBitPos = Writer.GetPosBits();
+
+	TArrayView<const FNetObjectReference> ExportsView = MakeArrayView(ExportContext->GetBatchExports().ReferencesPendingExportInCurrentBatch);
+
+	// Force exports to be written
+	{
+		FForceInlineExportScope ForceInlineExportScope(Context.GetInternalContext());
 
 		for (const FNetObjectReference& Reference : ExportsView)
 		{
 			const bool bIsClientAssigned = Reference.PathToken.IsValid();
 			const FNetRefHandle Handle = Reference.GetRefHandle();
-
 			if ((bIsClientAssigned && !ExportContext->IsExported(Reference.GetPathToken())) || (IsAuthority() && Reference.CanBeExported() && !ExportContext->IsExported(Handle)))
 			{
 				Writer.WriteBool(true);
 				WriteFullReference(Context, Reference);
 			}
 		}
+
+		// Write stop bit
+		Writer.WriteBool(false);
+	}
+
+	// We also write any must be mapped exports
+	WriteMustBeMappedExports(Context, ExportsView);
+
+	// Reset state of pending exports
+	ExportContext->ClearPendingExports();
+
+	if (Writer.IsOverflown())
+	{
+		return EWriteExportsResult::BitStreamOverflow;
+	}
+	else if ((Writer.GetPosBits() - ExportsStartBitPos) > 2U)
+	{
+		return EWriteExportsResult::WroteExports;
+	}
+	
+	Rollback.Rollback();
+
+	return EWriteExportsResult::NoExports;
+}
+
+bool FObjectReferenceCache::ReadExports(FNetSerializationContext& Context, TArray<FNetRefHandle>* MustBeMappedExports)
+{
+	FNetBitStreamReader& Reader = *Context.GetBitStreamReader();
+
+	bool bHasExportsToRead = Reader.ReadBool(); 
+	UE_NET_TRACE_SCOPE(Exports, Reader, Context.GetTraceCollector(), ENetTraceVerbosity::Verbose);
+
+	// Force inlined exports
+	{
+		FForceInlineExportScope ForceInlineExportScope(Context.GetInternalContext());
+		while (bHasExportsToRead && !Context.HasErrorOrOverflow())
+		{
+			FNetObjectReference Import;
+			ReadFullReference(Context, Import);
+			bHasExportsToRead = Reader.ReadBool();
+		}
 	}
 
+	if (!Context.HasErrorOrOverflow())
+	{
+		ReadMustBeMappedExports(Context, MustBeMappedExports);
+	}
+
+	return !Context.HasErrorOrOverflow();
+}
+
+bool FObjectReferenceCache::WriteMustBeMappedExports(FNetSerializationContext& Context, TArrayView<const FNetObjectReference> ExportsView) const
+{
+	FNetBitStreamWriter& Writer = *Context.GetBitStreamWriter();
+	FNetExportContext* ExportContext = Context.GetExportContext();
+
+	// $TODO: Will be enabled in later changelist
+	//if (IsAuthority() && ShouldAsyncLoad() && ExportContext)
+	//{
+	//	UE_NET_TRACE_SCOPE(MustBeMappedExports, Writer, Context.GetTraceCollector(), ENetTraceVerbosity::Verbose);
+
+	//	for (const FNetObjectReference& Reference : ExportsView)
+	//	{
+	//		const FNetRefHandle Handle = Reference.GetRefHandle();
+	//		const FCachedNetObjectReference* CachedObject = Reference.CanBeExported() ? ReferenceHandleToCachedReference.Find(Handle) : nullptr;
+	//		if (CachedObject && !CachedObject->bNoLoad)
+	//		{
+	//			UE_NET_TRACE_OBJECT_SCOPE(Handle, *Context.GetBitStreamWriter(), Context.GetTraceCollector(), ENetTraceVerbosity::Trace);
+	//			Writer.WriteBool(true);
+	//			WriteNetRefHandle(Context, Handle);
+	//		}
+	//	}
+	//}
+
 	// Write stop bit
 	Writer.WriteBool(false);
 
 	return !Writer.IsOverflown();
 }
 
-bool FObjectReferenceCache::ReadExports(FNetSerializationContext& Context)
+void FObjectReferenceCache::ReadMustBeMappedExports(FNetSerializationContext& Context, TArray<FNetRefHandle>* MustBeMappedExports)
 {
 	FNetBitStreamReader& Reader = *Context.GetBitStreamReader();
 
 	bool bHasExportsToRead = Reader.ReadBool(); 
-	UE_NET_TRACE_SCOPE(Exports, Reader, Context.GetTraceCollector(), ENetTraceVerbosity::Verbose);
+	UE_NET_TRACE_SCOPE(MustBeMappedExports, Reader, Context.GetTraceCollector(), ENetTraceVerbosity::Verbose);
 
 	while (bHasExportsToRead && !Context.HasErrorOrOverflow())
 	{
-		FNetObjectReference Import;
-		ReadFullReference(Context, Import);
-		bHasExportsToRead  = Reader.ReadBool();
-	}
+		const FNetRefHandle MustBeMappedHandle = ReadNetRefHandle(Context);
+		bHasExportsToRead = Reader.ReadBool();
 
-	return !Context.HasErrorOrOverflow();
+		if (MustBeMappedExports)
+		{
+			MustBeMappedExports->Add(MustBeMappedHandle);
+		}
+	}
 }
 
 FString FObjectReferenceCache::FullPath(FNetRefHandle RefHandle, const FNetObjectResolveContext& ResolveContext) const
@@ -1425,7 +1548,7 @@ FString FObjectReferenceCache::FullPath(FNetRefHandle RefHandle, const FNetObjec
 
 void FObjectReferenceCache::GenerateFullPath_r(FNetRefHandle RefHandle, const FNetObjectResolveContext& ResolveContext, FString& FullPath) const
 {
-	if ( !RefHandle.IsValid() )
+	if (!RefHandle.IsValid())
 	{
 		// This is the end of the outer chain, we're done
 		return;
diff --git a/Engine/Source/Runtime/Experimental/Iris/Core/Private/Iris/ReplicationSystem/ObjectReferenceCache.h b/Engine/Source/Runtime/Experimental/Iris/Core/Private/Iris/ReplicationSystem/ObjectReferenceCache.h
--- a/Engine/Source/Runtime/Experimental/Iris/Core/Private/Iris/ReplicationSystem/ObjectReferenceCache.h
+++ b/Engine/Source/Runtime/Experimental/Iris/Core/Private/Iris/ReplicationSystem/ObjectReferenceCache.h
@@ -81,12 +81,30 @@ public:
 	
 	// Read reference, as written by WriteReference
 	void ReadReference(FNetSerializationContext& Context, FNetObjectReference& OutRef);
- 
-	// Exports are expected to be part of the written state, so if WriteExports returns false or the BitStream is overflown
-	// it us up to the caller to roll back written data and exports
-	bool WriteExports(FNetSerializationContext& Context, TArrayView<const FNetObjectReference> ExportsView) const;
 
-	bool ReadExports(FNetSerializationContext& Context);
+	// Add exports to the set of pending exports for the current batch being written
+	void AddPendingExports(FNetSerializationContext& Context, TArrayView<const FNetObjectReference> ExportsView) const;
+
+	// Add export to the set of pending exports for the current batch being written
+	void AddPendingExport(FNetExportContext& ExportContext, const FNetObjectReference& Reference) const;
+
+	enum class EWriteExportsResult : unsigned
+	{
+		// We did write exports
+		WroteExports,
+
+		// BitStream overflow.
+		BitStreamOverflow,
+
+		// Some error occurred while serializing the object.
+		NoExports,
+	};
+
+	// Exports are expected to be part of the written state, so if the result is a BitStreamOverflow
+	// it is up to the caller to roll back written data and pending exports
+	EWriteExportsResult WritePendingExports(FNetSerializationContext& Context);
+
+	bool ReadExports(FNetSerializationContext& Context, TArray<FNetRefHandle>* MustBeMappedExports);
 
 	static FNetObjectReference MakeNetObjectReference(FNetRefHandle Handle);
 
@@ -138,6 +156,11 @@ private:
 	static FNetObjectReference MakeNetObjectReference(FNetRefHandle RefHandle, FNetToken RelativePath);
 	static FNetObjectReference MakeNetObjectReference(const FCachedNetObjectReference& CachedReference);
 
+	// Must be mapped exports are written for each batch that serializes object references, if async loading is enabled the client
+	// will defer application of data contained in the batch until the must be mapped exports are resolvable.
+	bool WriteMustBeMappedExports(FNetSerializationContext& Context, TArrayView<const FNetObjectReference> ExportsView) const;
+	void ReadMustBeMappedExports(FNetSerializationContext& Context, TArray<FNetRefHandle>* MustBeMappedExports);
+
 private:
 
 	// Map raw UObject pointer -> Handle
diff --git a/Engine/Source/Runtime/Experimental/Iris/Core/Private/Iris/ReplicationSystem/ReplicationReader.cpp b/Engine/Source/Runtime/Experimental/Iris/Core/Private/Iris/ReplicationSystem/ReplicationReader.cpp
--- a/Engine/Source/Runtime/Experimental/Iris/Core/Private/Iris/ReplicationSystem/ReplicationReader.cpp
+++ b/Engine/Source/Runtime/Experimental/Iris/Core/Private/Iris/ReplicationSystem/ReplicationReader.cpp
@@ -366,44 +366,137 @@ void FReplicationReader::DeserializeObjectStateDelta(FNetSerializationContext& C
 	}
 }
 
-void FReplicationReader::ReadObject(FNetSerializationContext& Context)
+uint32 FReplicationReader::ReadObjectBatch(FNetSerializationContext& Context)
 {
 	FNetBitStreamReader& Reader = *Context.GetBitStreamReader();
 
+	// Special handling for desrtuction infos
 	if (const bool bIsDestructionInfo = Reader.ReadBool())
 	{
 		FReplicationBridgeSerializationContext BridgeContext(Context, Parameters.ConnectionId, true);
 
+		// For destruction infos we inline the exports
+		FForceInlineExportScope ForceInlineExportScope(Context.GetInternalContext());
 		ReplicationBridge->ReadAndExecuteDestructionInfoFromRemote(BridgeContext);
 
 #if UE_NET_USE_READER_WRITER_SENTINEL
 		ReadAndVerifySentinelBits(&Reader, TEXT("DestructionInfo"), 8);
 #endif
 		
-		return;
+		return 1U;
 	}
 
 #if UE_NET_USE_READER_WRITER_SENTINEL
 	ReadAndVerifySentinelBits(&Reader, TEXT("ReadObject"), 8);
 #endif
 
+	uint32 ObjectsReadInBatch = 0U;
+	
+	// A batch starts with (RefHandleId | BatchSize | bHasBatchObjectData | bHasExports)
+	// If the batch has exports we must to seek to the end of the batch to read and process exports before reading/processing batch data
 	const FNetRefHandle IncompleteHandle = ReadNetRefHandleId(Reader);
+
+	uint32 BatchSize = 0U;
+	// Read Batch size
+	{
+		const uint32 NumBitsUsedForBatchSize = Parameters.NumBitsUsedForBatchSize;
+
+		UE_NET_TRACE_SCOPE(BatchSize, Reader, Context.GetTraceCollector(), ENetTraceVerbosity::Trace);
+		BatchSize = Reader.ReadBits(NumBitsUsedForBatchSize);
+	}
+
+	if (Context.HasErrorOrOverflow() || BatchSize > Reader.GetBitsLeft())
+	{
+		Context.SetError(GNetError_InvalidValue);
+		return 0U;
+	}
+
+	// This either marks the end of the data associated with this batch or the offset in the stream where exports are stored.
+	const uint32 BatchEndOrStartOfExportsPos = Reader.GetPosBits() + BatchSize;
+
+	uint32 ReadObjectCount = 0;
+
+	// Do we have state data or attachments for the owner of the batch?
+	const bool bHasBatchOwnerData = Reader.ReadBool();
+
+	// Do we have exports or not?
+	const bool bHasExports = Reader.ReadBool();
+
+	// First we need to read exports, they are stored at the end of the batch
+	uint32 BatchEndPos = BatchEndOrStartOfExportsPos;
+	if (bHasExports)
+	{
+		const uint32 ReturnPos = Reader.GetPosBits();
+
+		// Seek to the export section
+		Reader.Seek(BatchEndPos);
+
+		ObjectReferenceCache->ReadExports(Context, nullptr /* Currently we do not process MustBeMappedReferences */);
+		if (Context.HasErrorOrOverflow())
+		{
+			return 0U;
+		}
+
+		// Update BatchEndPos if we successfully read exports
+		BatchEndPos = Reader.GetPosBits();
+
+		// Seek back to state data
+		Reader.Seek(ReturnPos);
+	}
+
+	// If the batch owner had state, we read it now
+	if (bHasBatchOwnerData)
+	{
+		ReadObjectInBatch(Context, IncompleteHandle, false);
+		if (Context.HasErrorOrOverflow())
+		{
+			return 0U;
+		}
+		++ReadObjectCount;
+	}
+
+	// ReadSubObjects 
+	while (Reader.GetPosBits() < BatchEndOrStartOfExportsPos)
+	{
+		ReadObjectInBatch(Context, IncompleteHandle, true);
+		if (Context.HasErrorOrOverflow())
+		{
+			return 0U;
+		}
+		++ReadObjectCount;
+	}
+
+	// Skip to the end as we already have read any exports
+	Reader.Seek(BatchEndPos);
+
+	return ReadObjectCount;
+}
+
+void FReplicationReader::ReadObjectInBatch(FNetSerializationContext& Context, FNetRefHandle BatchHandle, bool bIsSubObject)
+{
+	FNetBitStreamReader& Reader = *Context.GetBitStreamReader();
+
+	// If we are reading 
+	FNetRefHandle IncompleteHandle = BatchHandle;
+	if (bIsSubObject)
+	{
+		IncompleteHandle = ReadNetRefHandleId(Reader);
+	}
 	
 	// Read replicated destroy header if necessary
 	const bool bReadReplicatedDestroyHeader = !IsObjectIndexForOOBAttachment(IncompleteHandle.GetId());
 	const uint32 ReplicatedDestroyHeaderFlags = bReadReplicatedDestroyHeader ? Reader.ReadBits(ReplicatedDestroyHeaderFlags_BitCount) : ReplicatedDestroyHeaderFlags_None;
 
 	const bool bHasState = Reader.ReadBool();
-	uint32 NewBaselineIndex = FDeltaCompressionBaselineManager::InvalidBaselineIndex;
-
 	if (bHasState)
 	{
-		ObjectReferenceCache->ReadExports(Context);
 #if UE_NET_USE_READER_WRITER_SENTINEL
 		ReadAndVerifySentinelBits(&Reader, TEXT("Exports"), 8);
 #endif
 	}
 
+	uint32 NewBaselineIndex = FDeltaCompressionBaselineManager::InvalidBaselineIndex;
+
 	const bool bIsInitialState = bHasState && Reader.ReadBool();
 	uint32 InternalIndex = ObjectIndexForOOBAttachment;
 
@@ -419,16 +512,12 @@ void FReplicationReader::ReadObject(FNetSerializationContext& Context)
 	{
 		UE_NET_TRACE_SCOPE(CreationInfo, *Context.GetBitStreamReader(), Context.GetTraceCollector(), ENetTraceVerbosity::Trace);
 
-#if UE_NET_REPLICATION_SUPPORT_SKIP_INITIAL_STATE
-		const uint32 NumBitsUsedForInitialStateSize = Parameters.NumBitsUsedForInitialStateSize;
-		const uint32 SkipSeekPos = Reader.ReadBits(NumBitsUsedForInitialStateSize) + Reader.GetPosBits();
-#endif
-
-		// SubObject data
+		// SubObject data for initial state
 		FNetRefHandle SubObjectOwnerHandle;
-		if (Reader.ReadBool())
+		if (bIsSubObject)
 		{
-			const FNetRefHandle IncompleteOwnerHandle = ReadNetRefHandleId(Reader);
+			// The owner is the same as the Batch owner
+			const FNetRefHandle IncompleteOwnerHandle = BatchHandle;
 				
 			FInternalNetRefIndex SubObjectOwnerInternalIndex = NetRefHandleManager->GetInternalIndex(IncompleteOwnerHandle);
 			if (Reader.IsOverflown() || SubObjectOwnerInternalIndex == FNetRefHandleManager::InvalidInternalIndex)
@@ -464,14 +553,7 @@ void FReplicationReader::ReadObject(FNetSerializationContext& Context)
 		const FReplicationBridgeCreateNetRefHandleResult CreateResult = ReplicationBridge->CallCreateNetRefHandleFromRemote(SubObjectOwnerHandle, IncompleteHandle, BridgeContext);
 		FNetRefHandle NetRefHandle = CreateResult.NetRefHandle;
 		if (!NetRefHandle.IsValid())
-		{
-#if UE_NET_REPLICATION_SUPPORT_SKIP_INITIAL_STATE
-			// If we support skipping
-			UE_LOG_REPLICATIONREADER_WARNING(TEXT("FReplicationReader::ReadObject Failed to instantiate %s, skipping over it assuming that object was streamed out"), *NetRefHandle.ToString());
-			Reader.Seek(SkipSeekPos);
-			return;
-#endif
-	
+		{	
 			UE_LOG_REPLICATIONREADER_ERROR(TEXT("FReplicationReader::ReadObject Unable to create handle for %s."), *IncompleteHandle.ToString());
 			Context.SetError(GNetError_InvalidNetHandle);
 			bHasErrors = true;
@@ -1293,6 +1375,7 @@ void FReplicationReader::DispatchEndReplication(FNetSerializationContext& Contex
 	}
 }
 
+
 void FReplicationReader::ReadObjects(FNetSerializationContext& Context, uint32 ObjectCountToRead)
 {
 	IRIS_PROFILER_SCOPE(ReplicationReader_ReadObjects);
@@ -1301,8 +1384,14 @@ void FReplicationReader::ReadObjects(FNetSerializationContext& Context, uint32 O
 	
 	while (ObjectCountToRead && !Context.HasErrorOrOverflow())
 	{
-		ReadObject(Context);
-		--ObjectCountToRead;
+		const uint32 ReadObjectCount = ReadObjectBatch(Context);
+		if (ReadObjectCount > ObjectCountToRead)
+		{			
+			Context.SetError(GNetError_BitStreamError);
+			break;
+		}
+
+		ObjectCountToRead -= ReadObjectCount;
 	}
 
 	ensureAlwaysMsgf(!Context.HasErrorOrOverflow(), TEXT("Overflow: %c Error: %s Bit stream bits left: %u position: %u"), "YN"[Context.HasError()], ToCStr(Context.GetError().ToString()), Reader.GetBitsLeft(), Reader.GetPosBits());
@@ -1366,8 +1455,8 @@ void FReplicationReader::ProcessHugeObjectAttachment(FNetSerializationContext& C
 	if (FNetTraceCollector* TraceCollector = Context.GetTraceCollector())
 	{
 		FNetBitStreamReader& Reader = *Context.GetBitStreamReader();
-		// Need a valid fake bitstream position, a position that starts before the end of the packet.
-		FNetTrace::FoldTraceCollector(TraceCollector, HugeObjectTraceCollector, GetBitStreamPositionForNetTrace(Reader) - 1U);
+		// Inject after all other trace events
+		FNetTrace::FoldTraceCollector(TraceCollector, HugeObjectTraceCollector, GetBitStreamPositionForNetTrace(Reader));
 	}
 #endif
 }
diff --git a/Engine/Source/Runtime/Experimental/Iris/Core/Private/Iris/ReplicationSystem/ReplicationReader.h b/Engine/Source/Runtime/Experimental/Iris/Core/Private/Iris/ReplicationSystem/ReplicationReader.h
--- a/Engine/Source/Runtime/Experimental/Iris/Core/Private/Iris/ReplicationSystem/ReplicationReader.h
+++ b/Engine/Source/Runtime/Experimental/Iris/Core/Private/Iris/ReplicationSystem/ReplicationReader.h
@@ -129,7 +129,10 @@ private:
 	uint32 ReadInternalIndex(FNetBitStreamReader& Reader) const;
 
 	// Read a new or updated object
-	void ReadObject(FNetSerializationContext& Context);
+	uint32 ReadObjectBatch(FNetSerializationContext& Context);
+
+	// Read object or subobject
+	void ReadObjectInBatch(FNetSerializationContext& Context, FNetRefHandle BatchHandle, bool bIsSubObject);
 
 	// Read all objects pending destroy
 	uint16 ReadObjectsPendingDestroy(FNetSerializationContext& Context);
diff --git a/Engine/Source/Runtime/Experimental/Iris/Core/Private/Iris/ReplicationSystem/ReplicationTypes.h b/Engine/Source/Runtime/Experimental/Iris/Core/Private/Iris/ReplicationSystem/ReplicationTypes.h
--- a/Engine/Source/Runtime/Experimental/Iris/Core/Private/Iris/ReplicationSystem/ReplicationTypes.h
+++ b/Engine/Source/Runtime/Experimental/Iris/Core/Private/Iris/ReplicationSystem/ReplicationTypes.h
@@ -5,9 +5,6 @@
 
 class UReplicationSystem;
 
-// If set to 1 we include the size of our creation state so that we can skip that on clients
-#define UE_NET_REPLICATION_SUPPORT_SKIP_INITIAL_STATE 0
-
 namespace UE::Net::Private
 {
 
@@ -22,9 +19,7 @@ struct FReplicationParameters
 	bool bAllowDelayingAttachmentsWithUnresolvedReferences = false;
 	uint32 SmallObjectBitThreshold = 160U; // Number of bits remaining in a packet for us to consider trying to serialize a replicated object
 	uint32 MaxFailedSmallObjectCount = 10U;	// Number of objects that we try to serialize after an initial stream overflow to fill up a packet, this can improve bandwidth usage but comes at a cpu cost
-#if UE_NET_REPLICATION_SUPPORT_SKIP_INITIAL_STATE
-	uint32 NumBitsUsedForInitialStateSize = 20U;
-#endif
+	uint32 NumBitsUsedForBatchSize = 20U;
 };
 
 enum EReplicatedDestroyHeaderFlags : uint32
diff --git a/Engine/Source/Runtime/Experimental/Iris/Core/Private/Iris/ReplicationSystem/ReplicationWriter.cpp b/Engine/Source/Runtime/Experimental/Iris/Core/Private/Iris/ReplicationSystem/ReplicationWriter.cpp
--- a/Engine/Source/Runtime/Experimental/Iris/Core/Private/Iris/ReplicationSystem/ReplicationWriter.cpp
+++ b/Engine/Source/Runtime/Experimental/Iris/Core/Private/Iris/ReplicationSystem/ReplicationWriter.cpp
@@ -171,6 +171,7 @@ void FReplicationWriter::FReplicationInfo::SetState(EReplicatedObjectState NewSt
 FReplicationWriter::FHugeObjectContext::FHugeObjectContext()
 : SendStatus(EHugeObjectSendStatus::Idle)
 , InternalIndex(0)
+, TraceCollector(nullptr)
 , DebugName(CreatePersistentNetDebugName(TEXT("HugeObjectState"), UE_ARRAY_COUNT(TEXT("HugeObjectState"))))
 , StartSendingTime(0)
 , EndSendingTime(0)
@@ -181,6 +182,12 @@ FReplicationWriter::FHugeObjectContext::FHugeObjectContext()
 #endif
 }
 
+FReplicationWriter::FHugeObjectContext::~FHugeObjectContext()
+{
+	UE_NET_TRACE_DESTROY_COLLECTOR(TraceCollector);
+	TraceCollector = nullptr;
+}
+
 // Default allocator for changemasks
 static FGlobalChangeMaskAllocator s_DefaultChangeMaskAllocator;
 
@@ -511,9 +518,11 @@ const FReplicationWriter::FReplicationInfo& FReplicationWriter::GetReplicationIn
 	return ReplicatedObjects[InternalIndex];
 }
 
-void FReplicationWriter::WriteNetRefHandleId(FNetBitStreamWriter& Writer, FNetRefHandle Handle)
+void FReplicationWriter::WriteNetRefHandleId(FNetSerializationContext& Context, FNetRefHandle Handle)
 {
-	WritePackedUint64(&Writer, Handle.GetId());
+	FNetBitStreamWriter* Writer = Context.GetBitStreamWriter();
+	UE_NET_TRACE_SCOPE(Handle, *Writer, Context.GetTraceCollector(), ENetTraceVerbosity::Trace);
+	WritePackedUint64(Writer, Handle.GetId());
 }
 
 uint32 FReplicationWriter::GetDefaultFlushFlags() const
@@ -1619,7 +1628,7 @@ uint32 FReplicationWriter::WriteObjectsPendingDestroy(FNetSerializationContext&
 		FNetBitStreamRollbackScope RollbackScope(Writer);
 
 		// Write handle with the needed bitCount
-		WriteNetRefHandleId(Writer, ObjectData.RefHandle);
+		WriteNetRefHandleId(Context, ObjectData.RefHandle);
 
 		// Write bit indicating if the static instance should be destroyed or not (could skip the bit for dynamic objects)
 		const bool bShouldDestroyInstance = ObjectData.RefHandle.IsDynamic() || NetRefHandleManager->GetIsDestroyedStartupObject(InternalIndex);
@@ -1757,7 +1766,7 @@ void FReplicationWriter::SerializeObjectStateDelta(FNetSerializationContext& Con
 	}
 }
 
-FReplicationWriter::EWriteObjectStatus FReplicationWriter::WriteObjectInBatch(FNetSerializationContext& Context, uint32 InternalIndex, uint32 WriteObjectFlags, FBatchInfo& OutBatchInfo)
+FReplicationWriter::EWriteObjectStatus FReplicationWriter::WriteObjectAndSubObjects(FNetSerializationContext& Context, uint32 InternalIndex, uint32 WriteObjectFlags, FBatchInfo& OutBatchInfo)
 {
 	IRIS_PROFILER_SCOPE(FReplicationWriter_WriteObjectInBatch);
 
@@ -1772,13 +1781,54 @@ FReplicationWriter::EWriteObjectStatus FReplicationWriter::WriteObjectInBatch(FN
 		return EWriteObjectStatus::InvalidState;
 	}
 
-	const bool bIsInitialState = IsInitialState(State);
+	// If this object or anything else included in the batch did not write any data we will rollback any data written for the object
+	FNetBitStreamRollbackScope ObjectRollbackScope(Writer);
 
-#if UE_NET_REPLICATION_SUPPORT_SKIP_INITIAL_STATE
-	uint32 InitialStateHeaderPos = 0U;
-	const uint32 NumBitsUsedForInitialStateSize = Parameters.NumBitsUsedForInitialStateSize;
+	const FNetRefHandleManager::FReplicatedObjectData& ObjectData = NetRefHandleManager->GetReplicatedObjectDataNoCheck(InternalIndex);
+	const FNetRefHandle NetRefHandle = ObjectData.RefHandle;
+
+#if UE_NET_TRACE_ENABLED
+	const FNetRefHandle NetRefHandleForTraceScope = (WriteObjectFlags & EWriteObjectFlag::WriteObjectFlag_HugeObject ? NetRefHandleManager->GetReplicatedObjectDataNoCheck(HugeObjectContext.InternalIndex).RefHandle : NetRefHandle);
+	UE_NET_TRACE_OBJECT_SCOPE(NetRefHandleForTraceScope, Writer, Context.GetTraceCollector(), ENetTraceVerbosity::Trace);
 #endif
 
+	// We only need to write batch 
+	const bool bWriteBatchInfo = !Info.IsSubObject;
+	uint32 InitialStateHeaderPos = 0U;
+	const uint32 NumBitsUsedForBatchSize = Parameters.NumBitsUsedForBatchSize;
+
+	// This is the beginning of what we treat as a batch on the receiving end
+	if (bWriteBatchInfo)
+	{
+		// Write bit indicating that we are not a destruction info.
+		const bool bIsDestructionInfo = false;
+		Writer.WriteBool(bIsDestructionInfo);
+
+	#if UE_NET_USE_READER_WRITER_SENTINEL
+		WriteSentinelBits(&Writer, 8);
+	#endif
+
+		// A batch starts with (RefHandleId | BatchSize | bHasBatchObjectData | bHasExports)
+		// We write the header up front, and then we seek back and update relevant info is the object + subobjects is successfully serialized along with necessary exports
+
+		// We send the Index of the handle to the remote end
+		// $IRIS: $TODO: consider sending the internal index instead to save bits and only send handle when we create the object, https://jira.it.epicgames.com/browse/UE-127373
+		WriteNetRefHandleId(Context, NetRefHandle);
+
+		InitialStateHeaderPos = Writer.GetPosBits();
+		{
+			UE_NET_TRACE_SCOPE(BatchSize, Writer, Context.GetTraceCollector(), ENetTraceVerbosity::Trace);
+			Writer.WriteBits(0U, NumBitsUsedForBatchSize);
+		}
+
+		// Did we write serialize any data related to batch owner
+		Writer.WriteBool(false);
+
+		// If the batch has exports, they are at the end of the batch
+		// We handle this on the reading side to avoid rewriting the entire object to insert exports up front.	
+		Writer.WriteBool(false);
+	}
+	
 	// Create a temporary batch entry. We don't want to push it to the batch info unless we're successful.
 	FBatchObjectInfo BatchEntry = {};
 
@@ -1796,9 +1846,9 @@ FReplicationWriter::EWriteObjectStatus FReplicationWriter::WriteObjectInBatch(FN
 	};
 
 	// Only write data for the object if we have data to write
-	const FNetRefHandleManager::FReplicatedObjectData& ObjectData = NetRefHandleManager->GetReplicatedObjectDataNoCheck(InternalIndex);
 	uint8* ReplicatedObjectStateBuffer = NetRefHandleManager->GetReplicatedObjectStateBufferNoCheck(InternalIndex);
 
+	const bool bIsInitialState = IsInitialState(State);
 	//IRIS_PROFILER_PROTOCOL_NAME(ObjectData.Protocol?ObjectData.Protocol->DebugName->Name:TEXT("NoProtocol"));
 
 	// Filter out changemasks that are not supposed to be replicated to this connection
@@ -1825,33 +1875,17 @@ FReplicationWriter::EWriteObjectStatus FReplicationWriter::WriteObjectInBatch(FN
 
 	Context.SetIsInitState(bIsInitialState);
 
-	// Update flag since dependent objects might change their status after creation
-	Info.IsSubObject = ObjectData.IsSubObject();
-
 	if (bHasState | bWriteAttachments | bSentTearOff | Info.SubObjectPendingDestroy)
 	{
-		const FNetRefHandle NetRefHandle = ObjectData.RefHandle;
-#if UE_NET_TRACE_ENABLED
-		const FNetRefHandle NetRefHandleForTraceScope = (WriteObjectFlags & EWriteObjectFlag::WriteObjectFlag_HugeObject ? NetRefHandleManager->GetReplicatedObjectDataNoCheck(HugeObjectContext.InternalIndex).RefHandle : NetRefHandle);
-		UE_NET_TRACE_OBJECT_SCOPE(NetRefHandleForTraceScope, Writer, Context.GetTraceCollector(), ENetTraceVerbosity::Trace);
-#endif
-		FNetBitStreamRollbackScope ObjectRollbackScope(Writer);
-
-		const bool bIsDestructionInfo = false;
-		Writer.WriteBool(bIsDestructionInfo);
-
-#if UE_NET_USE_READER_WRITER_SENTINEL
-		WriteSentinelBits(&Writer, 8);
-#endif
-
-		// We send the Index of the handle to the remote end
-		// $IRIS: $TODO: consider sending the internal index instead to save bits and only send handle when we create the object, https://jira.it.epicgames.com/browse/UE-127373
+		if (Info.IsSubObject)
 		{
-			UE_NET_TRACE_SCOPE(Handle, Writer, Context.GetTraceCollector(), ENetTraceVerbosity::Trace);
-			WriteNetRefHandleId(Writer, NetRefHandle);
+			// We send the Index of the handle to the remote end
+			// $IRIS: $TODO: consider sending the internal index instead to save bits and only send handle when we create the object, https://jira.it.epicgames.com/browse/UE-127373
+			UE_NET_TRACE_SCOPE(SubObjectHandle, Writer, Context.GetTraceCollector(), ENetTraceVerbosity::Trace);
+			WriteNetRefHandleId(Context, NetRefHandle);
 		}
 
-		// Store position of header bits
+		// Store position of destroy header bits
 		const uint32 ReplicatedDestroyHeaderBitPos = Writer.GetPosBits();
 
 		// We only need to write this for actual replicated objects
@@ -1864,6 +1898,10 @@ FReplicationWriter::EWriteObjectStatus FReplicationWriter::WriteObjectInBatch(FN
 
 		if (Writer.WriteBool(bHasState))
 		{
+#if UE_NET_USE_READER_WRITER_SENTINEL
+			WriteSentinelBits(&Writer, 8);
+#endif
+
 			BatchEntry.bSentState = 1;
 
 			// If the last transmitted baseline is acknowledged we can request a new baseline to be stored for the current state, we cannot compress against it until it has been acknowledged
@@ -1890,9 +1928,7 @@ FReplicationWriter::EWriteObjectStatus FReplicationWriter::WriteObjectInBatch(FN
 					{
 						CreatedBaselineIndex = NextBaselineIndex;
 
-						// $IRIS: $TODO: Consider conditionals, how to ensure that baseline stored on remote is complete if state of conditional changes?
 						// $IRIS: $TODO: Currently due to how repnotifies are implemented we might have to write an extra changemask when sending a new baseline to avoid extra calls to repnotifies
-
 						// Modify changemask to include any data we have in flight to ensure baseline integrity on receiving end
 						if (PatchupObjectChangeMaskWithInflightChanges(InternalIndex, Info))
 						{
@@ -1909,15 +1945,8 @@ FReplicationWriter::EWriteObjectStatus FReplicationWriter::WriteObjectInBatch(FN
 			FNetBitArrayView ChangeMask = MakeNetBitArrayView(Info.GetChangeMaskStoragePointer(), Info.ChangeMaskBitCount);
 			Context.SetChangeMask(&ChangeMask);
 
-			// Export references so that we can share the serialized state with other connection
-			if (!CollectAndWriteExports(Context, ReplicatedObjectStateBuffer, ObjectData.Protocol))
-			{
-				return EWriteObjectStatus::BitStreamOverflow;
-			}
-
-#if UE_NET_USE_READER_WRITER_SENTINEL
-			WriteSentinelBits(&Writer, 8);
-#endif
+			// Collect potential exports and append them to the list of pending exports to be exported with the batch
+			CollectAndAppendExports(Context, ReplicatedObjectStateBuffer, ObjectData.Protocol);
 			
 			// Role downgrade
 			{
@@ -1931,11 +1960,6 @@ FReplicationWriter::EWriteObjectStatus FReplicationWriter::WriteObjectInBatch(FN
 				{
 					UE_NET_TRACE_SCOPE(CreationInfo, Writer, Context.GetTraceCollector(), ENetTraceVerbosity::Trace);
 
-#if UE_NET_REPLICATION_SUPPORT_SKIP_INITIAL_STATE
-					InitialStateHeaderPos = Writer.GetPosBits();
-					Writer.WriteBits(0U, NumBitsUsedForInitialStateSize);
-#endif
-
 					// Warn if we cannot replicate this object
 					if (!ObjectData.InstanceProtocol)
 					{
@@ -1943,17 +1967,6 @@ FReplicationWriter::EWriteObjectStatus FReplicationWriter::WriteObjectInBatch(FN
 						return EWriteObjectStatus::NoInstanceProtocol;
 					}
 
-					// Real SubObjects needs to be tracked on client as well
-					if (Writer.WriteBool(ObjectData.IsSubObject()))
-					{
-						const FNetRefHandle OwnerRefHandle = NetRefHandleManager->GetReplicatedObjectDataNoCheck(ObjectData.SubObjectRootIndex).RefHandle;
-						WriteNetRefHandleId(Writer, OwnerRefHandle);
-						if (!OwnerRefHandle.IsValid() || Writer.IsOverflown())
-						{
-							return OwnerRefHandle.IsValid() ? EWriteObjectStatus::BitStreamOverflow : EWriteObjectStatus::InvalidOwner;
-						}
-					}
-
 					if (Writer.WriteBool(Info.IsDeltaCompressionEnabled))
 					{
 						// As we might fail to create a baseline for initial state we need to include it here.
@@ -2033,7 +2046,7 @@ FReplicationWriter::EWriteObjectStatus FReplicationWriter::WriteObjectInBatch(FN
 
 					Info.HasAttachments = BatchEntry.bHasUnsentAttachments;
 					
-					// If we should have had enough space to write a attachment, the attachement + exports must be huge and we need to fall back on using the huge object path
+					// If we should have had enough space to write a attachment, the attachment + exports must be huge and we need to fall back on using the huge object path
 					const uint32 SplitThreshold = PartialNetObjectAttachmentHandler->GetConfig()->GetBitCountSplitThreshold() * 2;
 					const bool bFallbackToHugeObjectPath = BatchEntry.AttachmentType != ENetObjectAttachmentType::HugeObject && BatchEntry.bHasUnsentAttachments && (BitsThatWasAvailableForAttachements >= SplitThreshold);
 					if (bFallbackToHugeObjectPath)
@@ -2041,7 +2054,7 @@ FReplicationWriter::EWriteObjectStatus FReplicationWriter::WriteObjectInBatch(FN
 						UE_LOG(LogIris, Verbose, TEXT("Failed to write huge attachment for object %s ( InternalIndex: %u ), forcing fallback on hugeobject for attachments"), *NetRefHandle.ToString(), InternalIndex);
 						Writer.DoOverflow();
 					}
-					else if (!(BatchEntry.bSentState || bSentTearOff || Info.SubObjectPendingDestroy))
+					else if (!(BatchEntry.bSentState || bSentTearOff || Info.SubObjectPendingDestroy || Info.HasDirtySubObjects))
 					{
 						// If we didn't send state and didn't send any attachments let's rollback
 						ObjectRollbackScope.Rollback();
@@ -2106,9 +2119,12 @@ FReplicationWriter::EWriteObjectStatus FReplicationWriter::WriteObjectInBatch(FN
 	CreatedBaselineIndex = FDeltaCompressionBaselineManager::InvalidBaselineIndex;
 
 	// Write dirty sub objects
+	bool bWroteSubObjectData = false;
 	if (Info.HasDirtySubObjects && !Info.IsSubObject)
 	{
 		bool bHasDirtySubObjects = false;
+
+		const uint32 SubObjectStartPos = Writer.GetPosBits();
 		
 		FReplicationConditionals::FSubObjectsToReplicateArray SubObjectsToReplicate;
 		ReplicationConditionals->GetSubObjectsToReplicate(Parameters.ConnectionId, InternalIndex, SubObjectsToReplicate);		
@@ -2121,7 +2137,7 @@ FReplicationWriter::EWriteObjectStatus FReplicationWriter::WriteObjectInBatch(FN
 			}
 
 			const int BatchObjectInfoCount = OutBatchInfo.ObjectInfos.Num();
-			EWriteObjectStatus SubObjectWriteStatus = WriteObjectInBatch(Context, SubObjectInternalIndex, WriteObjectFlags, OutBatchInfo);
+			EWriteObjectStatus SubObjectWriteStatus = WriteObjectAndSubObjects(Context, SubObjectInternalIndex, WriteObjectFlags, OutBatchInfo);
 			if (!IsWriteObjectSuccess(SubObjectWriteStatus))
 			{
 				// Need to remove the batch entry from BatchInfo.
@@ -2137,6 +2153,8 @@ FReplicationWriter::EWriteObjectStatus FReplicationWriter::WriteObjectInBatch(FN
 			}
 		}
 
+		bWroteSubObjectData = Writer.GetPosBits() != SubObjectStartPos;
+
 		// Update parent batch info
 		{
 			FBatchObjectInfo& ParentBatchEntry = OutBatchInfo.ObjectInfos[ParentBatchEntryIndex];
@@ -2144,34 +2162,91 @@ FReplicationWriter::EWriteObjectStatus FReplicationWriter::WriteObjectInBatch(FN
 		}
 	}
 
-	// If we have any dependent objects pending creation we currently include them in the batch in order to guarantee that they get sent in the same packet
-	// When we implement support for proper dependencies this can be changed.
-	for (const FDependentObjectInfo DependentObjectInfo : NetRefHandleManager->GetDependentObjectInfos(InternalIndex))
+	// ObjectBatch ends here
+	// We include the size of the data written so we can skip it if needed.
+	if (OutBatchInfo.ParentInternalIndex == InternalIndex)
 	{
-		const FInternalNetRefIndex DependentInternalIndex = DependentObjectInfo.NetRefIndex;
-		const bool bIsDependentInitialState = IsInitialState(GetReplicationInfo(DependentInternalIndex).GetState());
-		if (bIsDependentInitialState && !WriteContext.ObjectsWrittenThisPacket.GetBit(DependentInternalIndex))
+		const FBatchObjectInfo& ParentBatchEntry = OutBatchInfo.ObjectInfos[ParentBatchEntryIndex];
+
+		const uint32 WrittenBitsInBatch = (Writer.GetPosBits() - InitialStateHeaderPos) - NumBitsUsedForBatchSize;
+		
+		const bool bWroteData = (ParentBatchEntry.bSentState || ParentBatchEntry.bSentAttachments || bSentTearOff || Info.SubObjectPendingDestroy);
+		if (bWroteData || bWroteSubObjectData)
 		{
-			UE_NET_TRACE_SCOPE(DependentObjectData, *Context.GetBitStreamWriter(), Context.GetTraceCollector(), ENetTraceVerbosity::VeryVerbose);
-			EWriteObjectStatus DependentObjectWriteStatus = WriteObjectInBatch(Context, DependentInternalIndex, WriteObjectFlags, OutBatchInfo);
-			if (!IsWriteObjectSuccess(DependentObjectWriteStatus))
+			const FObjectReferenceCache::EWriteExportsResult WriteExportResult = ObjectReferenceCache->WritePendingExports(Context);
+
+			if (WriteExportResult == FObjectReferenceCache::EWriteExportsResult::BitStreamOverflow)
 			{
-				// Need to remove the batch entry from BatchInfo.
-				OutBatchInfo.ObjectInfos.RemoveAt(OutBatchInfo.ObjectInfos.Num() - 1);
-				return DependentObjectWriteStatus;
+				return EWriteObjectStatus::BitStreamOverflow;	
+			}
+
+			const bool bWroteExports = WriteExportResult == FObjectReferenceCache::EWriteExportsResult::WroteExports;
+
+			// Update header
+			{
+				FNetBitStreamWriteScope SizeScope(Writer, InitialStateHeaderPos);
+				Writer.WriteBits(WrittenBitsInBatch, NumBitsUsedForBatchSize);
+				Writer.WriteBool(bWroteData);
+				Writer.WriteBool(bWroteExports);
 			}
 		}
+		// If we did not write any data we rollback any written headers and report a succcess
+		else
+		{
+			// if we or our subobjects did not write any data, rollback and forget about everything
+			ObjectRollbackScope.Rollback();
+		}	
 	}
 
-	// We include the size of the data written so we can skip it if needed.
-#if UE_NET_REPLICATION_SUPPORT_SKIP_INITIAL_STATE
-	if (bIsInitialState)
+	return EWriteObjectStatus::Success;
+}
+
+FReplicationWriter::EWriteObjectStatus FReplicationWriter::WriteObjectInBatch(FNetSerializationContext& Context, uint32 InternalIndex, uint32 WriteObjectFlags, FBatchInfo& OutBatchInfo)
+{
+	IRIS_PROFILER_SCOPE(FReplicationWriter_WriteObjectInBatch);
+
+	// Reset pending exports
+	FNetExportContext* ExportContext = Context.GetExportContext();
+	if (ExportContext)
 	{
-		const uint32 InitialStateWrittenBits = (Writer.GetPosBits() - InitialStateHeaderPos) - NumBitsUsedForInitialStateSize;
-		FNetBitStreamWriteScope SizeScope(Writer, InitialStateHeaderPos);
-		Writer.WriteBits(InitialStateWrittenBits, NumBitsUsedForInitialStateSize);
+		ExportContext->ClearPendingExports();
+	}
+
+	// Write bit indicating if we are a destruction info or not
+	const EWriteObjectStatus WriteObjectStatus = WriteObjectAndSubObjects(Context, InternalIndex, WriteObjectFlags, OutBatchInfo);
+	if (!IsWriteObjectSuccess(WriteObjectStatus))
+	{
+		return WriteObjectStatus;
+	}
+
+	// Include dependent objects as separate batch
+	{
+		const uint32 OldBatchInfoParentInternalIndex = OutBatchInfo.ParentInternalIndex;
+		for (const FDependentObjectInfo DependentObjectInfo : NetRefHandleManager->GetDependentObjectInfos(InternalIndex))
+		{
+			const FInternalNetRefIndex DependentInternalIndex = DependentObjectInfo.NetRefIndex;
+			const bool bIsDependentInitialState = IsInitialState(GetReplicationInfo(DependentInternalIndex).GetState());
+			if (bIsDependentInitialState && !WriteContext.ObjectsWrittenThisPacket.GetBit(DependentInternalIndex))
+			{
+				UE_NET_TRACE_SCOPE(DependentObjectData, *Context.GetBitStreamWriter(), Context.GetTraceCollector(), ENetTraceVerbosity::VeryVerbose);
+
+				OutBatchInfo.ParentInternalIndex = DependentInternalIndex;
+				EWriteObjectStatus DependentObjectWriteStatus = WriteObjectInBatch(Context, DependentInternalIndex, WriteObjectFlags, OutBatchInfo);
+				if (!IsWriteObjectSuccess(DependentObjectWriteStatus))
+				{
+					// Need to remove the batch entry from BatchInfo.
+					OutBatchInfo.ObjectInfos.RemoveAt(OutBatchInfo.ObjectInfos.Num() - 1);
+
+					// Restore ParentInternalIndex
+					OutBatchInfo.ParentInternalIndex = OldBatchInfoParentInternalIndex;
+					return DependentObjectWriteStatus;
+				}
+			}
+		}
+
+		// Restore ParentInternalIndex
+		OutBatchInfo.ParentInternalIndex = OldBatchInfoParentInternalIndex;
 	}
-#endif
 
 	return EWriteObjectStatus::Success;
 }
@@ -2190,7 +2265,18 @@ int FReplicationWriter::PrepareAndSendHugeObjectPayload(FNetSerializationContext
 	const uint32 MaxHugeObjectPayLoadBytes = HugeObjectPayload.Num() * sizeof(HugeObjectStorageType);
 	HugeObjectWriter.InitBytes(HugeObjectPayload.GetData(), MaxHugeObjectPayLoadBytes);
 	FNetSerializationContext HugeObjectSerializationContext = Context.MakeSubContext(&HugeObjectWriter);
-	HugeObjectSerializationContext.SetTraceCollector(nullptr);
+
+#if UE_NET_TRACE_ENABLED
+	if (!HugeObjectContext.TraceCollector)
+	{
+		HugeObjectContext.TraceCollector = UE_NET_TRACE_CREATE_COLLECTOR(ENetTraceVerbosity::Trace);
+	}
+	else
+	{
+		HugeObjectContext.TraceCollector->Reset();
+	}
+	HugeObjectSerializationContext.SetTraceCollector(HugeObjectContext.TraceCollector);
+#endif
 
 	// Huge object header needed for the receiving side to be able to process this correctly.
 	FNetObjectBlob::FHeader HugeObjectHeader = {};
@@ -2209,28 +2295,32 @@ int FReplicationWriter::PrepareAndSendHugeObjectPayload(FNetSerializationContext
 	}
 	
 	// Push new ExportContext for the hugeobject-batch as we cannot share exports with an OOB object
-	HugeObjectContext.BatchExports.Reset();
-	FNetExports::FExportScope ExportScope = NetExports->MakeExportScope(HugeObjectSerializationContext, HugeObjectContext.BatchExports);
+	{
+		HugeObjectContext.BatchExports.Reset();
+		FNetExports::FExportScope ExportScope = NetExports->MakeExportScope(HugeObjectSerializationContext, HugeObjectContext.BatchExports);
 
-	// We can encounter other errors than bitstream overflow now that we've got a really large buffer to write to.
-	const EWriteObjectStatus WriteHugeObjectStatus = WriteObjectInBatch(HugeObjectSerializationContext, InternalIndex, WriteObjectFlags, BatchInfo);
+		UE_NET_TRACE_SCOPE(HugeObjectState, *HugeObjectSerializationContext.GetBitStreamWriter(), HugeObjectSerializationContext.GetTraceCollector(), ENetTraceVerbosity::Trace);
 
-	// If we cannot fit the object in the largest supported buffer then we will never fit the object.
-	if (WriteHugeObjectStatus == EWriteObjectStatus::BitStreamOverflow)
-	{
-		UE_LOG(LogIris, Error, TEXT("Unable to fit object ( InternalIndex: %u ) in maximum combined payload of %u bytes. Connection %u will be disconnected."), InternalIndex, MaxHugeObjectPayLoadBytes, Context.GetLocalConnectionId());
-		Context.SetError(NetError_ObjectStateTooLarge);
-		return -1;
-	}
+		// We can encounter other errors than bitstream overflow now that we've got a really large buffer to write to.
+		const EWriteObjectStatus WriteHugeObjectStatus = WriteObjectInBatch(HugeObjectSerializationContext, InternalIndex, WriteObjectFlags, BatchInfo);
 
-	// If we encounter some other error we can try sending a smaller object in the meantime.
-	if (!IsWriteObjectSuccess(WriteHugeObjectStatus))
-	{
-		// Need to call this to cleanup data from batch
-		HandleObjectBatchFailure(WriteHugeObjectStatus, BatchInfo, WriteBitStreamInfo);
+		// If we cannot fit the object in the largest supported buffer then we will never fit the object.
+		if (WriteHugeObjectStatus == EWriteObjectStatus::BitStreamOverflow)
+		{
+			UE_LOG(LogIris, Error, TEXT("Unable to fit object ( InternalIndex: %u ) in maximum combined payload of %u bytes. Connection %u will be disconnected."), InternalIndex, MaxHugeObjectPayLoadBytes, Context.GetLocalConnectionId());
+			Context.SetError(NetError_ObjectStateTooLarge);
+			return -1;
+		}
 
-		UE_LOG(LogIris, Verbose, TEXT("Problem writing huge object ( InternalIndex: %u ). WriteObjectStatus: %u. Trying smaller object."), InternalIndex, unsigned(WriteHugeObjectStatus));
-		return 0;
+		// If we encounter some other error we can try sending a smaller object in the meantime.
+		if (!IsWriteObjectSuccess(WriteHugeObjectStatus))
+		{
+			// Need to call this to cleanup data from batch
+			HandleObjectBatchFailure(WriteHugeObjectStatus, BatchInfo, WriteBitStreamInfo);
+
+			UE_LOG(LogIris, Verbose, TEXT("Problem writing huge object ( InternalIndex: %u ). WriteObjectStatus: %u. Trying smaller object."), InternalIndex, unsigned(WriteHugeObjectStatus));
+			return 0;
+		}
 	}
 
 	HugeObjectContext.InternalIndex = InternalIndex;
@@ -2340,7 +2430,10 @@ int FReplicationWriter::WriteObjectBatch(FNetSerializationContext& Context, uint
 		FBatchInfo BatchInfo;
 		BatchInfo.Type = (InternalIndex == ObjectIndexForOOBAttachment ? (WriteObjectFlags & WriteObjectFlag_HugeObject ? EBatchInfoType::HugeObject : EBatchInfoType::OOBAttachment) : EBatchInfoType::Object);
 		BatchInfo.ParentInternalIndex = InternalIndex;
+
+		// Write an objects and its SubObjects, if object has dependent objects pending creation we currently write them as well but they will be as an individual batch
 		const EWriteObjectStatus WriteObjectStatus = WriteObjectInBatch(Context, InternalIndex, WriteObjectFlags, BatchInfo);
+
 		if (IsWriteObjectSuccess(WriteObjectStatus))
 		{
 			FBatchRecord BatchRecord;
@@ -2409,6 +2502,8 @@ int FReplicationWriter::WriteDestructionInfo(FNetSerializationContext& Context,
 
 	FReplicationBridgeSerializationContext BridgeContext(Context, Parameters.ConnectionId, true);
 
+	// Push ForceInlineExportScope to inline exports instead of writing exports later.
+	FForceInlineExportScope ForceInlineExportScope(Context.GetInternalContext());
 	if (!ReplicationBridge->CallWriteNetRefHandleCreationInfo(BridgeContext, ObjectData.RefHandle))
 	{
 		// Trigger Rollback
@@ -2946,6 +3041,14 @@ UDataStream::EWriteResult FReplicationWriter::Write(FNetSerializationContext& Co
 		check(s_ValidateReplicationRecord(&ReplicationRecord, Parameters.MaxActiveReplicatedObjectCount + 1U, false));
 #endif
 
+#if UE_NET_TRACE_ENABLED
+		if (FNetTraceCollector* Collector = HugeObjectContext.TraceCollector)
+		{
+			FNetTrace::FoldTraceCollector(Context.GetTraceCollector(), Collector, GetBitStreamPositionForNetTrace(Writer));
+			Collector->Reset();
+		}
+#endif
+
 	}
 	else 
 	{
@@ -3089,24 +3192,20 @@ void FReplicationWriter::ClearHugeObjectContext(FHugeObjectContext& Context) con
 	Context.StartStallTime = 0;
 }
 
-bool FReplicationWriter::CollectAndWriteExports(FNetSerializationContext& Context, uint8* RESTRICT InternalBuffer, const FReplicationProtocol* Protocol) const
+void FReplicationWriter::CollectAndAppendExports(FNetSerializationContext& Context, uint8* RESTRICT InternalBuffer, const FReplicationProtocol* Protocol) const
 {
-	FNetReferenceCollector Collector(ENetReferenceCollectorTraits::OnlyCollectReferencesThatCanBeExported);
+	FNetExportContext* ExportContext = Context.GetExportContext();
+	if (!ExportContext)
+	{
+		return;
+	}
 
+	FNetReferenceCollector Collector(ENetReferenceCollectorTraits::OnlyCollectReferencesThatCanBeExported);
 	FReplicationProtocolOperationsInternal::CollectReferences(Context, Collector, InternalBuffer, Protocol);
 
-	if (Collector.GetCollectedReferences().Num())
-	{
-		TArray<FNetObjectReference, TInlineAllocator<32>> Exports;
-		for (const FNetReferenceCollector::FReferenceInfo& Info : MakeArrayView(Collector.GetCollectedReferences()))
-		{
-			Exports.AddUnique(Info.Reference);
-		}
-		return ObjectReferenceCache->WriteExports(Context, MakeArrayView(Exports));
-	}
-	else
+	for (const FNetReferenceCollector::FReferenceInfo& Info : MakeArrayView(Collector.GetCollectedReferences()))
 	{
-		return ObjectReferenceCache->WriteExports(Context, MakeArrayView<const FNetObjectReference>(nullptr, 0));
+		ObjectReferenceCache->AddPendingExport(*ExportContext, Info.Reference);
 	}
 }
 
diff --git a/Engine/Source/Runtime/Experimental/Iris/Core/Private/Iris/ReplicationSystem/ReplicationWriter.h b/Engine/Source/Runtime/Experimental/Iris/Core/Private/Iris/ReplicationSystem/ReplicationWriter.h
--- a/Engine/Source/Runtime/Experimental/Iris/Core/Private/Iris/ReplicationSystem/ReplicationWriter.h
+++ b/Engine/Source/Runtime/Experimental/Iris/Core/Private/Iris/ReplicationSystem/ReplicationWriter.h
@@ -291,11 +291,13 @@ private:
 	struct FHugeObjectContext
 	{
 		FHugeObjectContext();
+		~FHugeObjectContext();
 
 		EHugeObjectSendStatus SendStatus;
 		uint32 InternalIndex;
 		FBatchRecord BatchRecord;
 		FNetExportContext::FBatchExports BatchExports;
+		FNetTraceCollector* TraceCollector;
 		const FNetDebugName* DebugName;
 		// Cycle counter for when the huge object context went from idle to sending.
 		uint64 StartSendingTime;
@@ -367,7 +369,7 @@ private:
 	void SetState(uint32 InternalIndex, EReplicatedObjectState NewState);
 
 	// Write index part of handle
-	void WriteNetRefHandleId(FNetBitStreamWriter& Writer, FNetRefHandle RefHandle);
+	void WriteNetRefHandleId(FNetSerializationContext& Context, FNetRefHandle RefHandle);
 		
 	// Create new ObjectRecord
 	// Note: be aware that it will allocate a copy of the ChangeMask that needs to be handled if the record is not Committed
@@ -388,6 +390,9 @@ private:
 	// Write all objects pending destroy (or as many as we fit in the current packet)
 	uint32 WriteObjectsPendingDestroy(FNetSerializationContext& Context);
 
+	// Write object and SubObjects
+	EWriteObjectStatus WriteObjectAndSubObjects(FNetSerializationContext& Context, uint32 InternalIndex, uint32 WriteObjectFlags, FBatchInfo& OutBatchInfo);
+
 	// Write objects recursive
 	EWriteObjectStatus WriteObjectInBatch(FNetSerializationContext& Context, uint32 InternalIndex, uint32 WriteObjectFlags, FBatchInfo& OutBatchInfo);
 
@@ -454,7 +459,7 @@ private:
 
 	bool HasDataToSend(const FWriteContext& Context) const;
 
-	bool CollectAndWriteExports(FNetSerializationContext& Context, uint8* RESTRICT InternalBuffer, const FReplicationProtocol* Protocol) const;
+	void CollectAndAppendExports(FNetSerializationContext& Context, uint8* RESTRICT InternalBuffer, const FReplicationProtocol* Protocol) const;
 
 	bool IsWriteObjectSuccess(EWriteObjectStatus Status) const;
 
diff --git a/Engine/Source/Runtime/Experimental/Iris/Core/Private/Iris/Serialization/InternalNetSerializationContext.h b/Engine/Source/Runtime/Experimental/Iris/Core/Private/Iris/Serialization/InternalNetSerializationContext.h
--- a/Engine/Source/Runtime/Experimental/Iris/Core/Private/Iris/Serialization/InternalNetSerializationContext.h
+++ b/Engine/Source/Runtime/Experimental/Iris/Core/Private/Iris/Serialization/InternalNetSerializationContext.h
@@ -55,4 +55,35 @@ inline FInternalNetSerializationContext::FInternalNetSerializationContext()
 , bInlineObjectReferenceExports(0)
 {}
 
+// Scope used when we actually want to export references.
+class FForceInlineExportScope
+{
+public:
+	explicit FForceInlineExportScope(FInternalNetSerializationContext* InContext);
+	~FForceInlineExportScope();
+
+private:
+	FInternalNetSerializationContext* InternalContext;
+	uint32 bOldInlineObjectReferences = 0U;
+};
+
+inline FForceInlineExportScope::FForceInlineExportScope(FInternalNetSerializationContext* InContext)
+: InternalContext(InContext)
+{
+	if (InternalContext)
+	{
+		bOldInlineObjectReferences = InternalContext->bInlineObjectReferenceExports;
+		InternalContext->bInlineObjectReferenceExports = 1U;
+	}
+}
+
+inline FForceInlineExportScope::~FForceInlineExportScope()
+{
+	if (InternalContext)
+	{
+		// Restore state
+		InternalContext->bInlineObjectReferenceExports = bOldInlineObjectReferences;
+	}
+}
+
 }
diff --git a/Engine/Source/Runtime/Experimental/Iris/Core/Private/Iris/Serialization/NetExportContext.h b/Engine/Source/Runtime/Experimental/Iris/Core/Private/Iris/Serialization/NetExportContext.h
--- a/Engine/Source/Runtime/Experimental/Iris/Core/Private/Iris/Serialization/NetExportContext.h
+++ b/Engine/Source/Runtime/Experimental/Iris/Core/Private/Iris/Serialization/NetExportContext.h
@@ -8,6 +8,7 @@
 #include "Iris/ReplicationSystem/NetRefHandle.h"
 #include "Iris/ReplicationSystem/NetToken.h"
 #include "Iris/Serialization/NetSerializationContext.h"
+#include "Iris/Core/NetObjectReference.h"
 
 namespace UE::Net::Private
 {
@@ -17,6 +18,7 @@ class FNetExportContext
 public:
 	typedef TArray<FNetRefHandle, TInlineAllocator<32>> FExportsArray;
 	typedef TArray<FNetToken, TInlineAllocator<32>> FNetTokenExportsArray;
+	typedef TArray<FNetObjectReference, TInlineAllocator<32>> FPendingExportArray;
 
 	struct FAcknowledgedExports
 	{
@@ -30,11 +32,13 @@ public:
 		{
 			HandlesExportedInCurrentBatch.Empty();
 			NetTokensExportedInCurrentBatch.Empty();
+			ReferencesPendingExportInCurrentBatch.Empty();
 		}
 
 		// Exports in the current batch
 		FExportsArray HandlesExportedInCurrentBatch;
 		FNetTokenExportsArray NetTokensExportedInCurrentBatch;
+		FPendingExportArray ReferencesPendingExportInCurrentBatch;
 	};
 
 public:
@@ -53,6 +57,15 @@ public:
 	// Add a Handle to the current export batch
 	void AddExported(FNetToken Token);
 
+	// Add a reference to the current pending exports list.
+	void AddPendingExport(const FNetObjectReference& Ref);
+
+	// Returns true if the Reference is in PendingExports array
+	bool IsPendingExport(const FNetObjectReference& Ref) const;
+
+	// Clear the list of reference pending exports
+	void ClearPendingExports() { BatchExports.ReferencesPendingExportInCurrentBatch.Empty(); }
+
 	// Get current batch exports
 	const FNetExportContext::FBatchExports& GetBatchExports() const { return BatchExports; }
 
@@ -79,6 +92,7 @@ private:
 	FNetSerializationContext& Context;
 	int32 StartNumNetHandleExports;
 	int32 StartNumNetTokenExports;
+	int32 StartNumPendingExports;
 };
 
 inline FNetExportRollbackScope::FNetExportRollbackScope(FNetSerializationContext& InContext)
@@ -88,6 +102,7 @@ inline FNetExportRollbackScope::FNetExportRollbackScope(FNetSerializationContext
 
 	StartNumNetHandleExports = ExportContext ? ExportContext->BatchExports.HandlesExportedInCurrentBatch.Num() : 0;
 	StartNumNetTokenExports = ExportContext ? ExportContext->BatchExports.NetTokensExportedInCurrentBatch.Num() : 0;
+	StartNumPendingExports = ExportContext ? ExportContext->BatchExports.ReferencesPendingExportInCurrentBatch.Num() : 0;
 }
 
 inline void FNetExportRollbackScope::Rollback()
@@ -96,6 +111,7 @@ inline void FNetExportRollbackScope::Rollback()
 	{ 
 		ExportContext->BatchExports.HandlesExportedInCurrentBatch.SetNum(StartNumNetHandleExports);
 		ExportContext->BatchExports.NetTokensExportedInCurrentBatch.SetNum(StartNumNetTokenExports);
+		ExportContext->BatchExports.ReferencesPendingExportInCurrentBatch.SetNum(StartNumPendingExports);
 	}
 }
 
@@ -134,4 +150,15 @@ inline void FNetExportContext::AddExported(FNetToken Token)
 	BatchExports.NetTokensExportedInCurrentBatch.Add(Token);
 }
 
+inline void FNetExportContext::AddPendingExport(const FNetObjectReference& Ref)
+{
+	BatchExports.ReferencesPendingExportInCurrentBatch.AddUnique(Ref);
+}
+
+inline bool FNetExportContext::IsPendingExport(const FNetObjectReference& Ref) const
+{
+	return BatchExports.ReferencesPendingExportInCurrentBatch.Contains(Ref);
+}
+
+
 }
