[Perf][GC] Optimize per-component collision/breaking/trailing event dispatching
#[fyi] Benn.Gallagher, Vincent.Robert

diff --git a/Engine/Source/Runtime/Experimental/Chaos/Public/EventManager.h b/Engine/Source/Runtime/Experimental/Chaos/Public/EventManager.h
--- a/Engine/Source/Runtime/Experimental/Chaos/Public/EventManager.h
+++ b/Engine/Source/Runtime/Experimental/Chaos/Public/EventManager.h
@@ -44,6 +44,7 @@ namespace Chaos
 	public:
 		virtual ~IEventHandler() {}
 		virtual void HandleEvent(const void* EventData) const = 0;
+		virtual bool GetInterestedProxyOwners(TArray<UObject*>& Output) const = 0;
 
 	protected:
 		// internal use only
@@ -56,13 +57,17 @@ namespace Chaos
 	{
 	public:
 		typedef void (HandlerType::*FHandlerFunction)(const PayloadType&);
+		typedef TArray<UObject*> (HandlerType::*FInterestedProxyOwnerFunction)();
 
-		TRawEventHandler(HandlerType* InHandler, FHandlerFunction InFunction)
+		TRawEventHandler(HandlerType* InHandler, FHandlerFunction InFunction, FInterestedProxyOwnerFunction InFunctionProxyOwners = nullptr)
 			: Handler(InHandler)
 			, HandlerFunction(InFunction)
+			, InterestedProxyOwnersFunction(InFunctionProxyOwners)
 		{
 			check(Handler);
 			check(HandlerFunction);
+			// This can be a nullptr
+			//check(InterestedProxyOwnersFunction);
 		}
 
 		virtual void HandleEvent(const void* EventData) const override
@@ -70,6 +75,17 @@ namespace Chaos
 			(Handler->*HandlerFunction)(*(const PayloadType*)EventData);
 		}
 
+		virtual bool GetInterestedProxyOwners(TArray<UObject*>& Output) const override
+		{
+			if (!InterestedProxyOwnersFunction)
+			{
+				return false;
+			}
+			
+			Output = (Handler->*InterestedProxyOwnersFunction)();
+			return true;
+		}
+
 	protected:
 		void* GetHandler() const override
 		{
@@ -79,6 +95,9 @@ namespace Chaos
 	private:
 		HandlerType* Handler;
 		FHandlerFunction HandlerFunction;
+		// This function is used to get the proxies we are interested in, used for optimization
+		// Will be a nullptr if the handler is interested in all proxies
+		FInterestedProxyOwnerFunction InterestedProxyOwnersFunction;
 	};
 
 	/**
@@ -162,6 +181,23 @@ namespace Chaos
 		virtual void RegisterHandler(const FEventHandlerPtr& Handler)
 		{
 			HandlerArray.AddUnique(Handler);
+			TArray<UObject*> ProxyOwners;
+			bool bValidProxyFilter = Handler->GetInterestedProxyOwners(ProxyOwners);
+
+			if (bValidProxyFilter)
+			{
+				for (UObject* ProxyOwner : ProxyOwners)
+				{
+					ProxyOwnerToHandlerMap.Add(ProxyOwner, Handler);
+				}
+			}
+			else
+			{
+				if (GetProxyToIndexMap(EventBuffer.Get()->GetConsumerBuffer()) != nullptr) // Only if our type supports getting the ProxyToIndexMap do we bother adding this
+				{
+					HandlersNotInMap.AddUnique(Handler);
+				}				
+			}
 		}
 
 		/**
@@ -169,12 +205,36 @@ namespace Chaos
 		 */
 		virtual void UnregisterHandler(const void* InHandler)
 		{
+			TArray<TPair<UObject*, FEventHandlerPtr>> KeysAndValuesToRemove;
+			for (TPair<UObject*, FEventHandlerPtr>& KeyValue : ProxyOwnerToHandlerMap)
+			{
+				if (KeyValue.Get<1>() == InHandler)
+				{
+					KeysAndValuesToRemove.Add(KeyValue);
+				}
+			}
+
+			for (TPair<UObject*, FEventHandlerPtr>& KeyAndValue : KeysAndValuesToRemove)
+			{
+				ProxyOwnerToHandlerMap.Remove(KeyAndValue.Get<0>(), KeyAndValue.Get<1>());
+			}
+
+			for (int i = 0; i < HandlersNotInMap.Num(); i++)
+			{
+				if (HandlersNotInMap[i]->GetHandler() == InHandler)
+				{
+
+					HandlersNotInMap.RemoveAtSwap(i, 1, false);
+					break;
+				}
+			}
+
 			for (int i = 0; i < HandlerArray.Num(); i++)
 			{
 				if (HandlerArray[i]->GetHandler() == InHandler)
 				{
 					DeleteHandler(HandlerArray[i]);
-					HandlerArray.RemoveAt(i);
+					HandlerArray.RemoveAtSwap(i, 1, false);
 					break;
 				}
 			}
@@ -218,10 +278,42 @@ namespace Chaos
 				return;
 			}
 
-			for (FEventHandlerPtr Handler : HandlerArray)
+			const TMap<IPhysicsProxyBase*, TArray<int32>>* Map = GetProxyToIndexMap(Buffer); // Use t his map to get all proxies used in the event buffer
+			// Only take this path if we have fewer Events than Handlers
+			if (Map && Map->Num() + HandlersNotInMap.Num() < HandlerArray.Num())
 			{
-				Handler->HandleEvent(Buffer);
+				TSet<FEventHandlerPtr> UniqueHandlers;
+				UniqueHandlers.Reserve(Map->Num());
+				for (auto& KeyValue : *Map) // Only iterating over objects that are associated with events here
+				{
+					IPhysicsProxyBase* Proxy = KeyValue.Get<0>();
+					TArray<FEventHandlerPtr> EventHandlers;
+					ProxyOwnerToHandlerMap.MultiFind(Proxy->GetOwner(), EventHandlers);
+					for (FEventHandlerPtr EventHandler : EventHandlers)
+					{
+						UniqueHandlers.Add(EventHandler);
+					}
+				}
+
+				for (FEventHandlerPtr Handler : UniqueHandlers)
+				{
+					Handler->HandleEvent(Buffer);
+				}
+
+				for (FEventHandlerPtr Handler : HandlersNotInMap)
+				{
+					Handler->HandleEvent(Buffer);
+				}
+				return;
 			}
+			
+			// This path is taken if there are fewer Handlers than events or the handler does not support GetProxyToIndexMap
+			{
+				for (FEventHandlerPtr Handler : HandlerArray)
+				{
+					Handler->HandleEvent(Buffer);
+				}
+			}			
 		}
 
 	private:
@@ -242,6 +334,9 @@ namespace Chaos
 		 */
 		TUniquePtr<IBufferResource<PayloadType>> EventBuffer;
 
+		TMultiMap<UObject*, FEventHandlerPtr> ProxyOwnerToHandlerMap; // Used to prevent us from iterating through the whole HandlerArray
+		TArray<FEventHandlerPtr> HandlersNotInMap; // Handlers not added to ProxyOwnerToHandlerMap since they do not support it
+
 		/**
 		 * Delegate function registered to handle this event when it is dispatched
 		 */
@@ -317,7 +412,7 @@ namespace Chaos
 		 * Register a handler that will receive the dispatched events
 		 */
 		template<typename PayloadType, typename HandlerType>
-		void RegisterHandler(const EEventType& EventType, HandlerType* Handler, typename TRawEventHandler<PayloadType, HandlerType>::FHandlerFunction HandlerFunction)
+		void RegisterHandler(const EEventType& EventType, HandlerType* Handler, typename TRawEventHandler<PayloadType, HandlerType>::FHandlerFunction HandlerFunction, typename TRawEventHandler<PayloadType, HandlerType>::FInterestedProxyOwnerFunction InterestedProxyOwnerFunction = nullptr)
 		{
 			FScopeLock ScopeLock(&AccessDeferredHandlersLock);
 
@@ -326,13 +421,13 @@ namespace Chaos
 			// If we are currently dispatching events, defer handler registration until completion of dispatch to avoid deadlock
 			if (bCurrentlyDispatchingEvents)
 			{
-				DeferredHandlers.Add(TPair<FEventID, IEventHandler*>(EventID, new TRawEventHandler<PayloadType, HandlerType>(Handler, HandlerFunction)));
+				DeferredHandlers.Add(TPair<FEventID, IEventHandler*>(EventID, new TRawEventHandler<PayloadType, HandlerType>(Handler, HandlerFunction, InterestedProxyOwnerFunction)));
 			}
 			else
 			{
 				ContainerLock.WriteLock();
 				checkf(EventID < EventContainers.Num(), TEXT("Registering event Handler for an event ID that does not exist"));
-				EventContainers[EventID]->RegisterHandler(new TRawEventHandler<PayloadType, HandlerType>(Handler, HandlerFunction));
+				EventContainers[EventID]->RegisterHandler(new TRawEventHandler<PayloadType, HandlerType>(Handler, HandlerFunction, InterestedProxyOwnerFunction));
 				ContainerLock.WriteUnlock();
 			}	
 		}
diff --git a/Engine/Source/Runtime/Experimental/Chaos/Public/EventsData.h b/Engine/Source/Runtime/Experimental/Chaos/Public/EventsData.h
--- a/Engine/Source/Runtime/Experimental/Chaos/Public/EventsData.h
+++ b/Engine/Source/Runtime/Experimental/Chaos/Public/EventsData.h
@@ -235,4 +235,42 @@ namespace Chaos
 			return false;
 		}
 	}
+	
+	template<typename PayloadType>
+	const TMap<IPhysicsProxyBase*, TArray<int32>>* GetProxyToIndexMap(const PayloadType* Buffer)
+	{
+		if (!Buffer)
+		{
+			return nullptr;
+		}
+
+		if constexpr (std::is_same_v<PayloadType, FCollisionEventData>)
+		{
+			return &Buffer->PhysicsProxyToCollisionIndices.PhysicsProxyToIndicesMap;
+		}
+		else if constexpr (std::is_same_v<PayloadType, FBreakingEventData>)
+		{
+			return nullptr; // &Buffer->PhysicsProxyToBreakingIndices.PhysicsProxyToIndicesMap;
+		}
+		else if constexpr (std::is_same_v<PayloadType, FTrailingEventData>)
+		{
+			return nullptr; // &Buffer->PhysicsProxyToTrailingIndices.PhysicsProxyToIndicesMap;
+		}
+		else if constexpr (std::is_same_v<PayloadType, FRemovalEventData>)
+		{
+			return nullptr; // &Buffer->PhysicsProxyToRemovalIndices.PhysicsProxyToIndicesMap;
+		}
+		else if constexpr (std::is_same_v<PayloadType, FSleepingEventData>)
+		{
+			return nullptr;
+		}
+		else if constexpr (std::is_same_v<PayloadType, FCrumblingEventData>)
+		{
+			return nullptr; // &Buffer->PhysicsProxyToCrumblingIndices.PhysicsProxyToIndicesMap;
+		}
+		else
+		{
+			return nullptr;
+		}
+	}
 }
diff --git a/Engine/Source/Runtime/Experimental/ChaosSolverEngine/Private/Chaos/ChaosGameplayEventDispatcher.cpp b/Engine/Source/Runtime/Experimental/ChaosSolverEngine/Private/Chaos/ChaosGameplayEventDispatcher.cpp
--- a/Engine/Source/Runtime/Experimental/ChaosSolverEngine/Private/Chaos/ChaosGameplayEventDispatcher.cpp
+++ b/Engine/Source/Runtime/Experimental/ChaosSolverEngine/Private/Chaos/ChaosGameplayEventDispatcher.cpp
@@ -320,7 +320,8 @@ void UChaosGameplayEventDispatcher::RegisterChaosEvents()
 		if (Chaos::FPhysicsSolver* Solver = Scene->GetSolver())
 		{
 			Chaos::FEventManager* EventManager = Solver->GetEventManager();
-			EventManager->RegisterHandler<Chaos::FCollisionEventData>(Chaos::EEventType::Collision, this, &UChaosGameplayEventDispatcher::HandleCollisionEvents);
+			EventManager->RegisterHandler<Chaos::FCollisionEventData>(Chaos::EEventType::Collision, this, &UChaosGameplayEventDispatcher::HandleCollisionEvents, &UChaosGameplayEventDispatcher::GetInterestedProxyOwnersForCollisionEvents);
+			// Todo implement GetInterestedProxies for these as well 
 			EventManager->RegisterHandler<Chaos::FBreakingEventData>(Chaos::EEventType::Breaking, this, &UChaosGameplayEventDispatcher::HandleBreakingEvents);
 			EventManager->RegisterHandler<Chaos::FRemovalEventData>(Chaos::EEventType::Removal, this, &UChaosGameplayEventDispatcher::HandleRemovalEvents);
 			EventManager->RegisterHandler<Chaos::FCrumblingEventData>(Chaos::EEventType::Crumbling, this, &UChaosGameplayEventDispatcher::HandleCrumblingEvents);
@@ -347,6 +348,29 @@ void UChaosGameplayEventDispatcher::UnregisterChaosEvents()
 	}
 }
 
+TArray<UObject*> UChaosGameplayEventDispatcher::GetInterestedProxyOwnersForCollisionEvents()
+{
+	TArray<UObject*> Result;
+	FPhysScene_Chaos& Scene = *(GetWorld()->GetPhysicsScene());
+		
+	// look through all the components that someone is interested in and get all the proxies
+	for (decltype(CollisionEventRegistrations)::TIterator It(CollisionEventRegistrations); It; ++It)
+	{
+		UPrimitiveComponent* const Comp0 = Cast<UPrimitiveComponent>(It.Key());
+		const TArray<IPhysicsProxyBase*>* PhysicsProxyArray = Scene.GetOwnedPhysicsProxies(Comp0);
+
+		if (PhysicsProxyArray)
+		{
+			for (IPhysicsProxyBase* PhysicsProxy0 : *PhysicsProxyArray)
+			{
+				Result.AddUnique(PhysicsProxy0->GetOwner());
+			}
+		}
+	}
+
+	return Result;
+}
+
 void UChaosGameplayEventDispatcher::HandleCollisionEvents(const Chaos::FCollisionEventData& Event)
 {
 	// todo(chaos) : this code is very similar to FPhysScene_Chaos::HandleCollisionEvents, we should propably consolidate if possible or share as much code as possible 
diff --git a/Engine/Source/Runtime/Experimental/ChaosSolverEngine/Public/Chaos/ChaosGameplayEventDispatcher.h b/Engine/Source/Runtime/Experimental/ChaosSolverEngine/Public/Chaos/ChaosGameplayEventDispatcher.h
--- a/Engine/Source/Runtime/Experimental/ChaosSolverEngine/Public/Chaos/ChaosGameplayEventDispatcher.h
+++ b/Engine/Source/Runtime/Experimental/ChaosSolverEngine/Public/Chaos/ChaosGameplayEventDispatcher.h
@@ -280,6 +280,8 @@ private:
 	void RegisterChaosEvents();
 	void UnregisterChaosEvents();
 
+
+	TArray<UObject*> GetInterestedProxyOwnersForCollisionEvents();
 	// Chaos Event Handlers
 	void HandleCollisionEvents(const Chaos::FCollisionEventData& CollisionData);
 	void HandleBreakingEvents(const Chaos::FBreakingEventData& BreakingData);
