potential fix for crash due to a particle having edges to invalid particles when in a cluster union

diff --git a/Engine/Source/Runtime/Experimental/Chaos/Private/Chaos/PBDRigidClustering.cpp b/Engine/Source/Runtime/Experimental/Chaos/Private/Chaos/PBDRigidClustering.cpp
--- a/Engine/Source/Runtime/Experimental/Chaos/Private/Chaos/PBDRigidClustering.cpp
+++ b/Engine/Source/Runtime/Experimental/Chaos/Private/Chaos/PBDRigidClustering.cpp
@@ -1632,6 +1632,12 @@ namespace Chaos
 			ClusteredParticle->ClusterIds() = ClusterId();
 			ClusteredParticle->ClusterGroupIndex() = 0;
 
+			// If we have a parent, need to disconnect potentially from fellow children in connectivity edges.
+			RemoveNodeConnections(ClusteredParticle);
+
+			// Need to also check if the particle is a cluster union and remove from that as well.
+			ClusterUnionManager.HandleRemoveOperationWithClusterLookup({ ClusteredParticle }, EClusterUnionOperationTiming::Defer);
+
 			if (MChildren.Contains(ParentParticle))
 			{
 				FRigidHandleArray& Children = MChildren[ParentParticle];
diff --git a/Engine/Source/Runtime/Experimental/Chaos/Private/PhysicsProxy/GeometryCollectionPhysicsProxy.cpp b/Engine/Source/Runtime/Experimental/Chaos/Private/PhysicsProxy/GeometryCollectionPhysicsProxy.cpp
--- a/Engine/Source/Runtime/Experimental/Chaos/Private/PhysicsProxy/GeometryCollectionPhysicsProxy.cpp
+++ b/Engine/Source/Runtime/Experimental/Chaos/Private/PhysicsProxy/GeometryCollectionPhysicsProxy.cpp
@@ -2469,6 +2469,8 @@ void FGeometryCollectionPhysicsProxy::OnRemoveFromSolver(Chaos::FPBDRigidsSolver
 		}
 	}
 
+	Evolution->GetRigidClustering().GetClusterUnionManager().HandleDeferredClusterUnionUpdateProperties();
+
 	for (int i = 0; i < SolverParticleHandles.Num(); i++)
 	{
 		if (FClusterHandle* Handle = SolverParticleHandles[i])
