Fix crash, sleeping particles, they were not removed from active list

diff --git a/Engine/Source/Runtime/Experimental/Chaos/Public/Chaos/PBDRigidsSOAs.h b/Engine/Source/Runtime/Experimental/Chaos/Public/Chaos/PBDRigidsSOAs.h
--- a/Engine/Source/Runtime/Experimental/Chaos/Public/Chaos/PBDRigidsSOAs.h
+++ b/Engine/Source/Runtime/Experimental/Chaos/Public/Chaos/PBDRigidsSOAs.h
@@ -560,14 +560,7 @@ public:
 						RemoveGeometryCollectionParticle(PBDRigidGC);
 						InsertGeometryCollectionParticle(PBDRigidGC);
 					}
-					else if (FPBDRigidClusteredParticleHandle* PBDRigidClustered = Particle->CastToClustered())
-					{
-						RemoveFromActiveArray(PBDRigid, /*bStillDirty=*/true);
-					}
-					else
-					{
-						RemoveFromActiveArray(PBDRigid, /*bStillDirty=*/true);
-					}
+					RemoveFromActiveArray(PBDRigid, /*bStillDirty=*/true);
 
 					if (!DeferUpdateViews)
 					{
