Chaos : Fix crash in event manager where handlers were not properly removed from ProxyOwnerToHandlerMap

diff --git a/Engine/Source/Runtime/Experimental/Chaos/Public/EventManager.h b/Engine/Source/Runtime/Experimental/Chaos/Public/EventManager.h
--- a/Engine/Source/Runtime/Experimental/Chaos/Public/EventManager.h
+++ b/Engine/Source/Runtime/Experimental/Chaos/Public/EventManager.h
@@ -45,9 +45,6 @@ namespace Chaos
 		virtual ~IEventHandler() {}
 		virtual void HandleEvent(const void* EventData) const = 0;
 		virtual bool GetInterestedProxyOwners(TArray<UObject*>& Output) const = 0;
-
-	protected:
-		// internal use only
 		virtual void* GetHandler() const = 0;
 	};
 
@@ -86,7 +83,6 @@ namespace Chaos
 			return true;
 		}
 
-	protected:
 		void* GetHandler() const override
 		{
 			return Handler;
@@ -208,7 +204,8 @@ namespace Chaos
 			TArray<TPair<UObject*, FEventHandlerPtr>> KeysAndValuesToRemove;
 			for (TPair<UObject*, FEventHandlerPtr>& KeyValue : ProxyOwnerToHandlerMap)
 			{
-				if (KeyValue.Get<1>() == InHandler)
+				const FEventHandlerPtr Value = KeyValue.Value;
+				if (Value && Value->GetHandler() == InHandler)
 				{
 					KeysAndValuesToRemove.Add(KeyValue);
 				}
