Chaos : fix clustered particle being destroyed while being connected to other particles
- this can happen if a geometry collection was previously attached to a cluster union

diff --git a/Engine/Source/Runtime/Experimental/Chaos/Private/Chaos/PBDRigidClustering.cpp b/Engine/Source/Runtime/Experimental/Chaos/Private/Chaos/PBDRigidClustering.cpp
--- a/Engine/Source/Runtime/Experimental/Chaos/Private/Chaos/PBDRigidClustering.cpp
+++ b/Engine/Source/Runtime/Experimental/Chaos/Private/Chaos/PBDRigidClustering.cpp
@@ -1624,6 +1624,9 @@ namespace Chaos
 		TopLevelClusterParents.Remove(ClusteredParticle);
 		TopLevelClusterParentsStrained.Remove(ClusteredParticle);
 
+		// need to disconnect from any other particles ( this can be from being a child of a cluster or a cluster union)
+		RemoveNodeConnections(ClusteredParticle);
+
 		// disconnect from the parents
 		if (ClusteredParticle->ClusterIds().Id)
 		{
@@ -1632,9 +1635,6 @@ namespace Chaos
 			ClusteredParticle->ClusterIds() = ClusterId();
 			ClusteredParticle->ClusterGroupIndex() = 0;
 
-			// If we have a parent, need to disconnect potentially from fellow children in connectivity edges.
-			RemoveNodeConnections(ClusteredParticle);
-
 			// Need to also check if the particle is a cluster union and remove from that as well.
 			ClusterUnionManager.HandleRemoveOperationWithClusterLookup({ ClusteredParticle }, EClusterUnionOperationTiming::Defer);
 
