Modular vehicle network physics prediction and resim changes, still work in progress.

diff --git a/Engine/Source/Runtime/Experimental/Chaos/Private/PhysicsProxy/GeometryCollectionPhysicsProxy.cpp b/Engine/Source/Runtime/Experimental/Chaos/Private/PhysicsProxy/GeometryCollectionPhysicsProxy.cpp
--- a/Engine/Source/Runtime/Experimental/Chaos/Private/PhysicsProxy/GeometryCollectionPhysicsProxy.cpp
+++ b/Engine/Source/Runtime/Experimental/Chaos/Private/PhysicsProxy/GeometryCollectionPhysicsProxy.cpp
@@ -3216,7 +3216,9 @@ bool FGeometryCollectionPhysicsProxy::PullFromPhysicsState(const Chaos::FDirtyGe
 			}
 		}
 #endif
-		
+		TManagedArray<bool>* AnimationsActive =
+			GameThreadCollection.FindAttribute<bool>("AnimateTransformAttribute", FGeometryCollection::TransformGroup);
+
 		// second : interpolate-able ones
 		const bool bNeedInterpolation = (NextPullData!= nullptr);
 		if (bNeedInterpolation)
@@ -3225,8 +3227,10 @@ bool FGeometryCollectionPhysicsProxy::PullFromPhysicsState(const Chaos::FDirtyGe
 			const FGeometryCollectionResults& Next = NextPullData->Results;
 			for (int32 TransformGroupIndex = 0; TransformGroupIndex < NumTransforms; ++TransformGroupIndex)
 			{
+				const bool bAnimatingWhileDisabled = AnimationsActive ? (*AnimationsActive)[TransformGroupIndex] : false;
 				const bool WasDisabled = Prev.States[TransformGroupIndex].DisabledState;
-				if (!WasDisabled)
+
+				if (!WasDisabled || bAnimatingWhileDisabled)
 				{
 					if (GTParticles[TransformGroupIndex] == nullptr)
 					{
@@ -3279,9 +3283,6 @@ bool FGeometryCollectionPhysicsProxy::PullFromPhysicsState(const Chaos::FDirtyGe
 		}
 		else
 		{
-			TManagedArray<bool>* AnimationsActive =
-				GameThreadCollection.FindAttribute<bool>("AnimateTransformAttribute", FGeometryCollection::TransformGroup);
-
 			for (int32 TransformGroupIndex = 0; TransformGroupIndex < NumTransforms; ++TransformGroupIndex)
 			{
 				bool bAnimatingWhileDisabled = AnimationsActive ? (*AnimationsActive)[TransformGroupIndex] : false;
diff --git a/Engine/Source/Runtime/Experimental/ChaosVehicles/ChaosVehiclesCore/Private/SimModule/TransmissionModule.cpp b/Engine/Source/Runtime/Experimental/ChaosVehicles/ChaosVehiclesCore/Private/SimModule/TransmissionModule.cpp
--- a/Engine/Source/Runtime/Experimental/ChaosVehicles/ChaosVehiclesCore/Private/SimModule/TransmissionModule.cpp
+++ b/Engine/Source/Runtime/Experimental/ChaosVehicles/ChaosVehiclesCore/Private/SimModule/TransmissionModule.cpp
@@ -60,9 +60,9 @@ namespace Chaos
 			{
 				ChangeDown();
 			}
-			else if (Inputs.ControlInputs.SetGearNumber)
+			else if (Inputs.ControlInputs.GearNumber)
 			{
-				TargetGear = Inputs.ControlInputs.SetGearNumber;
+				TargetGear = Inputs.ControlInputs.GearNumber;
 			}
 		}
 
diff --git a/Engine/Source/Runtime/Experimental/ChaosVehicles/ChaosVehiclesCore/Public/SimModule/SimulationModuleBase.h b/Engine/Source/Runtime/Experimental/ChaosVehicles/ChaosVehiclesCore/Public/SimModule/SimulationModuleBase.h
--- a/Engine/Source/Runtime/Experimental/ChaosVehicles/ChaosVehiclesCore/Public/SimModule/SimulationModuleBase.h
+++ b/Engine/Source/Runtime/Experimental/ChaosVehicles/ChaosVehiclesCore/Public/SimModule/SimulationModuleBase.h
@@ -24,7 +24,8 @@ namespace Chaos
 	{
 
 		FControlInputs()
-			: Throttle(0)
+			: IsValid(false)
+			, Throttle(0)
 			, Brake(0)
 			, Steering(0)
 			, Clutch(0)
@@ -34,10 +35,12 @@ namespace Chaos
 			, Yaw(0)
 			, ChangeUp(false)
 			, ChangeDown(false)
-			, SetGearNumber(0)
+			, GearNumber(0)
+			, InputDebugIndex(0)
 		{
 		}
 
+		bool IsValid;
 		float Throttle;
 		float Brake;
 		float Steering;
@@ -48,7 +51,10 @@ namespace Chaos
 		float Yaw;
 		bool ChangeUp;
 		bool ChangeDown;
-		int SetGearNumber;
+		int GearNumber;
+		int InputDebugIndex;
+
+		int GetChecksum() { return (int)((Throttle+Brake+Steering)*100000.0f);}
 	};
 
 	struct CHAOSVEHICLESCORE_API FModuleHitResults
diff --git a/Engine/Source/Runtime/Experimental/GeometryCollectionEngine/Private/GeometryCollection/GeometryCollectionComponent.cpp b/Engine/Source/Runtime/Experimental/GeometryCollectionEngine/Private/GeometryCollection/GeometryCollectionComponent.cpp
--- a/Engine/Source/Runtime/Experimental/GeometryCollectionEngine/Private/GeometryCollection/GeometryCollectionComponent.cpp
+++ b/Engine/Source/Runtime/Experimental/GeometryCollectionEngine/Private/GeometryCollection/GeometryCollectionComponent.cpp
@@ -2929,7 +2929,10 @@ void UGeometryCollectionComponent::RequestUpdateRepData()
 {
 	if (FPhysScene* PhysScene = GetInnerChaosScene())
 	{
-		PhysScene->EnqueueAsyncPhysicsCommand(0, this,
+		Chaos::FPBDRigidsSolver* Solver = PhysicsProxy->GetSolver<Chaos::FPBDRigidsSolver>();
+		const int32 CurrentFrame = Solver->GetCurrentFrame();
+
+		PhysScene->EnqueueAsyncPhysicsCommand(CurrentFrame, this,
 			[this]()
 			{
 				UpdateRepData();
