Chaos - add collision detection stats

diff --git a/Engine/Source/Runtime/Experimental/Chaos/Private/ChaosStats.cpp b/Engine/Source/Runtime/Experimental/Chaos/Private/ChaosStats.cpp
--- a/Engine/Source/Runtime/Experimental/Chaos/Private/ChaosStats.cpp
+++ b/Engine/Source/Runtime/Experimental/Chaos/Private/ChaosStats.cpp
@@ -102,9 +102,9 @@ DEFINE_STAT(STAT_ForceUpdateField_LinearImpulse);
 DEFINE_STAT(STAT_Collisions_Detect);
 DEFINE_STAT(STAT_Collisions_ParticlePairBroadPhase);
 DEFINE_STAT(STAT_Collisions_SpatialBroadPhase);
-DEFINE_STAT(STAT_Collisions_AABBTree);
-DEFINE_STAT(STAT_Collisions_Restore);
-DEFINE_STAT(STAT_Collisions_Filtering);
+DEFINE_STAT(STAT_Collisions_MidPhase);
+DEFINE_STAT(STAT_Collisions_AssignMidPhases);
+DEFINE_STAT(STAT_Collisions_NarrowPhase);
 DEFINE_STAT(STAT_Collisions_GenerateCollisions);
 DEFINE_STAT(STAT_Collisions_Gather);
 DEFINE_STAT(STAT_Collisions_Scatter);
diff --git a/Engine/Source/Runtime/Experimental/Chaos/Public/Chaos/Collision/SpatialAccelerationBroadPhase.h b/Engine/Source/Runtime/Experimental/Chaos/Public/Chaos/Collision/SpatialAccelerationBroadPhase.h
--- a/Engine/Source/Runtime/Experimental/Chaos/Public/Chaos/Collision/SpatialAccelerationBroadPhase.h
+++ b/Engine/Source/Runtime/Experimental/Chaos/Public/Chaos/Collision/SpatialAccelerationBroadPhase.h
@@ -300,7 +300,6 @@ namespace Chaos
 			const FCollisionDetectorSettings& Settings,
 			IResimCacheBase* ResimCache)
 		{
-			SCOPE_CYCLE_COUNTER(STAT_Collisions_SpatialBroadPhase);
 
 			if (!ensure(SpatialAcceleration))
 			{
@@ -310,28 +309,33 @@ namespace Chaos
 
 			bNeedsResim = ResimCache && ResimCache->IsResimming();
 
-			if (const auto AABBTree = SpatialAcceleration->template As<TAABBTree<FAccelerationStructureHandle, TAABBTreeLeafArray<FAccelerationStructureHandle>>>())
 			{
-				ProduceOverlaps(Dt, *AABBTree, Allocator, Settings, ResimCache);
-			}
-			else if (const auto BV = SpatialAcceleration->template As<TBoundingVolume<FAccelerationStructureHandle>>())
-			{
-				ProduceOverlaps(Dt, *BV, Allocator, Settings, ResimCache);
-			}
-			else if (const auto AABBTreeBV = SpatialAcceleration->template As<TAABBTree<FAccelerationStructureHandle, TBoundingVolume<FAccelerationStructureHandle>>>())
-			{
-				ProduceOverlaps(Dt, *AABBTreeBV, Allocator, Settings, ResimCache);
-			}
-			else if (const auto Collection = SpatialAcceleration->template As<ISpatialAccelerationCollection<FAccelerationStructureHandle, FReal, 3>>())
-			{
-				Collection->PBDComputeConstraintsLowLevel(Dt, *this, Allocator, Settings, ResimCache);
-			}
-			else
-			{
-				check(false);  //question: do we want to support a dynamic dispatch version?
+				SCOPE_CYCLE_COUNTER(STAT_Collisions_SpatialBroadPhase);
+
+				if (const auto AABBTree = SpatialAcceleration->template As<TAABBTree<FAccelerationStructureHandle, TAABBTreeLeafArray<FAccelerationStructureHandle>>>())
+				{
+					ProduceOverlaps(Dt, *AABBTree, Allocator, Settings, ResimCache);
+				}
+				else if (const auto BV = SpatialAcceleration->template As<TBoundingVolume<FAccelerationStructureHandle>>())
+				{
+					ProduceOverlaps(Dt, *BV, Allocator, Settings, ResimCache);
+				}
+				else if (const auto AABBTreeBV = SpatialAcceleration->template As<TAABBTree<FAccelerationStructureHandle, TBoundingVolume<FAccelerationStructureHandle>>>())
+				{
+					ProduceOverlaps(Dt, *AABBTreeBV, Allocator, Settings, ResimCache);
+				}
+				else if (const auto Collection = SpatialAcceleration->template As<ISpatialAccelerationCollection<FAccelerationStructureHandle, FReal, 3>>())
+				{
+					Collection->PBDComputeConstraintsLowLevel(Dt, *this, Allocator, Settings, ResimCache);
+				}
+				else
+				{
+					check(false);  //question: do we want to support a dynamic dispatch version?
+				}
 			}
 
 			{
+				SCOPE_CYCLE_COUNTER(STAT_Collisions_MidPhase);
 				CSV_SCOPED_TIMING_STAT(PhysicsVerbose, DetectCollisions_AssignMidPhases);
 
 				// Find or assign a midphase to each overlapping particle pair
@@ -355,6 +359,7 @@ namespace Chaos
 		*/
 		void ProduceCollisions(FReal Dt)
 		{
+			SCOPE_CYCLE_COUNTER(STAT_Collisions_NarrowPhase);
 			CSV_SCOPED_TIMING_STAT(PhysicsVerbose, DetectCollisions_MidPhase);
 			const auto& ProcessMidPhasesWorker = [this, Dt](const int32 ContextIndex)
 			{
@@ -514,7 +519,6 @@ namespace Chaos
 			const bool bBody1Bounded = Particle1->HasBounds();
 			if (bBody1Bounded)
 			{
-				//SCOPE_CYCLE_COUNTER(STAT_Collisions_AABBTree);
 				// @todo(chaos): cache this on the particle?
 				FCollisionFilterData ParticleSimData;
 				FAccelerationStructureHandle::ComputeParticleSimFilterDataFromShapes(*Particle1, ParticleSimData);
@@ -529,7 +533,6 @@ namespace Chaos
 			}
 			else
 			{
-				//SCOPE_CYCLE_COUNTER(STAT_Collisions_AABBTree);
 				const auto& GlobalElems = InSpatialAcceleration.GlobalObjects();
 				Context.Overlaps.Reserve(GlobalElems.Num());
 
@@ -547,6 +550,8 @@ namespace Chaos
 		// @todo(chaos): optimize
 		void AssignMidPhases(Private::FBroadPhaseContext& BroadphaseContext)
 		{
+			SCOPE_CYCLE_COUNTER(STAT_Collisions_AssignMidPhases);
+
 			Private::FCollisionContextAllocator* ContextAllocator = BroadphaseContext.CollisionContext.GetAllocator();
 
 			BroadphaseContext.MidPhases.SetNum(BroadphaseContext.Overlaps.Num());
@@ -575,6 +580,8 @@ namespace Chaos
 		// Process all the midphases: generate constraints and execute the narrowphase
 		void ProcessMidPhases(const FReal Dt, const Private::FBroadPhaseContext& BroadphaseContext)
 		{
+			SCOPE_CYCLE_COUNTER(STAT_Collisions_GenerateCollisions);
+			
 			// Prefetch initial set of MidPhases
 			const int32 PrefetchLookahead = 4;
 			const int32 NumMidPhases = BroadphaseContext.MidPhases.Num();
diff --git a/Engine/Source/Runtime/Experimental/Chaos/Public/ChaosStats.h b/Engine/Source/Runtime/Experimental/Chaos/Public/ChaosStats.h
--- a/Engine/Source/Runtime/Experimental/Chaos/Public/ChaosStats.h
+++ b/Engine/Source/Runtime/Experimental/Chaos/Public/ChaosStats.h
@@ -126,9 +126,9 @@ DECLARE_CYCLE_STAT_EXTERN(TEXT("[Field Update] Impulse"), STAT_ForceUpdateField_
 DECLARE_CYCLE_STAT_EXTERN(TEXT("Collisions::Detect"), STAT_Collisions_Detect, STATGROUP_ChaosCollision, CHAOS_API);
 DECLARE_CYCLE_STAT_EXTERN(TEXT("Collisions::BroadPhase::ParticlePair"), STAT_Collisions_ParticlePairBroadPhase, STATGROUP_ChaosCollision, CHAOS_API);
 DECLARE_CYCLE_STAT_EXTERN(TEXT("Collisions::BroadPhase::Spatial"), STAT_Collisions_SpatialBroadPhase, STATGROUP_ChaosCollision, CHAOS_API);
-DECLARE_CYCLE_STAT_EXTERN(TEXT("Collisions::AABBTree"), STAT_Collisions_AABBTree, STATGROUP_ChaosCollision, CHAOS_API);
-DECLARE_CYCLE_STAT_EXTERN(TEXT("Collisions::Filtering"), STAT_Collisions_Filtering, STATGROUP_ChaosCollision, CHAOS_API);
-DECLARE_CYCLE_STAT_EXTERN(TEXT("Collisions::Restore"), STAT_Collisions_Restore, STATGROUP_ChaosCollision, CHAOS_API);
+DECLARE_CYCLE_STAT_EXTERN(TEXT("Collisions::MidPhase"), STAT_Collisions_MidPhase, STATGROUP_ChaosCollision, CHAOS_API);
+DECLARE_CYCLE_STAT_EXTERN(TEXT("Collisions::AssignMidPhases"), STAT_Collisions_AssignMidPhases, STATGROUP_ChaosCollision, CHAOS_API);
+DECLARE_CYCLE_STAT_EXTERN(TEXT("Collisions::NarrowPhase"), STAT_Collisions_NarrowPhase, STATGROUP_ChaosCollision, CHAOS_API);
 DECLARE_CYCLE_STAT_EXTERN(TEXT("Collisions::GenerateCollisions"), STAT_Collisions_GenerateCollisions, STATGROUP_ChaosCollision, CHAOS_API);
 DECLARE_CYCLE_STAT_EXTERN(TEXT("Collisions::Gather"), STAT_Collisions_Gather, STATGROUP_ChaosCollision, CHAOS_API);
 DECLARE_CYCLE_STAT_EXTERN(TEXT("Collisions::Scatter"), STAT_Collisions_Scatter, STATGROUP_ChaosCollision, CHAOS_API);
