Alter rigid clustering advance to not consider all top level particles for breaks but only particles that are strained. This greatly reduces the steady state cost of the breaking model and also avoids spikes when breaking as only the required particles are traversed.

diff --git a/Engine/Source/Runtime/Experimental/Chaos/Private/Chaos/PBDRigidClustering.cpp b/Engine/Source/Runtime/Experimental/Chaos/Private/Chaos/PBDRigidClustering.cpp
--- a/Engine/Source/Runtime/Experimental/Chaos/Private/Chaos/PBDRigidClustering.cpp
+++ b/Engine/Source/Runtime/Experimental/Chaos/Private/Chaos/PBDRigidClustering.cpp
@@ -1060,6 +1060,9 @@ namespace Chaos
 		}
 	}
 
+	static int32 GClusterBreakOnlyStrained = 1;
+	FAutoConsoleVariableRef CVarBreakMode(TEXT("p.chaos.clustering.breakonlystrained"), GClusterBreakOnlyStrained, 
+										  TEXT("If enabled we only process strained clusters for breaks, if disabled all clusters are traversed and checked"));
 
 	DECLARE_CYCLE_STAT(TEXT("TPBDRigidClustering<>::AdvanceClustering"), STAT_AdvanceClustering, STATGROUP_Chaos);
 	DECLARE_CYCLE_STAT(TEXT("TPBDRigidClustering<>::Update Impulse from Strain"), STAT_UpdateImpulseStrain, STATGROUP_Chaos);
@@ -1102,6 +1105,31 @@ namespace Chaos
 			//  That will trigger a break too.
 			//
 			bool bPotentialBreak = false;
+			TArray<FPBDRigidClusteredParticleHandle*> ParticlesToProcess;
+
+			auto ProcessClusteredParticle = [&ParticlesToProcess, &bPotentialBreak, this](FPBDRigidClusteredParticleHandle* Particle)
+			{
+				TArray<FRigidHandle>& ParentToChildren = MChildren[Particle];
+				for(FRigidHandle Child : ParentToChildren)
+				{
+					if(FClusterHandle ClusteredChild = Child->CastToClustered())
+					{
+						if(ClusteredChild->GetInternalStrains() <= 0.f)
+						{
+							ParticlesToProcess.Add(Particle);
+
+							ClusteredChild->CollisionImpulse() = FLT_MAX;
+							MCollisionImpulseArrayDirty = true;
+						}
+						else if(ClusteredChild->GetExternalStrain() > 0 || ClusteredChild->CollisionImpulse() > 0)
+						{
+							ParticlesToProcess.Add(Particle);
+							bPotentialBreak = true;
+						}
+					}
+				}
+			};
+
 			{
 				SCOPE_CYCLE_COUNTER(STAT_UpdateDirtyImpulses);
 				for (const auto& ActiveCluster : TopLevelClusterParentsStrained)
@@ -1110,49 +1138,75 @@ namespace Chaos
 					{
 						if (ActiveCluster->ClusterIds().NumChildren > 0) //active index is a cluster
 						{
-							TArray<FRigidHandle>& ParentToChildren = MChildren[ActiveCluster];
-							for (FRigidHandle Child : ParentToChildren)
+							if(FClusterUnion* ClusterUnion = ClusterUnionManager.FindClusterUnionFromParticle(ActiveCluster))
 							{
-								if (FClusterHandle ClusteredChild = Child->CastToClustered())
+								if(ClusterUnion->InternalCluster == ActiveCluster)
 								{
-									if (ClusteredChild->GetInternalStrains() <= 0.f)
-									{
-										ClusteredChild->CollisionImpulse() = FLT_MAX;
-										MCollisionImpulseArrayDirty = true;
-									}
-									else if (!bPotentialBreak && ClusteredChild->GetExternalStrain() > 0)
+									// ActiveCluster is itself a cluster union, so loop over its children and add those
+									// to process for breaking.
+									for(FPBDRigidParticleHandle* ChildParticle : ClusterUnion->ChildParticles)
 									{
-										bPotentialBreak = true;
+										if(ChildParticle)
+										{
+											if(FPBDRigidClusteredParticleHandle* ClusteredChild = ChildParticle->CastToClustered())
+											{
+												if(ClusteredChild->ClusterIds().NumChildren > 0)
+												{
+													ProcessClusteredParticle(ClusteredChild);
+												}
+											}
+										}
 									}
 								}
+								else
+								{
+									// Clustered is inside a clustered union, but not a clustered union itself
+									ProcessClusteredParticle(ActiveCluster);
+								}
+							}
+							else
+							{
+								ProcessClusteredParticle(ActiveCluster);
 							}
 						}
 					}
 				}
-				TopLevelClusterParentsStrained.Reset();
 			}
 
 			if (MCollisionImpulseArrayDirty || bPotentialBreak)
 			{
 				SCOPE_CYCLE_COUNTER(STAT_UpdateDirtyImpulses);
-				BreakingModel();
+
+				// Call our breaking model
+				// #TODO convert to visitor pattern to avoid TArray allocations above.
+				if(GClusterBreakOnlyStrained == 1)
+				{
+					BreakingModel(ParticlesToProcess);
+				}
+				else
+				{
+					BreakingModel();
+				}
 			} // end if MCollisionImpulseArrayDirty
+
+			TopLevelClusterParentsStrained.Reset();
 		}
 		Timer.Stop();
 		UE_LOG(LogChaos, Verbose, TEXT("Cluster Break Update Time is %f"), Time);
 	}
 
-	DECLARE_CYCLE_STAT(TEXT("TPBDRigidClustering<>::BreakingModel()"), STAT_BreakingModel, STATGROUP_Chaos);
+	DECLARE_CYCLE_STAT(TEXT("BreakingModel_AllParticles"), STAT_BreakingModel_AllParticles, STATGROUP_Chaos);
+	DECLARE_CYCLE_STAT(TEXT("BreakingModel"), STAT_BreakingModel, STATGROUP_Chaos);
 	void FRigidClustering::BreakingModel()
 	{
-		SCOPE_CYCLE_COUNTER(STAT_BreakingModel);
-
+		SCOPE_CYCLE_COUNTER(STAT_BreakingModel_AllParticles);
+		
 		// Clear the set tracking breaking collisions
 		BreakingCollisions.Empty();
 
 		//make copy because release cluster modifies active indices. We want to iterate over original active indices
 		TArray<FPBDRigidClusteredParticleHandle*> ClusteredParticlesToProcess;
-		for (auto& Particle : MEvolution.GetNonDisabledClusteredView())
+		for(FTransientPBDRigidParticleHandle& Particle : MEvolution.GetNonDisabledClusteredView())
 		{
 			if (FPBDRigidClusteredParticleHandle* Clustered = Particle.Handle()->CastToClustered())
 			{
@@ -1193,15 +1247,27 @@ namespace Chaos
 			}
 		}
 
-		for (FPBDRigidClusteredParticleHandle* ClusteredParticle : ClusteredParticlesToProcess)
+		BreakingModel(ClusteredParticlesToProcess);
+	}
+	
+	void FRigidClustering::BreakingModel(TArray<FPBDRigidClusteredParticleHandle*>& InParticles)
+	{
+		SCOPE_CYCLE_COUNTER(STAT_BreakingModel);
+
+		// Clear the set tracking breaking collisions
+		BreakingCollisions.Empty();
+
+		for(FPBDRigidClusteredParticleHandle* ClusteredParticle : InParticles)
 		{
-			ReleaseClusterParticles(ClusteredParticle);
+			if(ClusteredParticle->ClusterIds().NumChildren)
+			{
+				ReleaseClusterParticles(ClusteredParticle);
+			}
 		}
 
 		// This way if we break apart a large cluster union here (i.e. many of its children want to be released from ReleaseClusterParticles due to strain)
 		// we'll only update the cluster properties once here (connection graph, geometry, etc.).
 		ClusterUnionManager.HandleDeferredClusterUnionUpdateProperties();
-
 		// Restore some of the momentum of objects that were touching rigid clusters that broke
 		if (RestoreBreakingMomentumPercent > 0.f)
 		{
@@ -1417,27 +1483,25 @@ namespace Chaos
 	FAutoConsoleVariableRef CVarMinContactSpeedForStrainEval(TEXT("p.chaos.MinContactSpeedForStrainEval"), MinContactSpeedForStrainEval, TEXT("Minimum speed at the contact before accumulating for strain eval "));
 
 	DECLARE_CYCLE_STAT(TEXT("ComputeStrainFromCollision"), STAT_ComputeStrainFromCollision, STATGROUP_Chaos);
-	void 
-	FRigidClustering::ComputeStrainFromCollision(
-		const FPBDCollisionConstraints& CollisionRule)
+	void FRigidClustering::ComputeStrainFromCollision(const FPBDCollisionConstraints& CollisionRule)
 	{
 		SCOPE_CYCLE_COUNTER(STAT_ComputeStrainFromCollision);
 		FClusterMap& MParentToChildren = GetChildrenMap();
 
 		ResetCollisionImpulseArray();
 
-		for (const Chaos::FPBDCollisionConstraintHandle* ContactHandle : CollisionRule.GetConstConstraintHandles())
+		for (Chaos::FPBDCollisionConstraintHandle* ContactHandle : CollisionRule.GetConstraintHandles())
 		{
 			if (ContactHandle == nullptr)
 			{
 				continue;
 			}
 
-			TVector<const FGeometryParticleHandle*, 2> ConstrainedParticles = ContactHandle->GetConstrainedParticles();
+			TVector<FGeometryParticleHandle*, 2> ConstrainedParticles = ContactHandle->GetConstrainedParticles();
 			
 			// make sure we only compute things if one of the two particle is clustered
-			const FPBDRigidClusteredParticleHandle* ClusteredConstrainedParticles0 = ConstrainedParticles[0]->CastToClustered();
-			const FPBDRigidClusteredParticleHandle* ClusteredConstrainedParticles1 = ConstrainedParticles[1]->CastToClustered();
+			FPBDRigidClusteredParticleHandle* ClusteredConstrainedParticles0 = ConstrainedParticles[0]->CastToClustered();
+			FPBDRigidClusteredParticleHandle* ClusteredConstrainedParticles1 = ConstrainedParticles[1]->CastToClustered();
 			if (!ClusteredConstrainedParticles0 && !ClusteredConstrainedParticles1)
 			{
 				continue;
@@ -1465,8 +1529,9 @@ namespace Chaos
 				continue;
 			}
 
-			auto ComputeStrainLambda = [&](
-				const FPBDRigidClusteredParticleHandle* Cluster)
+			auto ComputeStrainLambda = [this, &ContactHandle](
+				const FPBDRigidClusteredParticleHandle* Cluster,
+				FRealSingle& OutTotalImpulseAccumulator)
 			{
 				const FVec3 ContactWorldLocation = ContactHandle->GetContact().CalculateWorldContactLocation();
 				
@@ -1514,6 +1579,7 @@ namespace Chaos
 							{
 								ClusteredChild->CollisionImpulses() += AccumulatedImpulse;
 								UpdateTopLevelParticle(ClusteredChild);
+								OutTotalImpulseAccumulator += AccumulatedImpulse;
 							}
 						}
 					}
@@ -1533,6 +1599,7 @@ namespace Chaos
 								{
 									ClusteredChild->CollisionImpulses() += AccumulatedImpulse;
 									UpdateTopLevelParticle(ClusteredChild);
+									OutTotalImpulseAccumulator += AccumulatedImpulse;
 								}
 							}
 						}
@@ -1540,23 +1607,28 @@ namespace Chaos
 				}
 			};
 
+			// We only need to dirty the impulse array if any of the active contacts actually added 
+			// a collision impulse to a particle. If they are all resting or otherwise non-impulsive
+			// contacts then we can skip dirtying the impulse array and avoid running the breaking
+			// model when we know nothing will break
+			FRealSingle TotalImpulses[] = { 0.0f, 0.0f };
+
 			if (ClusteredConstrainedParticles0)
 			{
-				ComputeStrainLambda(ClusteredConstrainedParticles0);
+				ComputeStrainLambda(ClusteredConstrainedParticles0, TotalImpulses[0]);
+				MCollisionImpulseArrayDirty |= TotalImpulses[0] > 0.0f;
 			}
 
 			if (ClusteredConstrainedParticles1)
 			{
-				ComputeStrainLambda(ClusteredConstrainedParticles1);
+				ComputeStrainLambda(ClusteredConstrainedParticles1, TotalImpulses[1]);
+				MCollisionImpulseArrayDirty |= TotalImpulses[1] > 0.0f;
 			}
-
-			MCollisionImpulseArrayDirty = true;
 		}
 	}
 
 	DECLARE_CYCLE_STAT(TEXT("ResetCollisionImpulseArray"), STAT_ResetCollisionImpulseArray, STATGROUP_Chaos);
-	void 
-	FRigidClustering::ResetCollisionImpulseArray()
+	void FRigidClustering::ResetCollisionImpulseArray()
 	{
 		SCOPE_CYCLE_COUNTER(STAT_ResetCollisionImpulseArray);
 		if (MCollisionImpulseArrayDirty)
diff --git a/Engine/Source/Runtime/Experimental/Chaos/Public/Chaos/PBDRigidClustering.h b/Engine/Source/Runtime/Experimental/Chaos/Public/Chaos/PBDRigidClustering.h
--- a/Engine/Source/Runtime/Experimental/Chaos/Public/Chaos/PBDRigidClustering.h
+++ b/Engine/Source/Runtime/Experimental/Chaos/Public/Chaos/PBDRigidClustering.h
@@ -208,7 +208,8 @@ public:
 	*    the children clusters.
 	*/
 	void BreakingModel();
-
+	void BreakingModel(TArray<FPBDRigidClusteredParticleHandle*>& InParticles);
+	
 	//
 	// Access
 	//
diff --git a/Engine/Source/Runtime/Experimental/Chaos/Public/Chaos/ParticleHandle.h b/Engine/Source/Runtime/Experimental/Chaos/Public/Chaos/ParticleHandle.h
--- a/Engine/Source/Runtime/Experimental/Chaos/Public/Chaos/ParticleHandle.h
+++ b/Engine/Source/Runtime/Experimental/Chaos/Public/Chaos/ParticleHandle.h
@@ -1268,7 +1268,7 @@ public:
 	const ClusterId& ClusterIds() const { return PBDRigidClusteredParticles->ClusterIds(ParticleIdx); }
 	ClusterId& ClusterIds() { return PBDRigidClusteredParticles->ClusterIds(ParticleIdx); }
 
-	FPBDRigidClusteredParticleHandle* Parent() { return (ClusterIds().Id)? ClusterIds().Id->CastToClustered(): nullptr; }
+	FPBDRigidClusteredParticleHandle* Parent() const { return (ClusterIds().Id)? ClusterIds().Id->CastToClustered(): nullptr; }
 
 	const TRigidTransform<T,d>& ChildToParent() const { return PBDRigidClusteredParticles->ChildToParent(ParticleIdx); }
 	TRigidTransform<T,d>& ChildToParent() { return PBDRigidClusteredParticles->ChildToParent(ParticleIdx); }
