clarify confusing comment on ManagedArrayAccessor re the ConstCollection and Collection members

diff --git a/Engine/Source/Runtime/Experimental/Chaos/Public/GeometryCollection/ManagedArrayAccessor.h b/Engine/Source/Runtime/Experimental/Chaos/Public/GeometryCollection/ManagedArrayAccessor.h
--- a/Engine/Source/Runtime/Experimental/Chaos/Public/GeometryCollection/ManagedArrayAccessor.h
+++ b/Engine/Source/Runtime/Experimental/Chaos/Public/GeometryCollection/ManagedArrayAccessor.h
@@ -176,8 +176,7 @@ public:
 
 private:
 
-	// const collection will be a null pointer, 
-	// while non-const will be valid.
+	// the non-const Collection will be null if the accessor is const, while ConstCollection is always set
 	const FManagedArrayCollection& ConstCollection;
 	FManagedArrayCollection* Collection = nullptr;
 
