Event optimization, breaking, removal and crumbling events are just processed on the interested proxies.

diff --git a/Engine/Source/Runtime/Experimental/Chaos/Public/EventsData.h b/Engine/Source/Runtime/Experimental/Chaos/Public/EventsData.h
--- a/Engine/Source/Runtime/Experimental/Chaos/Public/EventsData.h
+++ b/Engine/Source/Runtime/Experimental/Chaos/Public/EventsData.h
@@ -250,15 +250,15 @@ namespace Chaos
 		}
 		else if constexpr (std::is_same_v<PayloadType, FBreakingEventData>)
 		{
-			return nullptr; // &Buffer->PhysicsProxyToBreakingIndices.PhysicsProxyToIndicesMap;
+			return &Buffer->PhysicsProxyToBreakingIndices.PhysicsProxyToIndicesMap;
 		}
 		else if constexpr (std::is_same_v<PayloadType, FTrailingEventData>)
 		{
-			return nullptr; // &Buffer->PhysicsProxyToTrailingIndices.PhysicsProxyToIndicesMap;
+			return &Buffer->PhysicsProxyToTrailingIndices.PhysicsProxyToIndicesMap;
 		}
 		else if constexpr (std::is_same_v<PayloadType, FRemovalEventData>)
 		{
-			return nullptr; // &Buffer->PhysicsProxyToRemovalIndices.PhysicsProxyToIndicesMap;
+			return &Buffer->PhysicsProxyToRemovalIndices.PhysicsProxyToIndicesMap;
 		}
 		else if constexpr (std::is_same_v<PayloadType, FSleepingEventData>)
 		{
@@ -266,7 +266,7 @@ namespace Chaos
 		}
 		else if constexpr (std::is_same_v<PayloadType, FCrumblingEventData>)
 		{
-			return nullptr; // &Buffer->PhysicsProxyToCrumblingIndices.PhysicsProxyToIndicesMap;
+			return &Buffer->PhysicsProxyToCrumblingIndices.PhysicsProxyToIndicesMap;
 		}
 		else
 		{
diff --git a/Engine/Source/Runtime/Experimental/ChaosSolverEngine/Private/Chaos/ChaosGameplayEventDispatcher.cpp b/Engine/Source/Runtime/Experimental/ChaosSolverEngine/Private/Chaos/ChaosGameplayEventDispatcher.cpp
--- a/Engine/Source/Runtime/Experimental/ChaosSolverEngine/Private/Chaos/ChaosGameplayEventDispatcher.cpp
+++ b/Engine/Source/Runtime/Experimental/ChaosSolverEngine/Private/Chaos/ChaosGameplayEventDispatcher.cpp
@@ -321,10 +321,9 @@ void UChaosGameplayEventDispatcher::RegisterChaosEvents()
 		{
 			Chaos::FEventManager* EventManager = Solver->GetEventManager();
 			EventManager->RegisterHandler<Chaos::FCollisionEventData>(Chaos::EEventType::Collision, this, &UChaosGameplayEventDispatcher::HandleCollisionEvents, &UChaosGameplayEventDispatcher::GetInterestedProxyOwnersForCollisionEvents);
-			// Todo implement GetInterestedProxies for these as well 
-			EventManager->RegisterHandler<Chaos::FBreakingEventData>(Chaos::EEventType::Breaking, this, &UChaosGameplayEventDispatcher::HandleBreakingEvents);
-			EventManager->RegisterHandler<Chaos::FRemovalEventData>(Chaos::EEventType::Removal, this, &UChaosGameplayEventDispatcher::HandleRemovalEvents);
-			EventManager->RegisterHandler<Chaos::FCrumblingEventData>(Chaos::EEventType::Crumbling, this, &UChaosGameplayEventDispatcher::HandleCrumblingEvents);
+			EventManager->RegisterHandler<Chaos::FBreakingEventData>(Chaos::EEventType::Breaking, this, &UChaosGameplayEventDispatcher::HandleBreakingEvents, &UChaosGameplayEventDispatcher::GetInterestedProxyOwnersForBreakingEvents);
+			EventManager->RegisterHandler<Chaos::FRemovalEventData>(Chaos::EEventType::Removal, this, &UChaosGameplayEventDispatcher::HandleRemovalEvents, &UChaosGameplayEventDispatcher::GetInterestedProxyOwnersForRemovalEvents);
+			EventManager->RegisterHandler<Chaos::FCrumblingEventData>(Chaos::EEventType::Crumbling, this, &UChaosGameplayEventDispatcher::HandleCrumblingEvents, &UChaosGameplayEventDispatcher::GetInterestedProxyOwnersForCrumblingEvents);
 		}
 	}
 }
@@ -348,6 +347,22 @@ void UChaosGameplayEventDispatcher::UnregisterChaosEvents()
 	}
 }
 
+template <typename EventIterator>
+void UChaosGameplayEventDispatcher::FillPhysicsProxy(FPhysScene_Chaos& Scene, TArray<UObject*>& Result, EventIterator& It)
+{
+	UPrimitiveComponent* const Comp0 = Cast<UPrimitiveComponent>(It.Key());
+	const TArray<IPhysicsProxyBase*>* PhysicsProxyArray = Scene.GetOwnedPhysicsProxies(Comp0);
+
+	if (PhysicsProxyArray)
+	{
+		for (IPhysicsProxyBase* PhysicsProxy0 : *PhysicsProxyArray)
+		{
+			Result.AddUnique(PhysicsProxy0->GetOwner());
+		}
+	}
+}
+
+
 TArray<UObject*> UChaosGameplayEventDispatcher::GetInterestedProxyOwnersForCollisionEvents()
 {
 	TArray<UObject*> Result;
@@ -356,16 +371,49 @@ TArray<UObject*> UChaosGameplayEventDispatcher::GetInterestedProxyOwnersForColli
 	// look through all the components that someone is interested in and get all the proxies
 	for (decltype(CollisionEventRegistrations)::TIterator It(CollisionEventRegistrations); It; ++It)
 	{
-		UPrimitiveComponent* const Comp0 = Cast<UPrimitiveComponent>(It.Key());
-		const TArray<IPhysicsProxyBase*>* PhysicsProxyArray = Scene.GetOwnedPhysicsProxies(Comp0);
+		FillPhysicsProxy(Scene, Result, It);
+	}
 
-		if (PhysicsProxyArray)
-		{
-			for (IPhysicsProxyBase* PhysicsProxy0 : *PhysicsProxyArray)
-			{
-				Result.AddUnique(PhysicsProxy0->GetOwner());
-			}
-		}
+	return Result;
+}
+
+TArray<UObject*> UChaosGameplayEventDispatcher::GetInterestedProxyOwnersForBreakingEvents()
+{
+	TArray<UObject*> Result;
+	FPhysScene_Chaos& Scene = *(GetWorld()->GetPhysicsScene());
+
+	// look through all the components that someone is interested in and get all the proxies
+	for (decltype(BreakEventRegistrations)::TIterator It(BreakEventRegistrations); It; ++It)
+	{
+		FillPhysicsProxy(Scene, Result, It);
+	}
+
+	return Result;
+}
+
+TArray<UObject*> UChaosGameplayEventDispatcher::GetInterestedProxyOwnersForRemovalEvents()
+{
+	TArray<UObject*> Result;
+	FPhysScene_Chaos& Scene = *(GetWorld()->GetPhysicsScene());
+
+	// look through all the components that someone is interested in and get all the proxies
+	for (decltype(RemovalEventRegistrations)::TIterator It(RemovalEventRegistrations); It; ++It)
+	{
+		FillPhysicsProxy(Scene, Result, It);
+	}
+
+	return Result;
+}
+
+TArray<UObject*> UChaosGameplayEventDispatcher::GetInterestedProxyOwnersForCrumblingEvents()
+{
+	TArray<UObject*> Result;
+	FPhysScene_Chaos& Scene = *(GetWorld()->GetPhysicsScene());
+
+	// look through all the components that someone is interested in and get all the proxies
+	for (decltype(CrumblingEventRegistrations)::TIterator It(CrumblingEventRegistrations); It; ++It)
+	{
+		FillPhysicsProxy(Scene, Result, It);
 	}
 
 	return Result;
diff --git a/Engine/Source/Runtime/Experimental/ChaosSolverEngine/Public/Chaos/ChaosGameplayEventDispatcher.h b/Engine/Source/Runtime/Experimental/ChaosSolverEngine/Public/Chaos/ChaosGameplayEventDispatcher.h
--- a/Engine/Source/Runtime/Experimental/ChaosSolverEngine/Public/Chaos/ChaosGameplayEventDispatcher.h
+++ b/Engine/Source/Runtime/Experimental/ChaosSolverEngine/Public/Chaos/ChaosGameplayEventDispatcher.h
@@ -280,8 +280,14 @@ private:
 	void RegisterChaosEvents();
 	void UnregisterChaosEvents();
 
+	template <typename EventIterator>
+	void FillPhysicsProxy(FPhysScene_Chaos& Scene, TArray<UObject*>& Result, EventIterator& It);
 
 	TArray<UObject*> GetInterestedProxyOwnersForCollisionEvents();
+	TArray<UObject*> GetInterestedProxyOwnersForBreakingEvents();
+	TArray<UObject*> GetInterestedProxyOwnersForRemovalEvents();
+	TArray<UObject*> GetInterestedProxyOwnersForCrumblingEvents();
+
 	// Chaos Event Handlers
 	void HandleCollisionEvents(const Chaos::FCollisionEventData& CollisionData);
 	void HandleBreakingEvents(const Chaos::FBreakingEventData& BreakingData);
