performance optimization for adding/removing particles into the cluster union

- use more heavy use of batching of add/remove/update child to parent operations
- use the point implicit method for generating the connection graph instead of delaunay triangulation

diff --git a/Engine/Source/Runtime/Engine/Private/PhysicsEngine/ClusterUnionComponent.cpp b/Engine/Source/Runtime/Engine/Private/PhysicsEngine/ClusterUnionComponent.cpp
--- a/Engine/Source/Runtime/Engine/Private/PhysicsEngine/ClusterUnionComponent.cpp
+++ b/Engine/Source/Runtime/Engine/Private/PhysicsEngine/ClusterUnionComponent.cpp
@@ -199,7 +199,7 @@ void UClusterUnionComponent::OnCreatePhysicsState()
 
 	// TODO: Expose these parameters via the component.
 	Chaos::FClusterCreationParameters Parameters{ 0.3f, 100, false, false };
-	Parameters.ConnectionMethod = Chaos::FClusterCreationParameters::EConnectionMethod::DelaunayTriangulation;
+	Parameters.ConnectionMethod = Chaos::FClusterCreationParameters::EConnectionMethod::PointImplicit;
 
 	FChaosUserData::Set<UPrimitiveComponent>(&PhysicsUserData, this);
 
diff --git a/Engine/Source/Runtime/Experimental/Chaos/Private/Chaos/ClusterUnionManager.cpp b/Engine/Source/Runtime/Experimental/Chaos/Private/Chaos/ClusterUnionManager.cpp
--- a/Engine/Source/Runtime/Experimental/Chaos/Private/Chaos/ClusterUnionManager.cpp
+++ b/Engine/Source/Runtime/Experimental/Chaos/Private/Chaos/ClusterUnionManager.cpp
@@ -139,6 +139,7 @@ namespace Chaos
 	DECLARE_CYCLE_STAT(TEXT("FClusterUnionManager::AddPendingExplicitIndexOperation"), STAT_AddPendingExplicitIndexOperation, STATGROUP_Chaos);
 	void FClusterUnionManager::AddPendingExplicitIndexOperation(FClusterUnionExplicitIndex Index, EClusterUnionOperation Op, const TArray<FPBDRigidParticleHandle*>& Particles)
 	{
+		check(Op != EClusterUnionOperation::UpdateChildToParent);
 		SCOPE_CYCLE_COUNTER(STAT_AddPendingExplicitIndexOperation);
 		AddPendingOperation(PendingExplicitIndexOperations, Index, Op, Particles);
 	}
@@ -146,6 +147,7 @@ namespace Chaos
 	DECLARE_CYCLE_STAT(TEXT("FClusterUnionManager::AddPendingClusterIndexOperation"), STAT_AddPendingClusterIndexOperation, STATGROUP_Chaos);
 	void FClusterUnionManager::AddPendingClusterIndexOperation(FClusterUnionIndex Index, EClusterUnionOperation Op, const TArray<FPBDRigidParticleHandle*>& Particles)
 	{
+		check(Op != EClusterUnionOperation::UpdateChildToParent);
 		SCOPE_CYCLE_COUNTER(STAT_AddPendingClusterIndexOperation);
 		AddPendingOperation(PendingClusterIndexOperations, Index, Op, Particles);
 	}
@@ -171,9 +173,17 @@ namespace Chaos
 		}
 		PendingExplicitIndexOperations.Empty();
 
-		for (const TPair<FClusterUnionIndex, FClusterOpMap>& OpMap : PendingClusterIndexOperations)
+		for (TPair<FClusterUnionIndex, FClusterOpMap>& OpMap : PendingClusterIndexOperations)
 		{
-			for (const TPair<EClusterUnionOperation, TArray<FPBDRigidParticleHandle*>>& Op : OpMap.Value)
+			// Is this sort necessary? Better to be safe than sorry. Since we need to guarantee that the UpdateChildToParent happens after add.
+			OpMap.Value.KeyStableSort(
+				[](EClusterUnionOperation A, EClusterUnionOperation B)
+				{
+					return static_cast<int32>(A) < static_cast<int32>(B);
+				}
+			);
+
+			for (TPair<EClusterUnionOperation, TArray<FPBDRigidParticleHandle*>>& Op : OpMap.Value)
 			{
 				switch (Op.Key)
 				{
@@ -184,6 +194,9 @@ namespace Chaos
 				case EClusterUnionOperation::Remove:
 					HandleRemoveOperation(OpMap.Key, Op.Value, EClusterUnionOperationTiming::Immediate);
 					break;
+				case EClusterUnionOperation::UpdateChildToParent:
+					HandleUpdateChildToParentOperation(OpMap.Key, Op.Value);
+					break;
 				}
 			}
 		}
@@ -585,56 +598,80 @@ namespace Chaos
 	void FClusterUnionManager::UpdateClusterUnionParticlesChildToParent(FClusterUnionIndex Index, const TArray<FPBDRigidParticleHandle*>& Particles, const TArray<FTransform>& ChildToParent, bool bLock)
 	{
 		SCOPE_CYCLE_COUNTER(STAT_UpdateClusterUnionParticlesChildToParent);
+		AddPendingOperation(PendingClusterIndexOperations, Index, EClusterUnionOperation::UpdateChildToParent, Particles);
 
-		if (FClusterUnion* ClusterUnion = FindClusterUnion(Index))
+		for (int32 InputIndex = 0; InputIndex < Particles.Num() && InputIndex < ChildToParent.Num(); ++InputIndex)
 		{
-			// We need to keep track of all the particles we touched here. We need this because UpdateClusterUnionParticlesChildToParent
-			// is an *authoritative* update on the ChildToParent of the particle. However, UpdateAllClusterUnionProperties in certain cases
-			// may try to recompute the ChildToParent using the position of the particle. To counteract this, we will force the ChildToParent
-			// to be *temporarily* locked for the duration of UpdateAllClusterUnionProperties. However, unless bLock is true, we will restore the
-			// lock state of the particle to what it was previously.
-			TArray<TPair<FPBDRigidClusteredParticleHandle*, bool>> DirtyParticleLockStates;
-			if (!bLock)
+			if (!Particles[InputIndex])
 			{
-				DirtyParticleLockStates.Reserve(Particles.Num());
+				continue;
 			}
 
-			for (int32 ParticleIndex = 0; ParticleIndex < Particles.Num() && ParticleIndex < ChildToParent.Num(); ++ParticleIndex)
+			if (FPBDRigidClusteredParticleHandle* ClusterParticle = Particles[InputIndex]->CastToClustered())
 			{
-				FPBDRigidParticleHandle* Particle = Particles[ParticleIndex];
-				if (!ensure(Particle))
+				FClusterUnionChildToParentUpdate& Update = PendingChildToParentUpdates.FindOrAdd(ClusterParticle);
+				// If the current existing update wants to lock and we're not also locking, we can discard this new update.
+				if (Update.bLock && !bLock)
 				{
-					return;
+					continue;
 				}
+				Update.ChildToParent = ChildToParent[InputIndex];
+				Update.bLock = bLock;
+			}
+		}
+	}
 
-				const int32 ChildIndex = ClusterUnion->ChildParticles.Find(Particle->CastToRigidParticle());
-				if (ChildIndex != INDEX_NONE && ClusterUnion->InternalCluster)
+	DECLARE_CYCLE_STAT(TEXT("FClusterUnionManager::HandleUpdateChildToParentOperation"), STAT_HandleUpdateChildToParentOperation, STATGROUP_Chaos);
+	void FClusterUnionManager::HandleUpdateChildToParentOperation(FClusterUnionIndex ClusterIndex, const TArray<FPBDRigidParticleHandle*>& Particles)
+	{
+		SCOPE_CYCLE_COUNTER(STAT_UpdateClusterUnionParticlesChildToParent);
+		FClusterUnion* ClusterUnion = ClusterUnions.Find(ClusterIndex);
+		if (!ClusterUnion || Particles.IsEmpty())
+		{
+			return;
+		}
+
+		// We need to keep track of all the particles we touched here. We need this because UpdateClusterUnionParticlesChildToParent
+		// is an *authoritative* update on the ChildToParent of the particle. However, UpdateAllClusterUnionProperties in certain cases
+		// may try to recompute the ChildToParent using the position of the particle. To counteract this, we will force the ChildToParent
+		// to be *temporarily* locked for the duration of UpdateAllClusterUnionProperties. However, unless bLock is true, we will restore the
+		// lock state of the particle to what it was previously.
+		TArray<TPair<FPBDRigidClusteredParticleHandle*, bool>> DirtyParticleLockStates;
+		for (FPBDRigidParticleHandle* Particle : Particles)
+		{
+			if (!ensure(Particle))
+			{
+				return;
+			}
+
+			const int32 ChildIndex = ClusterUnion->ChildParticles.Find(Particle->CastToRigidParticle());
+			if (ChildIndex != INDEX_NONE && ClusterUnion->InternalCluster)
+			{
+				if (FPBDRigidClusteredParticleHandle* ChildHandle = ClusterUnion->ChildParticles[ChildIndex]->CastToClustered())
 				{
-					if (FPBDRigidClusteredParticleHandle* ChildHandle = ClusterUnion->ChildParticles[ChildIndex]->CastToClustered())
+					if (const FClusterUnionChildToParentUpdate* Update = PendingChildToParentUpdates.Find(ChildHandle))
 					{
-						ChildHandle->SetChildToParent(ChildToParent[ParticleIndex]);
+						ChildHandle->SetChildToParent(Update->ChildToParent);
 
-						if (!bLock)
+						if (!Update->bLock)
 						{
 							DirtyParticleLockStates.Add({ ChildHandle, ChildHandle->IsChildToParentLocked() });
 						}
 						ChildHandle->SetChildToParentLocked(true);
+						PendingChildToParentUpdates.Remove(ChildHandle);
 					}
 				}
 			}
+		}
 
-			UpdateAllClusterUnionProperties(*ClusterUnion, false);
-
-			if (!bLock)
-			{
-				for (TPair<FPBDRigidClusteredParticleHandle*, bool>& Pair : DirtyParticleLockStates)
-				{
-					Pair.Key->SetChildToParentLocked(Pair.Value);
-				}
-			}
+		UpdateAllClusterUnionProperties(*ClusterUnion, false);
 
-			MEvolution.GetParticles().MarkTransientDirtyParticle(ClusterUnion->InternalCluster);
-			MEvolution.DirtyParticle(*ClusterUnion->InternalCluster);
+		for (TPair<FPBDRigidClusteredParticleHandle*, bool>& Pair : DirtyParticleLockStates)
+		{
+			Pair.Key->SetChildToParentLocked(Pair.Value);
 		}
+
+		MEvolution.GetParticles().MarkTransientDirtyParticle(ClusterUnion->InternalCluster);
+		MEvolution.DirtyParticle(*ClusterUnion->InternalCluster);
 	}
 }
\ No newline at end of file
diff --git a/Engine/Source/Runtime/Experimental/Chaos/Private/Chaos/PBDRigidsEvolutionGBF.cpp b/Engine/Source/Runtime/Experimental/Chaos/Private/Chaos/PBDRigidsEvolutionGBF.cpp
--- a/Engine/Source/Runtime/Experimental/Chaos/Private/Chaos/PBDRigidsEvolutionGBF.cpp
+++ b/Engine/Source/Runtime/Experimental/Chaos/Private/Chaos/PBDRigidsEvolutionGBF.cpp
@@ -374,11 +374,6 @@ void FPBDRigidsEvolutionGBF::AdvanceOneTimeStepImpl(const FReal Dt, const FSubSt
 	// Update the collision solver type (used to support runtime comparisons of solver types for debugging/testing)
 	UpdateCollisionSolverType();
 
-	{
-		SCOPE_CYCLE_COUNTER(STAT_Evolution_UnclusterUnions);
-		Clustering.UnionClusterGroups();
-	}
-
 	{
 		SCOPE_CYCLE_COUNTER(STAT_Evolution_Integrate);
 		CSV_SCOPED_TIMING_STAT(PhysicsVerbose, StepSolver_Integrate);
diff --git a/Engine/Source/Runtime/Experimental/Chaos/Private/PBDRigidsSolver.cpp b/Engine/Source/Runtime/Experimental/Chaos/Private/PBDRigidsSolver.cpp
--- a/Engine/Source/Runtime/Experimental/Chaos/Private/PBDRigidsSolver.cpp
+++ b/Engine/Source/Runtime/Experimental/Chaos/Private/PBDRigidsSolver.cpp
@@ -384,6 +384,11 @@ namespace Chaos
 			{
 				SCOPE_CYCLE_COUNTER(STAT_EvolutionAndKinematicUpdate);
 
+				// This needs to be before BeginFrame since there's a possibility that we'll need to invalidate particles when updating cluster unions.
+				// Invalidating particles will cause us to iterate through and remove active collisions. If this is after BeginFrame but before we run
+				// the next broadphase, we'll pass the epoch check but have an invalid index into a now-empty active constraints array.
+				MSolver->GetEvolution()->GetRigidClustering().UnionClusterGroups();
+
 				// clear out the collision constraints as they will be stale from last frame if AdvanceOneTimeStep never gets called due to TimeRemaining being less than MinDeltaTime 
 				// @todo(chaos): maybe we can pull data at a better time instead to avoid collision-specific code here for event dispatch
 				MSolver->GetEvolution()->GetCollisionConstraints().BeginFrame();
diff --git a/Engine/Source/Runtime/Experimental/Chaos/Private/PhysicsProxy/ClusterUnionPhysicsProxy.cpp b/Engine/Source/Runtime/Experimental/Chaos/Private/PhysicsProxy/ClusterUnionPhysicsProxy.cpp
--- a/Engine/Source/Runtime/Experimental/Chaos/Private/PhysicsProxy/ClusterUnionPhysicsProxy.cpp
+++ b/Engine/Source/Runtime/Experimental/Chaos/Private/PhysicsProxy/ClusterUnionPhysicsProxy.cpp
@@ -155,7 +155,7 @@ namespace Chaos
 				FReadPhysicsObjectInterface_Internal Interface = FPhysicsObjectInternalInterface::GetRead();
 				if (FPBDRigidsEvolutionGBF* Evolution = GetEvolution(this))
 				{
-					Evolution->GetRigidClustering().GetClusterUnionManager().HandleAddOperation(ClusterUnionIndex, Interface.GetAllRigidParticles(Objects), false);
+					Evolution->GetRigidClustering().GetClusterUnionManager().AddPendingClusterIndexOperation(ClusterUnionIndex, EClusterUnionOperation::Add, Interface.GetAllRigidParticles(Objects));
 				}
 			}
 		);
@@ -174,7 +174,7 @@ namespace Chaos
 				FReadPhysicsObjectInterface_Internal Interface = FPhysicsObjectInternalInterface::GetRead();
 				if (FPBDRigidsEvolutionGBF* Evolution = GetEvolution(this))
 				{
-					Evolution->GetRigidClustering().GetClusterUnionManager().HandleRemoveOperation(ClusterUnionIndex, Interface.GetAllRigidParticles(Objects.Array()), EClusterUnionOperationTiming::Immediate);
+					Evolution->GetRigidClustering().GetClusterUnionManager().AddPendingClusterIndexOperation(ClusterUnionIndex, EClusterUnionOperation::Remove, Interface.GetAllRigidParticles(Objects.Array()));
 				}
 			}
 		);
@@ -239,8 +239,10 @@ namespace Chaos
 		);
 	}
 
+	DECLARE_CYCLE_STAT(TEXT("FClusterUnionPhysicsProxy::PushToPhysicsState"), STAT_ClusterUnionPhysicsProxyPushToPhysicsState, STATGROUP_Chaos);
 	void FClusterUnionPhysicsProxy::PushToPhysicsState(const FDirtyPropertiesManager& Manager, int32 DataIdx, const FDirtyProxy& Dirty)
 	{
+		SCOPE_CYCLE_COUNTER(STAT_ClusterUnionPhysicsProxyPushToPhysicsState);
 		if (!ensure(Solver) || !ensure(Particle_Internal))
 		{
 			return;
@@ -307,8 +309,10 @@ namespace Chaos
 		}
 	}
 
+	DECLARE_CYCLE_STAT(TEXT("FClusterUnionPhysicsProxy::PullFromPhysicsState"), STAT_ClusterUnionPhysicsProxyPullFromPhysicsState, STATGROUP_Chaos);
 	bool FClusterUnionPhysicsProxy::PullFromPhysicsState(const FDirtyClusterUnionData& PullData, int32 SolverSyncTimestamp, const FDirtyClusterUnionData* NextPullData, const FRealSingle* Alpha)
 	{
+		SCOPE_CYCLE_COUNTER(STAT_ClusterUnionPhysicsProxyPullFromPhysicsState);
 		if (!ensure(Particle_External))
 		{
 			return false;
@@ -420,8 +424,10 @@ namespace Chaos
 		return true;
 	}
 
+	DECLARE_CYCLE_STAT(TEXT("FClusterUnionPhysicsProxy::BufferPhysicsResults_Internal"), STAT_ClusterUnionPhysicsProxyBufferPhysicsResultsInternal, STATGROUP_Chaos);
 	void FClusterUnionPhysicsProxy::BufferPhysicsResults_Internal(FDirtyClusterUnionData& BufferData)
 	{
+		SCOPE_CYCLE_COUNTER(STAT_ClusterUnionPhysicsProxyBufferPhysicsResultsInternal);
 		BufferPhysicsResultsImp(this, Particle_Internal, BufferData);
 	
 		FPBDRigidsEvolutionGBF& Evolution = *static_cast<FPBDRigidsSolver*>(Solver)->GetEvolution();
@@ -455,8 +461,10 @@ namespace Chaos
 		}
 	}
 
+	DECLARE_CYCLE_STAT(TEXT("FClusterUnionPhysicsProxy::BufferPhysicsResults_External"), STAT_ClusterUnionPhysicsProxyBufferPhysicsResultsExternal, STATGROUP_Chaos);
 	void FClusterUnionPhysicsProxy::BufferPhysicsResults_External(FDirtyClusterUnionData& BufferData)
 	{
+		SCOPE_CYCLE_COUNTER(STAT_ClusterUnionPhysicsProxyBufferPhysicsResultsExternal);
 		BufferPhysicsResultsImp(this, Particle_External.Get(), BufferData);
 		BufferData.bIsAnchored = SyncedData_External.bIsAnchored;
 
@@ -470,8 +478,10 @@ namespace Chaos
 		}
 	}
 
+	DECLARE_CYCLE_STAT(TEXT("FClusterUnionPhysicsProxy::SyncRemoteData"), STAT_ClusterUnionPhysicsProxySyncRemoteData, STATGROUP_Chaos);
 	void FClusterUnionPhysicsProxy::SyncRemoteData(FDirtyPropertiesManager& Manager, int32 DataIdx, FDirtyChaosProperties& RemoteData) const
 	{
+		SCOPE_CYCLE_COUNTER(STAT_ClusterUnionPhysicsProxySyncRemoteData);
 		if (!ensure(Particle_External))
 		{
 			return;
diff --git a/Engine/Source/Runtime/Experimental/Chaos/Public/Chaos/ClusterUnionManager.h b/Engine/Source/Runtime/Experimental/Chaos/Public/Chaos/ClusterUnionManager.h
--- a/Engine/Source/Runtime/Experimental/Chaos/Public/Chaos/ClusterUnionManager.h
+++ b/Engine/Source/Runtime/Experimental/Chaos/Public/Chaos/ClusterUnionManager.h
@@ -21,7 +21,8 @@ namespace Chaos
 		// AddReleased is the original behavior where if the particle to be added is a cluster, we will release the cluster first
 		// and add its children instead.
 		AddReleased,
-		Remove
+		Remove,
+		UpdateChildToParent
 	};
 
 	enum class EClusterUnionOperationTiming
@@ -63,7 +64,11 @@ namespace Chaos
 		bool bNeedsXRInitialization = true;
 	};
 
-	
+	struct FClusterUnionChildToParentUpdate
+	{
+		FTransform ChildToParent;
+		bool bLock = false;
+	};
 
 	/**
 	 * This class is used by Chaos to create internal clusters that will
@@ -94,6 +99,9 @@ namespace Chaos
 		// Helper function to remove particles given only the particle handle. This will consult the lookup table to find which cluster the particle is in before doing a normal remove operation.
 		void HandleRemoveOperationWithClusterLookup(const TArray<FPBDRigidParticleHandle*>& InParticles, EClusterUnionOperationTiming UpdateClusterPropertiesTiming);
 
+		// Performs the queued child to parent modifications.
+		void HandleUpdateChildToParentOperation(FClusterUnionIndex ClusterIndex, const TArray<FPBDRigidParticleHandle*>& Particles);
+
 		// Will be called at the beginning of every time step to ensure that all the expected cluster unions have been modified.
 		void FlushPendingOperations();
 
@@ -127,6 +135,7 @@ namespace Chaos
 
 		TClusterIndexOpMap<FClusterUnionIndex> PendingClusterIndexOperations;
 		TClusterIndexOpMap<FClusterUnionExplicitIndex> PendingExplicitIndexOperations;
+		TMap<FPBDRigidClusteredParticleHandle*, FClusterUnionChildToParentUpdate> PendingChildToParentUpdates;
 
 		template<typename TIndex>
 		void AddPendingOperation(TClusterIndexOpMap<TIndex>& OpMap, TIndex Index, EClusterUnionOperation Op, const TArray<FPBDRigidParticleHandle*>& Particles)
