Events optimization:  Don't create temporary array when iterating a map

diff --git a/Engine/Source/Runtime/Experimental/Chaos/Public/EventManager.h b/Engine/Source/Runtime/Experimental/Chaos/Public/EventManager.h
--- a/Engine/Source/Runtime/Experimental/Chaos/Public/EventManager.h
+++ b/Engine/Source/Runtime/Experimental/Chaos/Public/EventManager.h
@@ -283,12 +283,10 @@ namespace Chaos
 				UniqueHandlers.Reserve(Map->Num());
 				for (auto& KeyValue : *Map) // Only iterating over objects that are associated with events here
 				{
-					IPhysicsProxyBase* Proxy = KeyValue.Get<0>();
-					TArray<FEventHandlerPtr> EventHandlers;
-					ProxyOwnerToHandlerMap.MultiFind(Proxy->GetOwner(), EventHandlers);
-					for (FEventHandlerPtr EventHandler : EventHandlers)
+					const IPhysicsProxyBase* Proxy = KeyValue.Get<0>();
+					for (TMultiMap<UObject*, FEventHandlerPtr>::TConstKeyIterator It = ProxyOwnerToHandlerMap.CreateConstKeyIterator(Proxy->GetOwner()); It; ++It)
 					{
-						UniqueHandlers.Add(EventHandler);
+						UniqueHandlers.Add(It.Value());
 					}
 				}
 
