Add UGeometryCollectionComponent::SetDamageThreshold so that we can call RecreatePhysicsState when damage threshold changes.

diff --git a/Engine/Source/Runtime/Experimental/GeometryCollectionEngine/Private/GeometryCollection/GeometryCollectionComponent.cpp b/Engine/Source/Runtime/Experimental/GeometryCollectionEngine/Private/GeometryCollection/GeometryCollectionComponent.cpp
--- a/Engine/Source/Runtime/Experimental/GeometryCollectionEngine/Private/GeometryCollection/GeometryCollectionComponent.cpp
+++ b/Engine/Source/Runtime/Experimental/GeometryCollectionEngine/Private/GeometryCollection/GeometryCollectionComponent.cpp
@@ -649,6 +649,19 @@ int32 UGeometryCollectionComponent::GetNumElements(FName Group) const
 	return Size > 0 ? Size : DynamicCollection->NumElements(Group);	//if not, maybe dynamic has the group
 }
 
+void UGeometryCollectionComponent::SetDamageThreshold(const TArray<float>& InDamageThreshold)
+{
+	// NOTE: Should only call this during construction, not during runtime
+
+	DamageThreshold = InDamageThreshold;
+
+	if (IsPhysicsStateCreated())
+	{
+		// NOTE: Probably should use a deferred version of this.
+		RecreatePhysicsState();
+	}
+}
+
 void UGeometryCollectionComponent::UpdateCachedBounds()
 {
 	ComponentSpaceBounds = ComputeBounds(FMatrix::Identity);
diff --git a/Engine/Source/Runtime/Experimental/GeometryCollectionEngine/Public/GeometryCollection/GeometryCollectionComponent.h b/Engine/Source/Runtime/Experimental/GeometryCollectionEngine/Public/GeometryCollection/GeometryCollectionComponent.h
--- a/Engine/Source/Runtime/Experimental/GeometryCollectionEngine/Public/GeometryCollection/GeometryCollectionComponent.h
+++ b/Engine/Source/Runtime/Experimental/GeometryCollectionEngine/Public/GeometryCollection/GeometryCollectionComponent.h
@@ -765,9 +765,15 @@ public:
 	UPROPERTY(EditAnywhere, BlueprintReadWrite, Category = "ChaosPhysics|Clustering")
 	int32 MaxSimulatedLevel;
 
-	UPROPERTY(EditAnywhere, BlueprintReadWrite, Category = "ChaosPhysics|Damage", meta = (EditCondition = "!bUseSizeSpecificDamageThreshold"))
+	UPROPERTY(EditAnywhere, BlueprintReadWrite, BlueprintGetter=GetDamageThreshold, BlueprintSetter=SetDamageThreshold, Category = "ChaosPhysics|Damage", meta = (EditCondition = "!bUseSizeSpecificDamageThreshold"))
 	TArray<float> DamageThreshold;
 
+	UFUNCTION(BlueprintPure, BlueprintInternalUseOnly)
+	TArray<float> GetDamageThreshold() const { return DamageThreshold; }
+
+	UFUNCTION(BlueprintCallable, BlueprintInternalUseOnly)
+	void SetDamageThreshold(const TArray<float>& InDamageThreshold);
+
 	/** Damage threshold for clusters at different levels. */
 	UPROPERTY(EditAnywhere, BlueprintReadWrite, Category = "ChaosPhysics|Damage")
 	bool bUseSizeSpecificDamageThreshold;
