Geometry collection : add a property to control at which level position and velocity are corrected for clusters
The smaller the number the more bandwidth will be saved by not sending the position and velocities to the client

diff --git a/Engine/Source/Runtime/Experimental/GeometryCollectionEngine/Private/GeometryCollection/GeometryCollectionComponent.cpp b/Engine/Source/Runtime/Experimental/GeometryCollectionEngine/Private/GeometryCollection/GeometryCollectionComponent.cpp
--- a/Engine/Source/Runtime/Experimental/GeometryCollectionEngine/Private/GeometryCollection/GeometryCollectionComponent.cpp
+++ b/Engine/Source/Runtime/Experimental/GeometryCollectionEngine/Private/GeometryCollection/GeometryCollectionComponent.cpp
@@ -336,6 +336,7 @@ UGeometryCollectionComponent::UGeometryCollectionComponent(const FObjectInitiali
 	, AbandonedCollisionProfileName(UCollisionProfile::CustomCollisionProfileName)
 	, ReplicationAbandonClusterLevel_DEPRECATED(0)
 	, ReplicationAbandonAfterLevel(0)
+	, ReplicationMaxPositionAndVelocityCorrectionLevel(100)
 	, bRenderStateDirty(true)
 	, bEnableBoneSelection(false)
 	, ViewLevel(-1)
@@ -1633,8 +1634,11 @@ void UGeometryCollectionComponent::UpdateRepData()
 					if (!bEnableAbandonAfterLevel || Level <= ReplicationAbandonAfterLevel)
 					{
 						// not already replicated and not abandoned level, start replicating cluster
-						ClustersToRep->Add(Root);
-						bClustersChanged = true;
+						if (Level <= ReplicationMaxPositionAndVelocityCorrectionLevel)
+						{
+							ClustersToRep->Add(Root);
+							bClustersChanged = true;
+						}
 					}
 
 					if(Root->InternalCluster() == false && Level > 0)
diff --git a/Engine/Source/Runtime/Experimental/GeometryCollectionEngine/Public/GeometryCollection/GeometryCollectionComponent.h b/Engine/Source/Runtime/Experimental/GeometryCollectionEngine/Public/GeometryCollection/GeometryCollectionComponent.h
--- a/Engine/Source/Runtime/Experimental/GeometryCollectionEngine/Public/GeometryCollection/GeometryCollectionComponent.h
+++ b/Engine/Source/Runtime/Experimental/GeometryCollectionEngine/Public/GeometryCollection/GeometryCollectionComponent.h
@@ -1176,6 +1176,16 @@ protected:
 	UPROPERTY(EditAnywhere, BlueprintReadWrite, Category = Network)
 	int32 ReplicationAbandonAfterLevel;
 
+	/**
+	* If replicating - the maximum level where clusters will have their position and velocity send over to the client for tracking and correcting
+	* When breaking, client will only receive the initial break velocity
+	* This helps save bandwidth where only the destruction state of the GC is to be replicated but the falling pieces do not need to be tracked precisely
+	* @note This will have an effect only if set to a value less than ReplicationAbandonAfterLevel
+	* @see ReplicationAbandonAfterLevel
+	*/
+	UPROPERTY(EditAnywhere, BlueprintReadWrite, Category = Network)
+	int32 ReplicationMaxPositionAndVelocityCorrectionLevel;
+
 	UPROPERTY(ReplicatedUsing=OnRep_RepData)
 	FGeometryCollectionRepData RepData;
 
