move the release of strained particles from cluster unions into BreakingModel to handle the cases of BreakCluster and BreakClustersByProxy in a single place

diff --git a/Engine/Source/Runtime/Experimental/Chaos/Private/Chaos/PBDRigidClustering.cpp b/Engine/Source/Runtime/Experimental/Chaos/Private/Chaos/PBDRigidClustering.cpp
--- a/Engine/Source/Runtime/Experimental/Chaos/Private/Chaos/PBDRigidClustering.cpp
+++ b/Engine/Source/Runtime/Experimental/Chaos/Private/Chaos/PBDRigidClustering.cpp
@@ -1134,11 +1134,18 @@ namespace Chaos
 				SCOPE_CYCLE_COUNTER(STAT_UpdateDirtyImpulses);
 				for (const auto& ActiveCluster : TopLevelClusterParentsStrained)
 				{
+					FClusterUnion* ClusterUnion = ClusterUnionManager.FindClusterUnionFromParticle(ActiveCluster);
+					if (ClusterUnion && ClusterUnion->InternalCluster != ActiveCluster)
+					{
+						// Need to pre-emptively remove the particle from the cluster union otherwise we won't be passing the disabled check.
+						ClusterUnionManager.HandleRemoveOperationWithClusterLookup({ ActiveCluster }, EClusterUnionOperationTiming::Defer);
+					}
+
 					if (!ActiveCluster->Disabled())
 					{
 						if (ActiveCluster->ClusterIds().NumChildren > 0) //active index is a cluster
 						{
-							if(FClusterUnion* ClusterUnion = ClusterUnionManager.FindClusterUnionFromParticle(ActiveCluster))
+							if(ClusterUnion)
 							{
 								if(ClusterUnion->InternalCluster == ActiveCluster)
 								{
@@ -1173,6 +1180,8 @@ namespace Chaos
 				}
 			}
 
+			ClusterUnionManager.HandleDeferredClusterUnionUpdateProperties();
+
 			if (MCollisionImpulseArrayDirty || bPotentialBreak)
 			{
 				SCOPE_CYCLE_COUNTER(STAT_UpdateDirtyImpulses);
@@ -1748,9 +1757,6 @@ namespace Chaos
 			return false;
 		}
 
-		// Pre-emptively remove the particle from cluster unions it might be in.
-		ClusterUnionManager.HandleRemoveOperationWithClusterLookup({ ClusteredParticle }, EClusterUnionOperationTiming::Immediate);
-
 		// max strain will allow to unconditionally release the children when strain is evaluated
 		constexpr FRealSingle MaxStrain = TNumericLimits<FRealSingle>::Max();
 		if (TArray<FPBDRigidParticleHandle*>* ChildrenHandles = GetChildrenMap().Find(ClusteredParticle))
