Chaos - add Chaos trace counters for use in Insights

diff --git a/Engine/Source/Runtime/Experimental/Chaos/Private/PBDRigidsSolver.cpp b/Engine/Source/Runtime/Experimental/Chaos/Private/PBDRigidsSolver.cpp
--- a/Engine/Source/Runtime/Experimental/Chaos/Private/PBDRigidsSolver.cpp
+++ b/Engine/Source/Runtime/Experimental/Chaos/Private/PBDRigidsSolver.cpp
@@ -28,6 +28,7 @@
 #include "Chaos/PullPhysicsDataImp.h"
 #include "Chaos/PhysicsSolverBaseImpl.h"
 
+#include "ProfilingDebugging/CountersTrace.h"
 #include "ProfilingDebugging/CsvProfiler.h"
 #include "ChaosLog.h"
 
@@ -39,10 +40,13 @@ DEFINE_LOG_CATEGORY_STATIC(LogPBDRigidsSolver, Log, All);
 DECLARE_DWORD_COUNTER_STAT(TEXT("NumDisabledBodies"), STAT_ChaosCounter_NumDisabledBodies, STATGROUP_ChaosCounters);
 DECLARE_DWORD_COUNTER_STAT(TEXT("NumBodies"), STAT_ChaosCounter_NumBodies, STATGROUP_ChaosCounters);
 DECLARE_DWORD_COUNTER_STAT(TEXT("NumDynamicBodies"), STAT_ChaosCounter_NumDynamicBodies, STATGROUP_ChaosCounters);
-DECLARE_DWORD_COUNTER_STAT(TEXT("NumActiveDynamicBodies"), STAT_ChaosCounter_NumActiveDynamicBodies, STATGROUP_ChaosCounters);
 DECLARE_DWORD_COUNTER_STAT(TEXT("NumKinematicBodies"), STAT_ChaosCounter_NumKinematicBodies, STATGROUP_ChaosCounters);
 DECLARE_DWORD_COUNTER_STAT(TEXT("NumStaticBodies"), STAT_ChaosCounter_NumStaticBodies, STATGROUP_ChaosCounters);
 DECLARE_DWORD_COUNTER_STAT(TEXT("NumGeomCollBodies"), STAT_ChaosCounter_NumGeometryCollectionBodies, STATGROUP_ChaosCounters);
+DECLARE_DWORD_COUNTER_STAT(TEXT("NumMovingBodies"), STAT_ChaosCounter_NumMovingBodies, STATGROUP_ChaosCounters);
+DECLARE_DWORD_COUNTER_STAT(TEXT("NumStaticShapes"), STAT_ChaosCounter_NumStaticShapes, STATGROUP_ChaosCounters);
+DECLARE_DWORD_COUNTER_STAT(TEXT("NumKinematicShapes"), STAT_ChaosCounter_NumKinematicShapes, STATGROUP_ChaosCounters);
+DECLARE_DWORD_COUNTER_STAT(TEXT("NumDynamicShapes"), STAT_ChaosCounter_NumDynamicShapes, STATGROUP_ChaosCounters);
 DECLARE_DWORD_COUNTER_STAT(TEXT("NumIslands"), STAT_ChaosCounter_NumIslands, STATGROUP_ChaosCounters);
 DECLARE_DWORD_COUNTER_STAT(TEXT("NumIslandGroups"), STAT_ChaosCounter_NumIslandGroups, STATGROUP_ChaosCounters);
 DECLARE_DWORD_COUNTER_STAT(TEXT("NumContacts"), STAT_ChaosCounter_NumContacts, STATGROUP_ChaosCounters);
@@ -56,11 +60,39 @@ DECLARE_DWORD_COUNTER_STAT(TEXT("NumUpdatedManifoldPoints"), STAT_ChaosCounter_N
 DECLARE_DWORD_COUNTER_STAT(TEXT("NumJoints"), STAT_ChaosCounter_NumJoints, STATGROUP_ChaosCounters);
 DECLARE_DWORD_COUNTER_STAT(TEXT("NumCharacterGroundConstraints"), STAT_ChaosCounter_NumCharacterGroundConstraints, STATGROUP_ChaosCounters);
 
+TRACE_DECLARE_INT_COUNTER(ChaosTraceCounter_NumDisabledBodies, TEXT("Chaos/Solver/Bodies/NumDisabled"));
+TRACE_DECLARE_INT_COUNTER(ChaosTraceCounter_NumBodies, TEXT("Chaos/Solver/Bodies/Num"));
+TRACE_DECLARE_INT_COUNTER(ChaosTraceCounter_NumStaticBodies, TEXT("Chaos/Solver/Bodies/NumStatic"));
+TRACE_DECLARE_INT_COUNTER(ChaosTraceCounter_NumKinematicBodies, TEXT("Chaos/Solver/Bodies/NumKinematic"));
+TRACE_DECLARE_INT_COUNTER(ChaosTraceCounter_NumDynamicBodies, TEXT("Chaos/Solver/Bodies/NumDynamic"));
+TRACE_DECLARE_INT_COUNTER(ChaosTraceCounter_NumGeometryCollectionBodies, TEXT("Chaos/Solver/Bodies/NumGC"));
+TRACE_DECLARE_INT_COUNTER(ChaosTraceCounter_NumMovingBodies, TEXT("Chaos/Solver/Bodies/NumMoving"));
+TRACE_DECLARE_INT_COUNTER(ChaosTraceCounter_NumStaticShapes, TEXT("Chaos/Solver/Collisions/NumStaticShapes"));
+TRACE_DECLARE_INT_COUNTER(ChaosTraceCounter_NumKinematicShapes, TEXT("Chaos/Solver/Collisions/NumKinematicShapes"));
+TRACE_DECLARE_INT_COUNTER(ChaosTraceCounter_NumDynamicShapes, TEXT("Chaos/Solver/Collisions/NumDynamicShapes"));
+TRACE_DECLARE_INT_COUNTER(ChaosTraceCounter_NumIslands, TEXT("Chaos/Solver/Islands/NumIslands"));
+TRACE_DECLARE_INT_COUNTER(ChaosTraceCounter_NumIslandGroups, TEXT("Chaos/Solver/Islands/NumIslandGroups"));
+TRACE_DECLARE_INT_COUNTER(ChaosTraceCounter_NumContacts, TEXT("Chaos/Solver/Collisions/NumConstraints"));
+TRACE_DECLARE_INT_COUNTER(ChaosTraceCounter_NumValidConstraints, TEXT("Chaos/Solver/Collisions/NumValidConstraints"));
+TRACE_DECLARE_INT_COUNTER(ChaosTraceCounter_NumActiveConstraints, TEXT("Chaos/Solver/Collisions/NumActiveConstraints"));
+TRACE_DECLARE_INT_COUNTER(ChaosTraceCounter_NumRestoredConstraints, TEXT("Chaos/Solver/Collisions/NumRestoredConstraints"));
+TRACE_DECLARE_INT_COUNTER(ChaosTraceCounter_NumManifoldPoints, TEXT("Chaos/Solver/Collisions/NumManifoldPoints"));
+TRACE_DECLARE_INT_COUNTER(ChaosTraceCounter_NumActiveManifoldPoints, TEXT("Chaos/Solver/Collisions/NumActiveManifoldPoints"));
+TRACE_DECLARE_INT_COUNTER(ChaosTraceCounter_NumRestoredManifoldPoints, TEXT("Chaos/Solver/Collisions/NumRestoredManifoldPoints"));
+TRACE_DECLARE_INT_COUNTER(ChaosTraceCounter_NumUpdatedManifoldPoints, TEXT("Chaos/Solver/Collisions/NumUpdatedManifoldPoints"));
+TRACE_DECLARE_INT_COUNTER(ChaosTraceCounter_NumJoints, TEXT("Chaos/Solver/Joints/NumConstraints"));
+TRACE_DECLARE_INT_COUNTER(ChaosTraceCounter_NumCharacterGroundConstraints, TEXT("Chaos/Solver/Character/NumConstraints"));
+
+
 // Stat Iteration counters
 DECLARE_DWORD_COUNTER_STAT(TEXT("NumPositionIterations"), STAT_ChaosCounter_NumPositionIterations, STATGROUP_ChaosCounters);
 DECLARE_DWORD_COUNTER_STAT(TEXT("NumVelocityIterations"), STAT_ChaosCounter_NumVelocityIterations, STATGROUP_ChaosCounters);
 DECLARE_DWORD_COUNTER_STAT(TEXT("NumProjectionIterations"), STAT_ChaosCounter_NumProjectionIterations, STATGROUP_ChaosCounters);
 
+TRACE_DECLARE_INT_COUNTER(ChaosTraceCounter_NumPositionIterations, TEXT("Chaos/Solver/Iterations/NumPosition"));
+TRACE_DECLARE_INT_COUNTER(ChaosTraceCounter_NumVelocityIterations, TEXT("Chaos/Solver/Iterations/NumVelocity"));
+TRACE_DECLARE_INT_COUNTER(ChaosTraceCounter_NumProjectionIterations, TEXT("Chaos/Solver/Iterations/NumProjection"));
+
 
 // DebugDraw CVars
 #if CHAOS_DEBUG_DRAW
@@ -1897,19 +1929,25 @@ namespace Chaos
 #ifndef CHAOS_COUNTER_STAT
 #define CHAOS_COUNTER_STAT(Name, Value)\
 SET_DWORD_STAT(STAT_ChaosCounter_##Name, Value); \
-CSV_CUSTOM_STAT(PhysicsCounters, Name, Value, ECsvCustomStatOp::Set);
+CSV_CUSTOM_STAT(PhysicsCounters, Name, Value, ECsvCustomStatOp::Set); \
+TRACE_COUNTER_SET(ChaosTraceCounter_##Name, Value)
 #endif
 
 
 	void FPBDRigidsSolver::UpdateStatCounters() const
 	{
+		const int32 NumStatic = GetEvolution()->GetParticles().GetActiveStaticParticlesView().Num();
+		const int32 NumKinematic = GetEvolution()->GetParticles().GetActiveKinematicParticlesView().Num();
+		const int32 NumDynamic = GetEvolution()->GetParticles().GetNonDisabledDynamicView().Num();
+		const int32 NumMoving = GetEvolution()->GetParticles().GetActiveDynamicMovingKinematicParticlesView().Num();
+
 		// Particle counts
 		CHAOS_COUNTER_STAT(NumDisabledBodies, GetEvolution()->GetParticles().GetAllParticlesView().Num() - GetEvolution()->GetParticles().GetNonDisabledView().Num());
 		CHAOS_COUNTER_STAT(NumBodies, GetEvolution()->GetParticles().GetNonDisabledView().Num());
-		CHAOS_COUNTER_STAT(NumDynamicBodies, GetEvolution()->GetParticles().GetNonDisabledDynamicView().Num());
-		CHAOS_COUNTER_STAT(NumActiveDynamicBodies, GetEvolution()->GetParticles().GetActiveParticlesView().Num());
-		CHAOS_COUNTER_STAT(NumKinematicBodies, GetEvolution()->GetParticles().GetActiveKinematicParticlesView().Num());
-		CHAOS_COUNTER_STAT(NumStaticBodies, GetEvolution()->GetParticles().GetActiveStaticParticlesView().Num());
+		CHAOS_COUNTER_STAT(NumDynamicBodies, NumDynamic);
+		CHAOS_COUNTER_STAT(NumKinematicBodies, NumKinematic);
+		CHAOS_COUNTER_STAT(NumStaticBodies, NumStatic);
+		CHAOS_COUNTER_STAT(NumMovingBodies, NumMoving);
 		CHAOS_COUNTER_STAT(NumGeometryCollectionBodies, (int32)GetEvolution()->GetParticles().GetGeometryCollectionParticles().Size());
 
 		// Constraint counts
@@ -1981,7 +2019,33 @@ CSV_CUSTOM_STAT(PhysicsCounters, Name, Value, ECsvCustomStatOp::Set);
 						}
 					}
 				}
+			}
+		}
 
+		int32 NumStaticShapes = 0;
+		int32 NumKinematicShapes = 0;
+		int32 NumDynamicShapes = 0;
+		for (const FTransientPBDRigidParticleHandle& Particle : Particles.GetActiveParticlesView())
+		{
+			const FConstGenericParticleHandle P = Particle.Handle();
+			if (Particle.Geometry().Get() != nullptr)
+			{
+				if (const FImplicitObjectUnion* Union = Particle.Geometry()->GetObject<FImplicitObjectUnion>())
+				{
+					const int32 NumShapes = Union->GetNumLeafObjects();
+					if (P->IsDynamic())
+					{
+						NumDynamicShapes += NumShapes;
+					}
+					else if (P->IsKinematic())
+					{
+						NumKinematicShapes += NumShapes;
+					}
+					else
+					{
+						NumStaticShapes += NumShapes;
+					}
+				}
 			}
 		}
 
@@ -1992,6 +2056,9 @@ CSV_CUSTOM_STAT(PhysicsCounters, Name, Value, ECsvCustomStatOp::Set);
 		CHAOS_COUNTER_STAT(NumActiveManifoldPoints, NumActiveManifoldPoints);
 		CHAOS_COUNTER_STAT(NumRestoredManifoldPoints, NumRestoredManifoldPoints);
 		CHAOS_COUNTER_STAT(NumUpdatedManifoldPoints, NumUpdatedManifoldPoints);
+		CHAOS_COUNTER_STAT(NumStaticShapes, NumStaticShapes);
+		CHAOS_COUNTER_STAT(NumKinematicShapes, NumKinematicShapes);
+		CHAOS_COUNTER_STAT(NumDynamicShapes, NumDynamicShapes);
 	}
 
 	void FPBDRigidsSolver::PostTickDebugDraw(FReal Dt) const
