fix issue where geometry collection events would not be registered properly after construction script re-registers the component

diff --git a/Engine/Source/Runtime/Experimental/GeometryCollectionEngine/Private/GeometryCollection/GeometryCollectionComponent.cpp b/Engine/Source/Runtime/Experimental/GeometryCollectionEngine/Private/GeometryCollection/GeometryCollectionComponent.cpp
--- a/Engine/Source/Runtime/Experimental/GeometryCollectionEngine/Private/GeometryCollection/GeometryCollectionComponent.cpp
+++ b/Engine/Source/Runtime/Experimental/GeometryCollectionEngine/Private/GeometryCollection/GeometryCollectionComponent.cpp
@@ -1399,6 +1399,13 @@ void UGeometryCollectionComponent::RegisterForEvents()
 {
 	if (EventDispatcher)
 	{
+		// The new interested proxy system relies on us having having a proxy registered with the scene at the time we register Chaos events.
+		// Hence we need to do a double-take here to force the re-registration of chaos events. This is guaranteed to happen every time we
+		// re-create the physics proxy (e.g. in the case of the construction script in PIE). In that case, UChaosGameplayEventDispatcher::OnRegister
+		// will run *before* the new physics proxy is created.
+		EventDispatcher->UnregisterChaosEvents();
+		EventDispatcher->RegisterChaosEvents();
+
 		if (BodyInstance.bNotifyRigidBodyCollision || bNotifyBreaks || bNotifyCollisions || bNotifyRemovals || bNotifyCrumblings)
 		{
 			Chaos::FPhysicsSolver* Solver = GetWorld()->GetPhysicsScene()->GetSolver();
