Geometry collection component : make sure we unregister chaos events if the proxy is destroyed

diff --git a/Engine/Source/Runtime/Experimental/GeometryCollectionEngine/Private/GeometryCollection/GeometryCollectionComponent.cpp b/Engine/Source/Runtime/Experimental/GeometryCollectionEngine/Private/GeometryCollection/GeometryCollectionComponent.cpp
--- a/Engine/Source/Runtime/Experimental/GeometryCollectionEngine/Private/GeometryCollection/GeometryCollectionComponent.cpp
+++ b/Engine/Source/Runtime/Experimental/GeometryCollectionEngine/Private/GeometryCollection/GeometryCollectionComponent.cpp
@@ -3086,6 +3086,9 @@ void UGeometryCollectionComponent::OnDestroyPhysicsState()
 		DummyBodyInstance.TermBody();
 	}
 
+	// we need to unregister the events because it relies on the proxy internally 
+	EventDispatcher->UnregisterChaosEvents();
+
 	if(PhysicsProxy)
 	{
 		FPhysScene_Chaos* Scene = GetInnerChaosScene();
