Fix stack over flow when building GPU nodes bounds/parameters

diff --git a/Engine/Source/Runtime/Engine/Classes/PhysicsField/PhysicsFieldComponent.h b/Engine/Source/Runtime/Engine/Classes/PhysicsField/PhysicsFieldComponent.h
--- a/Engine/Source/Runtime/Engine/Classes/PhysicsField/PhysicsFieldComponent.h
+++ b/Engine/Source/Runtime/Engine/Classes/PhysicsField/PhysicsFieldComponent.h
@@ -186,12 +186,6 @@ public:
 	 */
 	void UpdateInstance(const float TimeSeconds);
 
-	/** Update the offsets and params given a node */
-	void BuildNodeParams(FFieldNodeBase* FieldNode, const TMap<FFieldNodeBase*, float> CommandTimes, const float PreviousTime);
-
-	/** Update the bounds given a node */
-	static void BuildNodeBounds(FFieldNodeBase* FieldNode, FVector& MinBounds, FVector& MaxBounds, float& MaxMagnitude);
-
 	/** The field system resource. */
 	FPhysicsFieldResource* FieldResource = nullptr;
 
diff --git a/Engine/Source/Runtime/Engine/Private/PhysicsField/PhysicsFieldComponent.cpp b/Engine/Source/Runtime/Engine/Private/PhysicsField/PhysicsFieldComponent.cpp
--- a/Engine/Source/Runtime/Engine/Private/PhysicsField/PhysicsFieldComponent.cpp
+++ b/Engine/Source/Runtime/Engine/Private/PhysicsField/PhysicsFieldComponent.cpp
@@ -685,6 +685,11 @@ void FPhysicsFieldInstance::UpdateInstance(const float TimeSeconds)
 	TStaticArray<int32, EFieldPhysicsType::Field_PhysicsType_Max + 1> BoundsOffsets;
 	TArray<FVector4> BoundsMin, BoundsMax;
 
+	int32 NumOffsets = 0;
+	int32 NumParams = 0;
+	TArray<int32> TestOffsets;
+	TArray<float> TestParams;
+
 	if (FieldResource)
 	{
 		TargetsMin.Init(FVector(0), EFieldPhysicsType::Field_PhysicsType_Max);
@@ -771,11 +776,14 @@ void FPhysicsFieldInstance::UpdateInstance(const float TimeSeconds)
 			}
 			const int32 PreviousNodes = NodesOffsets.Num();
 			FVector MinBound(-FLT_MAX), MaxBound(FLT_MAX);
-			float MaxMagnitude = 1.0;
 			if (TargetNode)
 			{
-				BuildNodeBounds(TargetNode, MinBound, MaxBound, MaxMagnitude);
-				BuildNodeParams(TargetNode, CommandTimes, 0.0);
+				TargetNode->FillSetupCount(NumOffsets, NumParams);
+				NodesOffsets.Reserve(NumOffsets);
+				NodesParams.Reserve(NumParams);
+
+				TargetNode->FillSetupDatas(NodesOffsets, NodesParams, CommandTimes, 0.0);
+				TargetNode->ComputeFieldBounds(MinBound, MaxBound);
 				if (TargetRoots.Num() > 1) delete TargetNode;
 			}
 			TargetsOffsets[TargetType + 1] = NodesOffsets.Num() - PreviousNodes;
@@ -814,488 +822,6 @@ void FPhysicsFieldInstance::UpdateInstance(const float TimeSeconds)
 	}
 }
 
-void FPhysicsFieldInstance::BuildNodeParams(FFieldNodeBase* FieldNode, const TMap<FFieldNodeBase*,float> CommandTimes, const float PreviousTime)
-{
-	if (FieldNode)
-	{
-		const float NextTime = CommandTimes.Find(FieldNode) ? CommandTimes[FieldNode] : PreviousTime;
-		if (FieldNode->SerializationType() == FFieldNodeBase::ESerializationType::FieldNode_FUniformInteger)
-		{
-			FUniformInteger* LocalNode = StaticCast<FUniformInteger*>(FieldNode);
-			NodesOffsets.Add(NodesParams.Num());
-			NodesParams.Add(FieldNode->Type());
-			NodesParams.Add(FFieldNodeBase::ESerializationType::FieldNode_FUniformInteger);
-			NodesParams.Add(LocalNode->Magnitude);
-		}
-		else if (FieldNode->SerializationType() == FFieldNodeBase::ESerializationType::FieldNode_FRadialIntMask)
-		{
-			FRadialIntMask* LocalNode = StaticCast<FRadialIntMask*>(FieldNode);
-			NodesOffsets.Add(NodesParams.Num());
-			NodesParams.Add(FieldNode->Type());
-			NodesParams.Add(FFieldNodeBase::ESerializationType::FieldNode_FRadialIntMask);
-			NodesParams.Add(LocalNode->Radius);
-			NodesParams.Add(LocalNode->Position.X);
-			NodesParams.Add(LocalNode->Position.Y);
-			NodesParams.Add(LocalNode->Position.Z);
-			NodesParams.Add(LocalNode->InteriorValue);
-			NodesParams.Add(LocalNode->ExteriorValue);
-			NodesParams.Add(LocalNode->SetMaskCondition);
-		}
-		else if (FieldNode->SerializationType() == FFieldNodeBase::ESerializationType::FieldNode_FUniformScalar)
-		{
-			FUniformScalar* LocalNode = StaticCast<FUniformScalar*>(FieldNode);
-			NodesOffsets.Add(NodesParams.Num());
-			NodesParams.Add(FieldNode->Type());
-			NodesParams.Add(FFieldNodeBase::ESerializationType::FieldNode_FUniformScalar);
-			NodesParams.Add(LocalNode->Magnitude);
-		}
-		else if (FieldNode->SerializationType() == FFieldNodeBase::ESerializationType::FieldNode_FWaveScalar)
-		{
-			FWaveScalar* LocalNode = StaticCast<FWaveScalar*>(FieldNode);
-			NodesOffsets.Add(NodesParams.Num());
-			NodesParams.Add(FieldNode->Type());
-			NodesParams.Add(FFieldNodeBase::ESerializationType::FieldNode_FWaveScalar);
-			NodesParams.Add(LocalNode->Magnitude);
-			NodesParams.Add(LocalNode->Position.X);
-			NodesParams.Add(LocalNode->Position.Y);
-			NodesParams.Add(LocalNode->Position.Z);
-			NodesParams.Add(LocalNode->Wavelength);
-			NodesParams.Add(LocalNode->Period);
-			NodesParams.Add(NextTime);
-			NodesParams.Add(LocalNode->Function);
-			NodesParams.Add(LocalNode->Falloff);
-		}
-		else if (FieldNode->SerializationType() == FFieldNodeBase::ESerializationType::FieldNode_FRadialFalloff)
-		{
-			FRadialFalloff* LocalNode = StaticCast<FRadialFalloff*>(FieldNode);
-			NodesOffsets.Add(NodesParams.Num());
-			NodesParams.Add(FieldNode->Type());
-			NodesParams.Add(FFieldNodeBase::ESerializationType::FieldNode_FRadialFalloff);
-			NodesParams.Add(LocalNode->Magnitude);
-			NodesParams.Add(LocalNode->MinRange);
-			NodesParams.Add(LocalNode->MaxRange);
-			NodesParams.Add(LocalNode->Default);
-			NodesParams.Add(LocalNode->Radius);
-			NodesParams.Add(LocalNode->Position.X);
-			NodesParams.Add(LocalNode->Position.Y);
-			NodesParams.Add(LocalNode->Position.Z);
-			NodesParams.Add(LocalNode->Falloff);
-		}
-		else if (FieldNode->SerializationType() == FFieldNodeBase::ESerializationType::FieldNode_FPlaneFalloff)
-		{
-			FPlaneFalloff* LocalNode = StaticCast<FPlaneFalloff*>(FieldNode);
-			NodesOffsets.Add(NodesParams.Num());
-			NodesParams.Add(FieldNode->Type());
-			NodesParams.Add(FFieldNodeBase::ESerializationType::FieldNode_FPlaneFalloff);
-			NodesParams.Add(LocalNode->Magnitude);
-			NodesParams.Add(LocalNode->MinRange);
-			NodesParams.Add(LocalNode->MaxRange);
-			NodesParams.Add(LocalNode->Default);
-			NodesParams.Add(LocalNode->Distance);
-			NodesParams.Add(LocalNode->Position.X);
-			NodesParams.Add(LocalNode->Position.Y);
-			NodesParams.Add(LocalNode->Position.Z);
-			NodesParams.Add(LocalNode->Normal.X);
-			NodesParams.Add(LocalNode->Normal.Y);
-			NodesParams.Add(LocalNode->Normal.Z);
-			NodesParams.Add(LocalNode->Falloff);
-		}
-		else if (FieldNode->SerializationType() == FFieldNodeBase::ESerializationType::FieldNode_FBoxFalloff)
-		{
-			FBoxFalloff* LocalNode = StaticCast<FBoxFalloff*>(FieldNode);
-			NodesOffsets.Add(NodesParams.Num());
-			NodesParams.Add(FieldNode->Type());
-			NodesParams.Add(FFieldNodeBase::ESerializationType::FieldNode_FBoxFalloff);
-			NodesParams.Add(LocalNode->Magnitude);
-			NodesParams.Add(LocalNode->MinRange);
-			NodesParams.Add(LocalNode->MaxRange);
-			NodesParams.Add(LocalNode->Default);
-			NodesParams.Add(LocalNode->Transform.GetRotation().X);
-			NodesParams.Add(LocalNode->Transform.GetRotation().Y);
-			NodesParams.Add(LocalNode->Transform.GetRotation().Z);
-			NodesParams.Add(LocalNode->Transform.GetRotation().W);
-			NodesParams.Add(LocalNode->Transform.GetTranslation().X);
-			NodesParams.Add(LocalNode->Transform.GetTranslation().Y);
-			NodesParams.Add(LocalNode->Transform.GetTranslation().Z);
-			NodesParams.Add(LocalNode->Transform.GetScale3D().X);
-			NodesParams.Add(LocalNode->Transform.GetScale3D().Y);
-			NodesParams.Add(LocalNode->Transform.GetScale3D().Z);
-			NodesParams.Add(LocalNode->Falloff);
-		}
-		else if (FieldNode->SerializationType() == FFieldNodeBase::ESerializationType::FieldNode_FNoiseField)
-		{
-			FNoiseField* LocalNode = StaticCast<FNoiseField*>(FieldNode);
-			NodesOffsets.Add(NodesParams.Num());
-			NodesParams.Add(FieldNode->Type());
-			NodesParams.Add(FFieldNodeBase::ESerializationType::FieldNode_FNoiseField);
-			NodesParams.Add(LocalNode->MinRange);
-			NodesParams.Add(LocalNode->MaxRange);
-			NodesParams.Add(LocalNode->Transform.GetRotation().X);
-			NodesParams.Add(LocalNode->Transform.GetRotation().Y);
-			NodesParams.Add(LocalNode->Transform.GetRotation().Z);
-			NodesParams.Add(LocalNode->Transform.GetRotation().W);
-			NodesParams.Add(LocalNode->Transform.GetTranslation().X);
-			NodesParams.Add(LocalNode->Transform.GetTranslation().Y);
-			NodesParams.Add(LocalNode->Transform.GetTranslation().Z);
-			NodesParams.Add(LocalNode->Transform.GetScale3D().X);
-			NodesParams.Add(LocalNode->Transform.GetScale3D().Y);
-			NodesParams.Add(LocalNode->Transform.GetScale3D().Z);
-		}
-		else if (FieldNode->SerializationType() == FFieldNodeBase::ESerializationType::FieldNode_FUniformVector)
-		{
-			FUniformVector* LocalNode = StaticCast<FUniformVector*>(FieldNode);
-			NodesOffsets.Add(NodesParams.Num());
-			NodesParams.Add(FieldNode->Type());
-			NodesParams.Add(FFieldNodeBase::ESerializationType::FieldNode_FUniformVector);
-			NodesParams.Add(LocalNode->Magnitude);
-			NodesParams.Add(LocalNode->Direction.X);
-			NodesParams.Add(LocalNode->Direction.Y);
-			NodesParams.Add(LocalNode->Direction.Z);
-		}
-		else if (FieldNode->SerializationType() == FFieldNodeBase::ESerializationType::FieldNode_FRadialVector)
-		{
-			FRadialVector* LocalNode = StaticCast<FRadialVector*>(FieldNode);
-			NodesOffsets.Add(NodesParams.Num());
-			NodesParams.Add(FieldNode->Type());
-			NodesParams.Add(FFieldNodeBase::ESerializationType::FieldNode_FRadialVector);
-			NodesParams.Add(LocalNode->Magnitude);
-			NodesParams.Add(LocalNode->Position.X);
-			NodesParams.Add(LocalNode->Position.Y);
-			NodesParams.Add(LocalNode->Position.Z);
-		}
-		else if (FieldNode->SerializationType() == FFieldNodeBase::ESerializationType::FieldNode_FRandomVector)
-		{
-			FRandomVector* LocalNode = StaticCast<FRandomVector*>(FieldNode);
-			NodesOffsets.Add(NodesParams.Num());
-			NodesParams.Add(FieldNode->Type());
-			NodesParams.Add(FFieldNodeBase::ESerializationType::FieldNode_FRandomVector);
-			NodesParams.Add(LocalNode->Magnitude);
-		}
-		else if (FieldNode->SerializationType() == FFieldNodeBase::ESerializationType::FieldNode_FSumScalar)
-		{
-			FSumScalar* LocalNode = StaticCast<FSumScalar*>(FieldNode);
-
-			BuildNodeParams(LocalNode->ScalarRight.Get(), CommandTimes, NextTime);
-			BuildNodeParams(LocalNode->ScalarLeft.Get(), CommandTimes, NextTime);
-
-			NodesOffsets.Add(NodesParams.Num());
-			NodesParams.Add(FieldNode->Type());
-			NodesParams.Add(FFieldNodeBase::ESerializationType::FieldNode_FSumScalar);
-			NodesParams.Add(LocalNode->Magnitude);
-			NodesParams.Add(LocalNode->ScalarRight != nullptr);
-			NodesParams.Add(LocalNode->ScalarLeft != nullptr);
-			NodesParams.Add(LocalNode->Operation);
-		}
-		else if (FieldNode->SerializationType() == FFieldNodeBase::ESerializationType::FieldNode_FSumVector)
-		{
-			FSumVector* LocalNode = StaticCast<FSumVector*>(FieldNode);
-
-			BuildNodeParams(LocalNode->Scalar.Get(), CommandTimes, NextTime);
-			BuildNodeParams(LocalNode->VectorRight.Get(), CommandTimes, NextTime);
-			BuildNodeParams(LocalNode->VectorLeft.Get(), CommandTimes, NextTime);
-
-			NodesOffsets.Add(NodesParams.Num());
-			NodesParams.Add(FieldNode->Type());
-			NodesParams.Add(FFieldNodeBase::ESerializationType::FieldNode_FSumVector);
-			NodesParams.Add(LocalNode->Magnitude);
-			NodesParams.Add(LocalNode->Scalar.Get() != nullptr);
-			NodesParams.Add(LocalNode->VectorRight.Get() != nullptr);
-			NodesParams.Add(LocalNode->VectorLeft.Get() != nullptr);
-			NodesParams.Add(LocalNode->Operation);
-		}
-		else if (FieldNode->SerializationType() == FFieldNodeBase::ESerializationType::FieldNode_FConversionField)
-		{
-			if (FieldNode->Type() == FFieldNodeBase::EFieldType::EField_Int32)
-			{
-				FConversionField<float, int32>* LocalNode = StaticCast<FConversionField<float, int32>*>(FieldNode);
-
-				BuildNodeParams(LocalNode->InputField.Get(), CommandTimes, NextTime);
-
-				NodesOffsets.Add(NodesParams.Num());
-				NodesParams.Add(FieldNode->Type());
-				NodesParams.Add(FFieldNodeBase::ESerializationType::FieldNode_FConversionField);
-				NodesParams.Add(LocalNode->InputField.Get() != nullptr);
-			}
-			else if (FieldNode->Type() == FFieldNodeBase::EFieldType::EField_Float)
-			{
-				FConversionField<int32, float>* LocalNode = StaticCast<FConversionField<int32, float>*>(FieldNode);
-
-				BuildNodeParams(LocalNode->InputField.Get(), CommandTimes, NextTime);
-
-				NodesOffsets.Add(NodesParams.Num());
-				NodesParams.Add(FieldNode->Type());
-				NodesParams.Add(FFieldNodeBase::ESerializationType::FieldNode_FConversionField);
-				NodesParams.Add(LocalNode->InputField.Get() != nullptr);
-			}
-		}
-		else if (FieldNode->SerializationType() == FFieldNodeBase::ESerializationType::FieldNode_FCullingField)
-		{
-			if (FieldNode->Type() == FFieldNodeBase::EFieldType::EField_Int32)
-			{
-				FCullingField<int32>* LocalNode = StaticCast<FCullingField<int32>*>(FieldNode);
-
-				BuildNodeParams(LocalNode->Culling.Get(), CommandTimes, NextTime);
-				BuildNodeParams(LocalNode->Input.Get(), CommandTimes, NextTime);
-
-				NodesOffsets.Add(NodesParams.Num());
-				NodesParams.Add(FieldNode->Type());
-				NodesParams.Add(FFieldNodeBase::ESerializationType::FieldNode_FCullingField);
-				NodesParams.Add(LocalNode->Culling.Get() != nullptr);
-				NodesParams.Add(LocalNode->Input.Get() != nullptr);
-				NodesParams.Add(LocalNode->Operation);
-			}
-			else if (FieldNode->Type() == FFieldNodeBase::EFieldType::EField_Float)
-			{
-				FCullingField<float>* LocalNode = StaticCast<FCullingField<float>*>(FieldNode);
-
-				BuildNodeParams(LocalNode->Culling.Get(), CommandTimes, NextTime);
-				BuildNodeParams(LocalNode->Input.Get(), CommandTimes, NextTime);
-
-				NodesOffsets.Add(NodesParams.Num());
-				NodesParams.Add(FieldNode->Type());
-				NodesParams.Add(FFieldNodeBase::ESerializationType::FieldNode_FCullingField);
-				NodesParams.Add(LocalNode->Culling.Get() != nullptr);
-				NodesParams.Add(LocalNode->Input.Get() != nullptr);
-				NodesParams.Add(LocalNode->Operation);
-			}
-			else if (FieldNode->Type() == FFieldNodeBase::EFieldType::EField_FVector)
-			{
-				FCullingField<FVector>* LocalNode = StaticCast<FCullingField<FVector>*>(FieldNode);
-
-				BuildNodeParams(LocalNode->Culling.Get(), CommandTimes, NextTime);
-				BuildNodeParams(LocalNode->Input.Get(), CommandTimes, NextTime);
-
-				NodesOffsets.Add(NodesParams.Num());
-				NodesParams.Add(FieldNode->Type());
-				NodesParams.Add(FFieldNodeBase::ESerializationType::FieldNode_FCullingField);
-				NodesParams.Add(LocalNode->Culling.Get() != nullptr);
-				NodesParams.Add(LocalNode->Input.Get() != nullptr);
-				NodesParams.Add(LocalNode->Operation);
-			}
-		}
-	}
-}
-
-void FPhysicsFieldInstance::BuildNodeBounds(FFieldNodeBase* FieldNode, FVector& MinBounds, FVector& MaxBounds, float& MaxMagnitude)
-{
-	MinBounds = FVector(-FLT_MAX);
-	MaxBounds = FVector(FLT_MAX);
-	MaxMagnitude = 1.0;
-
-	if (FieldNode)
-	{
-		if (FieldNode->SerializationType() == FFieldNodeBase::ESerializationType::FieldNode_FUniformInteger)
-		{
-			FUniformInteger* LocalNode = StaticCast<FUniformInteger*>(FieldNode);
-			MaxMagnitude = LocalNode->Magnitude;
-		}
-		else if (FieldNode->SerializationType() == FFieldNodeBase::ESerializationType::FieldNode_FUniformScalar)
-		{
-			FUniformScalar* LocalNode = StaticCast<FUniformScalar*>(FieldNode);
-			MaxMagnitude = LocalNode->Magnitude;
-		}
-		else if (FieldNode->SerializationType() == FFieldNodeBase::ESerializationType::FieldNode_FWaveScalar)
-		{
-			FWaveScalar* LocalNode = StaticCast<FWaveScalar*>(FieldNode);
-			MaxMagnitude = LocalNode->Magnitude;
-		}
-		else if (FieldNode->SerializationType() == FFieldNodeBase::ESerializationType::FieldNode_FUniformVector)
-		{
-			FUniformVector* LocalNode = StaticCast<FUniformVector*>(FieldNode);
-			MaxMagnitude = LocalNode->Magnitude;
-		}
-		else if (FieldNode->SerializationType() == FFieldNodeBase::ESerializationType::FieldNode_FRadialVector)
-		{
-			FRadialVector* LocalNode = StaticCast<FRadialVector*>(FieldNode);
-			MaxMagnitude = LocalNode->Magnitude;
-		}
-		else if (FieldNode->SerializationType() == FFieldNodeBase::ESerializationType::FieldNode_FRandomVector)
-		{
-			FRandomVector* LocalNode = StaticCast<FRandomVector*>(FieldNode);
-			MaxMagnitude = LocalNode->Magnitude;
-		}
-		else if (FieldNode->SerializationType() == FFieldNodeBase::ESerializationType::FieldNode_FRadialIntMask)
-		{
-			FRadialIntMask* LocalNode = StaticCast<FRadialIntMask*>(FieldNode);
-
-			MinBounds = (LocalNode->ExteriorValue == 0) ? LocalNode->Position - FVector(LocalNode->Radius) : FVector(-FLT_MAX);
-			MaxBounds = (LocalNode->ExteriorValue == 0) ? LocalNode->Position + FVector(LocalNode->Radius) : FVector(FLT_MAX);
-			MaxMagnitude = 1.0;
-		}
-		else if (FieldNode->SerializationType() == FFieldNodeBase::ESerializationType::FieldNode_FPlaneFalloff)
-		{
-			FPlaneFalloff* LocalNode = StaticCast<FPlaneFalloff*>(FieldNode);
-			MinBounds = FVector(-FLT_MAX);
-			MaxBounds = FVector(FLT_MAX);
-			MaxMagnitude = LocalNode->Magnitude;
-		}
-		else if (FieldNode->SerializationType() == FFieldNodeBase::ESerializationType::FieldNode_FRadialFalloff)
-		{
-			FRadialFalloff* LocalNode = StaticCast<FRadialFalloff*>(FieldNode);
-
-			MinBounds = (LocalNode->Default == 0) ? LocalNode->Position - FVector(LocalNode->Radius) : FVector(-FLT_MAX);
-			MaxBounds = (LocalNode->Default == 0) ? LocalNode->Position + FVector(LocalNode->Radius) : FVector(FLT_MAX);
-			MaxMagnitude = LocalNode->Magnitude;
-		}
-		else if (FieldNode->SerializationType() == FFieldNodeBase::ESerializationType::FieldNode_FBoxFalloff)
-		{
-			FBoxFalloff* LocalNode = StaticCast<FBoxFalloff*>(FieldNode);
-			MinBounds = FVector(-FLT_MAX);
-			MaxBounds = FVector(FLT_MAX);
-			if (LocalNode->Default == 0)
-			{
-				const FBox UnitBox(FVector(-50), FVector(50));
-				const FBox BoundingBox = UnitBox.TransformBy(LocalNode->Transform);
-				MinBounds = BoundingBox.Min;
-				MaxBounds = BoundingBox.Max;
-			}
-			MaxMagnitude = LocalNode->Magnitude;
-		}
-		else if (FieldNode->SerializationType() == FFieldNodeBase::ESerializationType::FieldNode_FSumScalar)
-		{
-			FSumScalar* LocalNode = StaticCast<FSumScalar*>(FieldNode);
-
-			FVector MinBoundsA(-FLT_MAX), MaxBoundsA(FLT_MAX), MinBoundsB(-FLT_MAX), MaxBoundsB(FLT_MAX);
-			float MaxMagnitudeA = 0.0, MaxMagnitudeB = 0.0;
-			BuildNodeBounds(LocalNode->ScalarRight.Get(), MinBoundsA, MaxBoundsA, MaxMagnitudeA);
-			BuildNodeBounds(LocalNode->ScalarLeft.Get(), MinBoundsB, MaxBoundsB, MaxMagnitudeB);
-
-			if (LocalNode->Operation == EFieldOperationType::Field_Multiply ||
-				LocalNode->Operation == EFieldOperationType::Field_Divide)
-			{
-				MinBounds = MaxVector(MinBoundsA, MinBoundsB);
-				MaxBounds = MinVector(MaxBoundsA, MaxBoundsB);
-				MaxMagnitude = (LocalNode->Operation == EFieldOperationType::Field_Multiply) ? MaxMagnitudeA * MaxMagnitudeB : MaxMagnitudeA / MaxMagnitudeB;
-			}
-			else if (LocalNode->Operation == EFieldOperationType::Field_Add ||
-				LocalNode->Operation == EFieldOperationType::Field_Substract)
-			{
-				MinBounds = MinVector(MinBoundsA, MinBoundsB);
-				MaxBounds = MaxVector(MaxBoundsA, MaxBoundsB);
-				MaxMagnitude = FMath::Max(MaxMagnitudeA,MaxMagnitudeB);
-			}
-		}
-		else if (FieldNode->SerializationType() == FFieldNodeBase::ESerializationType::FieldNode_FSumVector)
-		{
-			FSumVector* LocalNode = StaticCast<FSumVector*>(FieldNode);
-
-			FVector MinBoundsA(-FLT_MAX), MaxBoundsA(FLT_MAX), MinBoundsB(-FLT_MAX), MaxBoundsB(FLT_MAX), MinBoundsC(-FLT_MAX), MaxBoundsC(FLT_MAX);
-			float MaxMagnitudeA = 0.0, MaxMagnitudeB = 0.0, MaxMagnitudeC = 0.0;
-			BuildNodeBounds(LocalNode->Scalar.Get(), MinBoundsA, MaxBoundsA, MaxMagnitudeA);
-			BuildNodeBounds(LocalNode->VectorRight.Get(), MinBoundsB, MaxBoundsB, MaxMagnitudeB);
-			BuildNodeBounds(LocalNode->VectorLeft.Get(), MinBoundsC, MaxBoundsC, MaxMagnitudeC);
-
-			if (LocalNode->Operation == EFieldOperationType::Field_Multiply ||
-				LocalNode->Operation == EFieldOperationType::Field_Divide)
-			{
-				MinBounds = MaxVector(MinBoundsB, MinBoundsC);
-				MaxBounds = MinVector(MaxBoundsB, MaxBoundsC);
-				MaxMagnitude = (LocalNode->Operation == EFieldOperationType::Field_Multiply) ? MaxMagnitudeB * MaxMagnitudeC : MaxMagnitudeB / MaxMagnitudeC;
-			}
-			else if (LocalNode->Operation == EFieldOperationType::Field_Add ||
-					 LocalNode->Operation == EFieldOperationType::Field_Substract)
-			{
-				MinBounds = MinVector(MinBoundsB, MinBoundsC);
-				MaxBounds = MaxVector(MaxBoundsB, MaxBoundsC);
-				MaxMagnitude = FMath::Max(MaxMagnitudeB, MaxMagnitudeC);
-			}
-			MinBounds = MaxVector(MinBounds, MinBoundsA);
-			MaxBounds = MinVector(MaxBounds, MaxBoundsA);
-			MaxMagnitude = MaxMagnitude * MaxMagnitudeA;
-		}
-		else if (FieldNode->SerializationType() == FFieldNodeBase::ESerializationType::FieldNode_FCullingField)
-		{
-			if (FieldNode->Type() == FFieldNodeBase::EFieldType::EField_Int32)
-			{
-				FCullingField<int32>* LocalNode = StaticCast<FCullingField<int32>*>(FieldNode);
-
-				FVector MinBoundsA(-FLT_MAX), MaxBoundsA(FLT_MAX), MinBoundsB(-FLT_MAX), MaxBoundsB(FLT_MAX);
-				float MaxMagnitudeA = 0.0, MaxMagnitudeB = 0.0;
-				BuildNodeBounds(LocalNode->Culling.Get(), MinBoundsA, MaxBoundsA, MaxMagnitudeA);
-				BuildNodeBounds(LocalNode->Input.Get(), MinBoundsB, MaxBoundsB, MaxMagnitudeB);
-
-				if (LocalNode->Operation == EFieldCullingOperationType::Field_Culling_Inside)
-				{
-					MinBounds = MinVector(MinBoundsA, MinBoundsB);
-					MaxBounds = MaxVector(MaxBoundsA, MaxBoundsB);
-				}
-				else if (LocalNode->Operation == EFieldCullingOperationType::Field_Culling_Outside)
-				{
-					MinBounds = MaxVector(MinBoundsA, MinBoundsB);
-					MaxBounds = MinVector(MaxBoundsA, MaxBoundsB);
-				}
-				MaxMagnitude = MaxMagnitudeB;
-			}
-			else if (FieldNode->Type() == FFieldNodeBase::EFieldType::EField_Float)
-			{
-				FCullingField<float>* LocalNode = StaticCast<FCullingField<float>*>(FieldNode);
-
-				FVector MinBoundsA(-FLT_MAX), MaxBoundsA(FLT_MAX), MinBoundsB(-FLT_MAX), MaxBoundsB(FLT_MAX);
-				float MaxMagnitudeA = 0.0, MaxMagnitudeB = 0.0;
-				BuildNodeBounds(LocalNode->Culling.Get(), MinBoundsA, MaxBoundsA, MaxMagnitudeA);
-				BuildNodeBounds(LocalNode->Input.Get(), MinBoundsB, MaxBoundsB, MaxMagnitudeB);
-
-				if (LocalNode->Operation == EFieldCullingOperationType::Field_Culling_Inside)
-				{
-					MinBounds = MinVector(MinBoundsA, MinBoundsB);
-					MaxBounds = MaxVector(MaxBoundsA, MaxBoundsB);
-				}
-				else if (LocalNode->Operation == EFieldCullingOperationType::Field_Culling_Outside)
-				{
-					MinBounds = MaxVector(MinBoundsA, MinBoundsB);
-					MaxBounds = MinVector(MaxBoundsA, MaxBoundsB);
-				}
-				MaxMagnitude = MaxMagnitudeB;
-			}
-			else if (FieldNode->Type() == FFieldNodeBase::EFieldType::EField_FVector)
-			{
-				FCullingField<FVector>* LocalNode = StaticCast<FCullingField<FVector>*>(FieldNode);
-
-				FVector MinBoundsA(-FLT_MAX), MaxBoundsA(FLT_MAX), MinBoundsB(-FLT_MAX), MaxBoundsB(FLT_MAX);
-				float MaxMagnitudeA = 0.0, MaxMagnitudeB = 0.0;
-				BuildNodeBounds(LocalNode->Culling.Get(), MinBoundsA, MaxBoundsA, MaxMagnitudeA);
-				BuildNodeBounds(LocalNode->Input.Get(), MinBoundsB, MaxBoundsB, MaxMagnitudeB);
-
-				if (LocalNode->Operation == EFieldCullingOperationType::Field_Culling_Inside)
-				{
-					MinBounds = MinVector(MinBoundsA, MinBoundsB);
-					MaxBounds = MaxVector(MaxBoundsA, MaxBoundsB);
-				}
-				else if (LocalNode->Operation == EFieldCullingOperationType::Field_Culling_Outside)
-				{
-					MinBounds = MaxVector(MinBoundsA, MinBoundsB);
-					MaxBounds = MinVector(MaxBoundsA, MaxBoundsB);
-				}
-				MaxMagnitude = MaxMagnitudeB;
-			}
-		}
-		else if (FieldNode->SerializationType() == FFieldNodeBase::ESerializationType::FieldNode_FConversionField)
-		{
-			if (FieldNode->Type() == FFieldNodeBase::EFieldType::EField_Int32)
-			{
-				FConversionField<float, int32>* LocalNode = StaticCast<FConversionField<float, int32>*>(FieldNode);
-				FVector MinBoundsA(-FLT_MAX), MaxBoundsA(FLT_MAX);
-				float MaxMagnitudeA = 0.0;
-
-				BuildNodeBounds(LocalNode->InputField.Get(), MinBoundsA, MaxBoundsA, MaxMagnitudeA);
-				MinBounds = MinBoundsA;
-				MaxBounds = MaxBoundsA;
-				MaxMagnitude = MaxMagnitudeA;
-			}
-			else if (FieldNode->Type() == FFieldNodeBase::EFieldType::EField_Float)
-			{
-				FConversionField<int32, float>* LocalNode = StaticCast<FConversionField<int32, float>*>(FieldNode);
-				FVector MinBoundsA(-FLT_MAX), MaxBoundsA(FLT_MAX);
-				float MaxMagnitudeA = 0.0;
-
-				BuildNodeBounds(LocalNode->InputField.Get(), MinBoundsA, MaxBoundsA, MaxMagnitudeA);
-				MinBounds = MinBoundsA;
-				MaxBounds = MaxBoundsA;
-				MaxMagnitude = MaxMagnitudeA;
-			}
-		}
-	}
-}
 
 /**
 *	PhysicsFieldComponent
@@ -1459,7 +985,11 @@ void UPhysicsFieldComponent::BuildCommandBounds(FFieldSystemCommand& FieldComman
 	FieldCommand.BoundingBox.Max = FVector(FLT_MAX);
 	FieldCommand.MaxMagnitude = 1.0;
 
-	FPhysicsFieldInstance::BuildNodeBounds(FieldCommand.RootNode.Get(), FieldCommand.BoundingBox.Min, FieldCommand.BoundingBox.Max, FieldCommand.MaxMagnitude);
+	if (FieldCommand.RootNode.IsValid())
+	{
+		FieldCommand.MaxMagnitude = FieldCommand.RootNode->EvalMaxMagnitude();
+		FieldCommand.RootNode->ComputeFieldBounds(FieldCommand.BoundingBox.Min, FieldCommand.BoundingBox.Max);
+	}
 }
 
 void UPhysicsFieldComponent::AddTransientCommand(const FFieldSystemCommand& FieldCommand, const bool bIsWorldField)
diff --git a/Engine/Source/Runtime/Experimental/Chaos/Private/Field/FieldSystemNodes.cpp b/Engine/Source/Runtime/Experimental/Chaos/Private/Field/FieldSystemNodes.cpp
--- a/Engine/Source/Runtime/Experimental/Chaos/Private/Field/FieldSystemNodes.cpp
+++ b/Engine/Source/Runtime/Experimental/Chaos/Private/Field/FieldSystemNodes.cpp
@@ -6,6 +6,19 @@
 #include "Chaos/Vector.h"
 #include <type_traits>
 
+namespace 
+{
+	inline FVector MinVector(const FVector& VectorA, const FVector& VectorB)
+	{
+		return FVector(FMath::Min(VectorA.X, VectorB.X), FMath::Min(VectorA.Y, VectorB.Y), FMath::Min(VectorA.Z, VectorB.Z));
+	}
+
+	inline FVector MaxVector(const FVector& VectorA, const FVector& VectorB)
+	{
+		return FVector(FMath::Max(VectorA.X, VectorB.X), FMath::Max(VectorA.Y, VectorB.Y), FMath::Max(VectorA.Z, VectorB.Z));
+	}
+}
+
 FFieldNodeBase * FieldNodeFactory(FFieldNodeBase::EFieldType BaseType,FFieldNodeBase::ESerializationType Type)
 {
 	switch (Type)
@@ -94,6 +107,7 @@ void SerializeInternal(FArchive& Ar, Enum& Var)
 /**
 * FUniformInteger
 */
+
 void FUniformInteger::Evaluate(FFieldContext& Context, TFieldArrayView<int32>& Results) const
 {
 	int32 MagnitudeVal = Magnitude;
@@ -104,11 +118,13 @@ void FUniformInteger::Evaluate(FFieldContext& Context, TFieldArrayView<int32>& R
 		Results[Context.SampleIndices[SampleIndex].Result] = MagnitudeVal;
 	}
 }
+
 void FUniformInteger::Serialize(FArchive& Ar)
 {
 	Super::Serialize(Ar);
 	Ar << Magnitude;
 }
+
 bool FUniformInteger::operator==(const FFieldNodeBase& Node)
 {
 	if (Node.SerializationType() == SerializationType())
@@ -119,10 +135,28 @@ bool FUniformInteger::operator==(const FFieldNodeBase& Node)
 	return false;
 }
 
+void FUniformInteger::FillSetupCount(int32& NumOffsets, int32& NumParams) const
+{
+	Super::FillSetupCount(NumOffsets, NumParams);
+	NumParams += 1;
+}
+
+void FUniformInteger::FillSetupDatas(TArray<int32>& NodesOffsets, TArray<float>& NodesParams,
+	const TMap<FFieldNodeBase*, float>& CommandTimes, const float PreviousTime) const
+{
+	Super::FillSetupDatas(NodesOffsets, NodesParams, CommandTimes, PreviousTime);
+	NodesParams.Add(static_cast<float>(Magnitude));
+}
+
+float FUniformInteger::EvalMaxMagnitude() const
+{
+	return static_cast<float>(Magnitude);
+}
 
 /**
 * FRadialIntMask
 */
+
 void FRadialIntMask::Evaluate(FFieldContext& Context, TFieldArrayView<int32>& Results) const
 {
 	float Radius2 = Radius * Radius;
@@ -167,6 +201,7 @@ void FRadialIntMask::Evaluate(FFieldContext& Context, TFieldArrayView<int32>& Re
 		}
 	}
 }
+
 void FRadialIntMask::Serialize(FArchive& Ar)
 {
 	Super::Serialize(Ar);
@@ -176,6 +211,7 @@ void FRadialIntMask::Serialize(FArchive& Ar)
 	Ar << ExteriorValue;
 	SerializeInternal<ESetMaskConditionType>(Ar, SetMaskCondition);
 }
+
 bool FRadialIntMask::operator==(const FFieldNodeBase& Node)
 {
 	if (Node.SerializationType() == SerializationType())
@@ -191,9 +227,40 @@ bool FRadialIntMask::operator==(const FFieldNodeBase& Node)
 	return false;
 }
 
+void FRadialIntMask::FillSetupCount(int32& NumOffsets, int32& NumParams) const
+{
+	Super::FillSetupCount(NumOffsets, NumParams);
+	NumParams += 7;
+}
+
+void FRadialIntMask::FillSetupDatas(TArray<int32>& NodesOffsets, TArray<float>& NodesParams,
+	const TMap<FFieldNodeBase*, float>& CommandTimes, const float PreviousTime) const
+{
+	Super::FillSetupDatas(NodesOffsets, NodesParams, CommandTimes, PreviousTime);
+	NodesParams.Add(Radius);
+	NodesParams.Add(static_cast<float>(Position.X));
+	NodesParams.Add(static_cast<float>(Position.Y));
+	NodesParams.Add(static_cast<float>(Position.Z));
+	NodesParams.Add(static_cast<float>(InteriorValue));
+	NodesParams.Add(static_cast<float>(ExteriorValue));
+	NodesParams.Add(static_cast<float>(SetMaskCondition));
+}
+
+float FRadialIntMask::EvalMaxMagnitude() const
+{
+	return 1.0f;
+}
+
+void FRadialIntMask::ComputeFieldBounds(FVector& MinBounds, FVector& MaxBounds) const
+{
+	MinBounds = (ExteriorValue == 0) ? Position - FVector(Radius) : FVector(-FLT_MAX);
+	MaxBounds = (ExteriorValue == 0) ? Position + FVector(Radius) : FVector(FLT_MAX);
+}
+
 /**
 * FUniformScalar
 */
+
 void FUniformScalar::Evaluate(FFieldContext& Context, TFieldArrayView<float>& Results) const
 {
 	float MagnitudeVal = Magnitude;
@@ -204,11 +271,13 @@ void FUniformScalar::Evaluate(FFieldContext& Context, TFieldArrayView<float>& Re
 		Results[Context.SampleIndices[SampleIndex].Result] = MagnitudeVal;
 	}
 }
+
 void FUniformScalar::Serialize(FArchive& Ar)
 {
 	Super::Serialize(Ar);
 	Ar << Magnitude;
 }
+
 bool FUniformScalar::operator==(const FFieldNodeBase& Node)
 {
 	if (Node.SerializationType() == SerializationType())
@@ -219,6 +288,24 @@ bool FUniformScalar::operator==(const FFieldNodeBase& Node)
 	return false;
 }
 
+void FUniformScalar::FillSetupCount(int32& NumOffsets, int32& NumParams) const
+{
+	Super::FillSetupCount(NumOffsets, NumParams);
+	NumParams += 1;
+}
+
+void FUniformScalar::FillSetupDatas(TArray<int32>& NodesOffsets, TArray<float>& NodesParams,
+	const TMap<FFieldNodeBase*, float>& CommandTimes, const float PreviousTime) const
+{
+	Super::FillSetupDatas(NodesOffsets, NodesParams, CommandTimes, PreviousTime);
+	NodesParams.Add(Magnitude);
+}
+
+float FUniformScalar::EvalMaxMagnitude() const
+{
+	return Magnitude;
+}
+
 /**
 * FWaveScalar
 */
@@ -298,6 +385,7 @@ void FWaveScalar::Serialize(FArchive& Ar)
 	SerializeInternal<EWaveFunctionType>(Ar, Function);
 	SerializeInternal<EFieldFalloffType>(Ar, Falloff);
 }
+
 bool FWaveScalar::operator==(const FFieldNodeBase& Node)
 {
 	if (Node.SerializationType() == SerializationType())
@@ -313,6 +401,33 @@ bool FWaveScalar::operator==(const FFieldNodeBase& Node)
 	return false;
 }
 
+void FWaveScalar::FillSetupCount(int32& NumOffsets, int32& NumParams) const
+{
+	Super::FillSetupCount(NumOffsets, NumParams);
+	NumParams += 9;
+}
+
+void FWaveScalar::FillSetupDatas(TArray<int32>& NodesOffsets, TArray<float>& NodesParams,
+	const TMap<FFieldNodeBase*, float>& CommandTimes, const float PreviousTime) const
+{
+	const float NextTime = CommandTimes.Find(this) ? CommandTimes[this] : PreviousTime;
+
+	Super::FillSetupDatas(NodesOffsets, NodesParams, CommandTimes, PreviousTime);
+	NodesParams.Add(Magnitude);
+	NodesParams.Add(static_cast<float>(Position.X));
+	NodesParams.Add(static_cast<float>(Position.Y));
+	NodesParams.Add(static_cast<float>(Position.Z));
+	NodesParams.Add(Wavelength);
+	NodesParams.Add(Period);
+	NodesParams.Add(NextTime);
+	NodesParams.Add(static_cast<float>(Function));
+	NodesParams.Add(static_cast<float>(Falloff));
+}
+
+float FWaveScalar::EvalMaxMagnitude() const
+{
+	return Magnitude;
+}
 
 /**
 * Function Utils
@@ -409,6 +524,7 @@ void FRadialFalloff::Evaluate(FFieldContext& Context, TFieldArrayView<float>& Re
 		break;
 	}
 }
+
 void FRadialFalloff::Serialize(FArchive& Ar)
 {
 	Super::Serialize(Ar);
@@ -421,6 +537,7 @@ void FRadialFalloff::Serialize(FArchive& Ar)
 	SerializeInternal<EFieldFalloffType>(Ar, Falloff);
 
 }
+
 bool FRadialFalloff::operator==(const FFieldNodeBase& Node)
 {
 	if (Node.SerializationType() == SerializationType())
@@ -438,6 +555,38 @@ bool FRadialFalloff::operator==(const FFieldNodeBase& Node)
 	return false;
 }
 
+void FRadialFalloff::FillSetupCount(int32& NumOffsets, int32& NumParams) const
+{
+	Super::FillSetupCount(NumOffsets, NumParams);
+	NumParams += 9;
+}
+
+void FRadialFalloff::FillSetupDatas(TArray<int32>& NodesOffsets, TArray<float>& NodesParams,
+	const TMap<FFieldNodeBase*, float>& CommandTimes, const float PreviousTime) const
+{
+	Super::FillSetupDatas(NodesOffsets, NodesParams, CommandTimes, PreviousTime);
+	NodesParams.Add(Magnitude);
+	NodesParams.Add(MinRange);
+	NodesParams.Add(MaxRange);
+	NodesParams.Add(Default);
+	NodesParams.Add(Radius);
+	NodesParams.Add(static_cast<float>(Position.X));
+	NodesParams.Add(static_cast<float>(Position.Y));
+	NodesParams.Add(static_cast<float>(Position.Z));
+	NodesParams.Add(static_cast<float>(Falloff));
+}
+
+float FRadialFalloff::EvalMaxMagnitude() const
+{
+	return Magnitude;
+}
+
+void FRadialFalloff::ComputeFieldBounds(FVector& MinBounds, FVector& MaxBounds) const
+{
+	MinBounds = (Default == 0) ? Position - FVector(Radius) : FVector(-FLT_MAX);
+	MaxBounds = (Default == 0) ? Position + FVector(Radius) : FVector(FLT_MAX);
+}
+
 /**
 * FPlaneFalloff
 */
@@ -520,6 +669,41 @@ bool FPlaneFalloff::operator==(const FFieldNodeBase& Node)
 	return false;
 }
 
+void FPlaneFalloff::FillSetupCount(int32& NumOffsets, int32& NumParams) const
+{
+	Super::FillSetupCount(NumOffsets, NumParams);
+	NumParams += 12;
+}
+
+void FPlaneFalloff::FillSetupDatas(TArray<int32>& NodesOffsets, TArray<float>& NodesParams,
+	const TMap<FFieldNodeBase*, float>& CommandTimes, const float PreviousTime) const
+{
+	Super::FillSetupDatas(NodesOffsets, NodesParams, CommandTimes, PreviousTime);
+	NodesParams.Add(Magnitude);
+	NodesParams.Add(MinRange);
+	NodesParams.Add(MaxRange);
+	NodesParams.Add(Default);
+	NodesParams.Add(Distance);
+	NodesParams.Add(static_cast<float>(Position.X));
+	NodesParams.Add(static_cast<float>(Position.Y));
+	NodesParams.Add(static_cast<float>(Position.Z));
+	NodesParams.Add(static_cast<float>(Normal.X));
+	NodesParams.Add(static_cast<float>(Normal.Y));
+	NodesParams.Add(static_cast<float>(Normal.Z));
+	NodesParams.Add(static_cast<float>(Falloff));
+}
+
+float FPlaneFalloff::EvalMaxMagnitude() const
+{
+	return Magnitude;
+}
+
+void FPlaneFalloff::ComputeFieldBounds(FVector& MinBounds, FVector& MaxBounds) const
+{
+	MinBounds = FVector(-FLT_MAX);
+	MaxBounds = FVector(FLT_MAX);
+}
+
 /**
 * FBoxFalloff
 */
@@ -551,8 +735,7 @@ void FBoxFalloff::Evaluator(const FFieldContext& Context, TFieldArrayView<float>
 	}
 }
 
-void
-FBoxFalloff::Evaluate(FFieldContext& Context, TFieldArrayView<float>& Results) const
+void FBoxFalloff::Evaluate(FFieldContext& Context, TFieldArrayView<float>& Results) const
 {
 	switch (Falloff)
 	{
@@ -573,6 +756,7 @@ FBoxFalloff::Evaluate(FFieldContext& Context, TFieldArrayView<float>& Results) c
 		break;
 	}
 }
+
 void FBoxFalloff::Serialize(FArchive& Ar)
 {
 	Super::Serialize(Ar);
@@ -583,6 +767,7 @@ void FBoxFalloff::Serialize(FArchive& Ar)
 	Ar << Transform;
 	SerializeInternal<EFieldFalloffType>(Ar, Falloff);
 }
+
 bool FBoxFalloff::operator==(const FFieldNodeBase& Node)
 {
 	if (Node.SerializationType() == SerializationType())
@@ -600,12 +785,56 @@ bool FBoxFalloff::operator==(const FFieldNodeBase& Node)
 	return false;
 }
 
+void FBoxFalloff::FillSetupCount(int32& NumOffsets, int32& NumParams) const
+{
+	Super::FillSetupCount(NumOffsets, NumParams);
+	NumParams += 15;
+}
+
+void FBoxFalloff::FillSetupDatas(TArray<int32>& NodesOffsets, TArray<float>& NodesParams,
+	const TMap<FFieldNodeBase*, float>& CommandTimes, const float PreviousTime) const
+{
+	Super::FillSetupDatas(NodesOffsets, NodesParams, CommandTimes, PreviousTime);
+	NodesParams.Add(Magnitude);
+	NodesParams.Add(MinRange);
+	NodesParams.Add(MaxRange);
+	NodesParams.Add(Default);
+	NodesParams.Add(static_cast<float>(Transform.GetRotation().X));
+	NodesParams.Add(static_cast<float>(Transform.GetRotation().Y));
+	NodesParams.Add(static_cast<float>(Transform.GetRotation().Z));
+	NodesParams.Add(static_cast<float>(Transform.GetRotation().W));
+	NodesParams.Add(static_cast<float>(Transform.GetTranslation().X));
+	NodesParams.Add(static_cast<float>(Transform.GetTranslation().Y));
+	NodesParams.Add(static_cast<float>(Transform.GetTranslation().Z));
+	NodesParams.Add(static_cast<float>(Transform.GetScale3D().X));
+	NodesParams.Add(static_cast<float>(Transform.GetScale3D().Y));
+	NodesParams.Add(static_cast<float>(Transform.GetScale3D().Z));
+	NodesParams.Add(static_cast<float>(Falloff));
+}
+
+float FBoxFalloff::EvalMaxMagnitude() const
+{
+	return Magnitude;
+}
+
+void FBoxFalloff::ComputeFieldBounds(FVector& MinBounds, FVector& MaxBounds) const
+{
+	MinBounds = FVector(-FLT_MAX);
+	MaxBounds = FVector(FLT_MAX);
+	if (Default == 0)
+	{
+		const FBox UnitBox(FVector(-50), FVector(50));
+		const FBox BoundingBox = UnitBox.TransformBy(Transform);
+		MinBounds = BoundingBox.Min;
+		MaxBounds = BoundingBox.Max;
+	}
+}
 
 /**
 * FNoiseField
 */
-void
-FNoiseField::Evaluate(FFieldContext& Context, TFieldArrayView<float>& Results) const
+
+void FNoiseField::Evaluate(FFieldContext& Context, TFieldArrayView<float>& Results) const
 {
 	const float DeltaRange = (MaxRange - MinRange);
 	const int32 NumSamples = Context.SampleIndices.Num();
@@ -629,6 +858,7 @@ FNoiseField::Evaluate(FFieldContext& Context, TFieldArrayView<float>& Results) c
 		Results[Index.Result] = ScaleFunctionResult(MinRange, DeltaRange, 1.0f, PerlinValue);
 	}
 }
+
 void FNoiseField::Serialize(FArchive& Ar)
 {
 	Super::Serialize(Ar);
@@ -636,6 +866,7 @@ void FNoiseField::Serialize(FArchive& Ar)
 	Ar << MaxRange;
 	Ar << Transform;
 }
+
 bool FNoiseField::operator==(const FFieldNodeBase& Node)
 {
 	if (Node.SerializationType() == SerializationType())
@@ -649,10 +880,39 @@ bool FNoiseField::operator==(const FFieldNodeBase& Node)
 	return false;
 }
 
+void FNoiseField::FillSetupCount(int32& NumOffsets, int32& NumParams) const
+{
+	Super::FillSetupCount(NumOffsets, NumParams);
+	NumParams += 12;
+}
+
+void FNoiseField::FillSetupDatas(TArray<int32>& NodesOffsets, TArray<float>& NodesParams,
+	const TMap<FFieldNodeBase*, float>& CommandTimes, const float PreviousTime) const
+{
+	Super::FillSetupDatas(NodesOffsets, NodesParams, CommandTimes, PreviousTime);
+	NodesParams.Add(MinRange);
+	NodesParams.Add(MaxRange);
+	NodesParams.Add(static_cast<float>(Transform.GetRotation().X));
+	NodesParams.Add(static_cast<float>(Transform.GetRotation().Y));
+	NodesParams.Add(static_cast<float>(Transform.GetRotation().Z));
+	NodesParams.Add(static_cast<float>(Transform.GetRotation().W));
+	NodesParams.Add(static_cast<float>(Transform.GetTranslation().X));
+	NodesParams.Add(static_cast<float>(Transform.GetTranslation().Y));
+	NodesParams.Add(static_cast<float>(Transform.GetTranslation().Z));
+	NodesParams.Add(static_cast<float>(Transform.GetScale3D().X));
+	NodesParams.Add(static_cast<float>(Transform.GetScale3D().Y));
+	NodesParams.Add(static_cast<float>(Transform.GetScale3D().Z));
+}
+
+float FNoiseField::EvalMaxMagnitude() const
+{
+	return MaxRange;
+}
 
 /**
 * FUniformVector
 */
+
 void FUniformVector::Evaluate(FFieldContext& Context, TFieldArrayView<FVector>& Results) const
 {
 	const FVector DirectionVal = Direction;
@@ -664,12 +924,14 @@ void FUniformVector::Evaluate(FFieldContext& Context, TFieldArrayView<FVector>&
 		Results[Context.SampleIndices[SampleIndex].Result] = MagnitudeVal * DirectionVal;
 	}
 }
+
 void FUniformVector::Serialize(FArchive& Ar)
 {
 	Super::Serialize(Ar);
 	Ar << Magnitude;
 	Ar << Direction;
 }
+
 bool FUniformVector::operator==(const FFieldNodeBase& Node)
 {
 	if (Node.SerializationType() == SerializationType())
@@ -682,9 +944,31 @@ bool FUniformVector::operator==(const FFieldNodeBase& Node)
 	return false;
 }
 
+void FUniformVector::FillSetupCount(int32& NumOffsets, int32& NumParams) const
+{
+	Super::FillSetupCount(NumOffsets, NumParams);
+	NumParams += 4;
+}
+
+void FUniformVector::FillSetupDatas(TArray<int32>& NodesOffsets, TArray<float>& NodesParams,
+	const TMap<FFieldNodeBase*, float>& CommandTimes, const float PreviousTime) const 
+{
+	Super::FillSetupDatas(NodesOffsets, NodesParams, CommandTimes, PreviousTime);
+	NodesParams.Add(Magnitude);
+	NodesParams.Add(static_cast<float>(Direction.X));
+	NodesParams.Add(static_cast<float>(Direction.Y));
+	NodesParams.Add(static_cast<float>(Direction.Z));
+}
+
+float FUniformVector::EvalMaxMagnitude() const
+{
+	return Magnitude;
+}
+
 /**
 * FRadialVector
 */
+
 void FRadialVector::Evaluate(FFieldContext& Context, TFieldArrayView<FVector>& Results) const
 {
 	const int32 NumSamples = Context.SampleIndices.Num();
@@ -697,12 +981,14 @@ void FRadialVector::Evaluate(FFieldContext& Context, TFieldArrayView<FVector>& R
 	}
 	
 }
+
 void FRadialVector::Serialize(FArchive& Ar)
 {
 	Super::Serialize(Ar);
 	Ar << Magnitude;
 	Ar << Position;
 }
+
 bool FRadialVector::operator==(const FFieldNodeBase& Node)
 {
 	if (Node.SerializationType() == SerializationType())
@@ -715,9 +1001,32 @@ bool FRadialVector::operator==(const FFieldNodeBase& Node)
 	return false;
 }
 
+void FRadialVector::FillSetupCount(int32& NumOffsets, int32& NumParams) const
+{
+	Super::FillSetupCount(NumOffsets, NumParams);
+	NumParams += 4;
+}
+
+void FRadialVector::FillSetupDatas(TArray<int32>& NodesOffsets, TArray<float>& NodesParams,
+	const TMap<FFieldNodeBase*, float>& CommandTimes, const float PreviousTime) const
+{
+	Super::FillSetupDatas(NodesOffsets, NodesParams, CommandTimes, PreviousTime);
+	NodesParams.Add(Magnitude);
+	NodesParams.Add(static_cast<float>(Position.X));
+	NodesParams.Add(static_cast<float>(Position.Y));
+	NodesParams.Add(static_cast<float>(Position.Z));
+}
+
+float FRadialVector::EvalMaxMagnitude() const
+{
+	return Magnitude;
+}
+
+
 /**
 * FRandomVector
 */
+
 void FRandomVector::Evaluate(FFieldContext& Context, TFieldArrayView<FVector>& Results) const
 {
 	const int32 NumSamples = Context.SampleIndices.Num();
@@ -733,11 +1042,13 @@ void FRandomVector::Evaluate(FFieldContext& Context, TFieldArrayView<FVector>& R
 		}
 	}
 }
+
 void FRandomVector::Serialize(FArchive& Ar)
 {
 	Super::Serialize(Ar);
 	Ar << Magnitude;
 }
+
 bool FRandomVector::operator==(const FFieldNodeBase& Node)
 {
 	if (Node.SerializationType() == SerializationType())
@@ -748,11 +1059,30 @@ bool FRandomVector::operator==(const FFieldNodeBase& Node)
 	return false;
 }
 
+void FRandomVector::FillSetupCount(int32& NumOffsets, int32& NumParams) const 
+{
+	Super::FillSetupCount(NumOffsets, NumParams);
+	NumParams += 1;
+}
+
+void FRandomVector::FillSetupDatas(TArray<int32>& NodesOffsets, TArray<float>& NodesParams,
+	const TMap<FFieldNodeBase*, float>& CommandTimes, const float PreviousTime) const 
+{
+	Super::FillSetupDatas(NodesOffsets, NodesParams, CommandTimes, PreviousTime);
+	NodesParams.Add(Magnitude);
+}
+
+float FRandomVector::EvalMaxMagnitude() const 
+{
+	return Magnitude;
+}
+
 
 /**
 * FSumScalar 
 * 
 */
+
 void FSumScalar::Evaluate(FFieldContext& ContextIn, TFieldArrayView<float>& Results) const
 {
 	TUniquePtr<FFieldSystemMetaDataResults<float> > ResultsData(new FFieldSystemMetaDataResults<float>(Results));
@@ -850,6 +1180,7 @@ void FSumScalar::Evaluate(FFieldContext& ContextIn, TFieldArrayView<float>& Resu
 		}
 	}
 }
+
 void FSumScalar::Serialize(FArchive& Ar)
 {
 	Super::Serialize(Ar);
@@ -858,6 +1189,7 @@ void FSumScalar::Serialize(FArchive& Ar)
 	SerializeInternal<float>(Ar, ScalarLeft);
 	SerializeInternal<EFieldOperationType>(Ar, Operation);
 }
+
 bool FSumScalar::operator==(const FFieldNodeBase& Node)
 {
 	if (Node.SerializationType() == SerializationType())
@@ -873,6 +1205,92 @@ bool FSumScalar::operator==(const FFieldNodeBase& Node)
 	return false;
 }
 
+void FSumScalar::FillSetupCount(int32& NumOffsets, int32& NumParams) const
+{
+	if (ScalarRight.IsValid())
+	{
+		ScalarRight->FillSetupCount(NumOffsets, NumParams);
+	}
+	if (ScalarLeft.IsValid())
+	{
+		ScalarLeft->FillSetupCount(NumOffsets, NumParams);
+	}
+	Super::FillSetupCount(NumOffsets, NumParams);
+	NumParams += 4;
+}
+
+void FSumScalar::FillSetupDatas(TArray<int32>& NodesOffsets, TArray<float>& NodesParams,
+	const TMap<FFieldNodeBase*, float>& CommandTimes, const float PreviousTime) const
+{
+	const float NextTime = CommandTimes.Find(this) ? CommandTimes[this] : PreviousTime;
+
+	if (ScalarRight.IsValid())
+	{
+		ScalarRight->FillSetupDatas(NodesOffsets, NodesParams, CommandTimes, NextTime);
+	}
+	if (ScalarLeft.IsValid())
+	{
+		ScalarLeft->FillSetupDatas(NodesOffsets, NodesParams, CommandTimes, NextTime);
+	}
+	Super::FillSetupDatas(NodesOffsets, NodesParams, CommandTimes, PreviousTime);
+	NodesParams.Add(Magnitude);
+	NodesParams.Add(static_cast<float>(ScalarRight != nullptr));
+	NodesParams.Add(static_cast<float>(ScalarLeft != nullptr));
+	NodesParams.Add(static_cast<float>(Operation));
+}
+
+float FSumScalar::EvalMaxMagnitude() const
+{
+	float MaxMagnitudeA = 0.0, MaxMagnitudeB = 0.0, MaxMagnitude = 0.0;
+	if (ScalarRight.IsValid())
+	{
+		MaxMagnitudeA = ScalarRight->EvalMaxMagnitude();
+	}
+	if (ScalarLeft.IsValid())
+	{
+		MaxMagnitudeB = ScalarLeft->EvalMaxMagnitude();
+	}
+	if (Operation == EFieldOperationType::Field_Multiply ||
+		Operation == EFieldOperationType::Field_Divide)
+	{
+		MaxMagnitude = (Operation == EFieldOperationType::Field_Multiply) ? MaxMagnitudeA * MaxMagnitudeB :
+			(FMath::Abs(MaxMagnitudeA) > FLT_EPSILON) ? MaxMagnitudeB / MaxMagnitudeA : 0.0f;
+	}
+	else if (Operation == EFieldOperationType::Field_Add ||
+		Operation == EFieldOperationType::Field_Substract)
+	{
+		MaxMagnitude = FMath::Max(MaxMagnitudeA, MaxMagnitudeB);
+	}
+
+	return Magnitude * MaxMagnitude;
+}
+
+void FSumScalar::ComputeFieldBounds(FVector& MinBounds, FVector& MaxBounds) const
+{
+	FVector MinBoundsA(-FLT_MAX), MaxBoundsA(FLT_MAX), MinBoundsB(-FLT_MAX), MaxBoundsB(FLT_MAX);
+	if (ScalarRight.IsValid())
+	{
+		ScalarRight->ComputeFieldBounds(MinBoundsA, MaxBoundsA);
+	}
+	if (ScalarLeft.IsValid())
+	{
+		ScalarLeft->ComputeFieldBounds(MinBoundsB, MaxBoundsB);
+	}
+
+	if (Operation == EFieldOperationType::Field_Multiply ||
+		Operation == EFieldOperationType::Field_Divide)
+	{
+		MinBounds = MaxVector(MinBoundsA, MinBoundsB);
+		MaxBounds = MinVector(MaxBoundsA, MaxBoundsB);
+	}
+	else if (Operation == EFieldOperationType::Field_Add ||
+		Operation == EFieldOperationType::Field_Substract)
+	{
+		MinBounds = MinVector(MinBoundsA, MinBoundsB);
+		MaxBounds = MaxVector(MaxBoundsA, MaxBoundsB);
+	}
+}
+
 /**
 * FSumVector
 */
@@ -1007,10 +1425,116 @@ bool FSumVector::operator==(const FFieldNodeBase& Node)
 	return false;
 }
 
+void FSumVector::FillSetupCount(int32& NumOffsets, int32& NumParams) const
+{
+	if (Scalar.IsValid())
+	{
+		Scalar->FillSetupCount(NumOffsets, NumParams);
+	}
+	if (VectorRight.IsValid())
+	{
+		VectorRight->FillSetupCount(NumOffsets, NumParams);
+	}
+	if (VectorLeft.IsValid())
+	{
+		VectorLeft->FillSetupCount(NumOffsets, NumParams);
+	}
+	Super::FillSetupCount(NumOffsets, NumParams);
+	NumParams += 5;
+}
+
+void FSumVector::FillSetupDatas(TArray<int32>& NodesOffsets, TArray<float>& NodesParams,
+	const TMap<FFieldNodeBase*, float>& CommandTimes, const float PreviousTime) const
+{
+	const float NextTime = CommandTimes.Find(this) ? CommandTimes[this] : PreviousTime;
+
+	if (Scalar.IsValid())
+	{
+		Scalar->FillSetupDatas(NodesOffsets, NodesParams, CommandTimes, NextTime);
+	}
+	if (VectorRight.IsValid())
+	{
+		VectorRight->FillSetupDatas(NodesOffsets, NodesParams, CommandTimes, NextTime);
+	}
+	if (VectorLeft.IsValid())
+	{
+		VectorLeft->FillSetupDatas(NodesOffsets, NodesParams, CommandTimes, NextTime);
+	}
+	Super::FillSetupDatas(NodesOffsets, NodesParams, CommandTimes, PreviousTime);
+	NodesParams.Add(Magnitude);
+	NodesParams.Add(static_cast<float>(Scalar != nullptr));
+	NodesParams.Add(static_cast<float>(VectorRight != nullptr));
+	NodesParams.Add(static_cast<float>(VectorLeft != nullptr));
+	NodesParams.Add(static_cast<float>(Operation));
+}
+
+float FSumVector::EvalMaxMagnitude() const
+{
+	float MaxMagnitudeA = 0.0, MaxMagnitudeB = 0.0, MaxMagnitudeC = 0.0, MaxMagnitude = 0.0;
+	if (Scalar.IsValid())
+	{
+		MaxMagnitudeA = Scalar->EvalMaxMagnitude();
+	}
+	if (VectorRight.IsValid())
+	{
+		MaxMagnitudeB = VectorRight->EvalMaxMagnitude();
+	}
+	if (VectorLeft.IsValid())
+	{
+		MaxMagnitudeC = VectorLeft->EvalMaxMagnitude();
+	}
+	if (Operation == EFieldOperationType::Field_Multiply ||
+		Operation == EFieldOperationType::Field_Divide)
+	{
+		MaxMagnitude = (Operation == EFieldOperationType::Field_Multiply) ? MaxMagnitudeA * MaxMagnitudeB :
+			(FMath::Abs(MaxMagnitudeA) > FLT_EPSILON) ? MaxMagnitudeB / MaxMagnitudeA : 0.0f;
+	}
+	else if (Operation == EFieldOperationType::Field_Add ||
+		Operation == EFieldOperationType::Field_Substract)
+	{
+		MaxMagnitude = FMath::Max(MaxMagnitudeA, MaxMagnitudeB);
+	}
+
+	return Magnitude * MaxMagnitudeA * MaxMagnitude;
+}
+
+void FSumVector::ComputeFieldBounds(FVector& MinBounds, FVector& MaxBounds) const
+{
+	FVector MinBoundsA(-FLT_MAX), MaxBoundsA(FLT_MAX), MinBoundsB(-FLT_MAX), MaxBoundsB(FLT_MAX), MinBoundsC(-FLT_MAX), MaxBoundsC(FLT_MAX);
+	if (Scalar.IsValid())
+	{
+		Scalar->ComputeFieldBounds(MinBoundsC, MaxBoundsC);
+	}
+	if (VectorRight.IsValid())
+	{
+		VectorRight->ComputeFieldBounds(MinBoundsA, MaxBoundsA);
+	}
+	if (VectorLeft.IsValid())
+	{
+		VectorLeft->ComputeFieldBounds(MinBoundsB, MaxBoundsB);
+	}
+
+	if (Operation == EFieldOperationType::Field_Multiply ||
+		Operation == EFieldOperationType::Field_Divide)
+	{
+		MinBounds = MaxVector(MinBoundsA, MinBoundsB);
+		MaxBounds = MinVector(MaxBoundsA, MaxBoundsB);
+	}
+	else if (Operation == EFieldOperationType::Field_Add ||
+		Operation == EFieldOperationType::Field_Substract)
+	{
+		MinBounds = MinVector(MinBoundsA, MinBoundsB);
+		MaxBounds = MaxVector(MaxBoundsA, MaxBoundsB);
+	}
+	MinBounds = MaxVector(MinBounds, MinBoundsC);
+	MaxBounds = MinVector(MaxBounds, MaxBoundsC);
+}
+
 
 /**
 * FConversionField<InT,OutT>
 */
+
 template<class InT, class OutT>
 void FConversionField<InT,OutT>::Evaluate(FFieldContext& Context, TFieldArrayView<OutT>& Results) const
 {
@@ -1036,12 +1560,14 @@ void FConversionField<InT,OutT>::Evaluate(FFieldContext& Context, TFieldArrayVie
 	}
 	ResultsArray.SetNum(BufferOffset, false);
 }
+
 template<class InT, class OutT>
 void FConversionField<InT, OutT>::Serialize(FArchive& Ar)
 {
 	Super::Serialize(Ar);
 	SerializeInternal<InT>(Ar, InputField);
 }
+
 template<class InT, class OutT>
 bool FConversionField<InT, OutT>::operator==(const FFieldNodeBase& Node)
 {
@@ -1054,12 +1580,62 @@ bool FConversionField<InT, OutT>::operator==(const FFieldNodeBase& Node)
 	return false;
 }
 
+template<class InT, class OutT>
+void FConversionField<InT, OutT>::FillSetupCount(int32& NumOffsets, int32& NumParams) const
+{
+	if (InputField.IsValid())
+	{
+		InputField->FillSetupCount(NumOffsets, NumParams);
+	}
+	Super::FillSetupCount(NumOffsets, NumParams);
+	NumParams += 1;
+}
+
+template<class InT, class OutT>
+void FConversionField<InT, OutT>::FillSetupDatas(TArray<int32>& NodesOffsets, TArray<float>& NodesParams,
+	const TMap<FFieldNodeBase*, float>& CommandTimes, const float PreviousTime) const
+{
+	const float NextTime = CommandTimes.Find(this) ? CommandTimes[this] : PreviousTime;
+
+	if (InputField.IsValid())
+	{
+		InputField->FillSetupDatas(NodesOffsets, NodesParams, CommandTimes, NextTime);
+	}
+
+	Super::FillSetupDatas(NodesOffsets, NodesParams, CommandTimes, PreviousTime);
+	NodesParams.Add(static_cast<float>(InputField != nullptr));
+}
+
+template<class InT, class OutT>
+float FConversionField<InT, OutT>::EvalMaxMagnitude() const
+{
+	float MaxMagnitude = 0.0;
+	if (InputField.IsValid())
+	{
+		MaxMagnitude = InputField->EvalMaxMagnitude();
+	}
+	return MaxMagnitude;
+}
+
+template<class InT, class OutT>
+void FConversionField<InT, OutT>::ComputeFieldBounds(FVector& MinBounds, FVector& MaxBounds) const
+{
+	FVector MinBoundsA(-FLT_MAX), MaxBoundsA(FLT_MAX);
+	if (InputField.IsValid())
+	{
+		InputField->ComputeFieldBounds(MinBoundsA, MaxBoundsA);
+	}
+	MinBounds = MinBoundsA;
+	MaxBounds = MaxBoundsA;
+}
+
 template class FConversionField<int32, float>;
 template class FConversionField<float, int32>;
 
 /**
 *  FCullingField<T>
 */
+
 template<class T>
 void FCullingField<T>::Evaluate(FFieldContext& Context, TFieldArrayView<T>& Results) const
 {
@@ -1137,6 +1713,7 @@ void FCullingField<T>::Evaluate(FFieldContext& Context, TFieldArrayView<T>& Resu
 		InputField->Evaluate(Context, Results);
 	}
 }
+
 template<class T>
 void FCullingField<T>::Serialize(FArchive& Ar)
 {
@@ -1145,6 +1722,7 @@ void FCullingField<T>::Serialize(FArchive& Ar)
 	SerializeInternal<T>(Ar, Input);
 	SerializeInternal<EFieldCullingOperationType>(Ar, Operation);
 }
+
 template<class T>
 bool FCullingField<T>::operator==(const FFieldNodeBase& Node)
 {
@@ -1158,6 +1736,79 @@ bool FCullingField<T>::operator==(const FFieldNodeBase& Node)
 	}
 	return false;
 }
+
+template<class T>
+void FCullingField<T>::FillSetupCount(int32& NumOffsets, int32& NumParams) const
+{
+	if (Culling.IsValid())
+	{
+		Culling->FillSetupCount(NumOffsets, NumParams);
+	}
+	if (Input.IsValid())
+	{
+		Input->FillSetupCount(NumOffsets, NumParams);
+	}
+	Super::FillSetupCount(NumOffsets, NumParams);
+	NumParams += 3;
+}
+
+template<class T>
+void FCullingField<T>::FillSetupDatas(TArray<int32>& NodesOffsets, TArray<float>& NodesParams,
+	const TMap<FFieldNodeBase*, float>& CommandTimes, const float PreviousTime) const
+{
+	const float NextTime = CommandTimes.Find(this) ? CommandTimes[this] : PreviousTime;
+
+	if (Culling.IsValid())
+	{
+		Culling->FillSetupDatas(NodesOffsets, NodesParams, CommandTimes, NextTime);
+	}
+	if (Input.IsValid())
+	{
+		Input->FillSetupDatas(NodesOffsets, NodesParams, CommandTimes, NextTime);
+	}
+
+	Super::FillSetupDatas(NodesOffsets, NodesParams, CommandTimes, PreviousTime);
+	NodesParams.Add(static_cast<float>(Culling != nullptr));
+	NodesParams.Add(static_cast<float>(Input != nullptr));
+	NodesParams.Add(static_cast<float>(Operation));
+}
+
+template<class T>
+float FCullingField<T>::EvalMaxMagnitude() const
+{
+	float MaxMagnitude = 0.0;
+	if (Input.IsValid())
+	{
+		MaxMagnitude = Input->EvalMaxMagnitude();
+	}
+	return  MaxMagnitude;
+}
+
+template<class T>
+void FCullingField<T>::ComputeFieldBounds(FVector& MinBounds, FVector& MaxBounds) const
+{
+	FVector MinBoundsA(-FLT_MAX), MaxBoundsA(FLT_MAX), MinBoundsB(-FLT_MAX), MaxBoundsB(FLT_MAX);
+
+	if (Culling.IsValid())
+	{
+		Culling->ComputeFieldBounds(MinBoundsA, MaxBoundsA);
+	}
+	if (Input.IsValid())
+	{
+		Input->ComputeFieldBounds(MinBoundsB, MaxBoundsB);
+	}
+	if (Operation == EFieldCullingOperationType::Field_Culling_Inside)
+	{
+		MinBounds = MinVector(MinBoundsA, MinBoundsB);
+		MaxBounds = MaxVector(MaxBoundsA, MaxBoundsB);
+	}
+	else if (Operation == EFieldCullingOperationType::Field_Culling_Outside)
+	{
+		MinBounds = MaxVector(MinBoundsA, MinBoundsB);
+		MaxBounds = MinVector(MaxBoundsA, MaxBoundsB);
+	}
+}
+
 template class CHAOS_API FCullingField<int32>;
 template class CHAOS_API FCullingField<float>;
 template class CHAOS_API FCullingField<FVector>;
diff --git a/Engine/Source/Runtime/Experimental/Chaos/Public/Field/FieldSystem.h b/Engine/Source/Runtime/Experimental/Chaos/Public/Field/FieldSystem.h
--- a/Engine/Source/Runtime/Experimental/Chaos/Public/Field/FieldSystem.h
+++ b/Engine/Source/Runtime/Experimental/Chaos/Public/Field/FieldSystem.h
@@ -446,6 +446,23 @@ public:
 	virtual void Serialize(FArchive& Ar) { Ar << Name; }
 	virtual bool operator==(const FFieldNodeBase& Node) { return Name.IsEqual(Node.GetName()); }
 
+	/** Count the number of offsets/params that will be used by the world physics field */
+	virtual void FillSetupCount(int32& NumOffsets, int32& NumParams) const {}
+
+	/** Fill the offsets/params arrays that will be used by the world physics field */
+	virtual void FillSetupDatas(TArray<int32>& NodesOffsets, TArray<float>& NodesParams,
+		const TMap<FFieldNodeBase*, float>& CommandTimes, const float PreviousTime) const {}
+
+	/** Evaluate the maximum magnitude of the field graph  */
+	virtual float EvalMaxMagnitude() const { return 1.0; }
+
+	/** Compute the min/max spatial bounds of the field */
+	virtual void ComputeFieldBounds(FVector& MinBounds, FVector& MaxBounds) const
+	{
+		MinBounds = FVector(-FLT_MAX);
+		MaxBounds = FVector(FLT_MAX);
+	}
+
 	FName GetName() const { return Name; }
 	void  SetName(const FName & NameIn) { Name = NameIn; }
 
@@ -473,6 +490,22 @@ public:
 
 	static EFieldType StaticType();
 	virtual EFieldType Type() const { return StaticType(); }
+
+	/** Count the number of offsets/params that will be used by the world physics field */
+	virtual void FillSetupCount(int32& NumOffsets, int32& NumParams) const override
+	{
+		++NumOffsets;
+		NumParams += 2;
+	}
+
+	/** Fill the offsets/params arrays that will be used by the world physics field */
+	virtual void FillSetupDatas(TArray<int32>& NodesOffsets, TArray<float>& NodesParams,
+		const TMap<FFieldNodeBase*, float>& CommandTimes, const float PreviousTime) const override
+	{
+		NodesOffsets.Add(NodesParams.Num());
+		NodesParams.Add(static_cast<float>(Type()));
+		NodesParams.Add(static_cast<float>(SerializationType()));
+	}
 };
 
 template<> inline FFieldNodeBase::EFieldType FFieldNode<int32>::StaticType() { return EFieldType::EField_Int32; }
diff --git a/Engine/Source/Runtime/Experimental/Chaos/Public/Field/FieldSystemNodes.h b/Engine/Source/Runtime/Experimental/Chaos/Public/Field/FieldSystemNodes.h
--- a/Engine/Source/Runtime/Experimental/Chaos/Public/Field/FieldSystemNodes.h
+++ b/Engine/Source/Runtime/Experimental/Chaos/Public/Field/FieldSystemNodes.h
@@ -23,13 +23,22 @@ public:
 	virtual FFieldNodeBase * NewCopy() const override { return new FUniformInteger(Magnitude); }
 	virtual ~FUniformInteger() {}
 
-	void Evaluate(FFieldContext&, TFieldArrayView<int32>& Results) const override;
-	virtual bool operator==(const FFieldNodeBase& Node);
+	virtual void Evaluate(FFieldContext&, TFieldArrayView<int32>& Results) const override;
+	virtual bool operator==(const FFieldNodeBase& Node) override;
 
 	/** Serialization API */
-	virtual FFieldNodeBase::ESerializationType SerializationType() const { return FFieldNodeBase::ESerializationType::FieldNode_FUniformInteger; }
+	virtual FFieldNodeBase::ESerializationType SerializationType() const override { return FFieldNodeBase::ESerializationType::FieldNode_FUniformInteger; }
 	virtual void Serialize(FArchive& Ar) override;
 
+	/** Count the number of offsets/params that will be used by the world physics field */
+	virtual void FillSetupCount(int32& NumOffsets, int32& NumParams) const override;
+
+	/** Fill the offsets/params arrays that will be used by the world physics field */
+	virtual void FillSetupDatas(TArray<int32>& NodesOffsets, TArray<float>& NodesParams,
+		const TMap<FFieldNodeBase*, float>& CommandTimes, const float PreviousTime) const override;
+
+	/** Evaluate the maximum magnitude of the field graph  */
+	virtual float EvalMaxMagnitude() const override;
 
 	int32 Magnitude;
 };
@@ -60,13 +69,26 @@ public:
 	virtual FFieldNodeBase * NewCopy() const override { return new FRadialIntMask(Radius,Position,InteriorValue,ExteriorValue,SetMaskCondition); }
 	virtual ~FRadialIntMask() {}
 
-	void Evaluate(FFieldContext&, TFieldArrayView<int32>& Results) const override;
-	virtual bool operator==(const FFieldNodeBase& Node);
+	virtual void Evaluate(FFieldContext&, TFieldArrayView<int32>& Results) const override;
+	virtual bool operator==(const FFieldNodeBase& Node) override;
 
 	/** Serialization API */
-	virtual FFieldNodeBase::ESerializationType SerializationType() const { return FFieldNodeBase::ESerializationType::FieldNode_FRadialIntMask; }
+	virtual FFieldNodeBase::ESerializationType SerializationType() const override { return FFieldNodeBase::ESerializationType::FieldNode_FRadialIntMask; }
 	virtual void Serialize(FArchive& Ar) override;
 
+	/** Count the number of offsets/params that will be used by the world physics field */
+	virtual void FillSetupCount(int32& NumOffsets, int32& NumParams) const override;
+
+	/** Fill the offsets/params arrays that will be used by the world physics field */
+	virtual void FillSetupDatas(TArray<int32>& NodesOffsets, TArray<float>& NodesParams,
+		const TMap<FFieldNodeBase*, float>& CommandTimes, const float PreviousTime) const override;
+
+	/** Evaluate the maximum magnitude of the field graph  */
+	virtual float EvalMaxMagnitude() const override;
+
+	/** Compute the min/max spatial bounds of the field */
+	virtual void ComputeFieldBounds(FVector& MinBounds, FVector& MaxBounds) const override;
+
 	float Radius;
 	FVector Position;
 	int32 InteriorValue;
@@ -97,9 +119,19 @@ public:
 	virtual bool operator==(const FFieldNodeBase& Node);
 
 	/** Serialization API */
-	virtual FFieldNodeBase::ESerializationType SerializationType() const { return FFieldNodeBase::ESerializationType::FieldNode_FUniformScalar; }
+	virtual FFieldNodeBase::ESerializationType SerializationType() const override { return FFieldNodeBase::ESerializationType::FieldNode_FUniformScalar; }
 	virtual void Serialize(FArchive& Ar) override;
 
+	/** Count the number of offsets/params that will be used by the world physics field */
+	virtual void FillSetupCount(int32& NumOffsets, int32& NumParams) const override;
+
+	/** Fill the offsets/params arrays that will be used by the world physics field */
+	virtual void FillSetupDatas(TArray<int32>& NodesOffsets, TArray<float>& NodesParams,
+		const TMap<FFieldNodeBase*, float>& CommandTimes, const float PreviousTime) const override;
+
+	/** Evaluate the maximum magnitude of the field graph  */
+	virtual float EvalMaxMagnitude() const override;
+
 	float Magnitude;
 };
 
@@ -127,13 +159,24 @@ public:
 
 	virtual ~FWaveScalar() {}
 
-	void Evaluate(FFieldContext&, TFieldArrayView<float>& Results) const override;
-	virtual bool operator==(const FFieldNodeBase& Node);
+	virtual void Evaluate(FFieldContext&, TFieldArrayView<float>& Results) const override;
+	virtual bool operator==(const FFieldNodeBase& Node) override;
 
 	/** Serialization API */
-	virtual FFieldNodeBase::ESerializationType SerializationType() const { return FFieldNodeBase::ESerializationType::FieldNode_FWaveScalar; }
+	virtual FFieldNodeBase::ESerializationType SerializationType() const override { return FFieldNodeBase::ESerializationType::FieldNode_FWaveScalar; }
 	virtual void Serialize(FArchive& Ar) override;
 
+	/** Count the number of offsets/params that will be used by the world physics field */
+	virtual void FillSetupCount(int32& NumOffsets, int32& NumParams) const override;
+
+	/** Fill the offsets/params arrays that will be used by the world physics field */
+	virtual void FillSetupDatas(TArray<int32>& NodesOffsets, TArray<float>& NodesParams,
+		const TMap<FFieldNodeBase*, float>& CommandTimes, const float PreviousTime) const override;
+
+	/** Evaluate the maximum magnitude of the field graph  */
+	virtual float EvalMaxMagnitude() const override;
+
+
 	float Magnitude;
 	FVector Position;
 	float Wavelength;
@@ -173,13 +216,26 @@ public:
 
 	virtual ~FRadialFalloff() {}
 
-	void Evaluate(FFieldContext&, TFieldArrayView<float>& Results) const override;
-	virtual bool operator==(const FFieldNodeBase& Node);
+	virtual void Evaluate(FFieldContext&, TFieldArrayView<float>& Results) const override;
+	virtual bool operator==(const FFieldNodeBase& Node) override;
 
 	/** Serialization API */
-	virtual FFieldNodeBase::ESerializationType SerializationType() const { return FFieldNodeBase::ESerializationType::FieldNode_FRadialFalloff; }
+	virtual FFieldNodeBase::ESerializationType SerializationType() const override { return FFieldNodeBase::ESerializationType::FieldNode_FRadialFalloff; }
 	virtual void Serialize(FArchive& Ar) override;
 
+	/** Count the number of offsets/params that will be used by the world physics field */
+	virtual void FillSetupCount(int32& NumOffsets, int32& NumParams) const override;
+
+	/** Fill the offsets/params arrays that will be used by the world physics field */
+	virtual void FillSetupDatas(TArray<int32>& NodesOffsets, TArray<float>& NodesParams,
+		const TMap<FFieldNodeBase*, float>& CommandTimes, const float PreviousTime) const override;
+
+	/** Evaluate the maximum magnitude of the field graph  */
+	virtual float EvalMaxMagnitude() const override;
+
+	/** Compute the min/max spatial bounds of the field */
+	virtual void ComputeFieldBounds(FVector& MinBounds, FVector& MaxBounds) const override;
+
 	float Magnitude;
 	float MinRange;
 	float MaxRange;
@@ -225,13 +281,27 @@ public:
 	virtual FFieldNodeBase * NewCopy() const override { return new FPlaneFalloff(Magnitude, MinRange, MaxRange, Default, Distance, Position, Normal, Falloff); }
 	virtual ~FPlaneFalloff() {}
 
-	void Evaluate(FFieldContext&, TFieldArrayView<float>& Results) const override;
-	virtual bool operator==(const FFieldNodeBase& Node);
+	virtual void Evaluate(FFieldContext&, TFieldArrayView<float>& Results) const override;
+	virtual bool operator==(const FFieldNodeBase& Node) override;
 
 	/** Serialization API */
-	virtual FFieldNodeBase::ESerializationType SerializationType() const { return FFieldNodeBase::ESerializationType::FieldNode_FPlaneFalloff; }
+	virtual FFieldNodeBase::ESerializationType SerializationType() const override { return FFieldNodeBase::ESerializationType::FieldNode_FPlaneFalloff; }
 	virtual void Serialize(FArchive& Ar) override;
 
+	/** Count the number of offsets/params that will be used by the world physics field */
+	virtual void FillSetupCount(int32& NumOffsets, int32& NumParams) const override;
+
+	/** Fill the offsets/params arrays that will be used by the world physics field */
+	virtual void FillSetupDatas(TArray<int32>& NodesOffsets, TArray<float>& NodesParams,
+		const TMap<FFieldNodeBase*, float>& CommandTimes, const float PreviousTime) const override;
+
+	/** Evaluate the maximum magnitude of the field graph  */
+	virtual float EvalMaxMagnitude() const override;
+
+	/** Compute the min/max spatial bounds of the field */
+	virtual void ComputeFieldBounds(FVector& MinBounds, FVector& MaxBounds) const override;
+
+
 	float Magnitude;
 	float MinRange;
 	float MaxRange;
@@ -274,13 +344,26 @@ public:
 	virtual FFieldNodeBase * NewCopy() const override { return new FBoxFalloff(Magnitude, MinRange, MaxRange, Default, Transform, Falloff); }
 	virtual ~FBoxFalloff() {}
 
-	void Evaluate(FFieldContext&, TFieldArrayView<float>& Results) const override;
-	virtual bool operator==(const FFieldNodeBase& Node);
+	virtual void Evaluate(FFieldContext&, TFieldArrayView<float>& Results) const override;
+	virtual bool operator==(const FFieldNodeBase& Node) override;
 
 	/** Serialization API */
-	virtual FFieldNodeBase::ESerializationType SerializationType() const { return FFieldNodeBase::ESerializationType::FieldNode_FBoxFalloff; }
+	virtual FFieldNodeBase::ESerializationType SerializationType() const override { return FFieldNodeBase::ESerializationType::FieldNode_FBoxFalloff; }
 	virtual void Serialize(FArchive& Ar) override;
 
+	/** Count the number of offsets/params that will be used by the world physics field */
+	virtual void FillSetupCount(int32& NumOffsets, int32& NumParams) const override;
+
+	/** Fill the offsets/params arrays that will be used by the world physics field */
+	virtual void FillSetupDatas(TArray<int32>& NodesOffsets, TArray<float>& NodesParams,
+		const TMap<FFieldNodeBase*, float>& CommandTimes, const float PreviousTime) const override;
+
+	/** Evaluate the maximum magnitude of the field graph  */
+	virtual float EvalMaxMagnitude() const override;
+
+	/** Compute the min/max spatial bounds of the field */
+	virtual void ComputeFieldBounds(FVector& MinBounds, FVector& MaxBounds) const override;
+
 	float Magnitude;
 	float MinRange;
 	float MaxRange;
@@ -312,21 +395,28 @@ public:
 	virtual FFieldNodeBase * NewCopy() const override { return new FNoiseField(MinRange,MaxRange,Transform); }
 	virtual ~FNoiseField() {}
 
-	void Evaluate(FFieldContext&, TFieldArrayView<float>& Results) const override;	
-	virtual bool operator==(const FFieldNodeBase& Node);
+	virtual void Evaluate(FFieldContext&, TFieldArrayView<float>& Results) const override;
+	virtual bool operator==(const FFieldNodeBase& Node) override;
 
 	/** Serialization API */
-	virtual FFieldNodeBase::ESerializationType SerializationType() const { return FFieldNodeBase::ESerializationType::FieldNode_FNoiseField; }
+	virtual FFieldNodeBase::ESerializationType SerializationType() const override { return FFieldNodeBase::ESerializationType::FieldNode_FNoiseField; }
 	virtual void Serialize(FArchive& Ar) override;
 
+	/** Count the number of offsets/params that will be used by the world physics field */
+	virtual void FillSetupCount(int32& NumOffsets, int32& NumParams) const override;
+
+	/** Fill the offsets/params arrays that will be used by the world physics field */
+	virtual void FillSetupDatas(TArray<int32>& NodesOffsets, TArray<float>& NodesParams,
+		const TMap<FFieldNodeBase*, float>& CommandTimes, const float PreviousTime) const override;
+
+	/** Evaluate the maximum magnitude of the field graph  */
+	virtual float EvalMaxMagnitude() const override;
+
 	float MinRange;
 	float MaxRange;
 	FTransform Transform;
 };
 
-
-
-
 /**
 * UniformVector
 **/
@@ -346,18 +436,27 @@ public:
 	virtual FFieldNodeBase * NewCopy() const override { return new FUniformVector(Magnitude,Direction); }
 	virtual ~FUniformVector() {}
 
-	void Evaluate(FFieldContext&, TFieldArrayView<FVector>& Results) const override;
-	virtual bool operator==(const FFieldNodeBase& Node);
+	virtual void Evaluate(FFieldContext&, TFieldArrayView<FVector>& Results) const override;
+	virtual bool operator==(const FFieldNodeBase& Node) override;
 
 	/** Serialization API */
-	virtual FFieldNodeBase::ESerializationType SerializationType() const { return FFieldNodeBase::ESerializationType::FieldNode_FUniformVector; }
+	virtual FFieldNodeBase::ESerializationType SerializationType() const override { return FFieldNodeBase::ESerializationType::FieldNode_FUniformVector; }
 	virtual void Serialize(FArchive& Ar) override;
 
+	/** Count the number of offsets/params that will be used by the world physics field */
+	virtual void FillSetupCount(int32& NumOffsets, int32& NumParams) const override;
+
+	/** Fill the offsets/params arrays that will be used by the world physics field */
+	virtual void FillSetupDatas(TArray<int32>& NodesOffsets, TArray<float>& NodesParams,
+		const TMap<FFieldNodeBase*, float>& CommandTimes, const float PreviousTime) const override;
+
+	/** Evaluate the maximum magnitude of the field graph  */
+	virtual float EvalMaxMagnitude() const override;
+
 	float Magnitude;
 	FVector Direction;
 };
 
-
 /**
 * RadialVector
 **/
@@ -377,19 +476,28 @@ public:
 	virtual FFieldNodeBase * NewCopy() const override { return new FRadialVector(Magnitude,Position); }
 	virtual ~FRadialVector() {}
 
-	void Evaluate(FFieldContext&, TFieldArrayView<FVector>& Results) const override;
-	virtual bool operator==(const FFieldNodeBase& Node);
+	virtual void Evaluate(FFieldContext&, TFieldArrayView<FVector>& Results) const override;
+	virtual bool operator==(const FFieldNodeBase& Node) override;
 
 	/** Serialization API */
-	virtual FFieldNodeBase::ESerializationType SerializationType() const { return FFieldNodeBase::ESerializationType::FieldNode_FRadialVector; }
+	virtual FFieldNodeBase::ESerializationType SerializationType() const override { return FFieldNodeBase::ESerializationType::FieldNode_FRadialVector; }
 	virtual void Serialize(FArchive& Ar) override;
 
+	/** Count the number of offsets/params that will be used by the world physics field */
+	virtual void FillSetupCount(int32& NumOffsets, int32& NumParams) const override;
+
+	/** Fill the offsets/params arrays that will be used by the world physics field */
+	virtual void FillSetupDatas(TArray<int32>& NodesOffsets, TArray<float>& NodesParams,
+		const TMap<FFieldNodeBase*, float>& CommandTimes, const float PreviousTime) const override;
+
+	/** Evaluate the maximum magnitude of the field graph  */
+	virtual float EvalMaxMagnitude() const override;
+
 	float Magnitude;
 	FVector Position;
 
 };
 
-
 /**
 * RandomVector
 **/
@@ -407,15 +515,24 @@ public:
 	virtual FFieldNodeBase * NewCopy() const override { return new FRandomVector(Magnitude); }
 	virtual ~FRandomVector() {}
 
-	void Evaluate(FFieldContext&, TFieldArrayView<FVector>& Results) const override;
-	virtual bool operator==(const FFieldNodeBase& Node);
+	virtual void Evaluate(FFieldContext&, TFieldArrayView<FVector>& Results) const override;
+	virtual bool operator==(const FFieldNodeBase& Node) override;
 
 	/** Serialization API */
-	virtual FFieldNodeBase::ESerializationType SerializationType() const { return FFieldNodeBase::ESerializationType::FieldNode_FRandomVector; }
+	virtual FFieldNodeBase::ESerializationType SerializationType() const override { return FFieldNodeBase::ESerializationType::FieldNode_FRandomVector; }
 	virtual void Serialize(FArchive& Ar) override;
 
-	float Magnitude;
+	/** Count the number of offsets/params that will be used by the world physics field */
+	virtual void FillSetupCount(int32& NumOffsets, int32& NumParams) const override;
+
+	/** Fill the offsets/params arrays that will be used by the world physics field */
+	virtual void FillSetupDatas(TArray<int32>& NodesOffsets, TArray<float>& NodesParams,
+		const TMap<FFieldNodeBase*, float>& CommandTimes, const float PreviousTime) const override;
+
+	/** Evaluate the maximum magnitude of the field graph  */
+	virtual float EvalMaxMagnitude() const override;
 
+	float Magnitude;
 };
 
 /**
@@ -447,13 +564,26 @@ public:
 	}
 	virtual ~FSumScalar() {}
 
-	void Evaluate(FFieldContext&, TFieldArrayView<float>& Results) const override;
-	virtual bool operator==(const FFieldNodeBase& Node);
+	virtual void Evaluate(FFieldContext&, TFieldArrayView<float>& Results) const override;
+	virtual bool operator==(const FFieldNodeBase& Node) override;
 
 	/** Serialization API */
-	virtual FFieldNodeBase::ESerializationType SerializationType() const { return FFieldNodeBase::ESerializationType::FieldNode_FSumScalar; }
+	virtual FFieldNodeBase::ESerializationType SerializationType() const override { return FFieldNodeBase::ESerializationType::FieldNode_FSumScalar; }
 	virtual void Serialize(FArchive& Ar) override;
 
+	/** Count the number of offsets/params that will be used by the world physics field */
+	virtual void FillSetupCount(int32& NumOffsets, int32& NumParams) const override;
+
+	/** Fill the offsets/params arrays that will be used by the world physics field */
+	virtual void FillSetupDatas(TArray<int32>& NodesOffsets, TArray<float>& NodesParams,
+		const TMap<FFieldNodeBase*, float>& CommandTimes, const float PreviousTime) const override;
+
+	/** Evaluate the maximum magnitude of the field graph  */
+	virtual float EvalMaxMagnitude() const override;
+
+	/** Compute the min/max spatial bounds of the field */
+	virtual void ComputeFieldBounds(FVector& MinBounds, FVector& MaxBounds) const override;
+
 	float Magnitude;
 	TUniquePtr< FFieldNode<float> > ScalarRight;
 	TUniquePtr< FFieldNode<float> > ScalarLeft;
@@ -492,13 +622,26 @@ public:
 	}
 	virtual ~FSumVector() {}
 
-	void Evaluate(FFieldContext&, TFieldArrayView<FVector>& Results) const override;
-	virtual bool operator==(const FFieldNodeBase& Node);
+	virtual void Evaluate(FFieldContext&, TFieldArrayView<FVector>& Results) const override;
+	virtual bool operator==(const FFieldNodeBase& Node) override;
 
 	/** Serialization API */
-	virtual FFieldNodeBase::ESerializationType SerializationType() const { return FFieldNodeBase::ESerializationType::FieldNode_FSumVector; }
+	virtual FFieldNodeBase::ESerializationType SerializationType() const override { return FFieldNodeBase::ESerializationType::FieldNode_FSumVector; }
 	virtual void Serialize(FArchive& Ar) override;
 
+	/** Count the number of offsets/params that will be used by the world physics field */
+	virtual void FillSetupCount(int32& NumOffsets, int32& NumParams) const override;
+
+	/** Fill the offsets/params arrays that will be used by the world physics field */
+	virtual void FillSetupDatas(TArray<int32>& NodesOffsets, TArray<float>& NodesParams,
+		const TMap<FFieldNodeBase*, float>& CommandTimes, const float PreviousTime) const override;
+
+	/** Evaluate the maximum magnitude of the field graph  */
+	virtual float EvalMaxMagnitude() const override;
+
+	/** Compute the min/max spatial bounds of the field */
+	virtual void ComputeFieldBounds(FVector& MinBounds, FVector& MaxBounds) const override;
+
 	float Magnitude;
 	TUniquePtr< FFieldNode<float> > Scalar;
 	TUniquePtr< FFieldNode<FVector> > VectorRight;
@@ -527,13 +670,26 @@ public:
 	}
 	virtual ~FConversionField() {}
 
-	void Evaluate(FFieldContext&, TFieldArrayView<OutT>& Results) const override;
-	virtual bool operator==(const FFieldNodeBase& Node);
+	virtual void Evaluate(FFieldContext&, TFieldArrayView<OutT>& Results) const override;
+	virtual bool operator==(const FFieldNodeBase& Node) override;
 
 	/** Serialization API */
-	virtual FFieldNodeBase::ESerializationType SerializationType() const { return FFieldNodeBase::ESerializationType::FieldNode_FConversionField; }
+	virtual FFieldNodeBase::ESerializationType SerializationType() const override { return FFieldNodeBase::ESerializationType::FieldNode_FConversionField; }
 	virtual void Serialize(FArchive& Ar) override;
 
+	/** Count the number of offsets/params that will be used by the world physics field */
+	virtual void FillSetupCount(int32& NumOffsets, int32& NumParams) const override;
+
+	/** Fill the offsets/params arrays that will be used by the world physics field */
+	virtual void FillSetupDatas(TArray<int32>& NodesOffsets, TArray<float>& NodesParams,
+		const TMap<FFieldNodeBase*, float>& CommandTimes, const float PreviousTime) const override;
+
+	/** Evaluate the maximum magnitude of the field graph  */
+	virtual float EvalMaxMagnitude() const override;
+
+	/** Compute the min/max spatial bounds of the field */
+	virtual void ComputeFieldBounds(FVector& MinBounds, FVector& MaxBounds) const override;
+
 	TUniquePtr< FFieldNode<InT> > InputField;
 };
 
@@ -567,13 +723,26 @@ public:
 	}
 	virtual ~FCullingField() {}
 
-	void Evaluate(FFieldContext&, TFieldArrayView<T>& Results) const override;
-	virtual bool operator==(const FFieldNodeBase& Node);
+	virtual void Evaluate(FFieldContext&, TFieldArrayView<T>& Results) const override;
+	virtual bool operator==(const FFieldNodeBase& Node) override;
 
 	/** Serialization API */
-	virtual FFieldNodeBase::ESerializationType SerializationType() const { return FFieldNodeBase::ESerializationType::FieldNode_FCullingField; }
+	virtual FFieldNodeBase::ESerializationType SerializationType() const override { return FFieldNodeBase::ESerializationType::FieldNode_FCullingField; }
 	virtual void Serialize(FArchive& Ar) override;
 
+	/** Count the number of offsets/params that will be used by the world physics field */
+	virtual void FillSetupCount(int32& NumOffsets, int32& NumParams) const override;
+
+	/** Fill the offsets/params arrays that will be used by the world physics field */
+	virtual void FillSetupDatas(TArray<int32>& NodesOffsets, TArray<float>& NodesParams,
+		const TMap<FFieldNodeBase*, float>& CommandTimes, const float PreviousTime) const override;
+
+	/** Evaluate the maximum magnitude of the field graph  */
+	virtual float EvalMaxMagnitude() const override;
+
+	/** Compute the min/max spatial bounds of the field */
+	virtual void ComputeFieldBounds(FVector& MinBounds, FVector& MaxBounds) const override;
+
 	TUniquePtr< FFieldNode<float> > Culling;
 	TUniquePtr< FFieldNode<T> > Input;
 	EFieldCullingOperationType Operation;
@@ -586,7 +755,7 @@ extern template class CHAOS_API FCullingField<FVector>;
 #endif
 
 /**
-* FCullingField
+* FReturnResultsTerminal
 **/
 
 template<class T>
@@ -610,7 +779,6 @@ public:
 	/** Serialization API */
 	virtual FFieldNodeBase::ESerializationType SerializationType() const { return FFieldNodeBase::ESerializationType::FieldNode_FReturnResultsTerminal; }
 	virtual void Serialize(FArchive& Ar) override;
-
 };
 
 #if PLATFORM_MAC
