Chaos : Fix crash where a cluster particle destroyed while still part of a cluster union was re-added to a map and leaving a dangling pointer

diff --git a/Engine/Source/Runtime/Experimental/Chaos/Private/Chaos/PBDRigidClustering.cpp b/Engine/Source/Runtime/Experimental/Chaos/Private/Chaos/PBDRigidClustering.cpp
--- a/Engine/Source/Runtime/Experimental/Chaos/Private/Chaos/PBDRigidClustering.cpp
+++ b/Engine/Source/Runtime/Experimental/Chaos/Private/Chaos/PBDRigidClustering.cpp
@@ -1701,9 +1701,6 @@ namespace Chaos
 			ensure(ClusteredParticle->ClusterIds().Id == nullptr);
 		}
 
-	// reset the structures
-		TopLevelClusterParents.Remove(ClusteredParticle);
-		TopLevelClusterParentsStrained.Remove(ClusteredParticle);
 
 		// need to disconnect from any other particles ( this can be from being a child of a cluster or a cluster union)
 		RemoveNodeConnections(ClusteredParticle);
@@ -1746,6 +1743,13 @@ namespace Chaos
 		{
 			ParentParticle = nullptr;
 		}
+
+		// reset the structures
+		// Note: this needs to be at the end to make sure that no other operations above may re-add it 
+		// ( for example HandleRemoveOperationWithClusterLookup )
+		TopLevelClusterParents.Remove(ClusteredParticle);
+		TopLevelClusterParentsStrained.Remove(ClusteredParticle);
+
 		return ParentParticle;
 
 	}
