Fixed mismatched user-defined copy constructor/assignment operators.

diff --git a/Engine/Source/Runtime/Experimental/ChaosCore/Public/Chaos/Vector.h b/Engine/Source/Runtime/Experimental/ChaosCore/Public/Chaos/Vector.h
--- a/Engine/Source/Runtime/Experimental/ChaosCore/Public/Chaos/Vector.h
+++ b/Engine/Source/Runtime/Experimental/ChaosCore/Public/Chaos/Vector.h
@@ -44,6 +44,9 @@ namespace Chaos
 		static const int NumElements = d;
 		using FTraits = TVectorTraits<T, d>;
 
+		TVector(const TVector&) = default;
+		TVector& operator=(const TVector&) = default;
+
 		TVector() {}
 		explicit TVector(const FElement& Element)
 		{
@@ -113,7 +116,8 @@ namespace Chaos
 			}
 		}
 
-		TVector<T, d>& operator=(const TVector<T, d>& Other)
+		template<class T2>
+		TVector<T, d>& operator=(const TVector<T2, d>& Other)
 		{
 			for (int32 i = 0; i < NumElements; ++i)
 			{
@@ -984,6 +988,11 @@ namespace Chaos
 	public:
 		using FElement = T;
 
+PRAGMA_DISABLE_DEPRECATION_WARNINGS
+		TVector(const TVector&) = default;
+		TVector& operator=(const TVector&) = default;
+PRAGMA_ENABLE_DEPRECATION_WARNINGS
+
 		FORCEINLINE TVector() {}
 		FORCEINLINE explicit TVector(T InX)
 		    : X(InX), Y(InX), Z(InX) {}
@@ -1021,7 +1030,8 @@ namespace Chaos
 			Stream.write(reinterpret_cast<const char*>(&Z), sizeof(T));
 		}
 #endif
-		FORCEINLINE TVector<T, 3>& operator=(const TVector<T, 3>& Other)
+		template<class T2>
+		FORCEINLINE TVector<T, 3>& operator=(const TVector<T2, 3>& Other)
 		{
 			X = Other.X;
 			Y = Other.Y;
@@ -1167,6 +1177,11 @@ PRAGMA_ENABLE_DEPRECATION_WARNINGS
 	public:
 		using FElement = int32;
 
+PRAGMA_DISABLE_DEPRECATION_WARNINGS
+		TVector(const TVector&) = default;
+		TVector& operator=(const TVector&) = default;
+PRAGMA_ENABLE_DEPRECATION_WARNINGS
+
 		FORCEINLINE TVector()
 		{}
 		FORCEINLINE explicit TVector(const FElement InX)
@@ -1209,7 +1224,8 @@ PRAGMA_ENABLE_DEPRECATION_WARNINGS
 		}
 #endif
 
-		FORCEINLINE TVector<FElement, 2>& operator=(const TVector<FElement, 2>& Other)
+		template<typename OtherT>
+		FORCEINLINE TVector<int32, 2>& operator=(const TVector<OtherT, 2>& Other)
 		{
 			X = Other.X;
 			Y = Other.Y;
