Added CVar to reduce proportion of dispatched collision events through GC event dispatchers

diff --git a/Engine/Source/Runtime/Experimental/ChaosSolverEngine/Private/Chaos/ChaosGameplayEventDispatcher.cpp b/Engine/Source/Runtime/Experimental/ChaosSolverEngine/Private/Chaos/ChaosGameplayEventDispatcher.cpp
--- a/Engine/Source/Runtime/Experimental/ChaosSolverEngine/Private/Chaos/ChaosGameplayEventDispatcher.cpp
+++ b/Engine/Source/Runtime/Experimental/ChaosSolverEngine/Private/Chaos/ChaosGameplayEventDispatcher.cpp
@@ -419,6 +419,11 @@ TArray<UObject*> UChaosGameplayEventDispatcher::GetInterestedProxyOwnersForCrumb
 	return Result;
 }
 
+int32 GReducedCollisionDispatch = 0;
+float GReducedCollisionDispatchAmount = 0.1f;
+FAutoConsoleVariableRef CVar_ReducedCollisionDispatch(TEXT("p.chaos.gpdispatch.collisions.reduce"), GReducedCollisionDispatch, TEXT("Enable reduced dispatch of collision events where only a subset of total collisions are actually dispatched to outside systems [0 - off, 1- on]"));
+FAutoConsoleVariableRef CVar_ReducedCollisionDispatchAmount(TEXT("p.chaos.gpdispatch.collisions.reduceamount"), GReducedCollisionDispatchAmount, TEXT("Percentage of collisions to actually dispatch [0 - 1]"));
+
 void UChaosGameplayEventDispatcher::HandleCollisionEvents(const Chaos::FCollisionEventData& Event)
 {
 	// todo(chaos) : this code is very similar to FPhysScene_Chaos::HandleCollisionEvents, we should propably consolidate if possible or share as much code as possible 
@@ -453,8 +458,23 @@ void UChaosGameplayEventDispatcher::HandleCollisionEvents(const Chaos::FCollisio
 						TArray<int32> const* const CollisionIndices = PhysicsProxyToCollisionIndicesMap.Find(PhysicsProxy0);
 						if (CollisionIndices)
 						{
-							for (int32 EncodedCollisionIdx : *CollisionIndices)
+							const int32 NumCollisionIndices = CollisionIndices->Num();
+
+							if(NumCollisionIndices == 0)
+							{
+								continue;
+							}
+
+							const int32 NumToDispatch = GReducedCollisionDispatch == 1 ?
+								FMath::Max((int32)((float)NumCollisionIndices * FMath::Clamp(GReducedCollisionDispatchAmount, 0.0f, 1.0f)), 1) :
+								NumCollisionIndices;
+
+							const int32 Interval = NumCollisionIndices / NumToDispatch;
+
+							for(int32 Index = 0; Index < NumCollisionIndices; Index += Interval)
 							{
+								int32 EncodedCollisionIdx = (*CollisionIndices)[Index];
+
 								bool bSwapOrder;
 								int32 CollisionIdx = Chaos::FEventManager::DecodeCollisionIndex(EncodedCollisionIdx, bSwapOrder);
 
