Add explicit remap arrays to the ISM Pool to manage mapping between the ISM Pool instance index and the ISM render instance index.
This costs some memory but makes the cost of instance removal much cheaper. We no longer need to loop and scan multiple arrays to find our instance to rename.
Added a bit more commenting around the code.

diff --git a/Engine/Source/Runtime/Experimental/GeometryCollectionEngine/Private/GeometryCollection/GeometryCollectionISMPoolComponent.cpp b/Engine/Source/Runtime/Experimental/GeometryCollectionEngine/Private/GeometryCollection/GeometryCollectionISMPoolComponent.cpp
--- a/Engine/Source/Runtime/Experimental/GeometryCollectionEngine/Private/GeometryCollection/GeometryCollectionISMPoolComponent.cpp
+++ b/Engine/Source/Runtime/Experimental/GeometryCollectionEngine/Private/GeometryCollection/GeometryCollectionISMPoolComponent.cpp
@@ -97,22 +97,32 @@ FGeometryCollectionISM::FGeometryCollectionISM(AActor* OwmingActor, const FGeome
 	ISMComponent = ISMC;
 }
 
-int32 FGeometryCollectionISM::AddInstanceGroup(int32 InstanceCount, TArrayView<const float> CustomDataFloats)
+FInstanceGroups::FInstanceGroupId FGeometryCollectionISM::AddInstanceGroup(int32 InstanceCount, TArrayView<const float> CustomDataFloats)
 {
-	const int32 InstanceGroupIndex = InstanceGroups.AddGroup(InstanceCount);
-
 	// When adding new group it will always have a single range
-	const FInstanceGroups::FInstanceGroupRange& NewInstanceGroup = InstanceGroups.GetGroup(InstanceGroupIndex);
-
-	ISMComponent->PreAllocateInstancesMemory(InstanceCount);
+	const FInstanceGroups::FInstanceGroupId InstanceGroupIndex = InstanceGroups.AddGroup(InstanceCount);
+	const FInstanceGroups::FInstanceGroupRange& NewInstanceGroup = InstanceGroups.GroupRanges[InstanceGroupIndex];
 
 	FTransform ZeroScaleTransform;
 	ZeroScaleTransform.SetIdentityZeroScale();
 	TArray<FTransform> ZeroScaleTransforms;
 	ZeroScaleTransforms.Init(ZeroScaleTransform, InstanceCount);
 
-	ISMComponent->AddInstances(ZeroScaleTransforms, false, true);
+	ISMComponent->PreAllocateInstancesMemory(InstanceCount);
+	TArray<int32> RenderInstances = ISMComponent->AddInstances(ZeroScaleTransforms, true, true);
+
+	// Ensure that remapping arrays are big enough to hold any new items.
+	InstanceIndexToRenderIndex.SetNum(InstanceGroups.GetMaxInstanceIndex(), false);
+	RenderIndexToInstanceIndex.SetNum(ISMComponent->PerInstanceSMData.Num(), false);
+	// Store mapping between our fixed instance index and the mutable ISM render index.
+	// todo: Improve ISM API so that we don't need to pay the memory overhead here to manage this.
+	for (int32 InstanceIndex = 0; InstanceIndex < InstanceCount; ++InstanceIndex)
+	{
+		InstanceIndexToRenderIndex[NewInstanceGroup.Start + InstanceIndex] = RenderInstances[InstanceIndex];
+		RenderIndexToInstanceIndex[RenderInstances[InstanceIndex]] = NewInstanceGroup.Start + InstanceIndex;
+	}
 
+	// Set any custom data.
 	if (CustomDataFloats.Num())
 	{
 		const int32 NumCustomDataFloats = ISMComponent->NumCustomDataFloats;
@@ -120,7 +130,7 @@ int32 FGeometryCollectionISM::AddInstanceGroup(int32 InstanceCount, TArrayView<c
 		{
 			for (int32 InstanceIndex = 0; InstanceIndex < InstanceCount; ++InstanceIndex)
 			{
-				ISMComponent->SetCustomData(NewInstanceGroup.InstanceIdToIndex[InstanceIndex], CustomDataFloats.Slice(InstanceIndex * NumCustomDataFloats, NumCustomDataFloats));
+				ISMComponent->SetCustomData(RenderInstances[InstanceIndex], CustomDataFloats.Slice(InstanceIndex * NumCustomDataFloats, NumCustomDataFloats));
 			}
 		}
 	}
@@ -137,6 +147,7 @@ FGeometryCollectionMeshInfo FGeometryCollectionISMPool::AddISM(UGeometryCollecti
 	{
 		if (FreeList.Num())
 		{
+			// Take an ISM from the current FreeList instead of allocating a new slot.
 			Info.ISMIndex = FreeList.Last();
 			FreeList.RemoveAt(FreeList.Num() - 1);
 			ISMs[Info.ISMIndex] = FGeometryCollectionISM(OwningComponent->GetOwner(), MeshInstance);
@@ -162,10 +173,10 @@ bool FGeometryCollectionISMPool::BatchUpdateInstancesTransforms(FGeometryCollect
 	if (ISMs.IsValidIndex(MeshInfo.ISMIndex))
 	{
 		FGeometryCollectionISM& ISM = ISMs[MeshInfo.ISMIndex];
-		const FInstanceGroups::FInstanceGroupRange& InstanceGroup = ISM.InstanceGroups.GetGroup(MeshInfo.InstanceGroupIndex);
-		ensure((StartInstanceIndex + NewInstancesTransforms.Num()) <= InstanceGroup.Count());
+		const FInstanceGroups::FInstanceGroupRange& InstanceGroup = ISM.InstanceGroups.GroupRanges[MeshInfo.InstanceGroupIndex];
+		ensure((StartInstanceIndex + NewInstancesTransforms.Num()) <= InstanceGroup.Count);
 
-		int32 StartIndex = InstanceGroup.InstanceIdToIndex[StartInstanceIndex];
+		int32 StartIndex = ISM.InstanceIndexToRenderIndex[InstanceGroup.Start];
 		int32 TransformIndex = 0;
 		int32 BatchCount = 1;
 		TArray<FTransform> BatchTransforms; //< Can't use TArrayView because blueprint function doesn't support that 
@@ -174,10 +185,11 @@ bool FGeometryCollectionISMPool::BatchUpdateInstancesTransforms(FGeometryCollect
 		for (int InstanceIndex = StartInstanceIndex + 1; InstanceIndex < NewInstancesTransforms.Num(); ++InstanceIndex)
 		{
 			// flush batch?
-			if (InstanceGroup.InstanceIdToIndex[InstanceIndex] != (StartIndex + BatchCount))
+			int32 RenderIndex = ISM.InstanceIndexToRenderIndex[InstanceGroup.Start + InstanceIndex];
+			if (RenderIndex != (StartIndex + BatchCount))
 			{
 				ISM.ISMComponent->BatchUpdateInstancesTransforms(StartIndex, BatchTransforms, bWorldSpace, bMarkRenderStateDirty, bTeleport);
-				StartIndex = InstanceGroup.InstanceIdToIndex[InstanceIndex];
+				StartIndex = RenderIndex;
 				BatchTransforms.SetNum(0, false);
 				BatchCount = 0;
 			}
@@ -197,8 +209,16 @@ void FGeometryCollectionISMPool::RemoveISM(const FGeometryCollectionMeshInfo& Me
 	if (ISMs.IsValidIndex(MeshInfo.ISMIndex))
 	{
 		FGeometryCollectionISM& ISM = ISMs[MeshInfo.ISMIndex];
-		const FInstanceGroups::FInstanceGroupRange& InstanceGroup = ISM.InstanceGroups.GetGroup(MeshInfo.InstanceGroupIndex);
-		ISM.ISMComponent->RemoveInstances(InstanceGroup.InstanceIdToIndex);
+		const FInstanceGroups::FInstanceGroupRange& InstanceGroup = ISM.InstanceGroups.GroupRanges[MeshInfo.InstanceGroupIndex];
+		
+		TArray<int32> InstancesToRemove;
+		InstancesToRemove.SetNum(InstanceGroup.Count);
+		for (int32 InstanceIndex = 0; InstanceIndex < InstanceGroup.Count; ++InstanceIndex)
+		{
+			InstancesToRemove[InstanceIndex] = ISM.InstanceIndexToRenderIndex[InstanceGroup.Start + InstanceIndex];
+		}
+		ISM.ISMComponent->RemoveInstances(InstancesToRemove);
+
 		ISM.InstanceGroups.RemoveGroup(MeshInfo.InstanceGroupIndex);
 	
 		if (ISM.InstanceGroups.IsEmpty() && ISM.ISMComponent->PerInstanceSMData.Num() == 0)
@@ -209,6 +229,7 @@ void FGeometryCollectionISMPool::RemoveISM(const FGeometryCollectionMeshInfo& Me
 			ISM.ISMComponent->UnregisterComponent();
 			ISM.ISMComponent->DestroyComponent();
 			
+			ISM.InstanceGroups.Reset();
 			MeshToISMIndex.Remove(ISM.MeshInstance);
 			ISMComponentToISMIndex.Remove(ISM.ISMComponent);
 			FreeList.Add(MeshInfo.ISMIndex);
@@ -233,11 +254,11 @@ void FGeometryCollectionISMPool::OnISMInstanceIndexUpdated(UInstancedStaticMeshC
 	{
 		if (IndexUpdateData.Type == FInstancedStaticMeshDelegates::EInstanceIndexUpdateType::Removed)
 		{
-			ISM.InstanceGroups.IndexRemoved(IndexUpdateData.Index);
+			ISM.IndexRemoved(IndexUpdateData.Index);
 		}
 		else if (IndexUpdateData.Type == FInstancedStaticMeshDelegates::EInstanceIndexUpdateType::Relocated)
 		{
-			ISM.InstanceGroups.IndexReallocated(IndexUpdateData.OldIndex, IndexUpdateData.Index);
+			ISM.IndexReallocated(IndexUpdateData.OldIndex, IndexUpdateData.Index);
 		}
 	}
 }
@@ -314,3 +335,25 @@ bool UGeometryCollectionISMPoolComponent::BatchUpdateInstancesTransforms(FMeshGr
 	UE_LOG(LogChaos, Warning, TEXT("UGeometryCollectionISMPoolComponent : Trying to update instance with mesh group (%d) that not exists"), MeshGroupId);
 	return false;
 }
+
+void UGeometryCollectionISMPoolComponent::GetResourceSizeEx(FResourceSizeEx& CumulativeResourceSize)
+{
+	Super::GetResourceSizeEx(CumulativeResourceSize);
+
+	int32 SizeBytes =
+		MeshGroups.GetAllocatedSize()
+		+ Pool.MeshToISMIndex.GetAllocatedSize()
+		+ Pool.ISMComponentToISMIndex.GetAllocatedSize()
+		+ Pool.ISMs.GetAllocatedSize()
+		+ Pool.FreeList.GetAllocatedSize();
+	
+	for (FGeometryCollectionISM ISM : Pool.ISMs)
+	{
+		SizeBytes += ISM.InstanceIndexToRenderIndex.GetAllocatedSize()
+			+ ISM.RenderIndexToInstanceIndex.GetAllocatedSize()
+			+ ISM.InstanceGroups.GroupRanges.GetAllocatedSize()
+			+ ISM.InstanceGroups.FreeList.GetAllocatedSize();
+	}
+
+	CumulativeResourceSize.AddDedicatedSystemMemoryBytes(SizeBytes);
+}
diff --git a/Engine/Source/Runtime/Experimental/GeometryCollectionEngine/Public/GeometryCollection/GeometryCollectionISMPoolComponent.h b/Engine/Source/Runtime/Experimental/GeometryCollectionEngine/Public/GeometryCollection/GeometryCollectionISMPoolComponent.h
--- a/Engine/Source/Runtime/Experimental/GeometryCollectionEngine/Public/GeometryCollection/GeometryCollectionISMPoolComponent.h
+++ b/Engine/Source/Runtime/Experimental/GeometryCollectionEngine/Public/GeometryCollection/GeometryCollectionISMPoolComponent.h
@@ -10,106 +10,93 @@
 #include "GeometryCollectionISMPoolComponent.generated.h"
 
 class AActor;
-class UGeometryComponent;
 class UInstancedStaticMeshComponent;
 class UGeometryCollectionISMPoolComponent;
 
-struct FGeometryCollectionISMInstance
-{
-	int32 StartIndex;
-	int32 Count;
-};
-
+/** 
+ * Structure containing a set of allocated instance ranges in an FGeometryCollectionISM which is the manager for a single ISM component.
+ * The instance ranges don't change once allocated, and aren't the same as the actual render indices in the ISM.
+ * The reason that we don't store the the actual ISM render indices is that ISM component is free to reorder its instances whenever it likes.
+ */
 struct FInstanceGroups
 {
-public:
 	using FInstanceGroupId = int32;
 
+	/** A single continuous range associated with an FInstanceGroupId. */
 	struct FInstanceGroupRange
 	{
 		FInstanceGroupRange(int32 InStart, int32 InCount)
+			: Start(InStart)
+			, Count(InCount)
 		{
-			InstanceIdToIndex.Reserve(InCount);
-			for (int Index = 0; Index < InCount; ++Index)
-			{
-				InstanceIdToIndex.Add(InStart + Index);
-			}
 		}
 
-		bool TryIndexReallocate(int32 OldIndex, int32 NewIndex)
-		{
-			for (int32 i = 0; i < InstanceIdToIndex.Num(); ++i)
-			{
-				if (InstanceIdToIndex[i] == OldIndex)
-				{
-					InstanceIdToIndex[i] = NewIndex;
-					return true;
-				}
-			}
-			return false;
-		}
-
-		int32 Count() const { return InstanceIdToIndex.Num(); }
-		TArray<int32> InstanceIdToIndex;
+		int32 Start = 0;
+		int32 Count = 0;
 	};
 
-	bool IsEmpty() const
+	/** Reset all contents. */
+	void Reset()
 	{
-		return InstancesCount == 0;
+		TotalInstanceCount = 0;
+		TotalFreeInstanceCount = 0;
+		GroupRanges.Empty();
+		FreeList.Empty();
 	}
 
-	FInstanceGroupId AddGroup(int32 Count)
+	/** Returns true if no groups ranges are in use. */
+	bool IsEmpty() const
 	{
-		const int32 StartIndex = InstancesCount;
-		InstancesCount += Count;
-		GroupRanges.Emplace(NextGroupId, FInstanceGroupRange(StartIndex, Count));
-		return NextGroupId++;
+		return GroupRanges.Num() == FreeList.Num();
 	}
 
-	void RemoveGroup(FInstanceGroupId GroupId)
+	/** Returns the maximum instance index that we have allocated in a group. */
+	int32 GetMaxInstanceIndex() const
 	{
-		check(GroupRanges.Contains(GroupId));
-		const FInstanceGroupRange& GroupRangeToRemove = GroupRanges[GroupId];
-		InstancesCount -= GroupRangeToRemove.InstanceIdToIndex.Num();
-		GroupRanges.Remove(GroupId);
+		return TotalInstanceCount;
 	}
 
-	void IndexRemoved(int32 IndexToRemove)
+	/** Add a new group range. */
+	FInstanceGroupId AddGroup(int32 Count)
 	{
-		bool bFound = false;
-		for (TPair<FInstanceGroupId, FInstanceGroupRange>& GroupRange : GroupRanges)
+		// First check to see if we can recycle a group from the free list.
+		for (int32 Index = 0; Index < FreeList.Num(); ++Index)
 		{
-			if (GroupRange.Value.TryIndexReallocate(IndexToRemove, INDEX_NONE))
+			const FInstanceGroupId GroupId = FreeList[Index];
+			// todo: Could also allow allocating a subrange if Count is less than the group range count.
+			if (Count == GroupRanges[GroupId].Count)
 			{
-				bFound = true;
-				break;
+				TotalFreeInstanceCount -= Count;
+				FreeList.RemoveAtSwap(Index, 1, false);
+				return GroupId;
 			}
 		}
-		check(bFound);
+		
+		// Create a new group.
+		const int32 StartIndex = TotalInstanceCount;
+		TotalInstanceCount += Count;
+		const FInstanceGroupId GroupId = GroupRanges.Add(FInstanceGroupRange(StartIndex, Count));
+		return GroupId;
 	}
 
-	void IndexReallocated(int32 OldIndex, int32 NewIndex)
+	/** Remove a group range. */
+	void RemoveGroup(FInstanceGroupId GroupId)
 	{
-		bool bFound = false;
-		for (TPair<FInstanceGroupId, FInstanceGroupRange>& GroupRange : GroupRanges)
-		{
-			if (GroupRange.Value.TryIndexReallocate(OldIndex, NewIndex))
-			{
-				bFound = true;
-				break;
-			}
-		}
-		check(bFound);
+		// Remove the group by putting on free list for reuse.
+		// Actually removing it would require too much shuffling of the render instance index remapping.
+		TotalFreeInstanceCount += GroupRanges[GroupId].Count;
+		FreeList.Add(GroupId);
 	}
-	const FInstanceGroupRange& GetGroup(int32 GroupIndex) const { return GroupRanges[GroupIndex]; };
 
-private:
-	int32 InstancesCount = 0;
-	int32 NextGroupId = 0;
-	TMap<FInstanceGroupId, FInstanceGroupRange> GroupRanges;
+	int32 TotalInstanceCount = 0;
+	int32 TotalFreeInstanceCount = 0;
+	TArray<FInstanceGroupRange> GroupRanges;
+	TArray<FInstanceGroupId> FreeList;
 };
 
-/** */
+/** 
+ * A description for an ISM component.
+ */
 struct FISMComponentDescription
 {
 	bool bUseHISM = false;
@@ -151,9 +138,9 @@ FORCEINLINE uint32 GetTypeHash(const FISMComponentDescription& Desc)
 }
 
 /**
-* This represent a unique mesh with potentially overriden materials
-* if the array is empty , there's no overrides
-*/
+ * A mesh with potentially overriden materials and ISM property description.
+ * We batch instances into ISMs that have equivalent values for this structure.
+ */
 struct FGeometryCollectionStaticMeshInstance
 {
 	UStaticMesh* StaticMesh = nullptr;
@@ -194,42 +181,72 @@ FORCEINLINE uint32 GetTypeHash(const FGeometryCollectionStaticMeshInstance& Mesh
 	return CombinedHash;
 }
 
+/** Describes a group of instances within an ISM. */
 struct FGeometryCollectionMeshInfo
 {
 	int32 ISMIndex;
-	int32 InstanceGroupIndex;
+	FInstanceGroups::FInstanceGroupId InstanceGroupIndex;
 };
 
 struct FGeometryCollectionISMPool;
 
 /**
-* FGeometryCollectionMeshGroup
-* a mesh groupo contains various mesh with their instances
-*/
+ * A mesh group which is a collection of meshes and their related FGeometryCollectionMeshInfo.
+ * We group these with a single handle with the expectation that a client will want to own multiple meshs and release them together.
+ */
 struct FGeometryCollectionMeshGroup
 {
 	using FMeshId = int32;
 
+	/** Adds a new mesh with instance count. We expect to only add a unique mesh instance once to each group. Returns a ID that can be used to update the instances. */
 	FMeshId AddMesh(const FGeometryCollectionStaticMeshInstance& MeshInstance, int32 InstanceCount, const FGeometryCollectionMeshInfo& ISMInstanceInfo);
+	/** Update instance transforms for a group of instances. */
 	bool BatchUpdateInstancesTransforms(FGeometryCollectionISMPool& ISMPool, FMeshId MeshId, int32 StartInstanceIndex, const TArray<FTransform>& NewInstancesTransforms, bool bWorldSpace, bool bMarkRenderStateDirty, bool bTeleport);
+	/** Remove all of our managed meshes and associated instances. */
 	void RemoveAllMeshes(FGeometryCollectionISMPool& ISMPool);
 
-	TMap<FGeometryCollectionStaticMeshInstance, FMeshId> Meshes;
+	/** Array of allocated mesh infos. */
 	TArray<FGeometryCollectionMeshInfo> MeshInfos;
+	/** Map from mesh instance description to the its index in the MeshInfo array. */
+	TMap<FGeometryCollectionStaticMeshInstance, FMeshId> Meshes;
 };
 
+/** Structure containting all info for a single ISM. */
 struct FGeometryCollectionISM
 {
 	FGeometryCollectionISM(AActor* OwmingActor, const FGeometryCollectionStaticMeshInstance& InMeshInstance);
 
-	int32 AddInstanceGroup(int32 InstanceCount, TArrayView<const float> CustomDataFloats);
+	/** Add a group to the ISM. Returns the group index. */
+	FInstanceGroups::FInstanceGroupId AddInstanceGroup(int32 InstanceCount, TArrayView<const float> CustomDataFloats);
 
-	TObjectPtr<UInstancedStaticMeshComponent> ISMComponent;
+	/** Handle removal of a render instance by the ISM component. Our instance mapping needs to be updated. */
+	void IndexRemoved(int32 IndexToRemove)
+	{
+		const int32 InstanceIndex = RenderIndexToInstanceIndex[IndexToRemove];
+		InstanceIndexToRenderIndex[InstanceIndex] = -1;
+	}
+
+	/** Handle move of a render instance by the ISM component. Our instance mapping needs to be updated. */
+	void IndexReallocated(int32 OldIndex, int32 NewIndex)
+	{
+		const int32 InstanceIndex = RenderIndexToInstanceIndex[OldIndex];
+		InstanceIndexToRenderIndex[InstanceIndex] = NewIndex;
+		RenderIndexToInstanceIndex[NewIndex] = InstanceIndex;
+	}
+
+	/** Unique description of ISM component settings. */
 	FGeometryCollectionStaticMeshInstance MeshInstance;
+	/** Created ISM component. Will be nullptr when this slot has been recycled to FGeometryCollectionISMPool FreeList. */
+	TObjectPtr<UInstancedStaticMeshComponent> ISMComponent;
+	/** Groups of instances allocated in the ISM. */
 	FInstanceGroups InstanceGroups;
+	/** Mapping from our instance index to the ISM Component index. */
+	TArray<int32> InstanceIndexToRenderIndex;
+	/** Mapping from the ISM Component index to our instance index . */
+	TArray<int32> RenderIndexToInstanceIndex;
 };
 
-
+/** A pool of ISMs. */
 struct FGeometryCollectionISMPool
 {
 	using FISMIndex = int32;
@@ -238,6 +255,7 @@ struct FGeometryCollectionISMPool
 	bool BatchUpdateInstancesTransforms(FGeometryCollectionMeshInfo& MeshInfo, int32 StartInstanceIndex, const TArray<FTransform>& NewInstancesTransforms, bool bWorldSpace, bool bMarkRenderStateDirty, bool bTeleport);
 	void RemoveISM(const FGeometryCollectionMeshInfo& MeshInfo);
 	
+	/** Handle ISM instance updating. */
 	void OnISMInstanceIndexUpdated(UInstancedStaticMeshComponent* InComponent, TArrayView<const FInstancedStaticMeshDelegates::FInstanceIndexUpdateData> InIndexUpdates);
 
 	/** Clear all ISM components and associated data */
@@ -251,8 +269,8 @@ struct FGeometryCollectionISMPool
 
 
 /**
-* UGeometryCollectionISMPoolComponent
-*   Component that managed a pool of ISM in order to optimize render of geometry collections when not using fracture
+* UGeometryCollectionISMPoolComponent.
+* Component that manages a pool of ISM components in order to allow multiple client components that use the same meshes to the share ISMs.
 */
 UCLASS(meta = (BlueprintSpawnableComponent))
 class GEOMETRYCOLLECTIONENGINE_API UGeometryCollectionISMPoolComponent: public USceneComponent
@@ -266,6 +284,7 @@ public:
 	//~ Begin UActorComponent Interface
 	virtual void OnRegister() override;
 	virtual void OnUnregister() override;
+	virtual void GetResourceSizeEx(FResourceSizeEx& CumulativeResourceSize) override;
 	//~ End UActorComponent Interface
 
 	/** 
@@ -281,7 +300,7 @@ public:
 	/** Add a static mesh for a nmesh group */
 	FMeshId AddMeshToGroup(FMeshGroupId MeshGroupId, const FGeometryCollectionStaticMeshInstance& MeshInstance, int32 InstanceCount, TArrayView<const float> CustomDataFloats);
 
-	/** Add a static mesh for a nmesh group */
+	/** Add a static mesh for a mesh group */
 	bool BatchUpdateInstancesTransforms(FMeshGroupId MeshGroupId, FMeshId MeshId, int32 StartInstanceIndex, const TArray<FTransform>& NewInstancesTransforms, bool bWorldSpace = false, bool bMarkRenderStateDirty = false, bool bTeleport = false);
 
 	/** Instance Index updated on the InstancedStaticMeshComponent which might need to be handled by the pool instance groups */
