Fixed breaking model considering a parent multiple times when classifying required breaking particles.
Fixed children with strains not being processed when their parent cluster breaks
Fixed clustering test to actually run the tests - iterating over a moved-from array meant these weren't running

diff --git a/Engine/Source/Programs/HeadlessChaos/Private/HeadlessChaosTestClustering.cpp b/Engine/Source/Programs/HeadlessChaos/Private/HeadlessChaosTestClustering.cpp
--- a/Engine/Source/Programs/HeadlessChaos/Private/HeadlessChaosTestClustering.cpp
+++ b/Engine/Source/Programs/HeadlessChaos/Private/HeadlessChaosTestClustering.cpp
@@ -110,7 +110,8 @@ namespace ChaosTest {
 		}
 
 		FClusterCreationParameters ClusterParams;
-		Chaos::FPBDRigidParticleHandle* RootClusterHandle = Evolution.GetRigidClustering().CreateClusterParticle(0, MoveTemp(ClusterHandles), ClusterParams);
+		TArray<Chaos::FPBDRigidParticleHandle* > ClusterHandlesCopy = ClusterHandles;
+		Chaos::FPBDRigidParticleHandle* RootClusterHandle = Evolution.GetRigidClustering().CreateClusterParticle(0, MoveTemp(ClusterHandlesCopy), ClusterParams);
 		FVec3 InitialVelocity((FReal)50, (FReal)20, (FReal)100);
 
 		RootClusterHandle->V() = InitialVelocity;		
diff --git a/Engine/Source/Runtime/Experimental/Chaos/Private/Chaos/PBDRigidClustering.cpp b/Engine/Source/Runtime/Experimental/Chaos/Private/Chaos/PBDRigidClustering.cpp
--- a/Engine/Source/Runtime/Experimental/Chaos/Private/Chaos/PBDRigidClustering.cpp
+++ b/Engine/Source/Runtime/Experimental/Chaos/Private/Chaos/PBDRigidClustering.cpp
@@ -1110,24 +1110,33 @@ namespace Chaos
 			auto ProcessClusteredParticle = [&ParticlesToProcess, &bPotentialBreak, this](FPBDRigidClusteredParticleHandle* Particle)
 			{
 				TArray<FRigidHandle>& ParentToChildren = MChildren[Particle];
+
+				bool bAddParent = false;
 				for(FRigidHandle Child : ParentToChildren)
 				{
 					if(FClusterHandle ClusteredChild = Child->CastToClustered())
 					{
 						if(ClusteredChild->GetInternalStrains() <= 0.f)
 						{
-							ParticlesToProcess.Add(Particle);
-
+							bAddParent = true;
+							// #TODO remove need to set this here so we can early out as soon as we
+							// find one child that requires processing for breaks
 							ClusteredChild->CollisionImpulse() = FLT_MAX;
 							MCollisionImpulseArrayDirty = true;
 						}
 						else if(ClusteredChild->GetExternalStrain() > 0 || ClusteredChild->CollisionImpulse() > 0)
 						{
-							ParticlesToProcess.Add(Particle);
+							bAddParent = true;
 							bPotentialBreak = true;
 						}
 					}
 				}
+
+				// Ensure we only add the parent once.
+				if(bAddParent)
+				{
+					ParticlesToProcess.Add(Particle);
+				}
 			};
 
 			{
@@ -1182,6 +1191,9 @@ namespace Chaos
 
 			ClusterUnionManager.HandleDeferredClusterUnionUpdateProperties();
 
+			// Breaking can populate this again with relevant children - so we clear before running the breaking model
+			TopLevelClusterParentsStrained.Reset();
+
 			if (MCollisionImpulseArrayDirty || bPotentialBreak)
 			{
 				SCOPE_CYCLE_COUNTER(STAT_UpdateDirtyImpulses);
@@ -1197,8 +1209,6 @@ namespace Chaos
 					BreakingModel();
 				}
 			} // end if MCollisionImpulseArrayDirty
-
-			TopLevelClusterParentsStrained.Reset();
 		}
 		Timer.Stop();
 		UE_LOG(LogChaos, Verbose, TEXT("Cluster Break Update Time is %f"), Time);
