Chaos - temporary fix for potential dangling pointer issue when a GeometryCollection particle in a ClusterUnion is destroyed

diff --git a/Engine/Source/Runtime/Experimental/Chaos/Private/Chaos/ClusterUnionManager.cpp b/Engine/Source/Runtime/Experimental/Chaos/Private/Chaos/ClusterUnionManager.cpp
--- a/Engine/Source/Runtime/Experimental/Chaos/Private/Chaos/ClusterUnionManager.cpp
+++ b/Engine/Source/Runtime/Experimental/Chaos/Private/Chaos/ClusterUnionManager.cpp
@@ -115,7 +115,9 @@ namespace Chaos
 
 			if (Child->Geometry())
 			{
-				Objects.Add(TUniquePtr<FImplicitObject>(new TImplicitObjectTransformed<FReal, 3>(Child->Geometry(), Frame)));
+				FImplicitObjectTransformed* TransformedChildGeometry = new TImplicitObjectTransformed<FReal, 3>(Child->Geometry(), Frame);
+				TransformedChildGeometry->SetSharedObject(Child->SharedGeometry());
+				Objects.Add(TUniquePtr<FImplicitObject>(TransformedChildGeometry));
 			}
 		}
 
diff --git a/Engine/Source/Runtime/Experimental/Chaos/Public/Chaos/ImplicitObjectTransformed.h b/Engine/Source/Runtime/Experimental/Chaos/Public/Chaos/ImplicitObjectTransformed.h
--- a/Engine/Source/Runtime/Experimental/Chaos/Public/Chaos/ImplicitObjectTransformed.h
+++ b/Engine/Source/Runtime/Experimental/Chaos/Public/Chaos/ImplicitObjectTransformed.h
@@ -345,6 +345,8 @@ public:
 		TImplicitObjectTransformSerializeHelper(Ar, MObject);
 		Ar << MTransform;
 		TBox<T, d>::SerializeAsAABB(Ar, MLocalBoundingBox);
+
+		// NOTE: Not serializing SharedObject which is a temp fix and only used in the runtime
 	}
 
 	virtual uint32 GetTypeHash() const override
@@ -406,10 +408,21 @@ protected:
 
 private:
 	ObjectType MObject;
-	TUniquePtr<Chaos::FImplicitObject> MObjectOwner;
+	TUniquePtr<FImplicitObject> MObjectOwner;
 	TRigidTransform<T, d> MTransform;
 	TAABB<T, d> MLocalBoundingBox;
 
+	// TEMP: workaround for a clustering issue when a geometry collection in a cluster is destroyed 
+	// after the cluster's geometry has been shared between game and physics thread
+	// @todo(chaos): remove this when FImplicitObject is ref counted
+	FConstImplicitObjectPtr SharedObject;
+
+	friend class FClusterUnionManager;
+	void SetSharedObject(const FConstImplicitObjectPtr& InSharedObject)
+	{
+		SharedObject = InSharedObject;
+	}
+
 	//needed for serialization
 	TImplicitObjectTransformed() : FImplicitObject(EImplicitObject::HasBoundingBox, ImplicitObjectType::Transformed) {}
 
