remove GCs from cluster unions when they get destroyed

diff --git a/Engine/Source/Runtime/Experimental/Chaos/Private/Chaos/ClusterUnionManager.cpp b/Engine/Source/Runtime/Experimental/Chaos/Private/Chaos/ClusterUnionManager.cpp
--- a/Engine/Source/Runtime/Experimental/Chaos/Private/Chaos/ClusterUnionManager.cpp
+++ b/Engine/Source/Runtime/Experimental/Chaos/Private/Chaos/ClusterUnionManager.cpp
@@ -158,51 +158,54 @@ namespace Chaos
 	void FClusterUnionManager::FlushPendingOperations()
 	{
 		SCOPE_CYCLE_COUNTER(STAT_FlushPendingOperations);
-		if (PendingExplicitIndexOperations.IsEmpty() && PendingClusterIndexOperations.IsEmpty())
+		if (!PendingExplicitIndexOperations.IsEmpty())
 		{
-			return;
-		}
-
-		// Go through every explicit index operation and convert them into a normal cluster index operation.
-		// This could be made more efficient but shouldn't happen enough for it to really matter.
-		for (const TPair<FClusterUnionExplicitIndex, FClusterOpMap>& OpMap : PendingExplicitIndexOperations)
-		{
-			const FClusterUnionIndex UnionIndex = GetOrCreateClusterUnionIndexFromExplicitIndex(OpMap.Key);
-			for (const TPair<EClusterUnionOperation, TArray<FPBDRigidParticleHandle*>>& Op : OpMap.Value)
+			// Go through every explicit index operation and convert them into a normal cluster index operation.
+			// This could be made more efficient but shouldn't happen enough for it to really matter.
+			for (const TPair<FClusterUnionExplicitIndex, FClusterOpMap>& OpMap : PendingExplicitIndexOperations)
 			{
-				AddPendingClusterIndexOperation(UnionIndex, Op.Key, Op.Value);
+				const FClusterUnionIndex UnionIndex = GetOrCreateClusterUnionIndexFromExplicitIndex(OpMap.Key);
+				for (const TPair<EClusterUnionOperation, TArray<FPBDRigidParticleHandle*>>& Op : OpMap.Value)
+				{
+					AddPendingClusterIndexOperation(UnionIndex, Op.Key, Op.Value);
+				}
 			}
+			PendingExplicitIndexOperations.Empty();
 		}
-		PendingExplicitIndexOperations.Empty();
 
-		for (TPair<FClusterUnionIndex, FClusterOpMap>& OpMap : PendingClusterIndexOperations)
+		if (!PendingClusterIndexOperations.IsEmpty())
 		{
-			// Is this sort necessary? Better to be safe than sorry. Since we need to guarantee that the UpdateChildToParent happens after add.
-			OpMap.Value.KeyStableSort(
-				[](EClusterUnionOperation A, EClusterUnionOperation B)
-				{
-					return static_cast<int32>(A) < static_cast<int32>(B);
-				}
-			);
-
-			for (TPair<EClusterUnionOperation, TArray<FPBDRigidParticleHandle*>>& Op : OpMap.Value)
+			for (TPair<FClusterUnionIndex, FClusterOpMap>& OpMap : PendingClusterIndexOperations)
 			{
-				switch (Op.Key)
+				// Is this sort necessary? Better to be safe than sorry. Since we need to guarantee that the UpdateChildToParent happens after add.
+				OpMap.Value.KeyStableSort(
+					[](EClusterUnionOperation A, EClusterUnionOperation B)
+					{
+						return static_cast<int32>(A) < static_cast<int32>(B);
+					}
+				);
+
+				for (TPair<EClusterUnionOperation, TArray<FPBDRigidParticleHandle*>>& Op : OpMap.Value)
 				{
-				case EClusterUnionOperation::Add:
-				case EClusterUnionOperation::AddReleased:
-					HandleAddOperation(OpMap.Key, Op.Value, Op.Key == EClusterUnionOperation::AddReleased);
-					break;
-				case EClusterUnionOperation::Remove:
-					HandleRemoveOperation(OpMap.Key, Op.Value, EClusterUnionOperationTiming::Immediate);
-					break;
-				case EClusterUnionOperation::UpdateChildToParent:
-					HandleUpdateChildToParentOperation(OpMap.Key, Op.Value);
-					break;
+					switch (Op.Key)
+					{
+					case EClusterUnionOperation::Add:
+					case EClusterUnionOperation::AddReleased:
+						HandleAddOperation(OpMap.Key, Op.Value, Op.Key == EClusterUnionOperation::AddReleased);
+						break;
+					case EClusterUnionOperation::Remove:
+						HandleRemoveOperation(OpMap.Key, Op.Value, EClusterUnionOperationTiming::Immediate);
+						break;
+					case EClusterUnionOperation::UpdateChildToParent:
+						HandleUpdateChildToParentOperation(OpMap.Key, Op.Value);
+						break;
+					}
 				}
 			}
+			PendingClusterIndexOperations.Empty();
 		}
-		PendingClusterIndexOperations.Empty();
+
+		HandleDeferredClusterUnionUpdateProperties();
 	}
 
 	FClusterUnion* FClusterUnionManager::FindClusterUnionFromExplicitIndex(FClusterUnionExplicitIndex Index)
diff --git a/Engine/Source/Runtime/Experimental/Chaos/Private/PBDRigidsSolver.cpp b/Engine/Source/Runtime/Experimental/Chaos/Private/PBDRigidsSolver.cpp
--- a/Engine/Source/Runtime/Experimental/Chaos/Private/PBDRigidsSolver.cpp
+++ b/Engine/Source/Runtime/Experimental/Chaos/Private/PBDRigidsSolver.cpp
@@ -743,11 +743,17 @@ namespace Chaos
 
 		EnqueueCommandImmediate([InProxy, this]()
 			{
+				FClusterUnionManager& ClusterUnionManager = MEvolution->GetRigidClustering().GetClusterUnionManager();
 				for (FPBDRigidClusteredParticleHandle* Handle : InProxy->GetSolverParticleHandles())
 				{
 					if (Handle)
 					{
 						MEvolution->DisableParticle(Handle);
+
+						// It should be safe to defer here without an immediate call to HandleDeferredClusterUnionUpdateProperties.
+						// This change doesn't really have to go through until FClusterUnionManager::FlushPendingOperations which happens
+						// prior to trying to advance the frame.
+						ClusterUnionManager.HandleRemoveOperationWithClusterLookup({ Handle }, EClusterUnionOperationTiming::Defer);
 					}
 				}
 
