support exporting contact areas from proximity tool to connection graph for use in physics
+ change connection graph representation to be an array of graph edges, rather than a TSet per transform of connected transforms, and add associated data migration

diff --git a/Engine/Plugins/Experimental/ChaosEditor/Source/FractureEditor/Private/FractureModeSettings.cpp b/Engine/Plugins/Experimental/ChaosEditor/Source/FractureEditor/Private/FractureModeSettings.cpp
--- a/Engine/Plugins/Experimental/ChaosEditor/Source/FractureEditor/Private/FractureModeSettings.cpp
+++ b/Engine/Plugins/Experimental/ChaosEditor/Source/FractureEditor/Private/FractureModeSettings.cpp
@@ -26,6 +26,7 @@ void UFractureModeSettings::ApplyDefaultProximitySettings(FGeometryCollection& G
 	Properties.DistanceThreshold = ProximityDistanceThreshold;
 	Properties.Method = ProximityMethod;
 	Properties.ContactMethod = ProximityContactMethod;
+	Properties.ContactAreaMethod = ProximityConnectionContactAreaMethod;
 	Properties.RequireContactAmount = ProximityContactThreshold;
 	GeometryCollection.SetProximityProperties(Properties);
 }
diff --git a/Engine/Plugins/Experimental/ChaosEditor/Source/FractureEditor/Private/FractureToolProximity.cpp b/Engine/Plugins/Experimental/ChaosEditor/Source/FractureEditor/Private/FractureToolProximity.cpp
--- a/Engine/Plugins/Experimental/ChaosEditor/Source/FractureEditor/Private/FractureToolProximity.cpp
+++ b/Engine/Plugins/Experimental/ChaosEditor/Source/FractureEditor/Private/FractureToolProximity.cpp
@@ -24,6 +24,7 @@ void UFractureProximityActions::SaveAsDefaults()
 	ModeSettings->bProximityUseAsConnectionGraph = ProximityTool->ProximitySettings->bUseAsConnectionGraph;
 	ModeSettings->ProximityContactThreshold = ProximityTool->ProximitySettings->ContactThreshold;
 	ModeSettings->ProximityContactMethod = ProximityTool->ProximitySettings->ContactMethod;
+	ModeSettings->ProximityConnectionContactAreaMethod = ProximityTool->ProximitySettings->ContactAreaMethod;
 }
 
 void UFractureProximityActions::SetFromDefaults()
@@ -35,6 +36,7 @@ void UFractureProximityActions::SetFromDefaults()
 	ProximityTool->ProximitySettings->bUseAsConnectionGraph = ModeSettings->bProximityUseAsConnectionGraph;
 	ProximityTool->ProximitySettings->ContactThreshold = ModeSettings->ProximityContactThreshold;
 	ProximityTool->ProximitySettings->ContactMethod = ModeSettings->ProximityContactMethod;
+	ProximityTool->ProximitySettings->ContactAreaMethod = ModeSettings->ProximityConnectionContactAreaMethod;
 
 	ProximityTool->NotifyOfPropertyChangeByTool(this);
 }
@@ -233,6 +235,7 @@ int32 UFractureToolProximity::ExecuteFracture(const FFractureToolContext& Fractu
 		Properties.Method = ProximitySettings->Method;
 		Properties.DistanceThreshold = ProximitySettings->DistanceThreshold;
 		Properties.bUseAsConnectionGraph = ProximitySettings->bUseAsConnectionGraph;
+		Properties.ContactAreaMethod = ProximitySettings->ContactAreaMethod;
 		Properties.ContactMethod = ProximitySettings->ContactMethod;
 		Properties.RequireContactAmount = ProximitySettings->ContactThreshold;
 		
diff --git a/Engine/Plugins/Experimental/ChaosEditor/Source/FractureEditor/Private/FractureToolProximity.h b/Engine/Plugins/Experimental/ChaosEditor/Source/FractureEditor/Private/FractureToolProximity.h
--- a/Engine/Plugins/Experimental/ChaosEditor/Source/FractureEditor/Private/FractureToolProximity.h
+++ b/Engine/Plugins/Experimental/ChaosEditor/Source/FractureEditor/Private/FractureToolProximity.h
@@ -28,7 +28,7 @@ public:
 
 	// If hull-based proximity detection is enabled, amount to expand hulls when searching for overlapping neighbors
 	UPROPERTY(EditAnywhere, Category = Automatic, meta = (ClampMin = "0", 
-		EditCondition = "Method == EProximityMethod::ConvexHull || ContactMethod == EProximityContactMethod::ConvexHullSharpContact || ContactMethod == EProximityContactMethod::ConvexHullAreaContact"))
+		EditCondition = "Method == EProximityMethod::ConvexHull || ContactMethod == EProximityContactMethod::ConvexHullSharpContact || ContactMethod == EProximityContactMethod::ConvexHullAreaContact || ContactAreaMethod == EConnectionContactMethod::ConvexHullContactArea"))
 	double DistanceThreshold = 1;
 
 	// Method to use to determine the contact between two pieces, if the Contact Threshold is greater than 0
@@ -43,6 +43,10 @@ public:
 	UPROPERTY(EditAnywhere, Category = Automatic)
 	bool bUseAsConnectionGraph = false;
 
+	// Method to use for determining contact areas that define the strength of connections for destruction simulation
+	UPROPERTY(EditAnywhere, Category = Automatic)
+	EConnectionContactMethod ContactAreaMethod = EConnectionContactMethod::None;
+
 	// Whether to display the proximity graph edges
 	UPROPERTY(EditAnywhere, Category = Visualization)
 	bool bShowProximity = true;
diff --git a/Engine/Plugins/Experimental/ChaosEditor/Source/FractureEditor/Public/FractureModeSettings.h b/Engine/Plugins/Experimental/ChaosEditor/Source/FractureEditor/Public/FractureModeSettings.h
--- a/Engine/Plugins/Experimental/ChaosEditor/Source/FractureEditor/Public/FractureModeSettings.h
+++ b/Engine/Plugins/Experimental/ChaosEditor/Source/FractureEditor/Public/FractureModeSettings.h
@@ -84,8 +84,12 @@ public:
 	UPROPERTY(config, EditAnywhere, Category = "Fracture Mode|Proximity Detection Defaults", meta = (DisplayName = "Use As Connection Graph"))
 	bool bProximityUseAsConnectionGraph = false;
 
-	// Method to use to determine the contact between two pieces, if the Contact Threshold is greater than 0
-	UPROPERTY(config, EditAnywhere, Category = "Fracture Mode|Proximity Detection Defaults", meta = (DisplayName = "Contact Method"))
+	// Method to use to determine the area of the contact for transforms that are connected in the connection graph used for simulation. Only used if "Use As Connection Graph" is enabled.
+	UPROPERTY(config, EditAnywhere, Category = "Fracture Mode|Proximity Detection Defaults", meta = (DisplayName = "Connection Contact Method"))
+	EConnectionContactMethod ProximityConnectionContactAreaMethod = EConnectionContactMethod::None;
+
+	// Method to use to determine the contact between two pieces, if the Contact Threshold is greater than 0, for the purpose of filtering out too-small contacts
+	UPROPERTY(config, EditAnywhere, Category = "Fracture Mode|Proximity Detection Defaults", meta = (DisplayName = "Contact Filter Method"))
 	EProximityContactMethod ProximityContactMethod = EProximityContactMethod::MinOverlapInProjectionToMajorAxes;
 
 	// If greater than zero, proximity will be additionally filtered by a 'contact' threshold, in cm, to exclude grazing / corner proximity
diff --git a/Engine/Plugins/Experimental/GeometryCollectionPlugin/Source/GeometryCollectionNodes/Private/Dataflow/GeometryCollectionNodes.cpp b/Engine/Plugins/Experimental/GeometryCollectionPlugin/Source/GeometryCollectionNodes/Private/Dataflow/GeometryCollectionNodes.cpp
--- a/Engine/Plugins/Experimental/GeometryCollectionPlugin/Source/GeometryCollectionNodes/Private/Dataflow/GeometryCollectionNodes.cpp
+++ b/Engine/Plugins/Experimental/GeometryCollectionPlugin/Source/GeometryCollectionNodes/Private/Dataflow/GeometryCollectionNodes.cpp
@@ -769,6 +769,7 @@ void FProximityDataflowNode::Evaluate(Dataflow::FContext& Context, const FDatafl
 			Properties.ContactMethod = (EProximityContactMethod)FilterContactMethod;
 			Properties.DistanceThreshold = DistanceThreshold;
 			Properties.bUseAsConnectionGraph = bUseAsConnectionGraph;
+			Properties.ContactAreaMethod = (EConnectionContactMethod)ContactAreaMethod;
 			Properties.RequireContactAmount = ContactThreshold;
 
 			GeomCollection->SetProximityProperties(Properties);
diff --git a/Engine/Plugins/Experimental/GeometryCollectionPlugin/Source/GeometryCollectionNodes/Public/Dataflow/GeometryCollectionNodes.h b/Engine/Plugins/Experimental/GeometryCollectionPlugin/Source/GeometryCollectionNodes/Public/Dataflow/GeometryCollectionNodes.h
--- a/Engine/Plugins/Experimental/GeometryCollectionPlugin/Source/GeometryCollectionNodes/Public/Dataflow/GeometryCollectionNodes.h
+++ b/Engine/Plugins/Experimental/GeometryCollectionPlugin/Source/GeometryCollectionNodes/Public/Dataflow/GeometryCollectionNodes.h
@@ -1114,6 +1114,19 @@ enum class EProximityContactFilteringMethodEnum : uint8
 	Dataflow_Max                UMETA(Hidden)
 };
 
+UENUM(BlueprintType)
+enum class EConnectionContactAreaMethodEnum : uint8
+{
+	/** Do not compute contact areas */
+	Dataflow_ConnectionContactAreaMethod_None UMETA(DisplayName = "None"),
+	/** Compute approximate contact surface area via the intersection of convex hulls (allowing for optional offset) */
+	Dataflow_ProximityContactFilteringMethod_ConvexHullArea UMETA(DisplayName = "Convex Hull Area Contact"),
+	//~~~
+	//256th entry
+	Dataflow_Max                UMETA(Hidden)
+};
+
+
 /**
  *
  * Update the proximity (contact) graph for the bones in a Collection
@@ -1132,7 +1145,7 @@ public:
 
 	/** If hull-based proximity detection is enabled, amount to expand hulls when searching for overlapping neighbors */
 	UPROPERTY(EditAnywhere, Category = "Proximity", meta = (ClampMin = "0", 
-		EditCondition = "ProximityMethod == EProximityMethodEnum::Dataflow_ProximityMethod_ConvexHull || FilterContactMethod == EProximityContactFilteringMethodEnum::Dataflow_ProximityContactFilteringMethod_ConvexHullSharp || FilterContactMethod == EProximityContactFilteringMethodEnum::Dataflow_ProximityContactFilteringMethod_ConvexHullArea"))
+		EditCondition = "ProximityMethod == EProximityMethodEnum::Dataflow_ProximityMethod_ConvexHull || FilterContactMethod == EProximityContactFilteringMethodEnum::Dataflow_ProximityContactFilteringMethod_ConvexHullSharp || FilterContactMethod == EProximityContactFilteringMethodEnum::Dataflow_ProximityContactFilteringMethod_ConvexHullArea || ContactAreaMethod = EConnectionContactAreaMethodEnum::Dataflow_ProximityContactFilteringMethod_ConvexHullArea"))
 	float DistanceThreshold = 1;
 
 	// If greater than zero, proximity will be additionally filtered by a 'contact' threshold, in cm, to exclude grazing / corner proximity
@@ -1147,6 +1160,10 @@ public:
 	UPROPERTY(EditAnywhere, Category = "Proximity")
 	bool bUseAsConnectionGraph = false;
 
+	/** The method used to compute contact areas for simulation purposes (only when 'Use As Connection Graph' is enabled) */
+	UPROPERTY(EditAnywhere, Category = "Proximity")
+	EConnectionContactAreaMethodEnum ContactAreaMethod = EConnectionContactAreaMethodEnum::Dataflow_ConnectionContactAreaMethod_None;
+
 	/** GeometryCollection to update the proximity graph on */
 	UPROPERTY(meta = (DataflowInput, DataflowOutput, DataflowPassthrough = "Collection", DataflowIntrinsic))
 	FManagedArrayCollection Collection;
diff --git a/Engine/Source/Runtime/Core/Public/UObject/FortniteMainBranchObjectVersion.h b/Engine/Source/Runtime/Core/Public/UObject/FortniteMainBranchObjectVersion.h
--- a/Engine/Source/Runtime/Core/Public/UObject/FortniteMainBranchObjectVersion.h
+++ b/Engine/Source/Runtime/Core/Public/UObject/FortniteMainBranchObjectVersion.h
@@ -295,6 +295,9 @@ struct CORE_API FFortniteMainBranchObjectVersion
 		// Fix the Nanite landscape mesh non-deterministic DDC keys
 		FixNaniteLandscapeMeshDDCKey,
 
+		// Change how connection graphs are stored on Geometry Collections to an edge-array representation
+		ChaosGeometryCollectionConnectionEdgeGroup,
+
 		// -----<new versions can be added above this line>-------------------------------------------------
 		VersionPlusOne,
 		LatestVersion = VersionPlusOne - 1
diff --git a/Engine/Source/Runtime/Experimental/Chaos/Private/GeometryCollection/Facades/CollectionConnectionGraphFacade.cpp b/Engine/Source/Runtime/Experimental/Chaos/Private/GeometryCollection/Facades/CollectionConnectionGraphFacade.cpp
--- a/Engine/Source/Runtime/Experimental/Chaos/Private/GeometryCollection/Facades/CollectionConnectionGraphFacade.cpp
+++ b/Engine/Source/Runtime/Experimental/Chaos/Private/GeometryCollection/Facades/CollectionConnectionGraphFacade.cpp
@@ -7,8 +7,15 @@
 namespace GeometryCollection::Facades
 {
 
+	static const FName ConnectionGroupName = "ConnectionEdge";
+
+// Disable deprecations so that we can initialize the ConnectionsAttribute, though we no longer expect it to be on the Collection
+PRAGMA_DISABLE_DEPRECATION_WARNINGS
 	FCollectionConnectionGraphFacade::FCollectionConnectionGraphFacade(FManagedArrayCollection& InCollection)
 		: ConnectionsAttribute(InCollection, "Connections", FTransformCollection::TransformGroup)
+		, ConnectionEdgeStartAttribute(InCollection, "ConnectionEdgeStarts", ConnectionGroupName, FTransformCollection::TransformGroup)
+		, ConnectionEdgeEndAttribute(InCollection, "ConnectionEdgeEnds", ConnectionGroupName, FTransformCollection::TransformGroup)
+		, ConnectionEdgeContactAttribute(InCollection, "ConnectionEdgeContacts", ConnectionGroupName)
 #if UE_BUILD_DEBUG
 		, ParentAttribute(InCollection, "Parent", FTransformCollection::TransformGroup)
 #endif
@@ -17,44 +24,154 @@ namespace GeometryCollection::Facades
 
 	FCollectionConnectionGraphFacade::FCollectionConnectionGraphFacade(const FManagedArrayCollection& InCollection)
 		: ConnectionsAttribute(InCollection, "Connections", FTransformCollection::TransformGroup)
+		, ConnectionEdgeStartAttribute(InCollection, "ConnectionEdgeStarts", ConnectionGroupName, FTransformCollection::TransformGroup)
+		, ConnectionEdgeEndAttribute(InCollection, "ConnectionEdgeEnds", ConnectionGroupName, FTransformCollection::TransformGroup)
+		, ConnectionEdgeContactAttribute(InCollection, "ConnectionEdgeContacts", ConnectionGroupName)
 #if UE_BUILD_DEBUG
 		, ParentAttribute(InCollection, "Parent", FTransformCollection::TransformGroup)
 #endif
 	{
 	}
+PRAGMA_ENABLE_DEPRECATION_WARNINGS
 
 	bool FCollectionConnectionGraphFacade::IsValid() const
 	{
-		return ConnectionsAttribute.IsValid();
+		return ConnectionEdgeStartAttribute.IsValid()
+			&& ConnectionEdgeEndAttribute.IsValid();
+
 	}
 
 	void FCollectionConnectionGraphFacade::DefineSchema()
 	{
 		check(!IsConst());
-
-		ConnectionsAttribute.Add();
+		ConnectionEdgeStartAttribute.Add();
+		ConnectionEdgeEndAttribute.Add();
 	}
 	
 	void FCollectionConnectionGraphFacade::ClearAttributes()
 	{
 		check(!IsConst());
 
-		ConnectionsAttribute.Remove();
+		ConnectionEdgeStartAttribute.Remove();
+		ConnectionEdgeEndAttribute.Remove();
+	}
+
+	void FCollectionConnectionGraphFacade::ReserveAdditionalConnections(int32 NumAdditionalConnections)
+	{
+		check(!IsConst());
+
+		ConnectionEdgeStartAttribute.GetCollection()->Reserve(ConnectionEdgeStartAttribute.Num() + NumAdditionalConnections, ConnectionGroupName);
 	}
 
 	void FCollectionConnectionGraphFacade::Connect(int32 BoneA, int32 BoneB)
 	{
+		if (ConnectionEdgeContactAttribute.IsValid()) // if we have a contact attribute, also set it with a default value
+		{
+			constexpr float DefaultContact = 1.0f;
+			ConnectWithContact(BoneA, BoneB, DefaultContact);
+		}
+
 		check(!IsConst());
 
 #if UE_BUILD_DEBUG
 		// Expect to always have a parent array and connected bones should always have the same parent
 		checkSlow(ParentAttribute.IsValid() && ParentAttribute.Get()[BoneA] == ParentAttribute.Get()[BoneB]);
 #endif
-		TManagedArray<TSet<int32>>& Connections = ConnectionsAttribute.Modify();
-		Connections[BoneA].Add(BoneB);
-		Connections[BoneB].Add(BoneA);
+
+		TManagedArray<int32>& Starts = ConnectionEdgeStartAttribute.Modify();
+		TManagedArray<int32>& Ends = ConnectionEdgeEndAttribute.Modify();
+		int32 NewEdgeIdx = ConnectionEdgeStartAttribute.GetCollection()->AddElements(1, ConnectionGroupName);
+		Starts[NewEdgeIdx] = BoneA;
+		Ends[NewEdgeIdx] = BoneB;
+	}
+
+	void FCollectionConnectionGraphFacade::EnableContactAreas(bool bEnable, float DefaultValue)
+	{
+		if (bEnable)
+		{
+			ConnectionEdgeContactAttribute.AddAndFill(DefaultValue);
+		}
+		else
+		{
+			ConnectionEdgeContactAttribute.Remove();
+		}
+	}
+
+	void FCollectionConnectionGraphFacade::ConnectWithContact(int32 BoneA, int32 BoneB, float ContactArea)
+	{
+		check(!IsConst());
+
+		if (!ConnectionEdgeContactAttribute.IsValid())
+		{
+			EnableContactAreas(true);
+		}
+
+#if UE_BUILD_DEBUG
+		// Expect to always have a parent array and connected bones should always have the same parent
+		checkSlow(ParentAttribute.IsValid() && ParentAttribute[BoneA] == ParentAttribute[BoneB]);
+#endif
+
+		TManagedArray<int32>& Starts = ConnectionEdgeStartAttribute.Modify();
+		TManagedArray<int32>& Ends = ConnectionEdgeEndAttribute.Modify();
+		TManagedArray<float>& Contacts = ConnectionEdgeContactAttribute.Modify();
+		int32 NewEdgeIdx = ConnectionEdgeStartAttribute.GetCollection()->AddElements(1, ConnectionGroupName);
+		Starts[NewEdgeIdx] = BoneA;
+		Ends[NewEdgeIdx] = BoneB;
+		Contacts[NewEdgeIdx] = ContactArea;
+	}
+
+	void FCollectionConnectionGraphFacade::ResetConnections()
+	{
+		check(!IsConst());
+
+		ConnectionEdgeStartAttribute.GetCollection()->EmptyGroup(ConnectionGroupName);
+	}
+
+	TPair<int32,int32> FCollectionConnectionGraphFacade::GetConnection(int32 ConnectionIndex) const
+	{
+		return TPair<int32, int32>(
+			ConnectionEdgeStartAttribute[ConnectionIndex],
+			ConnectionEdgeEndAttribute[ConnectionIndex]
+		);
+	}
+
+	float FCollectionConnectionGraphFacade::GetConnectionContactArea(int32 ConnectionIndex) const
+	{
+		return ConnectionEdgeContactAttribute[ConnectionIndex];
+	}
+
+	bool FCollectionConnectionGraphFacade::HasContactAreas() const
+	{
+		return ConnectionEdgeContactAttribute.IsValid();
 	}
 
+	int32 FCollectionConnectionGraphFacade::NumConnections() const
+	{
+		return ConnectionEdgeStartAttribute.GetConstCollection().NumElements(ConnectionGroupName);
+	}
+
+	bool FCollectionConnectionGraphFacade::HasValidConnections() const
+	{
+		int32 NumTransforms = ConnectionEdgeStartAttribute.GetConstCollection().NumElements(FTransformCollection::TransformGroup);
+		const TManagedArray<int32>& Starts = ConnectionEdgeStartAttribute.Get();
+		const TManagedArray<int32>& Ends = ConnectionEdgeEndAttribute.Get();
+		for (int32 Idx = 0; Idx < Starts.Num(); ++Idx)
+		{
+			// All indices should be valid indices into the transforms group, and start and end should always be different
+			if (Starts[Idx] < 0 || Starts[Idx] >= NumTransforms || Ends[Idx] < 0 || Ends[Idx] >= NumTransforms || Starts[Idx] == Ends[Idx])
+			{
+				return false;
+			}
+#if UE_BUILD_DEBUG
+			// Expect to always have a parent array and connected bones should always have the same parent
+			if (!ParentAttribute.IsValid() || ParentAttribute.Get()[Starts[Idx]] != ParentAttribute.Get()[Ends[Idx]])
+			{
+				return false;
+			}
+#endif
+		}
+		return true;
+	}
 };
 
 
diff --git a/Engine/Source/Runtime/Experimental/Chaos/Private/GeometryCollection/Facades/CollectionHierarchyFacade.cpp b/Engine/Source/Runtime/Experimental/Chaos/Private/GeometryCollection/Facades/CollectionHierarchyFacade.cpp
--- a/Engine/Source/Runtime/Experimental/Chaos/Private/GeometryCollection/Facades/CollectionHierarchyFacade.cpp
+++ b/Engine/Source/Runtime/Experimental/Chaos/Private/GeometryCollection/Facades/CollectionHierarchyFacade.cpp
@@ -65,6 +65,11 @@ namespace Chaos::Facades
 		return ChildrenArray;
 	}
 
+	int32 FCollectionHierarchyFacade::GetParent(int32 TransformIndex) const
+	{
+		return ParentAttribute[TransformIndex];
+	}
+
 	int32 FCollectionHierarchyFacade::GetInitialLevel(int32 TransformIndex) const
 	{
 		if (HasLevelAttribute())
diff --git a/Engine/Source/Runtime/Experimental/Chaos/Private/GeometryCollection/GeometryCollection.cpp b/Engine/Source/Runtime/Experimental/Chaos/Private/GeometryCollection/GeometryCollection.cpp
--- a/Engine/Source/Runtime/Experimental/Chaos/Private/GeometryCollection/GeometryCollection.cpp
+++ b/Engine/Source/Runtime/Experimental/Chaos/Private/GeometryCollection/GeometryCollection.cpp
@@ -1359,6 +1359,32 @@ void FGeometryCollection::Serialize(Chaos::FChaosArchive& Ar)
 			}
 		}
 
+		if (Ar.CustomVer(FFortniteMainBranchObjectVersion::GUID) < FFortniteMainBranchObjectVersion::ChaosGeometryCollectionConnectionEdgeGroup)
+		{
+			// Migrate old Connections TSet<int32> data to arrays of edge data
+			// Note we intentionally do *not* use the facade here, as the migration is specific to how the data is at the current moment, and the facade may change w/ future data changes.
+			if (TManagedArray<TSet<int32>>* Connections = FindAttribute<TSet<int32>>("Connections", TransformGroup))
+			{
+				const FName ConnectionGroupName = "ConnectionEdge";
+				AddGroup(ConnectionGroupName);
+				TManagedArray<int32>& Starts = AddAttribute<int32>("ConnectionEdgeStarts", ConnectionGroupName, FConstructionParameters(FTransformCollection::TransformGroup, true));
+				TManagedArray<int32>& Ends = AddAttribute<int32>("ConnectionEdgeEnds", ConnectionGroupName, FConstructionParameters(FTransformCollection::TransformGroup, true));
+				for (int32 TransformIdx = 0; TransformIdx < NumElements(TransformGroup); ++TransformIdx)
+				{
+					for (int32 NbrIdx : (*Connections)[TransformIdx])
+					{
+						if (TransformIdx < NbrIdx)
+						{
+							int32 EdgeIdx = AddElements(1, ConnectionGroupName);
+							Starts[EdgeIdx] = TransformIdx;
+							Ends[EdgeIdx] = NbrIdx;
+						}
+					}
+				}
+				RemoveAttribute("Connections", TransformGroup);
+			}
+		}
+
 		// Finally, make sure expected interfaces are initialized
 		InitializeInterfaces();
 	}
diff --git a/Engine/Source/Runtime/Experimental/Chaos/Private/GeometryCollection/GeometryCollectionConnectionGraphUtility.cpp b/Engine/Source/Runtime/Experimental/Chaos/Private/GeometryCollection/GeometryCollectionConnectionGraphUtility.cpp
--- a/Engine/Source/Runtime/Experimental/Chaos/Private/GeometryCollection/GeometryCollectionConnectionGraphUtility.cpp
+++ b/Engine/Source/Runtime/Experimental/Chaos/Private/GeometryCollection/GeometryCollectionConnectionGraphUtility.cpp
@@ -2,6 +2,7 @@
 
 #include "GeometryCollection/GeometryCollectionConnectionGraphUtility.h"
 #include "GeometryCollection/GeometryCollection.h"
+#include "GeometryCollection/Facades/CollectionConnectionGraphFacade.h"
 #include "Containers/Set.h"
 #include "Voronoi/Voronoi.h"
 
@@ -310,13 +311,15 @@ void FGeometryCollectionConnectionGraphGenerator::UpdateConnectivityGraphUsingPo
 
 void FGeometryCollectionConnectionGraphGenerator::CommitToCollection(FConnectionGraph& Graph, FGeometryCollection& Collection, int32 ClusterTransformIndex) 
 {
-	// first make sure we have the connection array
-	TManagedArray<FConnections>& ConnectionsArray = Collection.AddAttribute<FConnections>("Connections", FGeometryCollection::TransformGroup);
+	GeometryCollection::Facades::FCollectionConnectionGraphFacade ConnectionFacade(Collection);
+	ConnectionFacade.DefineSchema();
 
-	const TSet<int32>& Children = Collection.Children[ClusterTransformIndex];
-	for (int ChildIndex=0; ChildIndex<Children.Num(); ++ChildIndex)
+	for (int32 ChildTransformIndex : Collection.Children[ClusterTransformIndex])
 	{
-		const int32 ChildTransformIndex = Children[FSetElementId::FromInteger(ChildIndex)]; 
-		ConnectionsArray[ChildTransformIndex] = MoveTemp(Graph[ChildIndex]); 
+		ConnectionFacade.ReserveAdditionalConnections(Graph[ChildTransformIndex].Num());
+		for (int32 ChildNbr : Graph[ChildTransformIndex])
+		{
+			ConnectionFacade.Connect(ChildTransformIndex, ChildNbr);
+		}
 	}
 }
diff --git a/Engine/Source/Runtime/Experimental/Chaos/Private/GeometryCollection/GeometryCollectionProximityPropertiesInterface.cpp b/Engine/Source/Runtime/Experimental/Chaos/Private/GeometryCollection/GeometryCollectionProximityPropertiesInterface.cpp
--- a/Engine/Source/Runtime/Experimental/Chaos/Private/GeometryCollection/GeometryCollectionProximityPropertiesInterface.cpp
+++ b/Engine/Source/Runtime/Experimental/Chaos/Private/GeometryCollection/GeometryCollectionProximityPropertiesInterface.cpp
@@ -7,6 +7,7 @@ const FName FGeometryCollectionProximityPropertiesInterface::ProximityProperties
 const FName FGeometryCollectionProximityPropertiesInterface::ProximityDistanceThreshold = "DistanceThreshold";
 const FName FGeometryCollectionProximityPropertiesInterface::ProximityDetectionMethod = "DetectionMethod";
 const FName FGeometryCollectionProximityPropertiesInterface::ProximityAsConnectionGraph = "AsConnectionGraph";
+const FName FGeometryCollectionProximityPropertiesInterface::ProximityConnectionGraphContactAreaMethod = "ConnectionGraphContactAmount";
 const FName FGeometryCollectionProximityPropertiesInterface::ProximityRequireContactAmount = "RequireContactAmount";
 const FName FGeometryCollectionProximityPropertiesInterface::ProximityContactMethod = "ContactMethod";
 
@@ -46,6 +47,11 @@ FGeometryCollectionProximityPropertiesInterface::InitializeInterface()
 	{
 		ManagedCollection->AddAttribute<float>(ProximityRequireContactAmount, ProximityPropertiesGroup);
 	}
+
+	if (!ManagedCollection->HasAttribute(ProximityConnectionGraphContactAreaMethod, ProximityPropertiesGroup))
+	{
+		ManagedCollection->AddAttribute<uint8>(ProximityConnectionGraphContactAreaMethod, ProximityPropertiesGroup);
+	}
 }
 
 void 
@@ -68,6 +74,7 @@ FGeometryCollectionProximityPropertiesInterface::GetProximityProperties() const
 	if (bHasProximityProperties)
 	{
 		const TManagedArray<bool>& AsConnectionGraph = ManagedCollection->GetAttribute<bool>(ProximityAsConnectionGraph, ProximityPropertiesGroup);
+		const TManagedArray<uint8>& ConnectionContactMethod = ManagedCollection->GetAttribute<uint8>(ProximityConnectionGraphContactAreaMethod, ProximityPropertiesGroup);
 		const TManagedArray<int32>& DetectionMethod = ManagedCollection->GetAttribute<int32>(ProximityDetectionMethod, ProximityPropertiesGroup);
 		const TManagedArray<float>& DistanceThreshold = ManagedCollection->GetAttribute<float>(ProximityDistanceThreshold, ProximityPropertiesGroup);
 		const TManagedArray<uint8>& ContactMethod = ManagedCollection->GetAttribute<uint8>(ProximityContactMethod, ProximityPropertiesGroup);
@@ -75,6 +82,7 @@ FGeometryCollectionProximityPropertiesInterface::GetProximityProperties() const
 
 		constexpr int32 DefaultIndex = 0;
 		Properties.bUseAsConnectionGraph = AsConnectionGraph[DefaultIndex];
+		Properties.ContactAreaMethod = (EConnectionContactMethod)ConnectionContactMethod[DefaultIndex];
 		Properties.DistanceThreshold = DistanceThreshold[DefaultIndex];
 		Properties.Method = (EProximityMethod)DetectionMethod[DefaultIndex];
 		Properties.ContactMethod = (EProximityContactMethod)ContactMethod[DefaultIndex];
@@ -98,12 +106,14 @@ FGeometryCollectionProximityPropertiesInterface::SetProximityProperties(const FP
 	}
 
 	TManagedArray<bool>& AsConnectionGraph = ManagedCollection->ModifyAttribute<bool>(ProximityAsConnectionGraph, ProximityPropertiesGroup);
+	TManagedArray<uint8>& ConnectionContactMethod = ManagedCollection->ModifyAttribute<uint8>(ProximityConnectionGraphContactAreaMethod, ProximityPropertiesGroup);
 	TManagedArray<int32>& Method = ManagedCollection->ModifyAttribute<int32>(ProximityDetectionMethod, ProximityPropertiesGroup);
 	TManagedArray<float>& DistanceThreshold = ManagedCollection->ModifyAttribute<float>(ProximityDistanceThreshold, ProximityPropertiesGroup);
 	TManagedArray<uint8>& ContactMethod = ManagedCollection->ModifyAttribute<uint8>(ProximityContactMethod, ProximityPropertiesGroup);
 	TManagedArray<float>& ContactAmount = ManagedCollection->ModifyAttribute<float>(ProximityRequireContactAmount, ProximityPropertiesGroup);
 
 	AsConnectionGraph[AttributeIndex] = InProximityAttributes.bUseAsConnectionGraph;
+	ConnectionContactMethod[AttributeIndex] = (uint8)InProximityAttributes.ContactAreaMethod;
 	DistanceThreshold[AttributeIndex] = InProximityAttributes.DistanceThreshold;
 	Method[AttributeIndex] = (int32)InProximityAttributes.Method;
 	ContactMethod[AttributeIndex] = (uint8)InProximityAttributes.ContactMethod;
diff --git a/Engine/Source/Runtime/Experimental/Chaos/Private/GeometryCollection/GeometryCollectionProximityUtility.cpp b/Engine/Source/Runtime/Experimental/Chaos/Private/GeometryCollection/GeometryCollectionProximityUtility.cpp
--- a/Engine/Source/Runtime/Experimental/Chaos/Private/GeometryCollection/GeometryCollectionProximityUtility.cpp
+++ b/Engine/Source/Runtime/Experimental/Chaos/Private/GeometryCollection/GeometryCollectionProximityUtility.cpp
@@ -3,6 +3,7 @@
 #include "GeometryCollection/GeometryCollectionProximityUtility.h"
 
 #include "GeometryCollection/Facades/CollectionConnectionGraphFacade.h"
+#include "GeometryCollection/Facades/CollectionHierarchyFacade.h"
 #include "GeometryCollection/GeometryCollectionClusteringUtility.h"
 #include "GeometryCollection/GeometryCollectionConvexUtility.h"
 #include "Chaos/Convex.h"
@@ -529,7 +530,7 @@ void FGeometryCollectionProximityUtility::ClearConnectionGraph()
 	ConnectionsFacade.ClearAttributes();
 }
 
-void FGeometryCollectionProximityUtility::CopyProximityToConnectionGraph()
+void FGeometryCollectionProximityUtility::CopyProximityToConnectionGraph(const TArray<FGeometryContactEdge>* ContactEdges)
 {
 	if (!Collection->HasAttribute("Proximity", FGeometryCollection::GeometryGroup))
 	{
@@ -537,61 +538,103 @@ void FGeometryCollectionProximityUtility::CopyProximityToConnectionGraph()
 		return;
 	}
 
-	const TManagedArray<TSet<int32>>& Proximity = Collection->GetAttribute<TSet<int32>>("Proximity", FGeometryCollection::GeometryGroup);
 	GeometryCollection::Facades::FCollectionConnectionGraphFacade ConnectionsFacade(*Collection);
 	ConnectionsFacade.DefineSchema();
-	TManagedArray<TSet<int32>>& Connections = ConnectionsFacade.ConnectionsAttribute.Modify();
-	Connections.Fill(TSet<int32>());
+	ConnectionsFacade.ResetConnections();
 
-	int32 NumBones = Connections.Num();
-	TArray<int32> Depths;
-	Depths.SetNumZeroed(NumBones);
-	for (int32 BoneIdx = 0; BoneIdx < NumBones; ++BoneIdx)
+	Chaos::Facades::FCollectionHierarchyFacade HierarchyFacade(*Collection);
+	HierarchyFacade.GenerateLevelAttribute();
+
+	TMap<UE::Geometry::FIndex2i, float> ContactSum;
+	// Helper to accumulate contact from a geometry connection to the appropriate transform connection
+	auto ProcessGeometryContact = [this, &ConnectionsFacade, &HierarchyFacade, &ContactSum](int32 GeometryA, int32 GeometryB, bool bAccumulateContactAreas = false, float Contact = 1.f)
+	{
+		int32 Bones[2]{ Collection->TransformIndex[GeometryA], Collection->TransformIndex[GeometryB] };
+		int32 Levels[2]{ HierarchyFacade.GetInitialLevel(Bones[0]), HierarchyFacade.GetInitialLevel(Bones[1]) };
+		int32 Parents[2]{ HierarchyFacade.GetParent(Bones[0]), HierarchyFacade.GetParent(Bones[1]) };
+		while (Parents[0] != INDEX_NONE && Parents[1] != INDEX_NONE)
+		{
+			if (Bones[0] == Bones[1]) // stop if both ends are the same bone
+			{
+				break;
+			}
+			if (Parents[0] == Parents[1])
+			{
+				UE::Geometry::FIndex2i ConnectBones(Bones[0], Bones[1]);
+				if (Bones[0] > Bones[1])
+				{
+					Swap(ConnectBones.A, ConnectBones.B);
+				}
+				float* FoundContact = ContactSum.Find(ConnectBones);
+				if (!FoundContact)
+				{
+					ContactSum.Add(ConnectBones, Contact);
+				}
+				else if (bAccumulateContactAreas)
+				{
+					*FoundContact += Contact;
+				}
+				break;
+			}
+
+			bool RaiseLevels[2]{ Levels[0] >= Levels[1], Levels[1] >= Levels[0] };
+			for (int32 Idx = 0; Idx < 2; ++Idx)
+			{
+				if (RaiseLevels[Idx])
+				{
+					Bones[Idx] = Parents[Idx];
+					Parents[Idx] = HierarchyFacade.GetParent(Bones[Idx]);
+					Levels[Idx]--;
+				}
+			}
+		}
+	};
+
+	FGeometryCollectionProximityPropertiesInterface::FProximityProperties Properties = Collection->GetProximityProperties();
+	if (Properties.ContactAreaMethod == EConnectionContactMethod::ConvexHullContactArea)
 	{
-		if (Collection->SimulationType[BoneIdx] == FGeometryCollection::ESimulationTypes::FST_None)
+		ConnectionsFacade.EnableContactAreas(true);
+		TArray<FGeometryContactEdge> LocalContactEdges;
+		const TArray<FGeometryContactEdge>* UseContactEdges = ContactEdges;
+		if (!UseContactEdges)
 		{
-			Depths[BoneIdx] = -1;
-			continue;
+			FGeometryCollectionConvexPropertiesInterface::FConvexCreationProperties ConvexProperties = Collection->GetConvexProperties();
+			TArray<FTransform> GlobalTransformArray;
+			GeometryCollectionAlgo::GlobalMatrices(Collection->Transform, Collection->Parent, GlobalTransformArray);
+			UE::GeometryCollectionConvexUtility::FConvexHulls ComputedHulls = FGeometryCollectionConvexUtility::ComputeLeafHulls(Collection, GlobalTransformArray, ConvexProperties.SimplificationThreshold, 0.0f /* Never shrink hulls for proximity detection */);
+
+			LocalContactEdges = ComputeConvexGeometryContactFromProximity(Collection, Properties.DistanceThreshold, ComputedHulls);
+
+			UseContactEdges = &LocalContactEdges;
 		}
-		int32 Depth = 0, WalkParent = BoneIdx;
-		while (Collection->Parent[WalkParent] != INDEX_NONE)
+		for (const FGeometryContactEdge& Edge : *UseContactEdges)
 		{
-			Depth++;
-			WalkParent = Collection->Parent[WalkParent];
+			ProcessGeometryContact(Edge.GeometryIndices[0], Edge.GeometryIndices[1], true, Edge.ContactArea);
 		}
-		Depths[BoneIdx] = Depth;
 	}
-	TArray<int32> AllLeaves;
-	for (int32 BoneIdx = 0; BoneIdx < NumBones; ++BoneIdx)
+	else // No contact areas
 	{
-		if (Collection->SimulationType[BoneIdx] == FGeometryCollection::ESimulationTypes::FST_None)
+		ConnectionsFacade.EnableContactAreas(false);
+		const TManagedArray<TSet<int32>>& Proximity = Collection->GetAttribute<TSet<int32>>("Proximity", FGeometryCollection::GeometryGroup);
+		for (int32 GeoIdx = 0; GeoIdx < Proximity.Num(); ++GeoIdx)
 		{
-			continue;
+			for (int32 NbrIdx : Proximity[GeoIdx])
+			{
+				ProcessGeometryContact(GeoIdx, NbrIdx, false);
+			}
 		}
-		int32 BoneDepth = Depths[BoneIdx];
-		int32 BoneParent = Collection->Parent[BoneIdx];
-		AllLeaves.Reset();
+	}
 
-		FGeometryCollectionClusteringUtility::GetLeafBones(Collection, BoneIdx, true, AllLeaves);
-		for (int32 LeafBone : AllLeaves)
+	// Copy out final contacts
+	for (TPair<UE::Geometry::FIndex2i, float> Contact : ContactSum)
+	{
+		if (ConnectionsFacade.HasContactAreas())
 		{
-			int32 LeafGeo = Collection->TransformToGeometryIndex[LeafBone];
-			for (int32 NbrGeo : Proximity[LeafGeo])
-			{
-				int32 NbrBone = Collection->TransformIndex[NbrGeo];
-				if (Depths[NbrBone] < BoneDepth)
-				{
-					continue; // cluster is closer to root than us, ignore it
-				}
-				while (NbrBone != INDEX_NONE && Depths[NbrBone] > BoneDepth)
-				{
-					NbrBone = Collection->Parent[NbrBone];
-				}
-				if (NbrBone != INDEX_NONE && NbrBone != BoneIdx && Collection->Parent[NbrBone] == BoneParent)
-				{
-					ConnectionsFacade.Connect(BoneIdx, NbrBone);
-				}
-			}
+			ConnectionsFacade.ConnectWithContact(Contact.Key.A, Contact.Key.B, Contact.Value);
+		}
+		else
+		{
+			ConnectionsFacade.Connect(Contact.Key.A, Contact.Key.B);
 		}
 	}
 }
@@ -602,8 +645,11 @@ void FGeometryCollectionProximityUtility::UpdateProximity(UE::GeometryCollection
 
 	FGeometryCollectionProximityPropertiesInterface::FProximityProperties Properties = Collection->GetProximityProperties();
 
-	bool bWantConvexContactEdges = Properties.RequireContactAmount > 0.0f &&
-		(Properties.ContactMethod == EProximityContactMethod::ConvexHullSharpContact || Properties.ContactMethod == EProximityContactMethod::ConvexHullAreaContact);
+	bool bWantConvexContactEdges = 
+		(Properties.RequireContactAmount > 0.0f &&
+		(Properties.ContactMethod == EProximityContactMethod::ConvexHullSharpContact || Properties.ContactMethod == EProximityContactMethod::ConvexHullAreaContact))
+		||
+		(Properties.ContactAreaMethod == EConnectionContactMethod::ConvexHullContactArea);
 	bool bWantLocalHulls = Properties.Method == EProximityMethod::ConvexHull || bWantConvexContactEdges;
 
 	FGeometryCollectionConvexPropertiesInterface::FConvexCreationProperties ConvexProperties = Collection->GetConvexProperties();
@@ -735,7 +781,7 @@ void FGeometryCollectionProximityUtility::UpdateProximity(UE::GeometryCollection
 
 	if (Properties.bUseAsConnectionGraph)
 	{
-		CopyProximityToConnectionGraph();
+		CopyProximityToConnectionGraph(Properties.ContactAreaMethod == EConnectionContactMethod::ConvexHullContactArea  ? &ContactEdges : nullptr);
 	}
 	else
 	{
diff --git a/Engine/Source/Runtime/Experimental/Chaos/Private/PhysicsProxy/GeometryCollectionPhysicsProxy.cpp b/Engine/Source/Runtime/Experimental/Chaos/Private/PhysicsProxy/GeometryCollectionPhysicsProxy.cpp
--- a/Engine/Source/Runtime/Experimental/Chaos/Private/PhysicsProxy/GeometryCollectionPhysicsProxy.cpp
+++ b/Engine/Source/Runtime/Experimental/Chaos/Private/PhysicsProxy/GeometryCollectionPhysicsProxy.cpp
@@ -36,6 +36,7 @@
 #include "Chaos/PhysicsObjectInternal.h"
 #include "Chaos/PhysicsObjectInterface.h"
 #include "GeometryCollection/Facades/CollectionAnchoringFacade.h"
+#include "GeometryCollection/Facades/CollectionConnectionGraphFacade.h"
 
 #ifndef TODO_REIMPLEMENT_INIT_COMMANDS
 #define TODO_REIMPLEMENT_INIT_COMMANDS 0
@@ -1215,8 +1216,8 @@ void FGeometryCollectionPhysicsProxy::InitializeBodiesPT(Chaos::FPBDRigidsSolver
 				}
 			}
 
-			const TManagedArray<TSet<int32>>* Connections = RestCollection->FindAttribute<TSet<int32>>("Connections", FGeometryCollection::TransformGroup);
-			const bool bGenerateConnectionGraph = (!Connections) || (Connections->Num() != RestCollection->Transform.Num()) || bGeometryCollectionAlwaysGenerateConnectionGraph;
+			const GeometryCollection::Facades::FCollectionConnectionGraphFacade ConnectionFacade(*RestCollection);
+			const bool bGenerateConnectionGraph = !ConnectionFacade.IsValid() || bGeometryCollectionAlwaysGenerateConnectionGraph || !ConnectionFacade.HasValidConnections();
 			if (bGenerateConnectionGraph)
 			{
 				// Set cluster connectivity.  TPBDRigidClustering::CreateClusterParticle() 
@@ -1240,44 +1241,32 @@ void FGeometryCollectionPhysicsProxy::InitializeBodiesPT(Chaos::FPBDRigidsSolver
 			}
 			else
 			{
-				if (Connections)
+				if (ConnectionFacade.IsValid())
 				{
-					// first reserve memory for the connection so we alolocate just the right amount of memory 
-					for (int32 TransformGroupIndex = 0; TransformGroupIndex < NumTransforms; ++TransformGroupIndex)
+					// TODO: is it worth computing the valence of each edge to Reserve() the particle edge arrays?
+
+					// Transfer connections from the Collection's ConnectionFacade over to the clustered particles
+					int32 NumConnections = ConnectionFacade.NumConnections();
+					for (int32 ConnectionIdx = 0; ConnectionIdx < NumConnections; ++ConnectionIdx)
 					{
-						if (FClusterHandle* ClusteredParticle = SolverParticleHandles[TransformGroupIndex])
+						TPair<int32,int32> Connection = ConnectionFacade.GetConnection(ConnectionIdx);
+						checkSlow(Connection.Key != Connection.Value); // Graph should never contain self-loops
+
+						bool bHasContactAreas = ConnectionFacade.HasContactAreas();
+						float ContactArea = 0.0f; // Note: this default value should be unused
+						if (bHasContactAreas)
 						{
-							const TSet<int32>& Siblings = (*Connections)[TransformGroupIndex];
-							if (Siblings.Num())
-							{
-								ClusteredParticle->ConnectivityEdges().Reserve(Siblings.Num());
-							}
+							ContactArea = ConnectionFacade.GetConnectionContactArea(ConnectionIdx);
 						}
-					}
 
-					// now set the connections
-					for (int32 TransformGroupIndex = 0; TransformGroupIndex < NumTransforms; ++TransformGroupIndex)
-					{
-						if (FClusterHandle* ClusteredParticle = SolverParticleHandles[TransformGroupIndex]) 
+						if (FClusterHandle* ClusteredParticle = SolverParticleHandles[Connection.Key])
 						{
-							const TSet<int32>& Siblings = (*Connections)[TransformGroupIndex];
-							for (const int32 SiblingTransformIndex: Siblings)
+							if (FClusterHandle* OtherClusteredParticle = SolverParticleHandles[Connection.Value])
 							{
-								if (FClusterHandle* OtherClusteredParticle = SolverParticleHandles[SiblingTransformIndex])
-								{
-									// we add both connection for integrity so let use the index order as a way to filter out the symmetrical part 
-									if (SiblingTransformIndex > TransformGroupIndex)
-									{
-										Chaos::TConnectivityEdge<Chaos::FReal> Edge;
-										Edge.Strain = (ClusteredParticle->GetInternalStrains() + OtherClusteredParticle->GetInternalStrains()) * 0.5f;
-									
-										Edge.Sibling = OtherClusteredParticle;
-										ClusteredParticle->ConnectivityEdges().Add(Edge);
-									
-										Edge.Sibling = ClusteredParticle;
-										OtherClusteredParticle->ConnectivityEdges().Add(Edge);
-									}
-								}
+								float EdgeStrain = bHasContactAreas ? ContactArea : (ClusteredParticle->GetInternalStrains() + OtherClusteredParticle->GetInternalStrains()) * 0.5f;
+								// Add symmetric Chaos::TConnectivityEdge<Chaos::FReal> edges to each particle
+								ClusteredParticle->ConnectivityEdges().Emplace(OtherClusteredParticle, EdgeStrain);
+								OtherClusteredParticle->ConnectivityEdges().Emplace(ClusteredParticle, EdgeStrain);
 							}
 						}
 					}
diff --git a/Engine/Source/Runtime/Experimental/Chaos/Public/GeometryCollection/Facades/CollectionConnectionGraphFacade.h b/Engine/Source/Runtime/Experimental/Chaos/Public/GeometryCollection/Facades/CollectionConnectionGraphFacade.h
--- a/Engine/Source/Runtime/Experimental/Chaos/Public/GeometryCollection/Facades/CollectionConnectionGraphFacade.h
+++ b/Engine/Source/Runtime/Experimental/Chaos/Public/GeometryCollection/Facades/CollectionConnectionGraphFacade.h
@@ -21,7 +21,10 @@ namespace GeometryCollection::Facades
 		bool IsValid() const;
 
 		/** Is the facade defined constant. */
-		bool IsConst() const { return ConnectionsAttribute.IsConst(); }
+		bool IsConst() const
+		{
+			return ConnectionEdgeStartAttribute.IsConst();
+		}
 
 		/** Create the facade attributes. */
 		void DefineSchema();
@@ -29,12 +32,43 @@ namespace GeometryCollection::Facades
 		/** Remove the attributes */
 		void ClearAttributes();
 
+		/** Connect two bones */
+		void ConnectWithContact(int32 BoneA, int32 BoneB, float ContactArea);
+
 		/** Connect two bones */
 		void Connect(int32 BoneA, int32 BoneB);
 
+		/** Enable or disable the Contact Area attribute */
+		void EnableContactAreas(bool bEnable, float DefaultContactArea = 1.0f);
+
+		/** Reserve space for a number of additional connections */
+		void ReserveAdditionalConnections(int32 NumAdditionalConnections);
+
 		/**  Connections between bones that have the same parent in the hierarchy */
+		UE_DEPRECATED(5.3, "We have switched to an edge array connection representation. Please use the accessor functions (GetConnection, NumConnections, etc) to access the arrays of edge data.")
 		TManagedArrayAccessor<TSet<int32>> ConnectionsAttribute;
 
+		// Get the transform indices for the ConnectionIndex
+		TPair<int32, int32> GetConnection(int32 ConnectionIndex) const;
+		// Get the contact area for the ConnectionIndex
+		float GetConnectionContactArea(int32 ConnectionIndex) const;
+
+		bool HasContactAreas() const;
+
+		// Number of connection edges
+		int32 NumConnections() const;
+
+		// Verifies the connections indices are valid indices into the Collection's Transform group
+		bool HasValidConnections() const;
+
+		// Remove all edge connections, but keep the connection attributes
+		void ResetConnections();
+
+	protected:
+		 TManagedArrayAccessor<int32> ConnectionEdgeStartAttribute;
+		 TManagedArrayAccessor<int32> ConnectionEdgeEndAttribute;
+		 TManagedArrayAccessor<float> ConnectionEdgeContactAttribute;
+
 #if UE_BUILD_DEBUG
 		/** Optional Parent array for validating connections in debug */
 		TManagedArrayAccessor<int32> ParentAttribute;
diff --git a/Engine/Source/Runtime/Experimental/Chaos/Public/GeometryCollection/Facades/CollectionHierarchyFacade.h b/Engine/Source/Runtime/Experimental/Chaos/Public/GeometryCollection/Facades/CollectionHierarchyFacade.h
--- a/Engine/Source/Runtime/Experimental/Chaos/Public/GeometryCollection/Facades/CollectionHierarchyFacade.h
+++ b/Engine/Source/Runtime/Experimental/Chaos/Public/GeometryCollection/Facades/CollectionHierarchyFacade.h
@@ -44,9 +44,12 @@ namespace Chaos::Facades
 		/** Get the root indicies */
 		TArray<int32> GetRootIndices() const;
 
-		/** Get direct children from a specific transform index */
+		/** Get direct children from a specific transform index. Assumes parent attribute is valid. */
 		TArray<int32> GetChildrenAsArray(int32 TransformIndex) const;
 
+		/** Get parent of a specific transform index */
+		int32 GetParent(int32 TransformIndex) const;
+
 		/** 
 		* Get initial level of a specific transform index 
 		* If the attribute is missing return INDEX_NONE
diff --git a/Engine/Source/Runtime/Experimental/Chaos/Public/GeometryCollection/GeometryCollectionProximityPropertiesInterface.h b/Engine/Source/Runtime/Experimental/Chaos/Public/GeometryCollection/GeometryCollectionProximityPropertiesInterface.h
--- a/Engine/Source/Runtime/Experimental/Chaos/Public/GeometryCollection/GeometryCollectionProximityPropertiesInterface.h
+++ b/Engine/Source/Runtime/Experimental/Chaos/Public/GeometryCollection/GeometryCollectionProximityPropertiesInterface.h
@@ -23,6 +23,8 @@ public :
 	static const FName ProximityDistanceThreshold;
 	// Attribute: Whether to use the computed proximity graph as a connection graph
 	static const FName ProximityAsConnectionGraph;
+	// Attribute: How to compute contact area for connection graph edges, to be used by the simulation
+	static const FName ProximityConnectionGraphContactAreaMethod;
 	// Attribute: If greater than zero, filter proximity connections by requiring an amount of 'contact' as computed by the Contact Method
 	// This is a second filter applied after initial proximity connections were determined by any Proximity Detection Method,
 	// and can be used to reduce unsupported or spurious/glancing proximities.
@@ -36,6 +38,7 @@ public :
 		float DistanceThreshold = 1.0f;
 		float RequireContactAmount = 0.0f;
 		EProximityContactMethod ContactMethod = EProximityContactMethod::MinOverlapInProjectionToMajorAxes;
+		EConnectionContactMethod ContactAreaMethod = EConnectionContactMethod::None;
 		bool bUseAsConnectionGraph = false;
 	};
 
diff --git a/Engine/Source/Runtime/Experimental/Chaos/Public/GeometryCollection/GeometryCollectionProximityUtility.h b/Engine/Source/Runtime/Experimental/Chaos/Public/GeometryCollection/GeometryCollectionProximityUtility.h
--- a/Engine/Source/Runtime/Experimental/Chaos/Public/GeometryCollection/GeometryCollectionProximityUtility.h
+++ b/Engine/Source/Runtime/Experimental/Chaos/Public/GeometryCollection/GeometryCollectionProximityUtility.h
@@ -34,6 +34,16 @@ enum class EProximityContactMethod : uint8
 	//~ TODO: Add other methods for filtering overlaps, e.g. based on approximate surface area of the contact
 };
 
+// How contact is computed on the connection graph
+UENUM()
+enum class EConnectionContactMethod : uint8
+{
+	// Do not compute contact areas
+	None,
+	// Define contact based on the surface area of the intersection of the convex hulls, allowing for optional offset
+	ConvexHullContactArea
+};
+
 class CHAOS_API FGeometryCollectionProximityUtility
 {
 public:
@@ -46,10 +56,6 @@ public:
 
 	void InvalidateProximity();
 
-	void CopyProximityToConnectionGraph();
-	void ClearConnectionGraph();
-
-
 	// Stores stats about the contact between two geometries
 	struct FGeometryContactEdge
 	{
@@ -71,6 +77,10 @@ public:
 	// Note: This computes connections from lower to higher geometry indices, assuming connections are symmetric
 	static TArray<FGeometryContactEdge> ComputeConvexGeometryContactFromProximity(FGeometryCollection* Collection, float DistanceTolerance, UE::GeometryCollectionConvexUtility::FConvexHulls& LocalHulls);
 
+	// @param ContactEdges	Optional pre-computed proximity contact edges, to be used for computing contact areas on the connection graph edges
+	void CopyProximityToConnectionGraph(const TArray<FGeometryContactEdge>* ContactEdges = nullptr);
+	void ClearConnectionGraph();
+
 private:
 	FGeometryCollection* Collection;
 };
diff --git a/Engine/Source/Runtime/Experimental/GeometryCollectionEngine/Private/GeometryCollection/GeometryCollectionObject.cpp b/Engine/Source/Runtime/Experimental/GeometryCollectionEngine/Private/GeometryCollection/GeometryCollectionObject.cpp
--- a/Engine/Source/Runtime/Experimental/GeometryCollectionEngine/Private/GeometryCollection/GeometryCollectionObject.cpp
+++ b/Engine/Source/Runtime/Experimental/GeometryCollectionEngine/Private/GeometryCollection/GeometryCollectionObject.cpp
@@ -1013,7 +1013,8 @@ void UGeometryCollection::Serialize(FArchive& Ar)
 	//for all versions loaded, make sure loaded content is built
  	if (Ar.IsLoading())
 	{
-		EnsureDataIsCooked(true /*bInitResourcs*/, Ar.IsTransacting(), Ar.IsPersistent());
+		// note: don't allow copy from DDC here, since we've already loaded the data above, and the DDC data does not include any data migrations performed by the load
+		EnsureDataIsCooked(true /*bInitResources*/, Ar.IsTransacting(), Ar.IsPersistent(), false /*bAllowCopyFromDDC*/);
 	}
 #endif
 }
