[Chaos VD] Implemented Recording using Trace

- Created a set of macros to trace data about Solver Frames, Solver Steps, and Solver Particles on Rigid Solvers and RBAN solvers
- Added a way to have the context between CVDTrace calls to know which solver started the call on a determined solver step or particle without having to copy that information to every object

diff --git a/Engine/Source/Programs/UnrealBuildTool/Configuration/ModuleRules.cs b/Engine/Source/Programs/UnrealBuildTool/Configuration/ModuleRules.cs
--- a/Engine/Source/Programs/UnrealBuildTool/Configuration/ModuleRules.cs
+++ b/Engine/Source/Programs/UnrealBuildTool/Configuration/ModuleRules.cs
@@ -1549,6 +1549,24 @@ namespace UnrealBuildTool
 			}
 		}
 
+		/// <summary>
+		/// Setup this module for Chaos Visual Debugger support (Required for recording debug data that will be visualized in the Chaos Visual Debugger tool)
+		/// </summary>
+		public void SetupModuleChaosVisualDebuggerSupport(ReadOnlyTargetRules Target)
+		{
+			bool bHasChaosVisualDebuggerSupport = Target.bCompileChaosVisualDebuggerSupport && Target.Configuration != UnrealTargetConfiguration.Shipping;
+			if (bHasChaosVisualDebuggerSupport)
+			{
+				PublicDependencyModuleNames.Add("ChaosVDRuntime");
+
+				PublicDefinitions.Add("WITH_CHAOS_VISUAL_DEBUGGER=1");
+			}
+			else
+			{
+				PublicDefinitions.Add("WITH_CHAOS_VISUAL_DEBUGGER=0");
+			}
+		}
+
 		/// <summary>
 		/// Setup this module for physics support (based on the settings in UEBuildConfiguration)
 		/// </summary>
@@ -1570,6 +1588,8 @@ namespace UnrealBuildTool
 
 			PublicDefinitions.Add("WITH_CLOTH_COLLISION_DETECTION=1");
 
+			SetupModuleChaosVisualDebuggerSupport(Target);
+
 			// Modules may still be relying on appropriate definitions for physics.
 			// Nothing in engine should use these anymore as they were all deprecated and 
 			// assumed to be in the following configuration from 5.1, this will cause
diff --git a/Engine/Source/Programs/UnrealBuildTool/Configuration/TargetRules.cs b/Engine/Source/Programs/UnrealBuildTool/Configuration/TargetRules.cs
--- a/Engine/Source/Programs/UnrealBuildTool/Configuration/TargetRules.cs
+++ b/Engine/Source/Programs/UnrealBuildTool/Configuration/TargetRules.cs
@@ -781,6 +781,14 @@ namespace UnrealBuildTool
 		/// </summary>
 		[RequiresUniqueBuildEnvironment]
 		public bool bUseChaosMemoryTracking = false;
+		
+		/// <summary>
+		/// Whether to compile in Chaos Visual Debugger (CVD) support features to record the state of the physics simulation
+		/// </summary>
+		[XmlConfigFile(Name = "bCompileChaosVisualDebuggerSupport")]
+		[CommandLine("-CompileChaosVisualDebuggerSupport")]
+		[RequiresUniqueBuildEnvironment]
+		public bool bCompileChaosVisualDebuggerSupport = false;
 
 		/// <summary>
 		/// Whether scene query acceleration is done by UE. The physx scene query structure is still created, but we do not use it.
@@ -2990,6 +2998,11 @@ namespace UnrealBuildTool
 			get { return Inner.bUseChaosMemoryTracking; }
 		}
 
+		public bool bCompileChaosVisualDebuggerSupport
+		{
+			get { return Inner.bCompileChaosVisualDebuggerSupport; }
+		}
+
 		public bool bUseChaosChecked
 		{
 			get { return Inner.bUseChaosChecked; }
diff --git a/Engine/Source/Runtime/AnimGraphRuntime/Private/BoneControllers/AnimNode_RigidBody.cpp b/Engine/Source/Runtime/AnimGraphRuntime/Private/BoneControllers/AnimNode_RigidBody.cpp
--- a/Engine/Source/Runtime/AnimGraphRuntime/Private/BoneControllers/AnimNode_RigidBody.cpp
+++ b/Engine/Source/Runtime/AnimGraphRuntime/Private/BoneControllers/AnimNode_RigidBody.cpp
@@ -3,6 +3,11 @@
 #include "BoneControllers/AnimNode_RigidBody.h"
 #include "AnimationRuntime.h"
 #include "Animation/AnimInstanceProxy.h"
+
+#if WITH_CHAOS_VISUAL_DEBUGGER
+#include "ChaosVDRuntimeModule.h"
+#endif
+
 #include "ClothCollisionSource.h"
 #include "GameFramework/Pawn.h"
 #include "HAL/Event.h"
@@ -1059,6 +1064,11 @@ void FAnimNode_RigidBody::InitPhysics(const UAnimInstance* InAnimInstance)
 	if(bEnabled)
 	{
 		PhysicsSimulation = new ImmediatePhysics::FSimulation();
+
+#if WITH_CHAOS_VISUAL_DEBUGGER
+		PhysicsSimulation->GetChaosVDContextData().Id = FChaosVDRuntimeModule::Get().GenerateUniqueID();
+#endif
+
 		const int32 NumBodies = UsePhysicsAsset->SkeletalBodySetups.Num();
 		Bodies.Empty(NumBodies);
 		BodyAnimData.Reset(NumBodies);
diff --git a/Engine/Source/Runtime/Engine/Private/PhysicsEngine/ImmediatePhysics/ImmediatePhysicsChaos/ImmediatePhysicsSimulation_Chaos.cpp b/Engine/Source/Runtime/Engine/Private/PhysicsEngine/ImmediatePhysics/ImmediatePhysicsChaos/ImmediatePhysicsSimulation_Chaos.cpp
--- a/Engine/Source/Runtime/Engine/Private/PhysicsEngine/ImmediatePhysics/ImmediatePhysicsChaos/ImmediatePhysicsSimulation_Chaos.cpp
+++ b/Engine/Source/Runtime/Engine/Private/PhysicsEngine/ImmediatePhysics/ImmediatePhysicsChaos/ImmediatePhysicsSimulation_Chaos.cpp
@@ -16,6 +16,7 @@
 #include "Chaos/Joint/ChaosJointLog.h"
 #include "Chaos/MassConditioning.h"
 #include "Chaos/PBDJointConstraints.h"
+#include "ChaosVisualDebugger/ChaosVisualDebuggerTrace.h"
 #include "Stats/StatsTrace.h"
 
 
@@ -358,6 +359,10 @@ namespace ImmediatePhysics_Chaos
 
 		Implementation = MakeUnique<FImplementation>();
 
+#if CHAOS_DEBUG_NAME
+		DebugName = TEXT("RBAN");
+#endif
+
 		// RBAN collision customization
 		Implementation->Collisions.DisableHandles();
 
@@ -793,6 +798,8 @@ namespace ImmediatePhysics_Chaos
 	void FSimulation::Simulate(FReal InDeltaTime, FReal MaxStepTime, int32 MaxSubSteps, const FVector& InGravity)
 	{
 		SCOPE_CYCLE_COUNTER(STAT_ImmediateSimulate_Chaos);
+		CVD_SCOPE_TRACE_SOLVER_FRAME(FSimulation, *this);
+
 		using namespace Chaos;
 
 		// Reject DeltaTime outliers
diff --git a/Engine/Source/Runtime/Engine/Public/Physics/ImmediatePhysics/ImmediatePhysicsChaos/ImmediatePhysicsSimulation_Chaos.h b/Engine/Source/Runtime/Engine/Public/Physics/ImmediatePhysics/ImmediatePhysicsChaos/ImmediatePhysicsSimulation_Chaos.h
--- a/Engine/Source/Runtime/Engine/Public/Physics/ImmediatePhysics/ImmediatePhysicsChaos/ImmediatePhysicsSimulation_Chaos.h
+++ b/Engine/Source/Runtime/Engine/Public/Physics/ImmediatePhysics/ImmediatePhysicsChaos/ImmediatePhysicsSimulation_Chaos.h
@@ -2,6 +2,7 @@
 
 #pragma once
 
+#include "ChaosVisualDebugger/ChaosVDContextProvider.h"
 #include "Physics/ImmediatePhysics/ImmediatePhysicsChaos/ImmediatePhysicsCore_Chaos.h"
 
 #include "Engine/EngineTypes.h"
@@ -106,6 +107,33 @@ namespace ImmediatePhysics_Chaos
 
 		struct FImplementation;
 		TUniquePtr<FImplementation> Implementation;
+
+#if CHAOS_DEBUG_NAME
+	public:
+		void SetDebugName(const FName& Name)
+		{
+			DebugName = Name;
+		}
+
+		const FName& GetDebugName() const
+		{
+			return DebugName;
+		}
+
+	private:
+	FName DebugName;
+#endif
+
+#if WITH_CHAOS_VISUAL_DEBUGGER
+	private:
+		FChaosVDContext CVDContextData;
+
+	public:
+		FChaosVDContext& GetChaosVDContextData()
+		{
+			return CVDContextData;
+		};
+#endif
 	};
 
 }
diff --git a/Engine/Source/Runtime/Experimental/Chaos/Chaos.Build.cs b/Engine/Source/Runtime/Experimental/Chaos/Chaos.Build.cs
--- a/Engine/Source/Runtime/Experimental/Chaos/Chaos.Build.cs
+++ b/Engine/Source/Runtime/Experimental/Chaos/Chaos.Build.cs
@@ -34,6 +34,8 @@ namespace UnrealBuildTool.Rules
 			{
 				PublicDefinitions.Add("CHAOS_MEMORY_TRACKING=0");
 			}
+
+			SetupModuleChaosVisualDebuggerSupport(Target);
 			
 			UnsafeTypeCastWarningLevel = WarningLevel.Error;
 
diff --git a/Engine/Source/Runtime/Experimental/Chaos/Private/Chaos/Evolution/PBDMinEvolution.cpp b/Engine/Source/Runtime/Experimental/Chaos/Private/Chaos/Evolution/PBDMinEvolution.cpp
--- a/Engine/Source/Runtime/Experimental/Chaos/Private/Chaos/Evolution/PBDMinEvolution.cpp
+++ b/Engine/Source/Runtime/Experimental/Chaos/Private/Chaos/Evolution/PBDMinEvolution.cpp
@@ -15,6 +15,7 @@
 #include "Chaos/PerParticlePBDGroundConstraint.h"
 #include "Chaos/PerParticlePBDUpdateFromDeltaPosition.h"
 #include "ChaosStats.h"
+#include "ChaosVisualDebugger/ChaosVisualDebuggerTrace.h"
 
 //PRAGMA_DISABLE_OPTIMIZATION
 
@@ -121,6 +122,7 @@ namespace Chaos
 	void FPBDMinEvolution::AdvanceOneTimeStep(const FReal Dt, const FReal StepFraction)
 	{
 		SCOPE_CYCLE_COUNTER(STAT_MinEvolution_AdvanceOneTimeStep);
+		CVD_SCOPE_TRACE_SOLVER_STEP();
 
 		Integrate(Dt);
 
@@ -140,6 +142,8 @@ namespace Chaos
 
 			ScatterOutput(Dt);
 		}
+
+		CVD_TRACE_PARTICLES(Particles.GetParticleHandles());
 	}
 
 	// A opportunity for systems to allocate buffers for the duration of the tick, if they have enough info to do so
diff --git a/Engine/Source/Runtime/Experimental/Chaos/Private/Chaos/Framework/PhysicsSolverBase.cpp b/Engine/Source/Runtime/Experimental/Chaos/Private/Chaos/Framework/PhysicsSolverBase.cpp
--- a/Engine/Source/Runtime/Experimental/Chaos/Private/Chaos/Framework/PhysicsSolverBase.cpp
+++ b/Engine/Source/Runtime/Experimental/Chaos/Private/Chaos/Framework/PhysicsSolverBase.cpp
@@ -11,6 +11,9 @@
 #include "Framework/Threading.h"
 #include "RewindData.h"
 
+#include "ChaosVisualDebugger/ChaosVisualDebuggerTrace.h"
+#include "ChaosVisualDebugger/ChaosVDContextProvider.h"
+
 namespace Chaos
 {	
 	extern int GSingleThreadedPhysics;
@@ -119,6 +122,22 @@ namespace Chaos
 		
 	}
 
+
+	FPhysicsSolverAdvanceTask::FPhysicsSolverAdvanceTask(FPhysicsSolverBase& InSolver, FPushPhysicsData* InPushData)
+		: Solver(InSolver)
+		, PushData(InPushData)
+	{
+		CVD_GET_CURRENT_CONTEXT(CVDContext);
+	}
+
+	void FPhysicsSolverAdvanceTask::DoTask(ENamedThreads::Type CurrentThread, const FGraphEventRef& MyCompletionGraphEvent)
+	{
+		CVD_SCOPE_CONTEXT(CVDContext);
+
+		AdvanceSolver();
+	}
+
+
 	void FPhysicsSolverAdvanceTask::AdvanceSolver()
 	{
 		LLM_SCOPE(ELLMTag::ChaosUpdate);
@@ -337,6 +356,9 @@ namespace Chaos
 
 	FGraphEventRef FPhysicsSolverBase::AdvanceAndDispatch_External(FReal InDt)
 	{
+
+		CVD_TRACE_SOLVER_START_FRAME(FPhysicsSolverBase, *this);
+
 		const bool bSubstepping = MMaxSubSteps > 1;
 		SetSolverSubstep_External(bSubstepping);
 		const FReal DtWithPause = bPaused_External ? 0.0f : InDt;
diff --git a/Engine/Source/Runtime/Experimental/Chaos/Private/Chaos/PBDRigidsEvolutionGBF.cpp b/Engine/Source/Runtime/Experimental/Chaos/Private/Chaos/PBDRigidsEvolutionGBF.cpp
--- a/Engine/Source/Runtime/Experimental/Chaos/Private/Chaos/PBDRigidsEvolutionGBF.cpp
+++ b/Engine/Source/Runtime/Experimental/Chaos/Private/Chaos/PBDRigidsEvolutionGBF.cpp
@@ -24,6 +24,8 @@
 #include "Misc/Paths.h"
 #include "HAL/FileManager.h"
 
+#include "ChaosVisualDebugger/ChaosVisualDebuggerTrace.h"
+
 //PRAGMA_DISABLE_OPTIMIZATION
 
 namespace Chaos
@@ -246,6 +248,8 @@ void FPBDRigidsEvolutionGBF::Advance(const FReal Dt,const FReal MaxStepDt,const
 
 void FPBDRigidsEvolutionGBF::AdvanceOneTimeStep(const FReal Dt,const FSubStepInfo& SubStepInfo)
 {
+	CVD_SCOPE_TRACE_SOLVER_STEP();
+		
 	PrepareTick();
 
 	AdvanceOneTimeStepImpl(Dt, SubStepInfo);
@@ -617,6 +621,8 @@ void FPBDRigidsEvolutionGBF::AdvanceOneTimeStepImpl(const FReal Dt, const FSubSt
 		GetCollisionConstraints().DetectProbeCollisions(Dt);
 	}
 
+	CVD_TRACE_PARTICLES(GetParticleHandles());
+
 #if !UE_BUILD_SHIPPING
 	if(SerializeEvolution)
 	{
diff --git a/Engine/Source/Runtime/Experimental/Chaos/Private/ChaosVisualDebugger/ChaosVDContextProvider.cpp b/Engine/Source/Runtime/Experimental/Chaos/Private/ChaosVisualDebugger/ChaosVDContextProvider.cpp
new file mode 100644
--- /dev/null
+++ b/Engine/Source/Runtime/Experimental/Chaos/Private/ChaosVisualDebugger/ChaosVDContextProvider.cpp
@@ -0,0 +1,38 @@
+// Copyright Epic Games, Inc. All Rights Reserved.
+
+#include "ChaosVisualDebugger/ChaosVDContextProvider.h"
+
+#if WITH_CHAOS_VISUAL_DEBUGGER
+bool FChaosVDThreadContext::GetCurrentContext(FChaosVDContext& OutContext)
+{
+	if (LocalContextStack.Num() > 0)
+	{
+		OutContext = LocalContextStack.Top();
+		return true;
+	}
+
+	return false;
+}
+
+const FChaosVDContext* FChaosVDThreadContext::GetCurrentContext()
+{
+	if (LocalContextStack.Num() > 0)
+	{
+		return &LocalContextStack.Top();
+	}
+	else
+	{
+		return nullptr;
+	}
+}
+
+void FChaosVDThreadContext::PushContext(const FChaosVDContext& InContext)
+{
+	LocalContextStack.Add(InContext);
+}
+
+void FChaosVDThreadContext::PopContext()
+{
+	LocalContextStack.Pop();
+}
+#endif //WITH_CHAOS_VISUAL_DEBUGGER
diff --git a/Engine/Source/Runtime/Experimental/Chaos/Private/ChaosVisualDebugger/ChaosVisualDebuggerTrace.cpp b/Engine/Source/Runtime/Experimental/Chaos/Private/ChaosVisualDebugger/ChaosVisualDebuggerTrace.cpp
--- a/Engine/Source/Runtime/Experimental/Chaos/Private/ChaosVisualDebugger/ChaosVisualDebuggerTrace.cpp
+++ b/Engine/Source/Runtime/Experimental/Chaos/Private/ChaosVisualDebugger/ChaosVisualDebuggerTrace.cpp
@@ -1,20 +1,143 @@
 // Copyright Epic Games, Inc. All Rights Reserved.
+
 #include "ChaosVisualDebugger/ChaosVisualDebuggerTrace.h"
-#include "Trace/Trace.inl"
 
-UE_TRACE_CHANNEL_DEFINE(PhysicsChannel);
+#if CHAOS_VISUAL_DEBUGGER_ENABLED
+#include "Chaos/ParticleHandle.h"
+#include "Chaos/Framework/PhysicsSolverBase.h"
+
+UE_TRACE_EVENT_DEFINE(ChaosVDLogger, ChaosVDSolverFrameStart)
+
+UE_TRACE_EVENT_DEFINE(ChaosVDLogger, ChaosVDSolverFrameEnd)
+
+UE_TRACE_CHANNEL_DEFINE(ChaosVDChannel);
+UE_TRACE_EVENT_DEFINE(ChaosVDLogger, ChaosVDParticle)
+UE_TRACE_EVENT_DEFINE(ChaosVDLogger, ChaosVDSolverStepStart)
+UE_TRACE_EVENT_DEFINE(ChaosVDLogger, ChaosVDSolverStepEnd)
 
-UE_TRACE_EVENT_BEGIN(Physics, ParticlePosition)
-UE_TRACE_EVENT_FIELD(Chaos::FReal, PositionX)
-UE_TRACE_EVENT_FIELD(Chaos::FReal, PositionY)
-UE_TRACE_EVENT_FIELD(Chaos::FReal, PositionZ)
-UE_TRACE_EVENT_END()
+void FChaosVisualDebuggerTrace::TraceParticle(const Chaos::FGeometryParticleHandle* ParticleHandle)
+{
+	const FChaosVDContext* CVDContextData = FChaosVDThreadContext::Get().GetCurrentContext();
+	if (!ensure(CVDContextData))
+	{
+		return;
+	}
+
+	TraceParticle(ParticleHandle, *CVDContextData);
+}
 
-//A very minimal test debug log function
-void ChaosVisualDebugger::ParticlePositionLog(const Chaos::FVec3& Position)
+void FChaosVisualDebuggerTrace::TraceParticle(const Chaos::FGeometryParticleHandle* ParticleHandle, const FChaosVDContext& ContextData)
 {
-	UE_TRACE_LOG(Physics, ParticlePosition, PhysicsChannel)
-					<< ParticlePosition.PositionX(Position.X)
-					<< ParticlePosition.PositionY(Position.Y)
-					<< ParticlePosition.PositionZ(Position.Z);
-}
\ No newline at end of file
+	if (!TRACE_PRIVATE_CHANNELEXPR_IS_ENABLED(ChaosVDChannel))
+	{
+		return;
+	}
+
+	if (!ParticleHandle)
+	{
+		UE_LOG(LogChaos, Warning, TEXT("Tried to Trace a null particle %hs"), __FUNCTION__);
+		return;
+	}
+
+	FString TempDebugName;
+#if CHAOS_DEBUG_NAME
+	if (const TSharedPtr<FString, ESPMode::ThreadSafe> DebugNamePtr = ParticleHandle->DebugName())
+	{
+		TempDebugName = *DebugNamePtr.Get();
+	}
+#else
+	TempDebugName = TEXT("NOT AVAILABLE - COMPILED OUT");
+#endif
+	
+	UE_TRACE_LOG(ChaosVDLogger, ChaosVDParticle, ChaosVDChannel)
+		<< ChaosVDParticle.SolverID(ContextData.Id)
+		<< ChaosVDParticle.Cycle(FPlatformTime::Cycles64())
+
+		<< ChaosVDParticle.ParticleID(ParticleHandle->UniqueIdx().Idx)
+		<< ChaosVDParticle.ParticleType(static_cast<uint8>(ParticleHandle->Type))
+		<< ChaosVDParticle.DebugName(*TempDebugName, TempDebugName.Len())
+
+		<< CVD_TRACE_VECTOR_ON_EVENT(ChaosVDParticle, Position, ParticleHandle->X())
+		<< CVD_TRACE_ROTATOR_ON_EVENT(ChaosVDParticle, Rotation, ParticleHandle->R())
+		<< CVD_TRACE_VECTOR_ON_EVENT(ChaosVDParticle, Velocity, Chaos::FConstGenericParticleHandle(ParticleHandle)->V())
+		<< CVD_TRACE_VECTOR_ON_EVENT(ChaosVDParticle, AngularVelocity, Chaos::FConstGenericParticleHandle(ParticleHandle)->W())
+
+		<< ChaosVDParticle.ObjectState(static_cast<int8>(ParticleHandle->ObjectState()));
+}
+
+void FChaosVisualDebuggerTrace::TraceParticles(const Chaos::TGeometryParticleHandles<Chaos::FReal, 3>& ParticleHandles)
+{
+	if (!TRACE_PRIVATE_CHANNELEXPR_IS_ENABLED(ChaosVDChannel))
+	{
+		return;
+	}
+
+	const FChaosVDContext* CVDContextData = FChaosVDThreadContext::Get().GetCurrentContext();
+	if (!ensure(CVDContextData))
+	{
+		return;
+	}
+
+	for (uint32 ParticleIndex = 0; ParticleIndex < ParticleHandles.Size(); ParticleIndex++)
+	{
+		// TODO: We should only trace the Particles that changed. probably using the Dirty Flags?
+		// Geometry data "uniqueness" will be handled by the trace helper
+		TraceParticle(ParticleHandles.Handle(ParticleIndex).Get(), *CVDContextData);
+	}
+}
+
+void FChaosVisualDebuggerTrace::TraceSolverFrameStart(const FChaosVDContext& ContextData, const FString& InDebugName)
+{
+	if (!ensure(ContextData.Id != INDEX_NONE))
+	{
+		return;
+	}
+
+	FChaosVDThreadContext::Get().PushContext(ContextData);
+
+	UE_TRACE_LOG(ChaosVDLogger, ChaosVDSolverFrameStart, ChaosVDChannel)
+		<< ChaosVDSolverFrameStart.SolverID(ContextData.Id)
+		<< ChaosVDSolverFrameStart.Cycle(FPlatformTime::Cycles64())
+		<< ChaosVDSolverFrameStart.DebugName(*InDebugName, InDebugName.Len());
+}
+
+void FChaosVisualDebuggerTrace::TraceSolverFrameEnd(const FChaosVDContext& ContextData)
+{
+	FChaosVDThreadContext::Get().PopContext();
+
+	if (!ensure(ContextData.Id != INDEX_NONE))
+	{
+		return;
+	}
+
+	UE_TRACE_LOG(ChaosVDLogger, ChaosVDSolverFrameEnd, ChaosVDChannel)
+		<< ChaosVDSolverFrameEnd.SolverID(ContextData.Id)
+		<< ChaosVDSolverFrameEnd.Cycle(FPlatformTime::Cycles64());
+}
+
+void FChaosVisualDebuggerTrace::TraceSolverStepStart()
+{
+	const FChaosVDContext* CVDContextData = FChaosVDThreadContext::Get().GetCurrentContext();
+	if (!ensure(CVDContextData))
+	{
+		return;
+	}
+
+	UE_TRACE_LOG(ChaosVDLogger, ChaosVDSolverStepStart, ChaosVDChannel)
+		<< ChaosVDSolverStepStart.Cycle(FPlatformTime::Cycles64())
+		<< ChaosVDSolverStepStart.SolverID(CVDContextData->Id);
+}
+
+void FChaosVisualDebuggerTrace::TraceSolverStepEnd()
+{
+	const FChaosVDContext* CVDContextData = FChaosVDThreadContext::Get().GetCurrentContext();
+	if (!ensure(CVDContextData))
+	{
+		return;
+	}
+
+	UE_TRACE_LOG(ChaosVDLogger, ChaosVDSolverStepEnd, ChaosVDChannel)
+		<< ChaosVDSolverStepEnd.Cycle(FPlatformTime::Cycles64())
+		<< ChaosVDSolverStepEnd.SolverID(CVDContextData->Id);
+}
+#endif
diff --git a/Engine/Source/Runtime/Experimental/Chaos/Private/PBDRigidsSolver.cpp b/Engine/Source/Runtime/Experimental/Chaos/Private/PBDRigidsSolver.cpp
--- a/Engine/Source/Runtime/Experimental/Chaos/Private/PBDRigidsSolver.cpp
+++ b/Engine/Source/Runtime/Experimental/Chaos/Private/PBDRigidsSolver.cpp
@@ -461,7 +461,6 @@ namespace Chaos
 					}
 
 					MSolver->GetEvolution()->AdvanceOneTimeStep(DeltaTime, MSubStepInfo);
-					MSolver->PostEvolutionVDBPush();
 					bFirstStep = false;
 				}
 
@@ -2188,21 +2187,6 @@ TRACE_COUNTER_SET(ChaosTraceCounter_##Name, Value)
 		return GetParticleProxy_PT(Handle.UniqueIdx());
 	}
 
-	void FPBDRigidsSolver::PostEvolutionVDBPush() const
-	{
-#if CHAOS_VISUAL_DEBUGGER_ENABLED
-		if (ChaosVisualDebuggerEnable)
-		{
-			const TGeometryParticleHandles<FReal, 3>&  AllParticleHandles = GetEvolution()->GetParticleHandles();
-			for (uint32 ParticelIndex = 0; ParticelIndex < AllParticleHandles.Size(); ParticelIndex++)
-			{
-				const TUniquePtr<FGeometryParticleHandle>& ParticleHandle = AllParticleHandles.Handle(ParticelIndex);
-				ChaosVisualDebugger::ParticlePositionLog(ParticleHandle->X());				
-			}
-		}
-#endif
-	}
-
 	void FPBDRigidsSolver::UpdateMaterial(Chaos::FMaterialHandle InHandle, const Chaos::FChaosPhysicsMaterial& InNewData)
 	{
 		TSolverSimMaterialScope<ELockType::Write> Scope(this);
diff --git a/Engine/Source/Runtime/Experimental/Chaos/Public/Chaos/Framework/PhysicsSolverBase.h b/Engine/Source/Runtime/Experimental/Chaos/Public/Chaos/Framework/PhysicsSolverBase.h
--- a/Engine/Source/Runtime/Experimental/Chaos/Public/Chaos/Framework/PhysicsSolverBase.h
+++ b/Engine/Source/Runtime/Experimental/Chaos/Public/Chaos/Framework/PhysicsSolverBase.h
@@ -13,6 +13,10 @@
 #include "Chaos/ChaosMarshallingManager.h"
 #include "Stats/Stats2.h"
 
+#if WITH_CHAOS_VISUAL_DEBUGGER
+#include "ChaosVisualDebugger/ChaosVDContextProvider.h"
+#endif
+
 class FChaosSolversModule;
 
 DECLARE_MULTICAST_DELEGATE_OneParam(FSolverPreAdvance, Chaos::FReal);
@@ -118,21 +122,22 @@ namespace Chaos
 	{
 	public:
 
-		FPhysicsSolverAdvanceTask(FPhysicsSolverBase& InSolver, FPushPhysicsData* InPushData)
-			: Solver(InSolver)
-			, PushData(InPushData)
-		{}
-
+		FPhysicsSolverAdvanceTask(FPhysicsSolverBase& InSolver, FPushPhysicsData* InPushData);
+		
 		TStatId GetStatId() const { RETURN_QUICK_DECLARE_CYCLE_STAT(FPhysicsSolverAdvanceTask, STATGROUP_TaskGraphTasks); }
 		static ENamedThreads::Type GetDesiredThread();
 		static ESubsequentsMode::Type GetSubsequentsMode() { return ESubsequentsMode::TrackSubsequents; }
-		void DoTask(ENamedThreads::Type CurrentThread, const FGraphEventRef& MyCompletionGraphEvent) { AdvanceSolver(); }
+		void DoTask(ENamedThreads::Type CurrentThread, const FGraphEventRef& MyCompletionGraphEvent);
 		void AdvanceSolver();
 
 	private:
 
 		FPhysicsSolverBase& Solver;
 		FPushPhysicsData* PushData;
+
+#if WITH_CHAOS_VISUAL_DEBUGGER
+		FChaosVDContext CVDContext;
+#endif
 	};
 
 	struct CHAOS_API FAllSolverTasks
@@ -749,6 +754,17 @@ namespace Chaos
 		void SetStealAdvanceTasks_ForTesting(bool bInStealAdvanceTasksForTesting);
 		void PopAndExecuteStolenAdvanceTask_ForTesting();
 #endif
+
+#if WITH_CHAOS_VISUAL_DEBUGGER
+	private:
+		FChaosVDContext CVDContextData;
+
+	public:
+		FChaosVDContext& GetChaosVDContextData()
+		{
+			return CVDContextData;
+		};
+#endif
 	};
 }
 
diff --git a/Engine/Source/Runtime/Experimental/Chaos/Public/ChaosVisualDebugger/ChaosVDContextProvider.h b/Engine/Source/Runtime/Experimental/Chaos/Public/ChaosVisualDebugger/ChaosVDContextProvider.h
new file mode 100644
--- /dev/null
+++ b/Engine/Source/Runtime/Experimental/Chaos/Public/ChaosVisualDebugger/ChaosVDContextProvider.h
@@ -0,0 +1,81 @@
+// Copyright Epic Games, Inc. All Rights Reserved.
+
+#pragma once
+
+#include "Containers/Array.h"
+#include "Containers/ContainerAllocationPolicies.h"
+#include "HAL/ThreadSingleton.h"
+
+#if WITH_CHAOS_VISUAL_DEBUGGER
+
+/** Chaos Visual Debugger data used to context for logging or debugging purposes */
+struct CHAOS_API FChaosVDContext
+{
+	int32 OwnerID = INDEX_NONE;
+	int32 Id = INDEX_NONE;
+};
+
+/** Singleton class that manages the thread local storage used to store CVD Context data */
+class CHAOS_API FChaosVDThreadContext : public TThreadSingleton<FChaosVDThreadContext>
+{
+public:
+
+	/** Copies the Current CVD context data into the provided struct
+	 * @return true if the copy was successful
+	 */
+	bool GetCurrentContext(FChaosVDContext& OutContext);
+	
+	/** Gets the current CVD context data -
+	 * Don't use of a function that will recursively push new context data as it might invalidate the pointer
+	 * @return Ptr to the Current CVD context data
+	 */
+	const FChaosVDContext* GetCurrentContext();
+	
+	/** Pushed a new CVD Context data to the local cvd context stack */
+	void PushContext(const FChaosVDContext& InContext);
+	
+	/** Removes the CVD Context data at the top of the local cvd context stack */
+	void PopContext();
+
+protected:
+	TArray<FChaosVDContext, TInlineAllocator<16>> LocalContextStack;
+};
+
+/** Utility Class that will push the provided CVD Context Data to the local thread storage
+ * and remove it when it goes out of scope
+ */
+struct CHAOS_API FChaosVDScopeContext
+{
+	FChaosVDScopeContext(const FChaosVDContext& InCVDContext)
+	{
+		FChaosVDThreadContext::Get().PushContext(InCVDContext);
+	}
+
+	~FChaosVDScopeContext()
+	{
+		FChaosVDThreadContext::Get().PopContext();
+	}
+};
+
+#ifndef CVD_GET_CURRENT_CONTEXT
+	#define CVD_GET_CURRENT_CONTEXT(OutContext) \
+		ensure(FChaosVDThreadContext::Get().GetCurrentContext(OutContext))
+#endif
+
+#ifndef CVD_SCOPE_CONTEXT
+	#define CVD_SCOPE_CONTEXT(InContext) \
+		FChaosVDScopeContext CVDScope(InContext);
+#endif
+
+#else // WITH_CHAOS_VISUAL_DEBUGGER
+
+#ifndef CVD_GET_CURRENT_CONTEXT
+	#define CVD_GET_CURRENT_CONTEXT(OutContext)
+#endif
+
+#ifndef CVD_SCOPE_CONTEXT
+	#define CVD_SCOPE_CONTEXT(InContext)
+#endif
+
+#endif // WITH_CHAOS_VISUAL_DEBUGGER
+
diff --git a/Engine/Source/Runtime/Experimental/Chaos/Public/ChaosVisualDebugger/ChaosVisualDebuggerTrace.h b/Engine/Source/Runtime/Experimental/Chaos/Public/ChaosVisualDebugger/ChaosVisualDebuggerTrace.h
--- a/Engine/Source/Runtime/Experimental/Chaos/Public/ChaosVisualDebugger/ChaosVisualDebuggerTrace.h
+++ b/Engine/Source/Runtime/Experimental/Chaos/Public/ChaosVisualDebugger/ChaosVisualDebuggerTrace.h
@@ -1,20 +1,234 @@
 // Copyright Epic Games, Inc. All Rights Reserved.
 
 #pragma once
-#include "Chaos/Core.h"
+#include "Chaos/ParticleHandleFwd.h"
 #include "Chaos/Vector.h"
 
+#ifndef CHAOS_VISUAL_DEBUGGER_ENABLED
+	#define CHAOS_VISUAL_DEBUGGER_ENABLED (WITH_CHAOS_VISUAL_DEBUGGER && UE_TRACE_ENABLED)
+#endif
+
+// Define NO-OP versions of our macros if we can't Trace
+#if !CHAOS_VISUAL_DEBUGGER_ENABLED
+	#ifndef CVD_TRACE_PARTICLE
+		#define CVD_TRACE_PARTICLE(ParticleHandle)
+	#endif
+
+	#ifndef CVD_TRACE_PARTICLES
+		#define CVD_TRACE_PARTICLES(ParticleHandles)
+	#endif
+
+	#ifndef CVD_TRACE_SOLVER_START_FRAME
+		#define CVD_TRACE_SOLVER_START_FRAME(SolverType, SolverRef)
+	#endif
+
+	#ifndef CVD_TRACE_SOLVER_END_FRAME
+		#define CVD_TRACE_SOLVER_END_FRAME(SolverType, SolverRef)
+	#endif
+
+	#ifndef CVD_SCOPE_TRACE_SOLVER_FRAME
+		#define CVD_SCOPE_TRACE_SOLVER_FRAME(SolverType, SolverRef)
+	#endif
+
+	#ifndef CVD_TRACE_SOLVER_STEP_START
+		#define CVD_TRACE_SOLVER_STEP_START()
+	#endif
+
+	#ifndef CVD_TRACE_SOLVER_STEP_END
+		#define CVD_TRACE_SOLVER_STEP_END()
+	#endif
 
-#if UE_TRACE_ENABLED
-#	define CHAOS_VISUAL_DEBUGGER_ENABLED 0   // Disabled while under development
+	#ifndef CVD_SCOPE_TRACE_SOLVER_STEP
+		#define CVD_SCOPE_TRACE_SOLVER_STEP()
+	#endif
 #else
-#	define CHAOS_VISUAL_DEBUGGER_ENABLED 0
+
+#include "Trace/Trace.h"
+#include "Trace/Trace.inl"
+
+#ifndef CVD_DEFINE_TRACE_VECTOR
+	#define CVD_DEFINE_TRACE_VECTOR(Type, Name) \
+		UE_TRACE_EVENT_FIELD(Type, Name##X) \
+		UE_TRACE_EVENT_FIELD(Type, Name##Y) \
+		UE_TRACE_EVENT_FIELD(Type, Name##Z)
+#endif
+
+#ifndef CVD_DEFINE_TRACE_ROTATOR
+	#define CVD_DEFINE_TRACE_ROTATOR(Type, Name) \
+	UE_TRACE_EVENT_FIELD(Type, Name##X) \
+	UE_TRACE_EVENT_FIELD(Type, Name##Y) \
+	UE_TRACE_EVENT_FIELD(Type, Name##Z) \
+	UE_TRACE_EVENT_FIELD(Type, Name##W)
+#endif
+
+#ifndef CVD_TRACE_VECTOR_ON_EVENT
+	#define CVD_TRACE_VECTOR_ON_EVENT(EventName, Name, Vector) \
+	EventName.Name##X(Vector.X) \
+	<< EventName.Name##Y(Vector.Y) \
+	<< EventName.Name##Z(Vector.Z)
+#endif
+
+#ifndef CVD_TRACE_ROTATOR_ON_EVENT
+	#define CVD_TRACE_ROTATOR_ON_EVENT(EventName, Name, Rotator) \
+	EventName.Name##X(Rotator.X) \
+	<< EventName.Name##Y(Rotator.Y) \
+	<< EventName.Name##Z(Rotator.Z) \
+	<< EventName.Name##W(Rotator.W)
+#endif
+
+UE_TRACE_CHANNEL_EXTERN(ChaosVDChannel, CHAOS_API)
+
+UE_TRACE_EVENT_BEGIN_EXTERN(ChaosVDLogger, ChaosVDSolverFrameStart)
+	UE_TRACE_EVENT_FIELD(int32, SolverID)
+	UE_TRACE_EVENT_FIELD(uint64, Cycle)
+	UE_TRACE_EVENT_FIELD(UE::Trace::WideString, DebugName)
+UE_TRACE_EVENT_END()
+
+UE_TRACE_EVENT_BEGIN_EXTERN(ChaosVDLogger, ChaosVDSolverFrameEnd)
+	UE_TRACE_EVENT_FIELD(int32, SolverID)
+	UE_TRACE_EVENT_FIELD(uint64, Cycle)
+UE_TRACE_EVENT_END()
+
+UE_TRACE_EVENT_BEGIN_EXTERN(ChaosVDLogger, ChaosVDParticle)
+	UE_TRACE_EVENT_FIELD(int32, SolverID)
+	UE_TRACE_EVENT_FIELD(uint64, Cycle)
+	UE_TRACE_EVENT_FIELD(int32, ParticleID)
+	UE_TRACE_EVENT_FIELD(uint8, ParticleType)
+	CVD_DEFINE_TRACE_VECTOR(Chaos::FReal, Position)
+	CVD_DEFINE_TRACE_ROTATOR(Chaos::FReal, Rotation)
+	CVD_DEFINE_TRACE_VECTOR(Chaos::FReal, Velocity)
+	CVD_DEFINE_TRACE_VECTOR(Chaos::FReal, AngularVelocity)
+	UE_TRACE_EVENT_FIELD(int8, ObjectState)
+	UE_TRACE_EVENT_FIELD(UE::Trace::WideString, DebugName)
+UE_TRACE_EVENT_END()
+
+UE_TRACE_EVENT_BEGIN_EXTERN(ChaosVDLogger, ChaosVDSolverStepStart)
+	UE_TRACE_EVENT_FIELD(int32, SolverID)
+	UE_TRACE_EVENT_FIELD(uint64, Cycle)
+	UE_TRACE_EVENT_FIELD(uint32, StepNumber)
+UE_TRACE_EVENT_END()
+
+UE_TRACE_EVENT_BEGIN_EXTERN(ChaosVDLogger, ChaosVDSolverStepEnd)
+	UE_TRACE_EVENT_FIELD(int32, SolverID)
+	UE_TRACE_EVENT_FIELD(uint64, Cycle)
+	UE_TRACE_EVENT_FIELD(uint32, StepNumber)
+UE_TRACE_EVENT_END()
+
+#ifndef CVD_TRACE_PARTICLE
+	#define CVD_TRACE_PARTICLE(ParticleHandle) \
+		FChaosVisualDebuggerTrace::TraceParticle(ParticleHandle);
+#endif
+
+#ifndef CVD_TRACE_PARTICLES
+	#define CVD_TRACE_PARTICLES(ParticleHandles) \
+		FChaosVisualDebuggerTrace::TraceParticles(ParticleHandles)
+#endif
+
+#ifndef CVD_TRACE_SOLVER_START_FRAME
+	#define CVD_TRACE_SOLVER_START_FRAME(SolverType, SolverRef) \
+		FChaosVDContext StartContextData; \
+		FChaosVisualDebuggerTrace::GetCVDContext<SolverType>(SolverRef, StartContextData); \
+		FChaosVisualDebuggerTrace::TraceSolverFrameStart(StartContextData, FChaosVisualDebuggerTrace::GetDebugName<SolverType>(SolverRef));
+#endif
+
+#ifndef CVD_TRACE_SOLVER_END_FRAME
+	#define CVD_TRACE_SOLVER_END_FRAME(SolverType, SolverRef) \
+		FChaosVDContext EndContextData; \
+		FChaosVisualDebuggerTrace::GetCVDContext<SolverType>(SolverRef, EndContextData); \
+		FChaosVisualDebuggerTrace::TraceSolverFrameEnd(EndContextData);
+#endif
+
+#ifndef CVD_SCOPE_TRACE_SOLVER_FRAME
+	#define CVD_SCOPE_TRACE_SOLVER_FRAME(SolverType, SolverRef) \
+		FChaosVDScopeSolverFrame<SolverType> ScopeSolverFrame(SolverRef);
+#endif
+
+#ifndef CVD_TRACE_SOLVER_STEP_START
+	#define CVD_TRACE_SOLVER_STEP_START() \
+		FChaosVisualDebuggerTrace::TraceSolverStepStart();
+#endif
+
+#ifndef CVD_TRACE_SOLVER_STEP_END
+	#define CVD_TRACE_SOLVER_STEP_END() \
+		FChaosVisualDebuggerTrace::TraceSolverStepEnd();
 #endif
 
+#ifndef CVD_SCOPE_TRACE_SOLVER_STEP
+	#define CVD_SCOPE_TRACE_SOLVER_STEP() \
+		FChaosVDScopeSolverStep ScopeSolverStep;
+#endif
 
-// Note: This is the absolute minimum runtime interface for implementing a vertical slice of the visual debugger
-class ChaosVisualDebugger
+struct FChaosVDContext;
+
+namespace Chaos
+{
+	class FPhysicsSolverBase;
+	template <typename T, int d>
+	class TGeometryParticleHandles;
+}
+
+/** Class containing  all the Tracing logic to record data for the Chaos Visual Debugger tool */
+class FChaosVisualDebuggerTrace
 {
 public:
-	static void CHAOS_API ParticlePositionLog(const Chaos::FVec3& Position);
+
+	static CHAOS_API void TraceParticle(const Chaos::FGeometryParticleHandle* ParticleHandle);
+	static CHAOS_API void TraceParticle(const Chaos::FGeometryParticleHandle* ParticleHandle, const FChaosVDContext& ContextData);
+	static CHAOS_API void TraceParticles(const Chaos::TGeometryParticleHandles<Chaos::FReal, 3>& ParticleHandles);
+	static CHAOS_API void TraceSolverFrameStart(const FChaosVDContext& ContextData, const FString& InDebugName);
+	static CHAOS_API void TraceSolverFrameEnd(const FChaosVDContext& ContextData);
+	static CHAOS_API void TraceSolverStepStart();
+	static CHAOS_API void TraceSolverStepEnd();
+
+	template<typename T>
+	static void GetCVDContext(T& ObjectWithContext, FChaosVDContext& OutCVDContext);
+
+	template<typename T>
+	static FString GetDebugName(T& ObjectWithContext);
+};
+
+template <typename T>
+void FChaosVisualDebuggerTrace::GetCVDContext(T& ObjectWithContext, FChaosVDContext& OutCVDContext)
+{
+	OutCVDContext = ObjectWithContext.GetChaosVDContextData();	
+}
+
+template <typename T>
+FString FChaosVisualDebuggerTrace::GetDebugName(T& ObjectWithDebugName)
+{
+#if CHAOS_DEBUG_NAME
+	return ObjectWithDebugName.GetDebugName().ToString();
+#else
+	return FString("COMPILED OUT");
+#endif
+}
+
+struct FChaosVDScopeSolverStep
+{
+	FChaosVDScopeSolverStep()
+	{
+		CVD_TRACE_SOLVER_STEP_START();
+	}
+
+	~FChaosVDScopeSolverStep()
+	{
+		CVD_TRACE_SOLVER_STEP_END();
+	}
+};
+
+template<typename T>
+struct FChaosVDScopeSolverFrame
+{
+	FChaosVDScopeSolverFrame(T& InSolverRef) : SolverRef(InSolverRef)
+	{
+		CVD_TRACE_SOLVER_START_FRAME(T, SolverRef);
+	}
+
+	~FChaosVDScopeSolverFrame()
+	{
+		CVD_TRACE_SOLVER_END_FRAME(T, SolverRef);
+	}
+
+	T& SolverRef;
 };
+#endif // CHAOS_VISUAL_DEBUGGER_ENABLED
diff --git a/Engine/Source/Runtime/Experimental/Chaos/Public/PBDRigidsSolver.h b/Engine/Source/Runtime/Experimental/Chaos/Public/PBDRigidsSolver.h
--- a/Engine/Source/Runtime/Experimental/Chaos/Public/PBDRigidsSolver.h
+++ b/Engine/Source/Runtime/Experimental/Chaos/Public/PBDRigidsSolver.h
@@ -246,7 +246,8 @@ namespace Chaos
 		void PostTickDebugDraw(FReal Dt) const;
 
 		// Visual debugger (VDB) push methods
-		void PostEvolutionVDBPush() const;
+		UE_DEPRECATED(5.4, "This method will be removed in the future")
+		void PostEvolutionVDBPush() const {};
 
 		TArray<FGeometryCollectionPhysicsProxy*>& GetGeometryCollectionPhysicsProxies_Internal()
 		{
diff --git a/Engine/Source/Runtime/PhysicsCore/Private/ChaosScene.cpp b/Engine/Source/Runtime/PhysicsCore/Private/ChaosScene.cpp
--- a/Engine/Source/Runtime/PhysicsCore/Private/ChaosScene.cpp
+++ b/Engine/Source/Runtime/PhysicsCore/Private/ChaosScene.cpp
@@ -12,6 +12,10 @@
 #include "ChaosLog.h"
 #include "ChaosStats.h"
 
+#if WITH_CHAOS_VISUAL_DEBUGGER
+#include "ChaosVDRuntimeModule.h"
+#endif
+
 #include "Field/FieldSystem.h"
 
 #include "PhysicsProxy/PerSolverFieldSystem.h"
@@ -32,6 +36,7 @@
 #include "PhysicsSettingsCore.h"
 #include "Chaos/PhysicsSolverBaseImpl.h"
 
+#include "ChaosVisualDebugger/ChaosVisualDebuggerTrace.h"
 
 DECLARE_CYCLE_STAT(TEXT("Update Kinematics On Deferred SkelMeshes"),STAT_UpdateKinematicsOnDeferredSkelMeshesChaos,STATGROUP_Physics);
 CSV_DEFINE_CATEGORY(ChaosPhysics,true);
@@ -74,6 +79,11 @@ FChaosScene::FChaosScene(
 		);
 	check(SceneSolver);
 
+#if WITH_CHAOS_VISUAL_DEBUGGER
+	SceneSolver->GetChaosVDContextData().OwnerID = GetChaosVDContextData().Id;
+	SceneSolver->GetChaosVDContextData().Id = FChaosVDRuntimeModule::Get().GenerateUniqueID();
+#endif
+
 	SceneSolver->PhysSceneHack = this;
 	SimCallback = SceneSolver->CreateAndRegisterSimCallbackObject_External<FChaosSceneSimCallback>();
 
diff --git a/Engine/Source/Runtime/PhysicsCore/Public/Chaos/ChaosScene.h b/Engine/Source/Runtime/PhysicsCore/Public/Chaos/ChaosScene.h
--- a/Engine/Source/Runtime/PhysicsCore/Public/Chaos/ChaosScene.h
+++ b/Engine/Source/Runtime/PhysicsCore/Public/Chaos/ChaosScene.h
@@ -18,6 +18,8 @@
 #include "Chaos/SimCallbackInput.h"
 #include "Chaos/SimCallbackObject.h"
 
+#include "ChaosVisualDebugger/ChaosVDContextProvider.h"
+
 #ifndef CHAOS_DEBUG_NAME
 #define CHAOS_DEBUG_NAME 0
 #endif
@@ -189,4 +191,15 @@ private:
 	FGraphEventArray CompletionEvents;
 
 	FChaosSceneSimCallback* SimCallback;
+
+#if WITH_CHAOS_VISUAL_DEBUGGER
+private:
+	FChaosVDContext CVDContextData;
+
+public:
+	FChaosVDContext& GetChaosVDContextData()
+	{
+		return CVDContextData;
+	};
+#endif
 };
