Chaos Cloth - Added a new agnostic weight map system.

#rnx

diff --git a/Engine/Plugins/Experimental/ChaosCloth/Source/Chaos/Private/ChaosCloth/ChaosClothConstraints.cpp b/Engine/Plugins/Experimental/ChaosCloth/Source/Chaos/Private/ChaosCloth/ChaosClothConstraints.cpp
--- a/Engine/Plugins/Experimental/ChaosCloth/Source/Chaos/Private/ChaosCloth/ChaosClothConstraints.cpp
+++ b/Engine/Plugins/Experimental/ChaosCloth/Source/Chaos/Private/ChaosCloth/ChaosClothConstraints.cpp
@@ -101,11 +101,38 @@ void FClothConstraints::Enable(bool bEnable)
 void FClothConstraints::AddRules(
 	const Softs::FCollectionPropertyConstFacade& ConfigProperties,
 	const FTriangleMesh& TriangleMesh,
-	const TArray<TConstArrayView<FRealSingle>>& WeightMaps,
+	const TArray<TConstArrayView<FRealSingle>>& WeightMapArray,
 	const TArray<TConstArrayView<TTuple<int32, int32, FRealSingle>>>& Tethers,
 	Softs::FSolverReal MeshScale,
-	bool bEnabled,
-	const FClothingPatternData* PatternData)
+	bool bEnabled)
+{
+	// Build new weight map container from the old legacy weight map enum
+	TMap<FString, TConstArrayView<FRealSingle>> WeightMaps;
+
+	const UEnum* const ChaosWeightMapTargetEnum = StaticEnum<EChaosWeightMapTarget>();
+	const int32 NumWeightMaps = (int32)ChaosWeightMapTargetEnum->GetMaxEnumValue() + 1;
+
+	WeightMaps.Reserve(NumWeightMaps);
+
+	for (int32 EnumIndex = 0; EnumIndex < ChaosWeightMapTargetEnum->NumEnums(); ++EnumIndex)
+	{
+		const int32 TargetIndex = (int32)ChaosWeightMapTargetEnum->GetValueByIndex(EnumIndex);
+		const FString WeightMapName = ChaosWeightMapTargetEnum->GetNameByIndex(EnumIndex).ToString();
+
+		WeightMaps.Add(WeightMapName, WeightMapArray[TargetIndex]);
+	}
+
+	// Call new AddRules function
+	AddRules(ConfigProperties, TriangleMesh, nullptr, WeightMaps, Tethers, MeshScale, bEnabled);
+}
+
+void FClothConstraints::AddRules(
+	const Softs::FCollectionPropertyConstFacade& ConfigProperties,
+	const FTriangleMesh& TriangleMesh,
+	const FClothingPatternData* PatternData,
+	const TMap<FString, TConstArrayView<FRealSingle>>& WeightMaps,
+	const TArray<TConstArrayView<TTuple<int32, int32, FRealSingle>>>& Tethers,
+	Softs::FSolverReal MeshScale, bool bEnabled)
 {
 	// Self collisions
 	CreateSelfCollisionConstraints(ConfigProperties, TriangleMesh);
@@ -190,29 +217,19 @@ void FClothConstraints::CreateSelfCollisionConstraints(const Softs::FCollectionP
 
 void FClothConstraints::CreateStretchConstraints(
 	const Softs::FCollectionPropertyConstFacade& ConfigProperties,
-	const TArray<TConstArrayView<FRealSingle>>& WeightMaps,
+	const TMap<FString, TConstArrayView<FRealSingle>>& WeightMaps,
 	const FTriangleMesh& TriangleMesh,
 	const FClothingPatternData* PatternData)
 {
 	if (PatternData && PatternData->PatternPositions.Num() && Softs::FXPBDStretchBiasElementConstraints::IsEnabled(ConfigProperties))
 	{
-		// TODO: separate warp, weft, Bias multipliers
-		const TConstArrayView<FRealSingle>& EdgeStiffnessMultipliers = WeightMaps[(int32)EChaosWeightMapTarget::EdgeStiffness];
-		const TConstArrayView<FRealSingle> DampingMultipliers;  // TODO: Damping multiplier
-		const TConstArrayView<FRealSingle> WarpWeftScaleMultipliers; // TODO: separate multipliers for warp and weft scale
-
 		XStretchBiasConstraints = MakeShared<Softs::FXPBDStretchBiasElementConstraints>(
 			Evolution->Particles(),
 			ParticleOffset,
 			NumParticles,
 			TriangleMesh,
 			PatternData->WeldedFaceVertexPatternPositions,
-			EdgeStiffnessMultipliers,
-			EdgeStiffnessMultipliers,
-			EdgeStiffnessMultipliers,
-			DampingMultipliers,
-			WarpWeftScaleMultipliers,
-			WarpWeftScaleMultipliers,
+			WeightMaps,
 			ConfigProperties,
 			/*bTrimKinematicConstraints =*/ true);
 
@@ -221,16 +238,12 @@ void FClothConstraints::CreateStretchConstraints(
 	}
 	else if (Softs::FXPBDEdgeSpringConstraints::IsEnabled(ConfigProperties))
 	{
-		const TConstArrayView<FRealSingle>& EdgeStiffnessMultipliers = WeightMaps[(int32)EChaosWeightMapTarget::EdgeStiffness];
-		const TConstArrayView<FRealSingle> DampingMultipliers;  // TODO: Damping multiplier
-
 		XEdgeConstraints = MakeShared<Softs::FXPBDEdgeSpringConstraints>(
 			Evolution->Particles(),
 			ParticleOffset,
 			NumParticles,
 			TriangleMesh.GetSurfaceElements(),
-			EdgeStiffnessMultipliers,
-			DampingMultipliers,
+			WeightMaps,
 			ConfigProperties,
 			/*bTrimKinematicConstraints =*/ true);
 
@@ -239,14 +252,12 @@ void FClothConstraints::CreateStretchConstraints(
 	}
 	else if (Softs::FPBDEdgeSpringConstraints::IsEnabled(ConfigProperties))
 	{
-		const TConstArrayView<FRealSingle>& EdgeStiffnessMultipliers = WeightMaps[(int32)EChaosWeightMapTarget::EdgeStiffness];
-
 		EdgeConstraints = MakeShared<Softs::FPBDEdgeSpringConstraints>(
 			Evolution->Particles(),
 			ParticleOffset,
 			NumParticles,
 			TriangleMesh.GetSurfaceElements(),
-			EdgeStiffnessMultipliers,
+			WeightMaps,
 			ConfigProperties,
 			/*bTrimKinematicConstraints =*/ true);
 
@@ -257,27 +268,18 @@ void FClothConstraints::CreateStretchConstraints(
 
 void FClothConstraints::CreateBendingConstraints(
 	const Softs::FCollectionPropertyConstFacade& ConfigProperties,
-	const TArray<TConstArrayView<FRealSingle>>& WeightMaps,
+	const TMap<FString, TConstArrayView<FRealSingle>>& WeightMaps,
 	const FTriangleMesh& TriangleMesh,
 	const FClothingPatternData* PatternData)
 {
 	if (PatternData && PatternData->PatternPositions.Num() && Softs::FXPBDAnisotropicBendingConstraints::IsEnabled(ConfigProperties))
 	{
-		const TConstArrayView<FRealSingle>& BendingStiffnessMultipliers = WeightMaps[(int32)EChaosWeightMapTarget::BendingStiffness];
-		const TConstArrayView<FRealSingle>& BucklingStiffnessMultipliers = WeightMaps[(int32)EChaosWeightMapTarget::BucklingStiffness];
-		const TConstArrayView<FRealSingle> DampingMultipliers;
 		XAnisoBendingElementConstraints = MakeShared<Softs::FXPBDAnisotropicBendingConstraints>(
 			Evolution->Particles(),
 			ParticleOffset, NumParticles,
 			TriangleMesh,
 			PatternData->WeldedFaceVertexPatternPositions,
-			BendingStiffnessMultipliers,
-			BendingStiffnessMultipliers,
-			BendingStiffnessMultipliers,
-			BucklingStiffnessMultipliers,
-			BucklingStiffnessMultipliers,
-			BucklingStiffnessMultipliers,
-			DampingMultipliers,
+			WeightMaps,
 			ConfigProperties,
 			/*bTrimKinematicConstraints =*/ true);
 
@@ -287,17 +289,12 @@ void FClothConstraints::CreateBendingConstraints(
 	else if (Softs::FXPBDBendingConstraints::IsEnabled(ConfigProperties))
 	{
 		TArray<Chaos::TVec4<int32>> BendingElements = TriangleMesh.GetUniqueAdjacentElements();
-		const TConstArrayView<FRealSingle>& BendingStiffnessMultipliers = WeightMaps[(int32)EChaosWeightMapTarget::BendingStiffness];
-		const TConstArrayView<FRealSingle>& BucklingStiffnessMultipliers = WeightMaps[(int32)EChaosWeightMapTarget::BucklingStiffness];
-		const TConstArrayView<FRealSingle> DampingMultipliers;
 
 		XBendingElementConstraints = MakeShared<Softs::FXPBDBendingConstraints>(
 			Evolution->Particles(),
 			ParticleOffset, NumParticles,
 			MoveTemp(BendingElements),
-			BendingStiffnessMultipliers,
-			BucklingStiffnessMultipliers,
-			DampingMultipliers,
+			WeightMaps,
 			ConfigProperties,
 			/*bTrimKinematicConstraints =*/ true);
 
@@ -307,15 +304,12 @@ void FClothConstraints::CreateBendingConstraints(
 	else if (Softs::FPBDBendingConstraints::IsEnabled(ConfigProperties))
 	{
 		TArray<Chaos::TVec4<int32>> BendingElements = TriangleMesh.GetUniqueAdjacentElements();
-		const TConstArrayView<FRealSingle>& BendingStiffnessMultipliers = WeightMaps[(int32)EChaosWeightMapTarget::BendingStiffness];
-		const TConstArrayView<FRealSingle>& BucklingStiffnessMultipliers = WeightMaps[(int32)EChaosWeightMapTarget::BucklingStiffness];
 
 		BendingElementConstraints = MakeShared<Softs::FPBDBendingConstraints>(
 			Evolution->Particles(),
 			ParticleOffset, NumParticles,
 			MoveTemp(BendingElements),
-			BendingStiffnessMultipliers,
-			BucklingStiffnessMultipliers,
+			WeightMaps,
 			ConfigProperties,
 			/*bTrimKinematicConstraints =*/ true);
 
@@ -325,15 +319,12 @@ void FClothConstraints::CreateBendingConstraints(
 	else if (Softs::FXPBDBendingSpringConstraints::IsEnabled(ConfigProperties))
 	{
 		const TArray<Chaos::TVec2<int32>> CrossEdges = TriangleMesh.GetUniqueAdjacentPoints();
-		const TConstArrayView<FRealSingle>& BendingStiffnessMultipliers = WeightMaps[(int32)EChaosWeightMapTarget::BendingStiffness];
-		const TConstArrayView<FRealSingle> DampingMultipliers;
 
 		XBendingConstraints = MakeShared<Softs::FXPBDBendingSpringConstraints>(
 			Evolution->Particles(),
 			ParticleOffset, NumParticles,
 			CrossEdges,
-			BendingStiffnessMultipliers,
-			DampingMultipliers,
+			WeightMaps,
 			ConfigProperties,
 			/*bTrimKinematicConstraints =*/ true);
 
@@ -343,14 +334,13 @@ void FClothConstraints::CreateBendingConstraints(
 	else if (Softs::FPBDBendingSpringConstraints::IsEnabled(ConfigProperties))
 	{
 		const TArray<Chaos::TVec2<int32>> CrossEdges = TriangleMesh.GetUniqueAdjacentPoints();
-		const TConstArrayView<FRealSingle>& BendingStiffnessMultipliers = WeightMaps[(int32)EChaosWeightMapTarget::BendingStiffness];
 
 		BendingConstraints = MakeShared<Softs::FPBDBendingSpringConstraints>(
 			Evolution->Particles(),
 			ParticleOffset,
 			NumParticles,
 			CrossEdges,
-			BendingStiffnessMultipliers,
+			WeightMaps,
 			ConfigProperties,
 			/*bTrimKinematicConstraints =*/ true);
 
@@ -361,19 +351,17 @@ void FClothConstraints::CreateBendingConstraints(
 
 void FClothConstraints::CreateAreaConstraints(
 	const Softs::FCollectionPropertyConstFacade& ConfigProperties,
-	const TArray<TConstArrayView<FRealSingle>>& WeightMaps,
+	const TMap<FString, TConstArrayView<FRealSingle>>& WeightMaps,
 	const FTriangleMesh& TriangleMesh)
 {
 	if (Softs::FXPBDAreaSpringConstraints::IsEnabled(ConfigProperties))
 	{
-		const TConstArrayView<FRealSingle>& AreaStiffnessMultipliers = WeightMaps[(int32)EChaosWeightMapTarget::AreaStiffness];
-
 		XAreaConstraints = MakeShared<Softs::FXPBDAreaSpringConstraints>(
 			Evolution->Particles(),
 			ParticleOffset,
 			NumParticles,
 			TriangleMesh.GetSurfaceElements(),
-			AreaStiffnessMultipliers,
+			WeightMaps,
 			ConfigProperties,
 			/*bTrimKinematicConstraints =*/ true);
 
@@ -382,14 +370,12 @@ void FClothConstraints::CreateAreaConstraints(
 	}
 	else if (Softs::FPBDAreaSpringConstraints::IsEnabled(ConfigProperties))
 	{
-		const TConstArrayView<FRealSingle>& AreaStiffnessMultipliers = WeightMaps[(int32)EChaosWeightMapTarget::AreaStiffness];
-
 		AreaConstraints = MakeShared<Softs::FPBDAreaSpringConstraints>(
 			Evolution->Particles(),
 			ParticleOffset,
 			NumParticles,
 			TriangleMesh.GetSurfaceElements(),
-			AreaStiffnessMultipliers,
+			WeightMaps,
 			ConfigProperties,
 			/*bTrimKinematicConstraints =*/ true);
 
@@ -400,23 +386,19 @@ void FClothConstraints::CreateAreaConstraints(
 
 void FClothConstraints::CreateLongRangeConstraints(
 	const Softs::FCollectionPropertyConstFacade& ConfigProperties,
-	const TArray<TConstArrayView<FRealSingle>>& WeightMaps,
+	const TMap<FString, TConstArrayView<FRealSingle>>& WeightMaps,
 	const TArray<TConstArrayView<TTuple<int32, int32, FRealSingle>>>& Tethers,
 	Softs::FSolverReal MeshScale)
 {
 	if (Softs::FPBDLongRangeConstraints::IsEnabled(ConfigProperties))
 	{
-		const TConstArrayView<FRealSingle>& TetherStiffnessMultipliers = WeightMaps[(int32)EChaosWeightMapTarget::TetherStiffness];
-		const TConstArrayView<FRealSingle>& TetherScaleMultipliers = WeightMaps[(int32)EChaosWeightMapTarget::TetherScale];
-
 		//  Now that we're only doing a single iteration of Long range constraints, and they're more of a fake constraint to jump start our initial guess, it's not clear that using XPBD makes sense here.
 		LongRangeConstraints = MakeShared<Softs::FPBDLongRangeConstraints>(
 			Evolution->Particles(),
 			ParticleOffset,
 			NumParticles,
 			Tethers,
-			TetherStiffnessMultipliers,
-			TetherScaleMultipliers,
+			WeightMaps,
 			ConfigProperties,
 			MeshScale);
 
@@ -426,17 +408,27 @@ void FClothConstraints::CreateLongRangeConstraints(
 
 void FClothConstraints::CreateMaxDistanceConstraints(
 	const Softs::FCollectionPropertyConstFacade& ConfigProperties,
-	const TArray<TConstArrayView<FRealSingle>>& WeightMaps,
+	const TMap<FString, TConstArrayView<FRealSingle>>& WeightMaps,
 	Softs::FSolverReal MeshScale)
 {
-	const TConstArrayView<FRealSingle>& MaxDistances = WeightMaps[(int32)EChaosWeightMapTarget::MaxDistance];
-	if (MaxDistances.Num() == NumParticles && Softs::FPBDSphericalConstraint::IsEnabled(ConfigProperties))
+	if (ConfigProperties.GetValue("UseLegacyConfig", false))
+	{
+		FString MaxDistanceString = TEXT("MaxDistance");
+		MaxDistanceString = ConfigProperties.GetStringValue(MaxDistanceString, MaxDistanceString);  // Uses the same string for both the default weight map name and the property name
+		const TConstArrayView<FRealSingle> MaxDistances = WeightMaps.FindRef(MaxDistanceString);
+		if (MaxDistances.Num() != NumParticles)
+		{
+			return;  // Legacy configs disable the constraint when the weight map is missing
+		}
+	}
+
+	if (Softs::FPBDSphericalConstraint::IsEnabled(ConfigProperties))
 	{
 		MaximumDistanceConstraints = MakeShared<Softs::FPBDSphericalConstraint>(
 			ParticleOffset,
 			NumParticles,
 			*AnimationPositions,
-			MaxDistances,
+			WeightMaps,
 			ConfigProperties,
 			MeshScale);
 
@@ -446,37 +438,45 @@ void FClothConstraints::CreateMaxDistanceConstraints(
 
 void FClothConstraints::CreateBackstopConstraints(
 	const Softs::FCollectionPropertyConstFacade& ConfigProperties,
-	const TArray<TConstArrayView<FRealSingle>>& WeightMaps,
+	const TMap<FString, TConstArrayView<FRealSingle>>& WeightMaps,
 	Softs::FSolverReal MeshScale)
 {
-	if (Softs::FPBDSphericalBackstopConstraint::IsEnabled(ConfigProperties))
+	if (ConfigProperties.GetValue("UseLegacyConfig", false))
 	{
-		const TConstArrayView<FRealSingle>& BackstopDistances = WeightMaps[(int32)EChaosWeightMapTarget::BackstopDistance];
-		const TConstArrayView<FRealSingle>& BackstopRadiuses = WeightMaps[(int32)EChaosWeightMapTarget::BackstopRadius];
+		FString BackstopRadiusString = TEXT("BackstopRadius");
+		FString BackstopDistanceString = TEXT("BackstopDistance");
+		BackstopRadiusString = ConfigProperties.GetStringValue(BackstopRadiusString, BackstopRadiusString);        // Uses the same string for both the default weight map name and the property name
+		BackstopDistanceString = ConfigProperties.GetStringValue(BackstopDistanceString, BackstopDistanceString);  //
+		const TConstArrayView<FRealSingle> BackstopRadiuses = WeightMaps.FindRef(BackstopRadiusString);
+		const TConstArrayView<FRealSingle> BackstopDistances = WeightMaps.FindRef(BackstopDistanceString);
 
-		if (BackstopRadiuses.Num() == NumParticles && BackstopDistances.Num() == NumParticles)
+		if (BackstopRadiuses.Num() != NumParticles || BackstopDistances.Num() != NumParticles)
 		{
-			BackstopConstraints = MakeShared<Softs::FPBDSphericalBackstopConstraint>(
-				ParticleOffset,
-				NumParticles,
-				*AnimationPositions,
-				*AnimationNormals,
-				BackstopRadiuses,
-				BackstopDistances,
-				ConfigProperties,
-				MeshScale);
-
-			++NumConstraintRules;
+			return;  // Legacy configs disable the constraint when the weight maps are missing
 		}
 	}
+
+	if (Softs::FPBDSphericalBackstopConstraint::IsEnabled(ConfigProperties))
+	{
+		BackstopConstraints = MakeShared<Softs::FPBDSphericalBackstopConstraint>(
+			ParticleOffset,
+			NumParticles,
+			*AnimationPositions,
+			*AnimationNormals,
+			WeightMaps,
+			ConfigProperties,
+			MeshScale);
+
+		++NumConstraintRules;
+	}
 }
 
-void FClothConstraints::CreateAnimDriveConstraints(const Softs::FCollectionPropertyConstFacade& ConfigProperties, const TArray<TConstArrayView<FRealSingle>>& WeightMaps)
+void FClothConstraints::CreateAnimDriveConstraints(
+	const Softs::FCollectionPropertyConstFacade& ConfigProperties,
+	const TMap<FString, TConstArrayView<FRealSingle>>& WeightMaps)
 {
 	if (Softs::FPBDAnimDriveConstraint::IsEnabled(ConfigProperties))
 	{
-		const TConstArrayView<FRealSingle>& AnimDriveStiffnessMultipliers = WeightMaps[(int32)EChaosWeightMapTarget::AnimDriveStiffness];
-		const TConstArrayView<FRealSingle>& AnimDriveDampingMultipliers = WeightMaps[(int32)EChaosWeightMapTarget::AnimDriveDamping];
 		check(AnimationVelocities); // Legacy code didn't use to have AnimationVelocities
 
 		AnimDriveConstraints = MakeShared<Softs::FPBDAnimDriveConstraint>(
@@ -484,8 +484,7 @@ void FClothConstraints::CreateAnimDriveConstraints(const Softs::FCollectionPrope
 			NumParticles,
 			*AnimationPositions,
 			*AnimationVelocities,
-			AnimDriveStiffnessMultipliers,
-			AnimDriveDampingMultipliers,
+			WeightMaps,
 			ConfigProperties);
 
 		++NumConstraintInits;  // Uses init to update the property tables
@@ -1189,52 +1188,53 @@ void FClothConstraints::SetSelfCollisionConstraints(const FTriangleMesh& Triangl
 
 void FClothConstraints::Update(
 	const Softs::FCollectionPropertyConstFacade& ConfigProperties,
+	const TMap<FString, TConstArrayView<FRealSingle>>& WeightMaps,
 	Softs::FSolverReal MeshScale,
 	Softs::FSolverReal MaxDistancesScale)
 {
 	if (XEdgeConstraints)
 	{
-		XEdgeConstraints->SetProperties(ConfigProperties);
+		XEdgeConstraints->SetProperties(ConfigProperties, WeightMaps);
 	}
 	else if (EdgeConstraints)
 	{
-		EdgeConstraints->SetProperties(ConfigProperties);
+		EdgeConstraints->SetProperties(ConfigProperties, WeightMaps);
 	}
 	if (XBendingConstraints)
 	{
-		XBendingConstraints->SetProperties(ConfigProperties);
+		XBendingConstraints->SetProperties(ConfigProperties, WeightMaps);
 	}
 	else if (BendingConstraints)
 	{
-		BendingConstraints->SetProperties(ConfigProperties);
+		BendingConstraints->SetProperties(ConfigProperties, WeightMaps);
 	}
 	if (XBendingElementConstraints)
 	{
-		XBendingElementConstraints->SetProperties(ConfigProperties);
+		XBendingElementConstraints->SetProperties(ConfigProperties, WeightMaps);
 	}
 	else if (BendingElementConstraints)
 	{
-		BendingElementConstraints->SetProperties(ConfigProperties);
+		BendingElementConstraints->SetProperties(ConfigProperties, WeightMaps);
 	}
 	if (XAreaConstraints)
 	{
-		XAreaConstraints->SetProperties(ConfigProperties);
+		XAreaConstraints->SetProperties(ConfigProperties, WeightMaps);
 	}
 	else if (AreaConstraints)
 	{
-		AreaConstraints->SetProperties(ConfigProperties);
+		AreaConstraints->SetProperties(ConfigProperties, WeightMaps);
 	}
 	if (LongRangeConstraints)
 	{
-		LongRangeConstraints->SetProperties(ConfigProperties, MeshScale);
+		LongRangeConstraints->SetProperties(ConfigProperties, WeightMaps, MeshScale);
 	}
 	if (MaximumDistanceConstraints)
 	{
-		MaximumDistanceConstraints->SetProperties(ConfigProperties, MeshScale * MaxDistancesScale);
+		MaximumDistanceConstraints->SetProperties(ConfigProperties, WeightMaps, MeshScale * MaxDistancesScale);
 	}
 	if (AnimDriveConstraints)
 	{
-		AnimDriveConstraints->SetProperties(ConfigProperties);
+		AnimDriveConstraints->SetProperties(ConfigProperties, WeightMaps);
 	}
 	if (SelfCollisionConstraints)
 	{
@@ -1242,10 +1242,19 @@ void FClothConstraints::Update(
 	}
 	if (BackstopConstraints)
 	{
-		BackstopConstraints->SetProperties(ConfigProperties, MeshScale);
+		BackstopConstraints->SetProperties(ConfigProperties, WeightMaps, MeshScale);
 	}
 }
 
+// Deprecated
+void FClothConstraints::Update(
+	const Softs::FCollectionPropertyConstFacade& ConfigProperties,
+	Softs::FSolverReal MeshScale,
+	Softs::FSolverReal MaxDistancesScale)
+{
+	Update(ConfigProperties, TMap<FString, TConstArrayView<FRealSingle>>(), MeshScale, MaxDistancesScale);
+}
+
 void FClothConstraints::SetEdgeProperties(const Softs::FSolverVec2& EdgeStiffness, const Softs::FSolverVec2& DampingRatio)
 {
 	if (EdgeConstraints)
@@ -1321,8 +1330,7 @@ void FClothConstraints::SetMaximumDistanceProperties(Softs::FSolverReal MeshScal
 {
 	if (MaximumDistanceConstraints)
 	{
-		constexpr Softs::FSolverReal MaxDistancesMultiplier = (Softs::FSolverReal)1.;
-		MaximumDistanceConstraints->SetScale(MaxDistancesMultiplier, MeshScale);
+		MaximumDistanceConstraints->SetScale(MeshScale);
 	}
 }
 
@@ -1353,8 +1361,7 @@ void FClothConstraints::SetBackstopProperties(bool bEnabled, Softs::FSolverReal
 	if (BackstopConstraints)
 	{
 		BackstopConstraints->SetEnabled(bEnabled);
-		constexpr Softs::FSolverReal BackstopDistancesMultiplier = (Softs::FSolverReal)1.;
-		BackstopConstraints->SetScale(BackstopDistancesMultiplier, MeshScale);
+		BackstopConstraints->SetScale(MeshScale);
 	}
 }
 
diff --git a/Engine/Plugins/Experimental/ChaosCloth/Source/Chaos/Private/ChaosCloth/ChaosClothVisualization.cpp b/Engine/Plugins/Experimental/ChaosCloth/Source/Chaos/Private/ChaosCloth/ChaosClothVisualization.cpp
--- a/Engine/Plugins/Experimental/ChaosCloth/Source/Chaos/Private/ChaosCloth/ChaosClothVisualization.cpp
+++ b/Engine/Plugins/Experimental/ChaosCloth/Source/Chaos/Private/ChaosCloth/ChaosClothVisualization.cpp
@@ -234,7 +234,8 @@ namespace Chaos
 				continue;
 			}
 
-			const TConstArrayView<FRealSingle>& MaxDistances = Cloth->GetWeightMaps(Solver)[(int32)EChaosWeightMapTarget::MaxDistance];
+
+			const TConstArrayView<FRealSingle>& MaxDistances = Cloth->GetWeightMapByProperty(Solver, TEXT("MaxDistance"));
 			if (!MaxDistances.Num())
 			{
 				continue;
@@ -1057,8 +1058,9 @@ namespace Chaos
 			if (const Softs::FPBDSphericalBackstopConstraint* const BackstopConstraint = ClothConstraints.GetBackstopConstraints().Get())
 			{
 				const bool bUseLegacyBackstop = BackstopConstraint->UseLegacyBackstop();
-				const TConstArrayView<FRealSingle>& BackstopDistances = Cloth->GetWeightMaps(Solver)[(int32)EChaosWeightMapTarget::BackstopDistance];
-				const TConstArrayView<FRealSingle>& BackstopRadiuses = Cloth->GetWeightMaps(Solver)[(int32)EChaosWeightMapTarget::BackstopRadius];
+
+				const TConstArrayView<FRealSingle>& BackstopDistances = Cloth->GetWeightMapByProperty(Solver, TEXT("BackstopDistance"));
+				const TConstArrayView<FRealSingle>& BackstopRadiuses = Cloth->GetWeightMapByProperty(Solver, TEXT("BackstopRadius"));
 				const TConstArrayView<Softs::FSolverVec3> AnimationPositions = Cloth->GetAnimationPositions(Solver);
 				const TConstArrayView<Softs::FSolverVec3> AnimationNormals = Cloth->GetAnimationNormals(Solver);
 				const TConstArrayView<Softs::FSolverVec3> ParticlePositions = Cloth->GetParticlePositions(Solver);
@@ -1116,8 +1118,8 @@ namespace Chaos
 			if (const Softs::FPBDSphericalBackstopConstraint* const BackstopConstraint = ClothConstraints.GetBackstopConstraints().Get())
 			{
 				const bool bUseLegacyBackstop = BackstopConstraint->UseLegacyBackstop();
-				const TConstArrayView<FRealSingle>& BackstopDistances = Cloth->GetWeightMaps(Solver)[(int32)EChaosWeightMapTarget::BackstopDistance];
-				const TConstArrayView<FRealSingle>& BackstopRadiuses = Cloth->GetWeightMaps(Solver)[(int32)EChaosWeightMapTarget::BackstopRadius];
+				const TConstArrayView<FRealSingle>& BackstopDistances = Cloth->GetWeightMapByProperty(Solver, TEXT("BackstopDistance"));
+				const TConstArrayView<FRealSingle>& BackstopRadiuses = Cloth->GetWeightMapByProperty(Solver, TEXT("BackstopRadius"));
 				const TConstArrayView<Softs::FSolverVec3> AnimationPositions = Cloth->GetAnimationPositions(Solver);
 				const TConstArrayView<Softs::FSolverVec3> AnimationNormals = Cloth->GetAnimationNormals(Solver);
 
@@ -1159,7 +1161,7 @@ namespace Chaos
 				continue;
 			}
 
-			const TConstArrayView<FRealSingle>& MaxDistances = Cloth->GetWeightMaps(Solver)[(int32)EChaosWeightMapTarget::MaxDistance];
+			const TConstArrayView<FRealSingle>& MaxDistances = Cloth->GetWeightMapByProperty(Solver, TEXT("MaxDistance"));
 			if (!MaxDistances.Num())
 			{
 				continue;
@@ -1212,7 +1214,7 @@ namespace Chaos
 			const FClothConstraints& ClothConstraints = Solver->GetClothConstraints(Offset);
 			if (const Softs::FPBDAnimDriveConstraint* const AnimDriveConstraint = ClothConstraints.GetAnimDriveConstraints().Get())
 			{
-				const TConstArrayView<FRealSingle>& AnimDriveStiffnessMultipliers = Cloth->GetWeightMaps(Solver)[(int32)EChaosWeightMapTarget::AnimDriveStiffness];
+				const TConstArrayView<FRealSingle>& AnimDriveStiffnessMultipliers = Cloth->GetWeightMapByProperty(Solver, TEXT("AnimDriveStiffness"));
 				const TConstArrayView<Softs::FSolverVec3> AnimationPositions = Cloth->GetAnimationPositions(Solver);
 				const TConstArrayView<Softs::FSolverVec3> ParticlePositions = Cloth->GetParticlePositions(Solver);
 				check(ParticlePositions.Num() == AnimationPositions.Num());
diff --git a/Engine/Plugins/Experimental/ChaosCloth/Source/Chaos/Private/ChaosCloth/ChaosClothingSimulation.cpp b/Engine/Plugins/Experimental/ChaosCloth/Source/Chaos/Private/ChaosCloth/ChaosClothingSimulation.cpp
--- a/Engine/Plugins/Experimental/ChaosCloth/Source/Chaos/Private/ChaosCloth/ChaosClothingSimulation.cpp
+++ b/Engine/Plugins/Experimental/ChaosCloth/Source/Chaos/Private/ChaosCloth/ChaosClothingSimulation.cpp
@@ -369,7 +369,8 @@ void FClothingSimulation::CreateActor(USkeletalMeshComponent* InOwnerComponent,
 
 	// Create cloth config runtime simulation object
 	const int32 ClothConfigIndex = Configs.Emplace(MakeUnique<FClothingSimulationConfig>());
-	Configs[ClothConfigIndex]->Initialize(ClothConfig, nullptr);
+	constexpr bool bUseLegacyConfig = true;  // Make the config a legacy cloth config, so that the constraints disable themselves with missing masks, ...etc.
+	Configs[ClothConfigIndex]->Initialize(ClothConfig, nullptr, bUseLegacyConfig);
 
 	// Create cloth runtime simulation object
 	const int32 ClothIndex = Cloths.Emplace(MakeUnique<FClothingSimulationCloth>(
@@ -436,7 +437,8 @@ void FClothingSimulation::UpdateSimulationFromSharedSimConfig()
 			SolverConfig = Configs[SolverConfigIndex].Get();
 			Solver->SetConfig(SolverConfig);
 		}
-		SolverConfig->Initialize(nullptr, ClothSharedSimConfig);
+		constexpr bool bUseLegacyConfig = true;  // Make the config a legacy cloth config, so that the constraints disable themselves with missing masks, ...etc.
+		SolverConfig->Initialize(nullptr, ClothSharedSimConfig, bUseLegacyConfig);
 	}
 }
 
@@ -790,7 +792,8 @@ void FClothingSimulation::RefreshClothConfig(const IClothingSimulationContext* I
 		// Update cloth config runtime simulation object
 		FClothingSimulationConfig* const Config = Cloth->GetConfig();
 		check(Config);
-		Config->Initialize(ClothConfig, nullptr);
+		constexpr bool bUseLegacyConfig = true;  // Make the config a legacy cloth config, so that the constraints disable themselves with missing masks, ...etc.
+		Config->Initialize(ClothConfig, nullptr, bUseLegacyConfig);
 
 		// Recreate cloth runtime simulation object
 		Cloth = MakeUnique<FClothingSimulationCloth>(
diff --git a/Engine/Plugins/Experimental/ChaosCloth/Source/Chaos/Private/ChaosCloth/ChaosClothingSimulationCloth.cpp b/Engine/Plugins/Experimental/ChaosCloth/Source/Chaos/Private/ChaosCloth/ChaosClothingSimulationCloth.cpp
--- a/Engine/Plugins/Experimental/ChaosCloth/Source/Chaos/Private/ChaosCloth/ChaosClothingSimulationCloth.cpp
+++ b/Engine/Plugins/Experimental/ChaosCloth/Source/Chaos/Private/ChaosCloth/ChaosClothingSimulationCloth.cpp
@@ -1,4 +1,5 @@
 // Copyright Epic Games, Inc. All Rights Reserved.
+
 #include "ChaosCloth/ChaosClothingSimulationCloth.h"
 #include "ChaosCloth/ChaosClothingSimulationSolver.h"
 #include "ChaosCloth/ChaosClothingSimulationMesh.h"
@@ -27,6 +28,7 @@ namespace ClothingSimulationClothDefault
 	constexpr float DampingCoefficient = 0.01f;
 	constexpr float Drag = 0.035f;
 	constexpr float Lift = 0.035f;
+	constexpr float Pressure = 0.f;
 	constexpr float AirDensity = 1.225e-6f;
 	constexpr float GravityScale = 1.f;
 	constexpr float GravityZOverride = -980.665f;
@@ -44,7 +46,7 @@ struct FClothingSimulationCloth::FLODData
 	// Input mesh
 	const int32 NumParticles;
 	const TConstArrayView<uint32> Indices;
-	const TArray<TConstArrayView<FRealSingle>> WeightMaps;
+	const TMap<FString, TConstArrayView<FRealSingle>> WeightMaps;
 	const TArray<TConstArrayView<TTuple<int32, int32, FRealSingle>>> Tethers;
 
 	const FClothingPatternData PatternData;
@@ -67,7 +69,7 @@ struct FClothingSimulationCloth::FLODData
 		const TConstArrayView<FVector2f>& InPatternPositions,
 		const TConstArrayView<uint32>& InPatternIndices,
 		const TConstArrayView<uint32>& InPatternToWeldedIndices,
-		const TArray<TConstArrayView<FRealSingle>>& InWeightMaps,
+		TMap<FString, TConstArrayView<FRealSingle>>&& InWeightMaps,
 		const TArray<TConstArrayView<TTuple<int32, int32, FRealSingle>>>& InTethers);
 
 	void Add(FClothingSimulationSolver* Solver, FClothingSimulationCloth* Cloth, int32 LODIndex);
@@ -88,11 +90,11 @@ FClothingSimulationCloth::FLODData::FLODData(
 	const TConstArrayView<FVector2f>& InPatternPositions,
 	const TConstArrayView<uint32>& InPatternIndices,
 	const TConstArrayView<uint32>& InPatternToWeldedIndices,
-	const TArray<TConstArrayView<FRealSingle>>& InWeightMaps,
+	TMap<FString, TConstArrayView<FRealSingle>>&& InWeightMaps,
 	const TArray<TConstArrayView<TTuple<int32, int32, FRealSingle>>>& InTethers)
 	: NumParticles(InNumParticles)
 	, Indices(InIndices)
-	, WeightMaps(InWeightMaps)
+	, WeightMaps(MoveTemp(InWeightMaps))
 	, Tethers(InTethers)
 	, PatternData(NumParticles, Indices, InPatternPositions, InPatternIndices, InPatternToWeldedIndices)
 {
@@ -149,18 +151,38 @@ PRAGMA_ENABLE_DEPRECATION_WARNINGS
 	// Initialize the normals, in case the sim data is queried before the simulation steps
 	UpdateNormals(Solver);
 
+	// Retrieve config properties
+	check(Cloth->Config);
+	const Softs::FCollectionPropertyFacade& ConfigProperties = Cloth->Config->GetProperties();
+
+	// Retrieve MaxDistance information (weight map and Low/High values)
+	int32 MaxDistancePropertyKeyIndex;
+	FString MaxDistanceString = TEXT("MaxDistance");
+	MaxDistanceString = ConfigProperties.GetStringValue(MaxDistanceString, MaxDistanceString, &MaxDistancePropertyKeyIndex);
+
+	FRealSingle MaxDistanceBase;
+	FRealSingle MaxDistanceRange;
+	if (MaxDistancePropertyKeyIndex != INDEX_NONE)
+	{
+		MaxDistanceBase = (FRealSingle)ConfigProperties.GetLowValue<float>(MaxDistancePropertyKeyIndex);
+		MaxDistanceRange = (FRealSingle)ConfigProperties.GetHighValue<float>(MaxDistancePropertyKeyIndex) - MaxDistanceBase;
+	}
+	else
+	{
+		MaxDistanceBase = (FRealSingle)0.;
+		MaxDistanceRange = (FRealSingle)1.;
+	}
+
+	const TConstArrayView<FRealSingle> MaxDistances = WeightMaps.FindRef(MaxDistanceString);
+
 	// Set the particle masses
-	const TConstArrayView<FRealSingle>& MaxDistances = WeightMaps[(int32)EChaosWeightMapTarget::MaxDistance];
 	static const FRealSingle KinematicDistanceThreshold = 0.1f;  // TODO: This is not the same value as set in the painting UI but we might want to expose this value as parameter
 	auto KinematicPredicate =
-		[MaxDistances](int32 Index)
+		[MaxDistances, MaxDistanceBase, MaxDistanceRange](int32 Index)
 		{
-			return MaxDistances.IsValidIndex(Index) && MaxDistances[Index] < KinematicDistanceThreshold;
+			return MaxDistances.IsValidIndex(Index) && MaxDistanceBase + MaxDistanceRange * MaxDistances[Index] < KinematicDistanceThreshold;
 		};
 
-	check(Cloth->Config);
-	const Softs::FCollectionPropertyFacade& ConfigProperties = Cloth->Config->GetProperties();
-
 	const int32 MassMode = ConfigProperties.GetValue<int32>(TEXT("MassMode"), ClothingSimulationClothDefault::MassMode);
 	const FRealSingle MassValue = (FRealSingle)ConfigProperties.GetValue<float>(TEXT("MassValue"), ClothingSimulationClothDefault::MassValue);
 	const FRealSingle MinPerParticleMass = (FRealSingle)ConfigProperties.GetValue<float>(TEXT("MinPerParticleMass"), ClothingSimulationClothDefault::MinPerParticleMass);
@@ -185,7 +207,7 @@ PRAGMA_ENABLE_DEPRECATION_WARNINGS
 
 	// Create constraints
 	const bool bEnabled = false;  // Set constraint disabled by default
-	ClothConstraints.AddRules(ConfigProperties, TriangleMesh, WeightMaps, Tethers, MeshScale, bEnabled, &PatternData);
+	ClothConstraints.AddRules(ConfigProperties, TriangleMesh, &PatternData, WeightMaps, Tethers, MeshScale, bEnabled);
 
 	// Update LOD stats
 	const TConstArrayView<Softs::FSolverReal> InvMasses(Solver->GetParticleInvMasses(Offset), NumParticles);
@@ -225,7 +247,7 @@ PRAGMA_DISABLE_DEPRECATION_WARNINGS  // TODO: CHAOS_IS_CLOTHINGSIMULATIONMESH_AB
 	const Softs::FSolverReal MeshScale = Cloth->Mesh->GetScale();
 PRAGMA_ENABLE_DEPRECATION_WARNINGS
 	const Softs::FSolverReal MaxDistancesScale = (Softs::FSolverReal)Cloth->MaxDistancesMultiplier;
-	ClothConstraints.Update(Cloth->Config->GetProperties(), MeshScale, MaxDistancesScale);
+	ClothConstraints.Update(Cloth->Config->GetProperties(), WeightMaps, MeshScale, MaxDistancesScale);
 }
 
 void FClothingSimulationCloth::FLODData::Enable(FClothingSimulationSolver* Solver, bool bEnable) const
@@ -479,18 +501,18 @@ PRAGMA_ENABLE_DEPRECATION_WARNINGS
 	// Max distance
 	{
 		const int32 MaxDistanceIndex = Properties.AddProperty(TEXT("MaxDistance"));
-		Properties.SetWeightedValue(MaxDistanceIndex, 1.f, 1.f);  // Backward compatibility with legacy mask means can't be different from 1.
+		Properties.SetWeightedValue(MaxDistanceIndex, 0.f, 1.f);  // Backward compatibility with legacy mask must use a unit range since the multiplier is in the mask
 		Properties.SetStringValue(MaxDistanceIndex, TEXT("MaxDistance"));
 	}
 
 	// Backstop
 	{
 		const int32 BackstopScaleIndex = Properties.AddProperty(TEXT("BackstopDistance"));
-		Properties.SetWeightedValue(BackstopScaleIndex, 1.f, 1.f);  // Backward compatibility with legacy mask means can't be different from 1.
+		Properties.SetWeightedValue(BackstopScaleIndex, 0.f, 1.f);  // Backward compatibility with legacy mask must use a unit range since the multiplier is in the mask
 		Properties.SetStringValue(BackstopScaleIndex, TEXT("BackstopDistance"));
 
 		const int32 BackstopRadiusIndex = Properties.AddProperty(TEXT("BackstopRadius"));
-		Properties.SetWeightedValue(BackstopRadiusIndex, 1.f, 1.f);  // Backward compatibility with legacy mask means can't be different from 1.
+		Properties.SetWeightedValue(BackstopRadiusIndex, 0.f, 1.f);  // Backward compatibility with legacy mask must use a unit range since the multiplier is in the mask
 		Properties.SetStringValue(BackstopRadiusIndex, TEXT("BackstopRadius"));
 
 		Properties.AddValue(TEXT("UseLegacyBackstop"), bInUseLegacyBackstop);
@@ -517,22 +539,40 @@ void FClothingSimulationCloth::SetMesh(FClothingSimulationMesh* InMesh)
 {
 	Mesh = InMesh;
 
+	// Build weight map index lookup
+	const TMap<FString, int32> WeightMapIndices = Mesh ? Mesh->GetWeightMapIndices() : TMap<FString, int32>();
+
 	// Reset LODs
 	const int32 NumLODs = Mesh ? Mesh->GetNumLODs() : 0;
 	LODData.Reset(NumLODs);
 
 	const ETetherMode TetherMode = (ETetherMode)Config->GetProperties().GetValue<int32>(TEXT("TetherMode"), ClothingSimulationClothDefault::TetherMode);
 	const bool bUseGeodesicTethers = (TetherMode == ETetherMode::Geodesic);
-	for (int32 Index = 0; Index < NumLODs; ++Index)
+	for (int32 LODIndex = 0; LODIndex < NumLODs; ++LODIndex)
 	{
+		// Regenerate LOD weight maps lookup map
+		TMap<FString, TConstArrayView<FRealSingle>> WeightMaps;
+		WeightMaps.Reserve(WeightMapIndices.Num());
+
+		const TArray<TConstArrayView<FRealSingle>> WeightMapArray = Mesh->GetWeightMaps(LODIndex);
+
+		for (const TPair<FString, int32>& WeightMapIndex : WeightMapIndices)
+		{
+			WeightMaps.Add(WeightMapIndex.Key,
+				WeightMapArray.IsValidIndex(WeightMapIndex.Value) ?
+					WeightMapArray[WeightMapIndex.Value] :
+					TConstArrayView<FRealSingle>());
+		}
+
+		// Add LOD data
 		LODData.Add(MakeUnique<FLODData>(
-			Mesh->GetNumPoints(Index),
-			Mesh->GetIndices(Index),
-			Mesh->GetPatternPositions(Index),
-			Mesh->GetPatternIndices(Index),
-			Mesh->GetPatternToWeldedIndices(Index),
-			Mesh->GetWeightMaps(Index),
-			Mesh->GetTethers(Index, bUseGeodesicTethers)));
+			Mesh->GetNumPoints(LODIndex),
+			Mesh->GetIndices(LODIndex),
+			Mesh->GetPatternPositions(LODIndex),
+			Mesh->GetPatternIndices(LODIndex),
+			Mesh->GetPatternToWeldedIndices(LODIndex),
+			MoveTemp(WeightMaps),
+			Mesh->GetTethers(LODIndex, bUseGeodesicTethers)));
 	}
 
 	// Iterate all known solvers
@@ -548,14 +588,15 @@ PRAGMA_ENABLE_DEPRECATION_WARNINGS
 
 void FClothingSimulationCloth::SetConfig(FClothingSimulationConfig* InConfig)
 {
+	// If the PropertyCollection is owned by this object, so does the current config object
+	if (PropertyCollection.IsValid())
+	{
+		delete Config;
+		PropertyCollection.Reset();
+	}
+
 	if (InConfig)
 	{
-		// If the PropertyCollection is owned by this object, so does the current config object
-		if (InConfig && PropertyCollection.IsValid())
-		{
-			delete Config;
-			PropertyCollection.Reset();
-		}
 		Config = InConfig;
 	}
 	else
@@ -755,11 +796,24 @@ const FTriangleMesh& FClothingSimulationCloth::GetTriangleMesh(const FClothingSi
 	return LODData.IsValidIndex(LODIndex) ? LODData[LODIndex]->SolverData.FindChecked(Solver).TriangleMesh : EmptyTriangleMesh;
 }
 
+// Deprecated for 5.3
 const TArray<TConstArrayView<FRealSingle>>& FClothingSimulationCloth::GetWeightMaps(const FClothingSimulationSolver* Solver) const
 {
-	const int32 LODIndex = LODIndices.FindChecked(Solver);
 	static const TArray<TConstArrayView<FRealSingle>> EmptyWeightMaps;
-	return LODData.IsValidIndex(LODIndex) ? LODData[LODIndex]->WeightMaps : EmptyWeightMaps;
+	return EmptyWeightMaps;  // Can't return a reference anymore, so returning an empty weightmap array is probably the safest for existing code
+}
+
+TConstArrayView<FRealSingle> FClothingSimulationCloth::GetWeightMapByName(const FClothingSimulationSolver* Solver, const FString& Name) const
+{
+	const int32 LODIndex = LODIndices.FindChecked(Solver);
+	return LODData.IsValidIndex(LODIndex) ? LODData[LODIndex]->WeightMaps.FindRef(Name) : TConstArrayView<FRealSingle>();
+}
+
+TConstArrayView<FRealSingle> FClothingSimulationCloth::GetWeightMapByProperty(const FClothingSimulationSolver* Solver, const FString& Property) const
+{
+	check(Config);
+	const FString PropertyString = Config->GetProperties().GetStringValue(Property);
+	return GetWeightMapByName(Solver, PropertyString);
 }
 
 const TArray<TConstArrayView<TTuple<int32, int32, float>>>& FClothingSimulationCloth::GetTethers(const FClothingSimulationSolver* Solver) const
@@ -842,6 +896,9 @@ PRAGMA_DISABLE_DEPRECATION_WARNINGS  // TODO: CHAOS_IS_CLOTHINGSIMULATIONMESH_AB
 	Mesh->Update(Solver, PrevLODIndex, LODIndex, PrevOffset, Offset);
 PRAGMA_ENABLE_DEPRECATION_WARNINGS
 
+	check(Config);
+	const Softs::FCollectionPropertyFacade& ConfigProperties = Config->GetProperties();
+
 	// LOD Switching
 	if (LODIndex != PrevLODIndex)
 	{
@@ -871,9 +928,15 @@ PRAGMA_ENABLE_DEPRECATION_WARNINGS
 				Solver->GetParticleVs(Offset));
 
 				// Update the wind velocity field for the new LOD mesh
-				const TConstArrayView<FRealSingle>& DragMultipliers = LODData[LODIndex]->WeightMaps[(int32)EChaosWeightMapTarget::Drag];
-				const TConstArrayView<FRealSingle>& LiftMultipliers = LODData[LODIndex]->WeightMaps[(int32)EChaosWeightMapTarget::Lift];
-				const TConstArrayView<FRealSingle>& PressureMultipliers = LODData[LODIndex]->WeightMaps[(int32)EChaosWeightMapTarget::Pressure];
+				FString DragString = TEXT("Drag");            // The default name is the property name
+				FString LiftString = TEXT("Lift");            //
+				FString PressureString = TEXT("Pressure");    //
+				DragString = ConfigProperties.GetStringValue(DragString, DragString);
+				LiftString = ConfigProperties.GetStringValue(LiftString, LiftString);
+				PressureString = ConfigProperties.GetStringValue(PressureString, PressureString);
+				const TConstArrayView<FRealSingle> DragMultipliers = LODData[LODIndex]->WeightMaps.FindRef(DragString);
+				const TConstArrayView<FRealSingle> LiftMultipliers = LODData[LODIndex]->WeightMaps.FindRef(LiftString);
+				const TConstArrayView<FRealSingle> PressureMultipliers = LODData[LODIndex]->WeightMaps.FindRef(PressureString);
 
 				Solver->SetWindAndPressureGeometry(GroupId, GetTriangleMesh(Solver), DragMultipliers, LiftMultipliers, PressureMultipliers);
 		}
@@ -886,9 +949,6 @@ PRAGMA_ENABLE_DEPRECATION_WARNINGS
 
 	if (LODIndex != INDEX_NONE)
 	{
-		check(Config);
-		const Softs::FCollectionPropertyFacade& ConfigProperties = Config->GetProperties();
-
 		// Update Cloth group parameters  TODO: Cloth groups should exist as their own node object so that they can be used by several cloth objects
 		LODData[LODIndex]->Update(Solver, this);
 
@@ -942,18 +1002,47 @@ PRAGMA_ENABLE_DEPRECATION_WARNINGS
 		const bool bUsePointBasedWindModel = ConfigProperties.GetValue<bool>(TEXT("UsePointBasedWindModel"));
 		Solver->AddExternalForces(GroupId, bUsePointBasedWindModel);
 
-		const FVec2f Pressure = ConfigProperties.GetWeightedFloatValue(TEXT("Pressure"));
+		int32 PressurePropertyKeyIndex;
+		const FVec2f Pressure(ConfigProperties.GetWeightedFloatValue(TEXT("Pressure"), ClothingSimulationClothDefault::Pressure, &PressurePropertyKeyIndex));
 		if (bUsePointBasedWindModel && ClothingSimulationClothConsoleVariables::CVarLegacyDisablesAccurateWind.GetValueOnAnyThread())
 		{
 			Solver->SetWindAndPressureProperties(GroupId, TVec2<FRealSingle>(0.), TVec2<FRealSingle>(0.), (FRealSingle)0., Pressure);  // Disable the wind velocity field
+			
+			if (ConfigProperties.IsStringDirty(PressurePropertyKeyIndex))
+			{
+				const FString& PressureString = ConfigProperties.GetStringValue(PressurePropertyKeyIndex);
+				Solver->SetWindAndPressureMultipliers(
+					GroupId,
+					TConstArrayView<FRealSingle>(),
+					TConstArrayView<FRealSingle>(),
+					LODData[LODIndex]->WeightMaps.FindRef(PressureString));
+			}
 		}
 		else
 		{
-			const FVec2f Drag = ConfigProperties.GetWeightedFloatValue(TEXT("Drag"), ClothingSimulationClothDefault::Drag);
-			const FVec2f Lift = ConfigProperties.GetWeightedFloatValue(TEXT("Lift"), ClothingSimulationClothDefault::Lift);
+			int32 DragPropertyKeyIndex;
+			int32 LiftPropertyKeyIndex;
+			const FVec2f Drag(ConfigProperties.GetWeightedFloatValue(TEXT("Drag"), ClothingSimulationClothDefault::Drag, &DragPropertyKeyIndex));
+			const FVec2f Lift(ConfigProperties.GetWeightedFloatValue(TEXT("Lift"), ClothingSimulationClothDefault::Lift, &LiftPropertyKeyIndex));
 			const FRealSingle AirDensity = ConfigProperties.GetValue<float>(TEXT("AirDensity"), ClothingSimulationClothDefault::AirDensity);
+
 			Solver->SetWindAndPressureProperties(GroupId, Drag, Lift, AirDensity, Pressure);
+
+			if (ConfigProperties.IsStringDirty(DragPropertyKeyIndex) ||
+				ConfigProperties.IsStringDirty(LiftPropertyKeyIndex) ||
+				ConfigProperties.IsStringDirty(PressurePropertyKeyIndex))
+			{
+				const FString& DragString = ConfigProperties.GetStringValue(DragPropertyKeyIndex);
+				const FString& LiftString = ConfigProperties.GetStringValue(LiftPropertyKeyIndex);
+				const FString& PressureString = ConfigProperties.GetStringValue(PressurePropertyKeyIndex);
+				Solver->SetWindAndPressureMultipliers(
+					GroupId,
+					LODData[LODIndex]->WeightMaps.FindRef(DragString),
+					LODData[LODIndex]->WeightMaps.FindRef(LiftString),
+					LODData[LODIndex]->WeightMaps.FindRef(PressureString));
+			}
 		}
+
 		const FVec3f WindVelocity = ConfigProperties.GetValue<FVector3f>(TEXT("WindVelocity"));
 		Solver->SetWindVelocity(GroupId, WindVelocity + Solver->GetWindVelocity());
 
@@ -988,6 +1077,9 @@ void FClothingSimulationCloth::PostUpdate(FClothingSimulationSolver* Solver)
 		// Update normals
 		LODData[LODIndex]->UpdateNormals(Solver);
 	}
+
+	// Clear the cloth config dirty flags
+	Config->GetProperties().ClearDirtyFlags();
 }
 
 void FClothingSimulationCloth::UpdateFromCache(const FClothingSimulationCacheData& CacheData)
diff --git a/Engine/Plugins/Experimental/ChaosCloth/Source/Chaos/Private/ChaosCloth/ChaosClothingSimulationConfig.cpp b/Engine/Plugins/Experimental/ChaosCloth/Source/Chaos/Private/ChaosCloth/ChaosClothingSimulationConfig.cpp
--- a/Engine/Plugins/Experimental/ChaosCloth/Source/Chaos/Private/ChaosCloth/ChaosClothingSimulationConfig.cpp
+++ b/Engine/Plugins/Experimental/ChaosCloth/Source/Chaos/Private/ChaosCloth/ChaosClothingSimulationConfig.cpp
@@ -20,7 +20,9 @@ namespace Chaos
 		Initialize(InPropertyCollection);
 	}
 
-	void FClothingSimulationConfig::Initialize(const UChaosClothConfig* ClothConfig, const UChaosClothSharedSimConfig* ClothSharedConfig)
+	FClothingSimulationConfig::~FClothingSimulationConfig() = default;
+
+	void FClothingSimulationConfig::Initialize(const UChaosClothConfig* ClothConfig, const UChaosClothSharedSimConfig* ClothSharedConfig, bool bUseLegacyConfig)
 	{
 		constexpr bool bEnable = true;
 		constexpr bool bAnimatable = true;
@@ -184,23 +186,26 @@ namespace Chaos
 			// Max distance
 			{
 				const int32 MaxDistanceIndex = Properties->AddProperty(TEXT("MaxDistance"));
-				Properties->SetWeightedValue(MaxDistanceIndex, 1.f, 1.f);  // Backward compatibility with legacy mask means it can't be different from 1.
+				Properties->SetWeightedValue(MaxDistanceIndex, 0.f, 1.f);  // Backward compatibility with legacy mask must use a unit range since the multiplier is in the mask
 				Properties->SetStringValue(MaxDistanceIndex, TEXT("MaxDistance"));
 			}
 
 			// Backstop
 			{
 				const int32 BackstopDistanceIndex = Properties->AddProperty(TEXT("BackstopDistance"));
-				Properties->SetWeightedValue(BackstopDistanceIndex, 1.f, 1.f);  // Backward compatibility with legacy mask means it can't be different from 1.
+				Properties->SetWeightedValue(BackstopDistanceIndex, 0.f, 1.f);  // Backward compatibility with legacy mask must use a unit range since the multiplier is in the mask
 				Properties->SetStringValue(BackstopDistanceIndex, TEXT("BackstopDistance"));
 
 				const int32 BackstopRadiusIndex = Properties->AddProperty(TEXT("BackstopRadius"));
-				Properties->SetWeightedValue(BackstopRadiusIndex, 1.f, 1.f);  // Backward compatibility with legacy mask means it can't be different from 1.
+				Properties->SetWeightedValue(BackstopRadiusIndex, 0.f, 1.f);  // Backward compatibility with legacy mask must use a unit range since the multiplier is in the mask
 				Properties->SetStringValue(BackstopRadiusIndex, TEXT("BackstopRadius"));
 
 				Properties->AddValue(TEXT("UseLegacyBackstop"), ClothConfig->bUseLegacyBackstop);
 			}
 		}
+
+		// Mark this as a potential legacy config, but leave the behavior control to the client code
+		Properties->AddValue(TEXT("UseLegacyConfig"), bUseLegacyConfig);
 	}
 
 	void FClothingSimulationConfig::Initialize(const TSharedPtr<const FManagedArrayCollection>& InPropertyCollection)
diff --git a/Engine/Plugins/Experimental/ChaosCloth/Source/Chaos/Private/ChaosCloth/ChaosClothingSimulationMesh.cpp b/Engine/Plugins/Experimental/ChaosCloth/Source/Chaos/Private/ChaosCloth/ChaosClothingSimulationMesh.cpp
--- a/Engine/Plugins/Experimental/ChaosCloth/Source/Chaos/Private/ChaosCloth/ChaosClothingSimulationMesh.cpp
+++ b/Engine/Plugins/Experimental/ChaosCloth/Source/Chaos/Private/ChaosCloth/ChaosClothingSimulationMesh.cpp
@@ -149,6 +149,42 @@ TConstArrayView<uint32> FClothingSimulationMesh::GetPatternToWeldedIndices(int32
 	return TConstArrayView<uint32>();
 }
 
+TArray<FName> FClothingSimulationMesh::GetWeightMapNames() const
+{
+	TArray<FName> WeightMapNames;
+
+	const UEnum* const ChaosWeightMapTargetEnum = StaticEnum<EChaosWeightMapTarget>();
+	const int32 NumWeightMaps = (int32)ChaosWeightMapTargetEnum->GetMaxEnumValue() + 1;
+
+	WeightMapNames.SetNum(NumWeightMaps);
+
+	for (int32 EnumIndex = 0; EnumIndex < ChaosWeightMapTargetEnum->NumEnums(); ++EnumIndex)
+	{
+		const int32 TargetIndex = (int32)ChaosWeightMapTargetEnum->GetValueByIndex(EnumIndex);
+		WeightMapNames[TargetIndex] = FName(ChaosWeightMapTargetEnum->GetNameStringByIndex(EnumIndex));
+	}
+	return WeightMapNames;
+}
+
+TMap<FString, int32> FClothingSimulationMesh::GetWeightMapIndices() const
+{
+	TMap<FString, int32> WeightMapIndices;
+
+	const UEnum* const ChaosWeightMapTargetEnum = StaticEnum<EChaosWeightMapTarget>();
+	const int32 NumWeightMaps = (int32)ChaosWeightMapTargetEnum->GetMaxEnumValue() + 1;
+
+	WeightMapIndices.Reserve(NumWeightMaps);
+
+	for (int32 EnumIndex = 0; EnumIndex < ChaosWeightMapTargetEnum->NumEnums(); ++EnumIndex)
+	{
+		const int32 TargetIndex = (int32)ChaosWeightMapTargetEnum->GetValueByIndex(EnumIndex);
+		const FString WeightMapName = ChaosWeightMapTargetEnum->GetNameStringByIndex(EnumIndex);
+
+		WeightMapIndices.Emplace(WeightMapName, TargetIndex);
+	}
+	return WeightMapIndices;
+}
+
 TArray<TConstArrayView<FRealSingle>> FClothingSimulationMesh::GetWeightMaps(int32 LODIndex) const
 {
 	TArray<TConstArrayView<FRealSingle>> WeightMaps;
diff --git a/Engine/Plugins/Experimental/ChaosCloth/Source/Chaos/Private/ChaosCloth/ChaosClothingSimulationSkeletalMesh.cpp b/Engine/Plugins/Experimental/ChaosCloth/Source/Chaos/Private/ChaosCloth/ChaosClothingSimulationSkeletalMesh.cpp
--- a/Engine/Plugins/Experimental/ChaosCloth/Source/Chaos/Private/ChaosCloth/ChaosClothingSimulationSkeletalMesh.cpp
+++ b/Engine/Plugins/Experimental/ChaosCloth/Source/Chaos/Private/ChaosCloth/ChaosClothingSimulationSkeletalMesh.cpp
@@ -108,6 +108,42 @@ PRAGMA_ENABLE_DEPRECATION_WARNINGS
 		return TConstArrayView<uint32>();
 	}
 
+	TArray<FName> FClothingSimulationSkeletalMesh::GetWeightMapNames() const
+	{
+		TArray<FName> WeightMapNames;
+
+		const UEnum* const ChaosWeightMapTargetEnum = StaticEnum<EChaosWeightMapTarget>();
+		const int32 NumWeightMaps = (int32)ChaosWeightMapTargetEnum->GetMaxEnumValue() + 1;
+		
+		WeightMapNames.SetNum(NumWeightMaps);
+
+		for (int32 EnumIndex = 0; EnumIndex < ChaosWeightMapTargetEnum->NumEnums(); ++EnumIndex)
+		{
+			const int32 TargetIndex = (int32)ChaosWeightMapTargetEnum->GetValueByIndex(EnumIndex);
+			WeightMapNames[TargetIndex] = FName(ChaosWeightMapTargetEnum->GetNameStringByIndex(EnumIndex));
+		}
+		return WeightMapNames;
+	}
+
+	TMap<FString, int32> FClothingSimulationSkeletalMesh::GetWeightMapIndices() const
+	{
+		TMap<FString, int32> WeightMapIndices;
+
+		const UEnum* const ChaosWeightMapTargetEnum = StaticEnum<EChaosWeightMapTarget>();
+		const int32 NumWeightMaps = (int32)ChaosWeightMapTargetEnum->GetMaxEnumValue() + 1;
+		
+		WeightMapIndices.Reserve(NumWeightMaps);
+
+		for (int32 EnumIndex = 0; EnumIndex < ChaosWeightMapTargetEnum->NumEnums(); ++EnumIndex)
+		{
+			const int32 TargetIndex = (int32)ChaosWeightMapTargetEnum->GetValueByIndex(EnumIndex);
+			const FString WeightMapName = ChaosWeightMapTargetEnum->GetNameStringByIndex(EnumIndex);
+
+			WeightMapIndices.Emplace(WeightMapName, TargetIndex);
+		}
+		return WeightMapIndices;
+	}
+
 	TArray<TConstArrayView<FRealSingle>> FClothingSimulationSkeletalMesh::GetWeightMaps(int32 LODIndex) const
 	{
 		TArray<TConstArrayView<FRealSingle>> WeightMaps;
diff --git a/Engine/Plugins/Experimental/ChaosCloth/Source/Chaos/Private/ChaosCloth/ChaosClothingSimulationSolver.cpp b/Engine/Plugins/Experimental/ChaosCloth/Source/Chaos/Private/ChaosCloth/ChaosClothingSimulationSolver.cpp
--- a/Engine/Plugins/Experimental/ChaosCloth/Source/Chaos/Private/ChaosCloth/ChaosClothingSimulationSolver.cpp
+++ b/Engine/Plugins/Experimental/ChaosCloth/Source/Chaos/Private/ChaosCloth/ChaosClothingSimulationSolver.cpp
@@ -786,13 +786,27 @@ void FClothingSimulationSolver::SetWindVelocity(uint32 GroupId, const TVec3<FRea
 	VelocityField.SetVelocity(Softs::FSolverVec3(InWindVelocity));
 }
 
-void FClothingSimulationSolver::SetWindAndPressureGeometry(uint32 GroupId, const FTriangleMesh& TriangleMesh, const TConstArrayView<FRealSingle>& DragMultipliers, const TConstArrayView<FRealSingle>& LiftMultipliers,
+void FClothingSimulationSolver::SetWindAndPressureGeometry(
+	uint32 GroupId,
+	const FTriangleMesh& TriangleMesh,
+	const TConstArrayView<FRealSingle>& DragMultipliers,
+	const TConstArrayView<FRealSingle>& LiftMultipliers,
 	const TConstArrayView<FRealSingle>& PressureMultipliers)
 {
 	Softs::FVelocityAndPressureField& VelocityField = Evolution->GetVelocityAndPressureField(GroupId);
 	VelocityField.SetGeometry(&TriangleMesh, DragMultipliers, LiftMultipliers, PressureMultipliers);
 }
 
+void FClothingSimulationSolver::SetWindAndPressureMultipliers(
+	uint32 GroupId,
+	const TConstArrayView<FRealSingle>& DragMultipliers,
+	const TConstArrayView<FRealSingle>& LiftMultipliers,
+	const TConstArrayView<FRealSingle>& PressureMultipliers)
+{
+	Softs::FVelocityAndPressureField& VelocityField = Evolution->GetVelocityAndPressureField(GroupId);
+	VelocityField.SetMultipliers(DragMultipliers, LiftMultipliers, PressureMultipliers);
+}
+
 void FClothingSimulationSolver::SetWindAndPressureProperties(uint32 GroupId, const TVec2<FRealSingle>& Drag, const TVec2<FRealSingle>& Lift, FRealSingle AirDensity, const TVec2<FRealSingle>& Pressure)
 {
 	Softs::FVelocityAndPressureField& VelocityField = Evolution->GetVelocityAndPressureField(GroupId);
@@ -1175,6 +1189,9 @@ void FClothingSimulationSolver::Update(Softs::FSolverReal InDeltaTime)
 			FClothingSimulationCloth* const Cloth = Cloths[ClothIndex];
 			Cloth->PostUpdate(this);
 		}, /*bForceSingleThreaded =*/ !bClothSolverParallelClothPostUpdate);
+
+		// Clear the solver config dirty flags
+		Config->GetProperties().ClearDirtyFlags();
 	}
 
 	// Save old space location for next update
diff --git a/Engine/Plugins/Experimental/ChaosCloth/Source/Chaos/Public/ChaosCloth/ChaosClothConstraints.h b/Engine/Plugins/Experimental/ChaosCloth/Source/Chaos/Public/ChaosCloth/ChaosClothConstraints.h
--- a/Engine/Plugins/Experimental/ChaosCloth/Source/Chaos/Public/ChaosCloth/ChaosClothConstraints.h
+++ b/Engine/Plugins/Experimental/ChaosCloth/Source/Chaos/Public/ChaosCloth/ChaosClothConstraints.h
@@ -44,12 +44,28 @@ namespace Chaos
 		void AddRules(
 			const Softs::FCollectionPropertyConstFacade& ConfigProperties,
 			const FTriangleMesh& TriangleMesh,
-			const TArray<TConstArrayView<FRealSingle>>& WeightMaps,
+			const FClothingPatternData* PatternData,
+			const TMap<FString, TConstArrayView<FRealSingle>>& WeightMaps,
 			const TArray<TConstArrayView<TTuple<int32, int32, FRealSingle>>>& Tethers,
 			Softs::FSolverReal MeshScale,
-			bool bEnabled,
-			const FClothingPatternData* PatternData = nullptr);
+			bool bEnabled);
 
+		UE_DEPRECATED(5.3, "Use AddRules() with WeightMaps and optional PatternData instead.")
+		void AddRules(
+			const Softs::FCollectionPropertyConstFacade& ConfigProperties,
+			const FTriangleMesh& TriangleMesh,
+			const TArray<TConstArrayView<FRealSingle>>& WeightMapArray,
+			const TArray<TConstArrayView<TTuple<int32, int32, FRealSingle>>>& Tethers,
+			Softs::FSolverReal MeshScale,
+			bool bEnabled);
+
+		void Update(
+			const Softs::FCollectionPropertyConstFacade& ConfigProperties,
+			const TMap<FString, TConstArrayView<FRealSingle>>& WeightMaps,
+			Softs::FSolverReal MeshScale,
+			Softs::FSolverReal MaxDistancesScale = (Softs::FSolverReal)1.);
+
+		UE_DEPRECATED(5.3, "Use Update() with WeightMaps instead.")
 		void Update(
 			const Softs::FCollectionPropertyConstFacade& ConfigProperties,
 			Softs::FSolverReal MeshScale,
@@ -182,34 +198,34 @@ namespace Chaos
 			const FTriangleMesh& TriangleMesh);
 		void CreateStretchConstraints(
 			const Softs::FCollectionPropertyConstFacade& ConfigProperties,
-			const TArray<TConstArrayView<FRealSingle>>& WeightMaps,
+			const TMap<FString, TConstArrayView<FRealSingle>>& WeightMaps,
 			const FTriangleMesh& TriangleMesh,
 			const FClothingPatternData* PatternData);
 		void CreateBendingConstraints(
 			const Softs::FCollectionPropertyConstFacade& ConfigProperties,
-			const TArray<TConstArrayView<FRealSingle>>& WeightMaps,
+			const TMap<FString, TConstArrayView<FRealSingle>>& WeightMaps,
 			const FTriangleMesh& TriangleMesh,
 			const FClothingPatternData* PatternData);
 		void CreateAreaConstraints(
 			const Softs::FCollectionPropertyConstFacade& ConfigProperties,
-			const TArray<TConstArrayView<FRealSingle>>& WeightMaps,
+			const TMap<FString, TConstArrayView<FRealSingle>>& WeightMaps,
 			const FTriangleMesh& TriangleMesh);
 		void CreateLongRangeConstraints(
 			const Softs::FCollectionPropertyConstFacade& ConfigProperties,
-			const TArray<TConstArrayView<FRealSingle>>& WeightMaps,
+			const TMap<FString, TConstArrayView<FRealSingle>>& WeightMaps,
 			const TArray<TConstArrayView<TTuple<int32, int32, FRealSingle>>>& Tethers,
 			Softs::FSolverReal MeshScale);
 		void CreateMaxDistanceConstraints(
 			const Softs::FCollectionPropertyConstFacade& ConfigProperties,
-			const TArray<TConstArrayView<FRealSingle>>& WeightMaps,
+			const TMap<FString, TConstArrayView<FRealSingle>>& WeightMaps,
 			Softs::FSolverReal MeshScale);
 		void CreateBackstopConstraints(
 			const Softs::FCollectionPropertyConstFacade& ConfigProperties,
-			const TArray<TConstArrayView<FRealSingle>>& WeightMaps,
+			const TMap<FString, TConstArrayView<FRealSingle>>& WeightMaps,
 			Softs::FSolverReal MeshScale);
 		void CreateAnimDriveConstraints(
 			const Softs::FCollectionPropertyConstFacade& ConfigProperties,
-			const TArray<TConstArrayView<FRealSingle>>& WeightMaps);
+			const TMap<FString, TConstArrayView<FRealSingle>>& WeightMaps);
 
 		//~ Begin deprecated constraints
 		TSharedPtr<Softs::FPBDSpringConstraints> EdgeConstraints_Deprecated;
diff --git a/Engine/Plugins/Experimental/ChaosCloth/Source/Chaos/Public/ChaosCloth/ChaosClothingSimulationCloth.h b/Engine/Plugins/Experimental/ChaosCloth/Source/Chaos/Public/ChaosCloth/ChaosClothingSimulationCloth.h
--- a/Engine/Plugins/Experimental/ChaosCloth/Source/Chaos/Public/ChaosCloth/ChaosClothingSimulationCloth.h
+++ b/Engine/Plugins/Experimental/ChaosCloth/Source/Chaos/Public/ChaosCloth/ChaosClothingSimulationCloth.h
@@ -161,7 +161,11 @@ PRAGMA_ENABLE_DEPRECATION_WARNINGS
 		int32 GetNumParticles(const FClothingSimulationSolver* Solver) const;
 		// Return the current LOD mesh.
 		const FTriangleMesh& GetTriangleMesh(const FClothingSimulationSolver* Solver) const;
-		// Return the current LOD weightmaps.
+		// Return the weight map of the specified name if available on the current LOD, or an empty array view otherwise.
+		TConstArrayView<FRealSingle> GetWeightMapByName(const FClothingSimulationSolver* Solver, const FString& Name) const;
+		// Return the weight map of the specified property name if it exists and is available on the current LOD, or an empty array view otherwise.
+		TConstArrayView<FRealSingle> GetWeightMapByProperty(const FClothingSimulationSolver* Solver, const FString& Property) const;
+		UE_DEPRECATED(5.3, "Returns an empty array from 5.3. Update your code with GetWeightMapByName and GetWeightMapByProperty to return the current LOD weight map instead.")
 		const TArray<TConstArrayView<FRealSingle>>& GetWeightMaps(const FClothingSimulationSolver* Solver) const;
 		// Return the current LOD tethers.
 		const TArray<TConstArrayView<TTuple<int32, int32, float>>>& GetTethers(const FClothingSimulationSolver* Solver) const;
diff --git a/Engine/Plugins/Experimental/ChaosCloth/Source/Chaos/Public/ChaosCloth/ChaosClothingSimulationConfig.h b/Engine/Plugins/Experimental/ChaosCloth/Source/Chaos/Public/ChaosCloth/ChaosClothingSimulationConfig.h
--- a/Engine/Plugins/Experimental/ChaosCloth/Source/Chaos/Public/ChaosCloth/ChaosClothingSimulationConfig.h
+++ b/Engine/Plugins/Experimental/ChaosCloth/Source/Chaos/Public/ChaosCloth/ChaosClothingSimulationConfig.h
@@ -24,15 +24,20 @@ namespace Chaos
 		FClothingSimulationConfig();
 		FClothingSimulationConfig(const TSharedPtr<const FManagedArrayCollection>& InPropertyCollection);
 
-		~FClothingSimulationConfig() = default;
+		~FClothingSimulationConfig();
 
 		FClothingSimulationConfig(const FClothingSimulationConfig&) = delete;
 		FClothingSimulationConfig(FClothingSimulationConfig&&) = delete;
 		FClothingSimulationConfig& operator=(const FClothingSimulationConfig&) = delete;
 		FClothingSimulationConfig& operator=(FClothingSimulationConfig&&) = delete;
 
-		/** Initialize configuration from cloth config UObjects. */
-		void Initialize(const UChaosClothConfig* ClothConfig, const UChaosClothSharedSimConfig* ClothSharedConfig);
+		/**
+		 * Initialize configuration from cloth config UObjects.
+		 * @param ClothConfig The cloth config UObject.
+		 * @param ClothSharedConfig The cloth solver shared config UObject.
+		 * @param bUseLegacyConfig Whether to make the config a legacy cloth config, so that the constraints disable themselves with missing masks, ...etc.
+		 */
+		void Initialize(const UChaosClothConfig* ClothConfig, const UChaosClothSharedSimConfig* ClothSharedConfig, bool bUseLegacyConfig = false);
 
 		/** Initialize config from a property collection. */
 		void Initialize(const TSharedPtr<const FManagedArrayCollection>& InPropertyCollection);
diff --git a/Engine/Plugins/Experimental/ChaosCloth/Source/Chaos/Public/ChaosCloth/ChaosClothingSimulationMesh.h b/Engine/Plugins/Experimental/ChaosCloth/Source/Chaos/Public/ChaosCloth/ChaosClothingSimulationMesh.h
--- a/Engine/Plugins/Experimental/ChaosCloth/Source/Chaos/Public/ChaosCloth/ChaosClothingSimulationMesh.h
+++ b/Engine/Plugins/Experimental/ChaosCloth/Source/Chaos/Public/ChaosCloth/ChaosClothingSimulationMesh.h
@@ -65,6 +65,8 @@ namespace Chaos
 		virtual TConstArrayView<uint32> GetIndices(int32 LODIndex) const;
 		virtual TConstArrayView<uint32> GetPatternIndices(int32 LODIndex) const;
 		virtual TConstArrayView<uint32> GetPatternToWeldedIndices(int32 LODIndex) const;
+		virtual TArray<FName> GetWeightMapNames() const;
+		virtual TMap<FString, int32> GetWeightMapIndices() const;
 		virtual TArray<TConstArrayView<FRealSingle>> GetWeightMaps(int32 LODIndex) const;
 		virtual TArray<TConstArrayView<TTuple<int32, int32, float>>> GetTethers(int32 LODIndex, bool bUseGeodesicTethers) const;
 		virtual int32 GetReferenceBoneIndex() const;
@@ -113,7 +115,13 @@ namespace Chaos
 		/* Return the specified LOD's map from pattern (unwelded) vertices to (welded) vertices, or empty array if patterns are not supported. */
 		virtual TConstArrayView<uint32> GetPatternToWeldedIndices(int32 LODIndex) const = 0;
 
-		/* Return the specified LOD's weightmap. */
+		/* Return all weight maps associated with this mesh returned in the same order as GetWeightMaps. */
+		virtual TArray<FName> GetWeightMapNames() const = 0;
+
+		/* Return a map of all weight map names associated with this mesh to the index in the array returned by GetWeightMaps. */
+		virtual TMap<FString, int32> GetWeightMapIndices() const = 0;
+
+		/* Return the specified LOD's weight map. */
 		virtual TArray<TConstArrayView<FRealSingle>> GetWeightMaps(int32 LODIndex) const = 0;
 
 		/* Return the tethers connections for the long range attachment into convenient parallel friendly batches. */
diff --git a/Engine/Plugins/Experimental/ChaosCloth/Source/Chaos/Public/ChaosCloth/ChaosClothingSimulationSkeletalMesh.h b/Engine/Plugins/Experimental/ChaosCloth/Source/Chaos/Public/ChaosCloth/ChaosClothingSimulationSkeletalMesh.h
--- a/Engine/Plugins/Experimental/ChaosCloth/Source/Chaos/Public/ChaosCloth/ChaosClothingSimulationSkeletalMesh.h
+++ b/Engine/Plugins/Experimental/ChaosCloth/Source/Chaos/Public/ChaosCloth/ChaosClothingSimulationSkeletalMesh.h
@@ -45,6 +45,8 @@ PRAGMA_ENABLE_DEPRECATION_WARNINGS
 		virtual TConstArrayView<uint32> GetIndices(int32 LODIndex) const override;
 		virtual TConstArrayView<uint32> GetPatternIndices(int32 LODIndex) const override;
 		virtual TConstArrayView<uint32> GetPatternToWeldedIndices(int32 LODIndex) const override;
+		virtual TArray<FName> GetWeightMapNames() const override;
+		virtual TMap<FString, int32> GetWeightMapIndices() const override;
 		virtual TArray<TConstArrayView<FRealSingle>> GetWeightMaps(int32 LODIndex) const override;
 		virtual TArray<TConstArrayView<TTuple<int32, int32, float>>> GetTethers(int32 LODIndex, bool bUseGeodesicTethers) const override;
 		virtual int32 GetReferenceBoneIndex() const override;
diff --git a/Engine/Plugins/Experimental/ChaosCloth/Source/Chaos/Public/ChaosCloth/ChaosClothingSimulationSolver.h b/Engine/Plugins/Experimental/ChaosCloth/Source/Chaos/Public/ChaosCloth/ChaosClothingSimulationSolver.h
--- a/Engine/Plugins/Experimental/ChaosCloth/Source/Chaos/Public/ChaosCloth/ChaosClothingSimulationSolver.h
+++ b/Engine/Plugins/Experimental/ChaosCloth/Source/Chaos/Public/ChaosCloth/ChaosClothingSimulationSolver.h
@@ -145,6 +145,13 @@ namespace Chaos
 			SetWindAndPressureGeometry(GroupId, TriangleMesh, DragMultipliers, LiftMultipliers, TConstArrayView<FRealSingle>());
 		}
 
+		// Set new weight maps for the wind and pressure multipliers
+		void SetWindAndPressureMultipliers(
+			uint32 GroupId,
+			const TConstArrayView<FRealSingle>& DragMultipliers,
+			const TConstArrayView<FRealSingle>& LiftMultipliers,
+			const TConstArrayView<FRealSingle>& PressureMultipliers);
+
 		// Set the wind and pressure properties.
 		void SetWindAndPressureProperties(uint32 GroupId, const TVec2<FRealSingle>& Drag, const TVec2<FRealSingle>& Lift, FRealSingle AirDensity = 1.225e-6f, const TVec2<FRealSingle>& Pressure = TVec2<FRealSingle>::ZeroVector);
 		UE_DEPRECATED(5.1, "Chaos::Softs::FVelocityField has been renamed FVelocityAndPressureField to match its new behavior.")
diff --git a/Engine/Plugins/Experimental/ChaosClothAsset/Source/ChaosClothAssetEngine/Private/ChaosClothAsset/ClothSimulationMesh.cpp b/Engine/Plugins/Experimental/ChaosClothAsset/Source/ChaosClothAssetEngine/Private/ChaosClothAsset/ClothSimulationMesh.cpp
--- a/Engine/Plugins/Experimental/ChaosClothAsset/Source/ChaosClothAssetEngine/Private/ChaosClothAsset/ClothSimulationMesh.cpp
+++ b/Engine/Plugins/Experimental/ChaosClothAsset/Source/ChaosClothAssetEngine/Private/ChaosClothAsset/ClothSimulationMesh.cpp
@@ -77,17 +77,49 @@ PRAGMA_ENABLE_DEPRECATION_WARNINGS
 		return TConstArrayView<uint32>(ClothSimulationModel.GetPatternToWeldedIndices(LODIndex));
 	}
 
-	TArray<TConstArrayView<::Chaos::FRealSingle>> FClothSimulationMesh::GetWeightMaps(int32 LODIndex) const
+	TArray<FName> FClothSimulationMesh::GetWeightMapNames() const
 	{
-		constexpr int32 MaxNumWeightMaps = 15; // TODO: Refactor how weight maps are used in base simulation classes
+		TArray<FName> WeightMapNames;
+		if (ClothSimulationModel.ClothSimulationLodModels.Num())
+		{
+			ClothSimulationModel.ClothSimulationLodModels[0].WeightMaps.GetKeys(WeightMapNames);
+		}
+		return WeightMapNames;
+	}
 
+	TMap<FString, int32> FClothSimulationMesh::GetWeightMapIndices() const
+	{
+		TMap<FString, int32> WeightMapIndices;
+
+		// Retrieve weight map names for this cloth
+		const TArray<FName> WeightMapNames = GetWeightMapNames();
+		WeightMapIndices.Reserve(WeightMapNames.Num());
+
+		for (int32 WeightMapIndex = 0; WeightMapIndex < WeightMapNames.Num(); ++WeightMapIndex)
+		{
+			const FName& WeightMapName = WeightMapNames[WeightMapIndex];
+			WeightMapIndices.Emplace(WeightMapName.ToString(), WeightMapIndex);
+		}
+		return WeightMapIndices;
+	}
+
+	TArray<TConstArrayView<::Chaos::FRealSingle>> FClothSimulationMesh::GetWeightMaps(int32 LODIndex) const
+	{
 		TArray<TConstArrayView<::Chaos::FRealSingle>> WeightMaps;
-		WeightMaps.SetNum(MaxNumWeightMaps);
 
-		// Set max distance weight map
-		constexpr int32 MaxDistanceWeightMapTarget = 1;		// EWeightMapTargetCommon::MaxDistance
-		WeightMaps[MaxDistanceWeightMapTarget] = TConstArrayView<::Chaos::FRealSingle>(ClothSimulationModel.ClothSimulationLodModels[LODIndex].MaxDistance);
+		// Retrieve weight map names for this cloth
+		const TArray<FName> WeightMapNames = GetWeightMapNames();
+		WeightMaps.Reserve(WeightMapNames.Num());
 
+		for (int32 WeightMapIndex = 0; WeightMapIndex < WeightMapNames.Num(); ++WeightMapIndex)
+		{
+			const FName& WeightMapName = WeightMapNames[WeightMapIndex];
+
+			const TArray<float>& WeightMap = ClothSimulationModel.ClothSimulationLodModels[LODIndex].WeightMaps.FindChecked(WeightMapName);
+			static_assert(std::is_same_v<::Chaos::FRealSingle, float>, "FRealSingle must be same as float for the Array View to match.");
+
+			WeightMaps.Emplace(TConstArrayView<::Chaos::FRealSingle>(WeightMap));
+		}
 		return WeightMaps;
 	}
 
diff --git a/Engine/Plugins/Experimental/ChaosClothAsset/Source/ChaosClothAssetEngine/Private/ChaosClothAsset/ClothSimulationMesh.h b/Engine/Plugins/Experimental/ChaosClothAsset/Source/ChaosClothAssetEngine/Private/ChaosClothAsset/ClothSimulationMesh.h
--- a/Engine/Plugins/Experimental/ChaosClothAsset/Source/ChaosClothAssetEngine/Private/ChaosClothAsset/ClothSimulationMesh.h
+++ b/Engine/Plugins/Experimental/ChaosClothAsset/Source/ChaosClothAssetEngine/Private/ChaosClothAsset/ClothSimulationMesh.h
@@ -38,6 +38,8 @@ PRAGMA_ENABLE_DEPRECATION_WARNINGS
 		virtual TConstArrayView<uint32> GetIndices(int32 LODIndex) const override;
 		virtual TConstArrayView<uint32> GetPatternIndices(int32 LODIndex) const override;
 		virtual TConstArrayView<uint32> GetPatternToWeldedIndices(int32 LODIndex) const override;
+		virtual TArray<FName> GetWeightMapNames() const override;
+		virtual TMap<FString, int32> GetWeightMapIndices() const override;
 		virtual TArray<TConstArrayView<::Chaos::FRealSingle>> GetWeightMaps(int32 LODIndex) const override;
 		virtual TArray<TConstArrayView<TTuple<int32, int32, float>>> GetTethers(int32 LODIndex, bool bUseGeodesicTethers) const override;
 		virtual int32 GetReferenceBoneIndex() const override;
diff --git a/Engine/Plugins/Experimental/ChaosClothAsset/Source/ChaosClothAssetEngine/Private/ChaosClothAsset/ClothSimulationModel.cpp b/Engine/Plugins/Experimental/ChaosClothAsset/Source/ChaosClothAssetEngine/Private/ChaosClothAsset/ClothSimulationModel.cpp
--- a/Engine/Plugins/Experimental/ChaosClothAsset/Source/ChaosClothAssetEngine/Private/ChaosClothAsset/ClothSimulationModel.cpp
+++ b/Engine/Plugins/Experimental/ChaosClothAsset/Source/ChaosClothAssetEngine/Private/ChaosClothAsset/ClothSimulationModel.cpp
@@ -104,6 +104,7 @@ namespace UE::Chaos::ClothAsset
 		return ReferenceBoneIndex;
 	}
 
+	// TODO: Check which welding implementation is best and if we should use this one instead
 	// Welded weights are averaged. TODO: add other welding options such as Max and Min
 	static TArray<float> WeldWeightMap(const TConstArrayView<float>& UnweldedMap, const TArray<uint32>& PatternToWeldedIndices, const int32 NumWelded)
 	{
@@ -133,12 +134,31 @@ namespace UE::Chaos::ClothAsset
 
 }  // End namespace UE::Chaos::ClothAsset
 
+bool FChaosClothSimulationLodModel::Serialize(FArchive& Ar)
+{
+	// Serialize normal tagged property data
+	if (Ar.IsLoading() || Ar.IsSaving())
+	{
+		UScriptStruct* const Struct = FChaosClothSimulationLodModel::StaticStruct();
+		Struct->SerializeTaggedProperties(Ar, (uint8*)this, Struct, nullptr);
+	}
+
+	// Serialize weight maps (not a tagged property)
+	Ar << WeightMaps;
+
+	// Return true to confirm that serialization has already been taken care of
+	return true;
+}
+
 FChaosClothSimulationModel::FChaosClothSimulationModel(const TSharedPtr<const FManagedArrayCollection>& ClothCollection, const FReferenceSkeleton& ReferenceSkeleton)
 {
 	using namespace UE::Chaos::ClothAsset;
 
 	const FCollectionClothConstFacade Cloth(ClothCollection);
 
+	// Retrieve weigh map names
+	const TArray<FName> WeightMapNames = Cloth.GetWeightMapNames();
+
 	// Initialize LOD models
 	const int32 NumLods = Cloth.GetNumLods();
 	ClothSimulationLodModels.SetNum(NumLods);
@@ -150,39 +170,68 @@ FChaosClothSimulationModel::FChaosClothSimulationModel(const TSharedPtr<const FM
 		TArray<int32> WeldingMap;
 		ClothLod.BuildSimulationMesh(LodModel.Positions, LodModel.Normals, LodModel.Indices, WeldingMap, LodModel.PatternPositions, LodModel.PatternIndices, LodModel.PatternToWeldedIndices);
 
-		static const FName MaxDistanceName = TEXT("MaxDistance");
-		const TConstArrayView<float> MaxDistanceValues = ClothLod.GetWeightMap(MaxDistanceName);
+		// Build welding index lookup  // TODO: Could probably do this better in the BuildSimulationMesh function
+		const int32 LodModelNumSimVertices = LodModel.Positions.Num();
+		const int32 ClothLodNumSimVertices = ClothLod.GetNumSimVertices();
 
-		if (MaxDistanceValues.Num() == LodModel.Positions.Num())
+		TArray<int32> ReverseWeldingMap;
+		ReverseWeldingMap.Reserve(LodModelNumSimVertices);
+		TMap<int32, TArray<int32>> SourceVertexIndicesMap;
+		SourceVertexIndicesMap.Reserve(ClothLodNumSimVertices);
+			
+		for (int32 VertexIndex = 0; VertexIndex < ClothLodNumSimVertices; ++VertexIndex)
 		{
-			LodModel.MaxDistance = MaxDistanceValues;
-		}
-		else if(MaxDistanceValues.Num() > 0)
-		{
-			LodModel.MaxDistance = WeldWeightMap(MaxDistanceValues, LodModel.PatternToWeldedIndices, LodModel.Positions.Num());
+			if (WeldingMap[VertexIndex] == VertexIndex)
+			{
+				SourceVertexIndicesMap.Emplace(VertexIndex, { VertexIndex } );
+				ReverseWeldingMap.Emplace(VertexIndex);
+			}
+			else
+			{
+				TArray<int32>& SourceVertexIndices = SourceVertexIndicesMap.FindChecked(WeldingMap[VertexIndex]);
+				SourceVertexIndices.Emplace(VertexIndex);
+			}
 		}
-		else
+
+		// Copy and weld (average) weight maps
+		LodModel.WeightMaps.Reserve(WeightMapNames.Num());
+
+		for (const FName& WeightMapName : WeightMapNames)
 		{
-			LodModel.MaxDistance.Init(1.0, LodModel.Positions.Num());
+			const TConstArrayView<float> ClothLodWeightMap = ClothLod.GetWeightMap(WeightMapName);
+			TArray<float>& LodModelWeightMap = LodModel.WeightMaps.Add(WeightMapName);
+
+			LodModelWeightMap.SetNumZeroed(LodModelNumSimVertices);
+
+			for (int32 WeldedIndex = 0; WeldedIndex < LodModelNumSimVertices; ++WeldedIndex)
+			{
+				const int32 BaseSourceVertexIndex = ReverseWeldingMap[WeldedIndex];
+				const TArray<int32>& SourceVertexIndices = SourceVertexIndicesMap.FindChecked(BaseSourceVertexIndex);
+				check(SourceVertexIndices.Num());
+
+				for (const int32 VertexIndex : SourceVertexIndices)
+				{
+					LodModelWeightMap[WeldedIndex] += ClothLodWeightMap[VertexIndex];
+				}
+
+				LodModelWeightMap[WeldedIndex] /= (float)SourceVertexIndices.Num();
+			}
 		}
 
+		// Weld bone influences
 		TConstArrayView<int32> NumBoneInfluences = ClothLod.GetSimNumBoneInfluences();
 		TConstArrayView<TArray<int32>> SimBoneIndices = ClothLod.GetSimBoneIndices();
 		TConstArrayView<TArray<float>> SimBoneWeights = ClothLod.GetSimBoneWeights();
-		LodModel.BoneData.SetNum(LodModel.Positions.Num());
+		LodModel.BoneData.SetNum(LodModelNumSimVertices);
 
-		uint32 WeldedIndex = 0;
-		for (int32 VertexIndex = 0; VertexIndex <  ClothLod.GetNumSimVertices(); ++VertexIndex)
+		for (int32 WeldedIndex = 0; WeldedIndex < LodModelNumSimVertices; ++WeldedIndex)
 		{
-			if (WeldingMap[VertexIndex] == VertexIndex)
+			const int32 SourceVertexIndex = ReverseWeldingMap[WeldedIndex];
+			LodModel.BoneData[WeldedIndex].NumInfluences = NumBoneInfluences[SourceVertexIndex];
+			for (int32 BoneIndex = 0; BoneIndex < LodModel.BoneData[WeldedIndex].NumInfluences; ++BoneIndex)
 			{
-				LodModel.BoneData[WeldedIndex].NumInfluences = NumBoneInfluences[VertexIndex];
-				for (int BoneIndex = 0; BoneIndex < LodModel.BoneData[WeldedIndex].NumInfluences; ++BoneIndex)
-				{
-					LodModel.BoneData[WeldedIndex].BoneIndices[BoneIndex] = SimBoneIndices[VertexIndex][BoneIndex];
-					LodModel.BoneData[WeldedIndex].BoneWeights[BoneIndex] = SimBoneWeights[VertexIndex][BoneIndex];
-				}
-				++WeldedIndex;
+				LodModel.BoneData[WeldedIndex].BoneIndices[BoneIndex] = SimBoneIndices[SourceVertexIndex][BoneIndex];
+				LodModel.BoneData[WeldedIndex].BoneWeights[BoneIndex] = SimBoneWeights[SourceVertexIndex][BoneIndex];
 			}
 		}
 	}
diff --git a/Engine/Plugins/Experimental/ChaosClothAsset/Source/ChaosClothAssetEngine/Private/ChaosClothAsset/ClothSimulationModel.h b/Engine/Plugins/Experimental/ChaosClothAsset/Source/ChaosClothAssetEngine/Private/ChaosClothAsset/ClothSimulationModel.h
--- a/Engine/Plugins/Experimental/ChaosClothAsset/Source/ChaosClothAssetEngine/Private/ChaosClothAsset/ClothSimulationModel.h
+++ b/Engine/Plugins/Experimental/ChaosClothAsset/Source/ChaosClothAssetEngine/Private/ChaosClothAsset/ClothSimulationModel.h
@@ -37,9 +37,6 @@ struct FChaosClothSimulationLodModel
 	UPROPERTY()
 	TArray<FClothVertBoneData> BoneData;
 
-	UPROPERTY()
-	TArray<float> MaxDistance;
-
 	/** 2d pattern positions. */
 	UPROPERTY()
 	TArray<FVector2f> PatternPositions;
@@ -51,6 +48,21 @@ struct FChaosClothSimulationLodModel
 	/** Map from PatternsPositions indices to Positions indices. */
 	UPROPERTY()
 	TArray<uint32> PatternToWeldedIndices;
+
+	/** Weight maps for storing painted attributes modifiers on constraint properties. */
+	TMap<FName, TArray<float>> WeightMaps;
+
+	// Custom serialize for weight maps
+	bool Serialize(FArchive& Ar);
+};
+
+template<>
+struct TStructOpsTypeTraits<FChaosClothSimulationLodModel> : public TStructOpsTypeTraitsBase2<FChaosClothSimulationLodModel>
+{
+	enum
+	{
+		WithSerializer = true,
+	};
 };
 
 /**
diff --git a/Engine/Source/Runtime/Experimental/Chaos/Private/Chaos/CollectionPropertyFacade.cpp b/Engine/Source/Runtime/Experimental/Chaos/Private/Chaos/CollectionPropertyFacade.cpp
--- a/Engine/Source/Runtime/Experimental/Chaos/Private/Chaos/CollectionPropertyFacade.cpp
+++ b/Engine/Source/Runtime/Experimental/Chaos/Private/Chaos/CollectionPropertyFacade.cpp
@@ -61,13 +61,14 @@ namespace Chaos::Softs
 	template<typename T, typename ElementType>
 	T FCollectionPropertyConstFacade::GetValue(int32 KeyIndex, const TConstArrayView<ElementType>& ValueArray) const
 	{
-		return ValueArray[KeyIndex];
+		return (T)ValueArray[KeyIndex];
 	}
 	template CHAOS_API FVector3f FCollectionPropertyConstFacade::GetValue<FVector3f, FVector3f>(int32 KeyIndex, const TConstArrayView<FVector3f>& ValueArray) const;
 	template CHAOS_API const FVector3f& FCollectionPropertyConstFacade::GetValue<const FVector3f&, FVector3f>(int32 KeyIndex, const TConstArrayView<FVector3f>& ValueArray) const;
 	template CHAOS_API FString FCollectionPropertyConstFacade::GetValue<FString, FString>(int32 KeyIndex, const TConstArrayView<FString>& ValueArray) const;
 	template CHAOS_API const FString& FCollectionPropertyConstFacade::GetValue<const FString&, FString>(int32 KeyIndex, const TConstArrayView<FString>& ValueArray) const;
 	template CHAOS_API uint8 FCollectionPropertyConstFacade::GetValue<uint8, uint8>(int32 KeyIndex, const TConstArrayView<uint8>& ValueArray) const;
+	template CHAOS_API ECollectionPropertyFlags FCollectionPropertyConstFacade::GetValue<ECollectionPropertyFlags, uint8>(int32 KeyIndex, const TConstArrayView<uint8>& ValueArray) const;
 
 	template<> CHAOS_API
 	bool FCollectionPropertyConstFacade::GetValue<bool, FVector3f>(int32 KeyIndex, const TConstArrayView<FVector3f>& ValueArray) const
@@ -87,16 +88,6 @@ namespace Chaos::Softs
 		return ValueArray[KeyIndex].X;
 	}
 
-	FVector2f FCollectionPropertyConstFacade::GetWeightedFloatValue(const FString& Key, const float& Default, int32* OutKeyIndex) const
-	{
-		const int32 KeyIndex = GetKeyIndex(Key);
-		if (OutKeyIndex)
-		{
-			*OutKeyIndex = KeyIndex;
-		}
-		return KeyIndex != INDEX_NONE ? GetWeightedFloatValue(KeyIndex) : FVector2f(Default, Default);
-	}
-
 	template <typename T>
 	TConstArrayView<T> FCollectionPropertyConstFacade::GetArray(const FName& Name) const
 	{
@@ -116,36 +107,33 @@ namespace Chaos::Softs
 	{
 	}
 
-	int32 FCollectionPropertyFacade::SetWeightedFloatValue(const FString& Key, const FVector2f& Value)
+	void FCollectionPropertyFacade::ClearDirtyFlags()
 	{
-		const int32 KeyIndex = GetKeyIndex(Key);
-		if (KeyIndex != INDEX_NONE)
+		for (ECollectionPropertyFlags& Flags : GetFlagsArray())
 		{
-			SetWeightedFloatValue(KeyIndex, Value);
+			EnumRemoveFlags(Flags, ECollectionPropertyFlags::StringDirty | ECollectionPropertyFlags::Dirty);
 		}
-		return KeyIndex;
 	}
 
-	void FCollectionPropertyFacade::EnableFlag(int32 KeyIndex, ECollectionPropertyFlag Flag, bool bEnable)
+	void FCollectionPropertyFacade::EnableFlags(int32 KeyIndex, ECollectionPropertyFlags Flags, bool bEnable)
 	{
+		const ECollectionPropertyFlags CurrentFlags = GetFlagsArray()[KeyIndex];
+
 		if (bEnable)
 		{
-			EnumAddFlags(GetFlagsArray()[KeyIndex], (uint8)Flag);
+			if (!EnumHasAllFlags(CurrentFlags, Flags))
+			{
+				EnumAddFlags(GetFlagsArray()[KeyIndex], Flags | ECollectionPropertyFlags::Dirty);  // Changing any flags adds the dirty flag
+			}
 		}
 		else
 		{
-			EnumRemoveFlags(GetFlagsArray()[KeyIndex], (uint8)Flag);
-		}
-	}
-
-	int32 FCollectionPropertyFacade::EnableFlag(const FString& Key, ECollectionPropertyFlag Flag, bool bEnable)
-	{
-		const int32 KeyIndex = GetKeyIndex(Key);
-		if (KeyIndex != INDEX_NONE)
-		{
-			EnableFlag(KeyIndex, Flag, bEnable);
+			if (EnumHasAnyFlags(CurrentFlags, Flags))
+			{
+				EnumRemoveFlags(GetFlagsArray()[KeyIndex], Flags);
+				EnumAddFlags(GetFlagsArray()[KeyIndex], ECollectionPropertyFlags::Dirty);  // Changing any flags adds the dirty flag, dirty flags can only be removed from the ClearDirtyFlags function
+			}
 		}
-		return KeyIndex;
 	}
 
 	FCollectionPropertyMutableFacade::FCollectionPropertyMutableFacade(const TSharedPtr<FManagedArrayCollection>& InManagedArrayCollection)
@@ -170,7 +158,9 @@ namespace Chaos::Softs
 	int32 FCollectionPropertyMutableFacade::AddProperty(const FString& Key, bool bEnabled, bool bAnimatable)
 	{
 		const int32 Index = GetManagedArrayCollection()->AddElements(1, PropertyFacadeNames::PropertyGroup);
-		const uint8 Flags = (bEnabled ? (uint8)ECollectionPropertyFlag::Enabled : 0) | (bAnimatable ? (uint8)ECollectionPropertyFlag::Animatable : 0);
+		const ECollectionPropertyFlags Flags =
+			(bEnabled ? ECollectionPropertyFlags::Enabled : ECollectionPropertyFlags::None) |
+			(bAnimatable ? ECollectionPropertyFlags::Animatable : ECollectionPropertyFlags::None);
 
 		// Update the arrayviews in case the new element triggered a reallocation 
 		UpdateArrays();
@@ -191,7 +181,9 @@ namespace Chaos::Softs
 		if (const int32 NumProperties = Keys.Num())
 		{
 			const int32 StartIndex = GetManagedArrayCollection()->AddElements(NumProperties, PropertyFacadeNames::PropertyGroup);
-			const uint8 Flags = (bEnabled ? (uint8)ECollectionPropertyFlag::Enabled : 0) | (bAnimatable ? (uint8)ECollectionPropertyFlag::Animatable : 0);
+			const ECollectionPropertyFlags Flags =
+				(bEnabled ? ECollectionPropertyFlags::Enabled : ECollectionPropertyFlags::None) |
+				(bAnimatable ? ECollectionPropertyFlags::Animatable : ECollectionPropertyFlags::None);
 
 			// Update the arrayviews in case the new elements triggered a reallocation 
 			UpdateArrays();
diff --git a/Engine/Source/Runtime/Experimental/Chaos/Private/Chaos/PBDAxialSpringConstraints.cpp b/Engine/Source/Runtime/Experimental/Chaos/Private/Chaos/PBDAxialSpringConstraints.cpp
--- a/Engine/Source/Runtime/Experimental/Chaos/Private/Chaos/PBDAxialSpringConstraints.cpp
+++ b/Engine/Source/Runtime/Experimental/Chaos/Private/Chaos/PBDAxialSpringConstraints.cpp
@@ -23,7 +23,7 @@ namespace Chaos::Softs {
 int32 Chaos_AxialSpring_ParallelConstraintCount = 100;
 FAutoConsoleVariableRef CVarChaosAxialSpringParallelConstraintCount(TEXT("p.Chaos.AxialSpring.ParallelConstraintCount"), Chaos_AxialSpring_ParallelConstraintCount, TEXT("If we have more constraints than this, use parallel-for in Apply."));
 
-void FPBDAxialSpringConstraints::InitColor(const FSolverParticles& InParticles, const int32 ParticleOffset, const int32 ParticleCount)
+void FPBDAxialSpringConstraints::InitColor(const FSolverParticles& InParticles)
 {
 	// In dev builds we always color so we can tune the system without restarting. See Apply()
 #if UE_BUILD_SHIPPING || UE_BUILD_TEST
@@ -188,4 +188,28 @@ void FPBDAxialSpringConstraints::Apply(FSolverParticles& Particles, const FSolve
 	}
 }
 
+void FPBDAreaSpringConstraints::SetProperties(
+	const FCollectionPropertyConstFacade& PropertyCollection,
+	const TMap<FString, TConstArrayView<FRealSingle>>& WeightMaps)
+{
+	if (IsAreaSpringStiffnessMutable(PropertyCollection))
+	{
+		const FSolverVec2 WeightedValue(GetWeightedFloatAreaSpringStiffness(PropertyCollection));
+		if (IsAreaSpringStiffnessStringDirty(PropertyCollection))
+		{
+			const FString& WeightMapName = GetAreaSpringStiffnessString(PropertyCollection);
+			Stiffness = FPBDStiffness(
+				WeightedValue,
+				WeightMaps.FindRef(WeightMapName),
+				TConstArrayView<TVec3<int32>>(Constraints),
+				ParticleOffset,
+				ParticleCount);
+		}
+		else
+		{
+			Stiffness.SetWeightedValue(WeightedValue);
+		}
+	}
+}
+
 }  // End namespace Chaos::Softs
diff --git a/Engine/Source/Runtime/Experimental/Chaos/Private/Chaos/PBDBendingConstraints.cpp b/Engine/Source/Runtime/Experimental/Chaos/Private/Chaos/PBDBendingConstraints.cpp
--- a/Engine/Source/Runtime/Experimental/Chaos/Private/Chaos/PBDBendingConstraints.cpp
+++ b/Engine/Source/Runtime/Experimental/Chaos/Private/Chaos/PBDBendingConstraints.cpp
@@ -23,7 +23,7 @@ namespace Chaos::Softs {
 int32 Chaos_Bending_ParallelConstraintCount = 100;
 FAutoConsoleVariableRef CVarChaosBendingParallelConstraintCount(TEXT("p.Chaos.Bending.ParallelConstraintCount"), Chaos_Bending_ParallelConstraintCount, TEXT("If we have more constraints than this, use parallel-for in Apply."));
 
-void FPBDBendingConstraints::InitColor(const FSolverParticles& InParticles, const int32 ParticleOffset, const int32 ParticleCount)
+void FPBDBendingConstraints::InitColor(const FSolverParticles& InParticles)
 {
 	// In dev builds we always color so we can tune the system without restarting. See Apply()
 #if UE_BUILD_SHIPPING || UE_BUILD_TEST
@@ -68,6 +68,52 @@ void FPBDBendingConstraints::InitColor(const FSolverParticles& InParticles, cons
 	}
 }
 
+void FPBDBendingConstraints::SetProperties(
+	const FCollectionPropertyConstFacade& PropertyCollection,
+	const TMap<FString, TConstArrayView<FRealSingle>>& WeightMaps)
+{
+	if (IsBendingElementStiffnessMutable(PropertyCollection))
+	{
+		const FSolverVec2 WeightedValue(GetWeightedFloatBendingElementStiffness(PropertyCollection));
+		if (IsBendingElementStiffnessStringDirty(PropertyCollection))
+		{
+			const FString& WeightMapName = GetBendingElementStiffnessString(PropertyCollection);
+			Stiffness = FPBDStiffness(
+				WeightedValue,
+				WeightMaps.FindRef(WeightMapName),
+				TConstArrayView<TVec2<int32>>(ConstraintSharedEdges),
+				ParticleOffset,
+				ParticleCount);
+		}
+		else
+		{
+			Stiffness.SetWeightedValue(WeightedValue);
+		}
+	}
+	if (IsBucklingRatioMutable(PropertyCollection))
+	{
+		BucklingRatio = (FSolverReal)FMath::Clamp(GetBucklingRatio(PropertyCollection), 0.f, 1.);
+	}
+	if (IsBucklingStiffnessMutable(PropertyCollection))
+	{
+		const FSolverVec2 WeightedValue(GetWeightedFloatBucklingStiffness(PropertyCollection));
+		if (IsBucklingStiffnessStringDirty(PropertyCollection))
+		{
+			const FString& WeightMapName = GetBucklingStiffnessString(PropertyCollection);
+			BucklingStiffness = FPBDStiffness(
+				WeightedValue,
+				WeightMaps.FindRef(WeightMapName),
+				TConstArrayView<TVec2<int32>>(ConstraintSharedEdges),
+				ParticleOffset,
+				ParticleCount);
+		}
+		else
+		{
+			BucklingStiffness.SetWeightedValue(WeightedValue);
+		}
+	}
+}
+
 void FPBDBendingConstraints::ApplyHelper(FSolverParticles& Particles, const FSolverReal Dt, const int32 ConstraintIndex, const FSolverReal ExpStiffnessValue, const FSolverReal ExpBucklingValue) const
 {
 
diff --git a/Engine/Source/Runtime/Experimental/Chaos/Private/Chaos/PBDLongRangeConstraints.cpp b/Engine/Source/Runtime/Experimental/Chaos/Private/Chaos/PBDLongRangeConstraints.cpp
--- a/Engine/Source/Runtime/Experimental/Chaos/Private/Chaos/PBDLongRangeConstraints.cpp
+++ b/Engine/Source/Runtime/Experimental/Chaos/Private/Chaos/PBDLongRangeConstraints.cpp
@@ -15,6 +15,39 @@ FAutoConsoleVariableRef CVarChaosLongRangeISPCEnabled(TEXT("p.Chaos.LongRange.IS
 
 namespace Chaos::Softs {
 
+void FPBDLongRangeConstraints::SetProperties(
+	const FCollectionPropertyConstFacade& PropertyCollection,
+	const TMap<FString, TConstArrayView<FRealSingle>>& WeightMaps,
+	FSolverReal MeshScale)
+{
+	if (IsTetherStiffnessMutable(PropertyCollection))
+	{
+		const FSolverVec2 WeightedValue(GetWeightedFloatTetherStiffness(PropertyCollection));
+		if (IsTetherStiffnessStringDirty(PropertyCollection))
+		{
+			const FString& WeightMapName = GetTetherStiffnessString(PropertyCollection);
+			Stiffness = FPBDStiffness(WeightedValue, WeightMaps.FindRef(WeightMapName), ParticleCount);
+		}
+		else
+		{
+			Stiffness.SetWeightedValue(WeightedValue);
+		}
+	}
+	if (IsTetherScaleMutable(PropertyCollection))
+	{
+		const FSolverVec2 WeightedValue = FSolverVec2(GetWeightedFloatTetherScale(PropertyCollection)).ClampAxes(MinTetherScale, MaxTetherScale) * MeshScale;
+		if (IsTetherScaleStringDirty(PropertyCollection))
+		{
+			const FString& WeightMapName = GetTetherScaleString(PropertyCollection);
+			TetherScale = FPBDWeightMap(WeightedValue, WeightMaps.FindRef(WeightMapName), ParticleCount);
+		}
+		else
+		{
+			TetherScale.SetWeightedValue(WeightedValue);
+		}
+	}
+}
+
 void FPBDLongRangeConstraints::Apply(FSolverParticles& Particles, const FSolverReal /*Dt*/) const
 {
 	TRACE_CPUPROFILER_EVENT_SCOPE(FPBDLongRangeConstraints_Apply);
diff --git a/Engine/Source/Runtime/Experimental/Chaos/Private/Chaos/PBDLongRangeConstraintsBase.cpp b/Engine/Source/Runtime/Experimental/Chaos/Private/Chaos/PBDLongRangeConstraintsBase.cpp
--- a/Engine/Source/Runtime/Experimental/Chaos/Private/Chaos/PBDLongRangeConstraintsBase.cpp
+++ b/Engine/Source/Runtime/Experimental/Chaos/Private/Chaos/PBDLongRangeConstraintsBase.cpp
@@ -36,6 +36,8 @@ FPBDLongRangeConstraintsBase::FPBDLongRangeConstraintsBase(
 	FSolverReal MaxStiffness,
 	FSolverReal MeshScale)
 	: Tethers(InTethers)
+	, ParticleOffset(InParticleOffset)
+	, ParticleCount(InParticleCount)
 	, Stiffness(
 		InStiffness,
 		StiffnessMultipliers,
@@ -47,7 +49,6 @@ FPBDLongRangeConstraintsBase::FPBDLongRangeConstraintsBase(
 		InScale.ClampAxes(MinTetherScale, MaxTetherScale) * MeshScale,
 		ScaleMultipliers,
 		InParticleCount)
-	, ParticleOffset(InParticleOffset)
 {
 	// Apply default properties
 	ApplyProperties((FSolverReal)(1. / FPBDStiffness::ParameterFrequency), 1);
diff --git a/Engine/Source/Runtime/Experimental/Chaos/Private/Chaos/PBDSphericalConstraint.cpp b/Engine/Source/Runtime/Experimental/Chaos/Private/Chaos/PBDSphericalConstraint.cpp
--- a/Engine/Source/Runtime/Experimental/Chaos/Private/Chaos/PBDSphericalConstraint.cpp
+++ b/Engine/Source/Runtime/Experimental/Chaos/Private/Chaos/PBDSphericalConstraint.cpp
@@ -18,56 +18,209 @@ static_assert(sizeof(ispc::FVector4f) == sizeof(Chaos::Softs::FPAndInvM), "sizeo
 
 namespace Chaos::Softs {
 
+void FPBDSphericalConstraint::SetProperties(
+	const FCollectionPropertyConstFacade& PropertyCollection,
+	const TMap<FString, TConstArrayView<FRealSingle>>& WeightMaps,
+	FSolverReal MeshScale)
+{
+	if (IsMaxDistanceMutable(PropertyCollection))
+	{
+		const FSolverVec2 WeightedFloatMaxDistance(FVector2f::Max(GetWeightedFloatMaxDistance(PropertyCollection), FVector2f::ZeroVector));
+		MaxDistanceBase = WeightedFloatMaxDistance[0];
+		MaxDistanceRange = WeightedFloatMaxDistance[1] - MaxDistanceBase;
+
+		if (IsMaxDistanceStringDirty(PropertyCollection))
+		{
+			SphereRadii = WeightMaps.FindRef(GetMaxDistanceString(PropertyCollection));
+		}
+	}
+
+	SetScale(MeshScale);
+}
+
 void FPBDSphericalConstraint::ApplyHelperISPC(FSolverParticles& Particles, const FSolverReal /*Dt*/) const
 {
 	check(bRealTypeCompatibleWithISPC);
 
-	const int32 ParticleCount = SphereRadii.Num();
-
 #if INTEL_ISPC
-	ispc::ApplySphericalConstraints(
-		(ispc::FVector4f*)Particles.GetPAndInvM().GetData(),
-		(const ispc::FVector3f*)AnimationPositions.GetData(),
-		SphereRadii.GetData(),
-		GetScale(),
-		ParticleOffset,
-		ParticleCount);
+	if (SphereRadii.Num() == ParticleCount)
+	{
+		ispc::ApplySphericalConstraintsWithMap(
+			(ispc::FVector4f*)Particles.GetPAndInvM().GetData(),
+			(const ispc::FVector3f*)AnimationPositions.GetData(),
+			SphereRadii.GetData(),
+			MaxDistanceBase * GetScale(),
+			MaxDistanceRange * GetScale(),
+			ParticleOffset,
+			ParticleCount);
+	}
+	else
+	{
+		ispc::ApplySphericalConstraintsWithoutMap(
+			(ispc::FVector4f*)Particles.GetPAndInvM().GetData(),
+			(const ispc::FVector3f*)AnimationPositions.GetData(),
+			MaxDistanceBase * GetScale(),
+			ParticleOffset,
+			ParticleCount);
+	}
 #endif
 }
+
+void FPBDSphericalBackstopConstraint::SetProperties(
+	const FCollectionPropertyConstFacade& PropertyCollection,
+	const TMap<FString, TConstArrayView<FRealSingle>>& WeightMaps,
+	FSolverReal MeshScale)
+{
+	if (IsBackstopRadiusMutable(PropertyCollection))
+	{
+		const FSolverVec2 WeightedFloatBackstopRadius(FVector2f::Max(GetWeightedFloatBackstopRadius(PropertyCollection), FVector2f::ZeroVector));
+		BackstopRadiusBase = (FSolverReal)WeightedFloatBackstopRadius[0];
+		BackstopRadiusRange = (FSolverReal)WeightedFloatBackstopRadius[1] - BackstopRadiusBase;
+
+		if (IsBackstopRadiusStringDirty(PropertyCollection))
+		{
+			SphereRadii = WeightMaps.FindRef(GetBackstopRadiusString(PropertyCollection));
+		}
+	}
+
+	if (IsBackstopDistanceMutable(PropertyCollection))
+	{
+		const FSolverVec2 WeightedFloatBackstopDistance(GetWeightedFloatBackstopDistance(PropertyCollection));  // Unlike radius, distance can be negative
+		BackstopDistanceBase = WeightedFloatBackstopDistance[0];
+		BackstopDistanceRange = WeightedFloatBackstopDistance[1] - BackstopDistanceBase;
+
+		if (IsBackstopDistanceStringDirty(PropertyCollection))
+		{
+			SphereOffsetDistances = WeightMaps.FindRef(GetBackstopDistanceString(PropertyCollection));
+		}
+	}
+
+	SetScale(MeshScale);
+}
+
 void FPBDSphericalBackstopConstraint::ApplyLegacyHelperISPC(FSolverParticles& Particles, const FSolverReal /*Dt*/) const
 {
 	check(bRealTypeCompatibleWithISPC);
 
-	const int32 ParticleCount = SphereRadii.Num();
-
 #if INTEL_ISPC
-	ispc::ApplyLegacySphericalBackstopConstraints(
-		(ispc::FVector4f*)Particles.GetPAndInvM().GetData(),
-		(const ispc::FVector3f*)AnimationPositions.GetData(),
-		(const ispc::FVector3f*)AnimationNormals.GetData(),
-		SphereOffsetDistances.GetData(),
-		SphereRadii.GetData(),
-		GetScale(),
-		ParticleOffset,
-		ParticleCount);
+	if (SphereRadii.Num() == ParticleCount)
+	{
+		if (SphereOffsetDistances.Num() == ParticleCount)
+		{
+			ispc::ApplyLegacySphericalBackstopConstraintsWithMaps(
+				(ispc::FVector4f*)Particles.GetPAndInvM().GetData(),
+				(const ispc::FVector3f*)AnimationPositions.GetData(),
+				(const ispc::FVector3f*)AnimationNormals.GetData(),
+				SphereOffsetDistances.GetData(),
+				SphereRadii.GetData(),
+				BackstopDistanceBase * GetScale(),
+				BackstopDistanceRange * GetScale(),
+				BackstopRadiusBase * GetScale(),
+				BackstopRadiusRange * GetScale(),
+				ParticleOffset,
+				ParticleCount);
+		}
+		else
+		{
+			ispc::ApplyLegacySphericalBackstopConstraintsWithRadiusMap(
+				(ispc::FVector4f*)Particles.GetPAndInvM().GetData(),
+				(const ispc::FVector3f*)AnimationPositions.GetData(),
+				(const ispc::FVector3f*)AnimationNormals.GetData(),
+				SphereRadii.GetData(),
+				BackstopDistanceBase * GetScale(),
+				BackstopRadiusBase * GetScale(),
+				BackstopRadiusRange * GetScale(),
+				ParticleOffset,
+				ParticleCount);
+		}
+	}
+	else if (SphereOffsetDistances.Num() == ParticleCount)
+	{
+		ispc::ApplyLegacySphericalBackstopConstraintsWithDistanceMap(
+			(ispc::FVector4f*)Particles.GetPAndInvM().GetData(),
+			(const ispc::FVector3f*)AnimationPositions.GetData(),
+			(const ispc::FVector3f*)AnimationNormals.GetData(),
+			SphereOffsetDistances.GetData(),
+			BackstopDistanceBase * GetScale(),
+			BackstopDistanceRange * GetScale(),
+			BackstopRadiusBase * GetScale(),
+			ParticleOffset,
+			ParticleCount);
+	}
+	else
+	{
+		ispc::ApplyLegacySphericalBackstopConstraintsWithoutMaps(
+			(ispc::FVector4f*)Particles.GetPAndInvM().GetData(),
+			(const ispc::FVector3f*)AnimationPositions.GetData(),
+			(const ispc::FVector3f*)AnimationNormals.GetData(),
+			BackstopDistanceBase * GetScale(),
+			BackstopRadiusBase * GetScale(),
+			ParticleOffset,
+			ParticleCount);
+	}
 #endif
 }
 
 void FPBDSphericalBackstopConstraint::ApplyHelperISPC(FSolverParticles& Particles, const FSolverReal /*Dt*/) const
 {
 	check(bRealTypeCompatibleWithISPC);
-	const int32 ParticleCount = SphereRadii.Num();
 
 #if INTEL_ISPC
-	ispc::ApplySphericalBackstopConstraints(
-		(ispc::FVector4f*)Particles.GetPAndInvM().GetData(),
-		(const ispc::FVector3f*)AnimationPositions.GetData(),
-		(const ispc::FVector3f*)AnimationNormals.GetData(),
-		SphereOffsetDistances.GetData(),
-		SphereRadii.GetData(),
-		GetScale(),
-		ParticleOffset,
-		ParticleCount);
+	if (SphereRadii.Num() == ParticleCount)
+	{
+		if (SphereOffsetDistances.Num() == ParticleCount)
+		{
+			ispc::ApplySphericalBackstopConstraintsWithMaps(
+				(ispc::FVector4f*)Particles.GetPAndInvM().GetData(),
+				(const ispc::FVector3f*)AnimationPositions.GetData(),
+				(const ispc::FVector3f*)AnimationNormals.GetData(),
+				SphereOffsetDistances.GetData(),
+				SphereRadii.GetData(),
+				BackstopDistanceBase * GetScale(),
+				BackstopDistanceRange * GetScale(),
+				BackstopRadiusBase * GetScale(),
+				BackstopRadiusRange * GetScale(),
+				ParticleOffset,
+				ParticleCount);
+		}
+		else
+		{
+			ispc::ApplySphericalBackstopConstraintsWithRadiusMap(
+				(ispc::FVector4f*)Particles.GetPAndInvM().GetData(),
+				(const ispc::FVector3f*)AnimationPositions.GetData(),
+				(const ispc::FVector3f*)AnimationNormals.GetData(),
+				SphereRadii.GetData(),
+				BackstopDistanceBase * GetScale(),
+				BackstopRadiusBase * GetScale(),
+				BackstopRadiusRange * GetScale(),
+				ParticleOffset,
+				ParticleCount);
+		}
+	}
+	else if (SphereOffsetDistances.Num() == ParticleCount)
+	{
+		ispc::ApplySphericalBackstopConstraintsWithDistanceMap(
+			(ispc::FVector4f*)Particles.GetPAndInvM().GetData(),
+			(const ispc::FVector3f*)AnimationPositions.GetData(),
+			(const ispc::FVector3f*)AnimationNormals.GetData(),
+			SphereOffsetDistances.GetData(),
+			BackstopDistanceBase * GetScale(),
+			BackstopDistanceRange * GetScale(),
+			BackstopRadiusBase * GetScale(),
+			ParticleOffset,
+			ParticleCount);
+	}
+	else
+	{
+		ispc::ApplySphericalBackstopConstraintsWithoutMaps(
+			(ispc::FVector4f*)Particles.GetPAndInvM().GetData(),
+			(const ispc::FVector3f*)AnimationPositions.GetData(),
+			(const ispc::FVector3f*)AnimationNormals.GetData(),
+			BackstopDistanceBase * GetScale(),
+			BackstopRadiusBase * GetScale(),
+			ParticleOffset,
+			ParticleCount);
+	}
 #endif
 }
 
diff --git a/Engine/Source/Runtime/Experimental/Chaos/Private/Chaos/PBDSphericalConstraint.ispc b/Engine/Source/Runtime/Experimental/Chaos/Private/Chaos/PBDSphericalConstraint.ispc
--- a/Engine/Source/Runtime/Experimental/Chaos/Private/Chaos/PBDSphericalConstraint.ispc
+++ b/Engine/Source/Runtime/Experimental/Chaos/Private/Chaos/PBDSphericalConstraint.ispc
@@ -5,10 +5,11 @@
 #include "Math/Vector.isph"
 #include "Chaos/PBDSofts.isph"
 
-export void ApplySphericalConstraints(uniform FVector4f ParticlesPandInvM[],
+export void ApplySphericalConstraintsWithMap(uniform FVector4f ParticlesPandInvM[],
 									const uniform FVector3f AnimationPositions[],
 									const uniform float SphereRadii[],
-									const uniform float SphereRadiiMultiplier,
+									const uniform float SphereRadiiBase,
+									const uniform float SphereRadiiRange,
 									const uniform int32 ParticleOffset,
 									const uniform int32 ParticleCount)
 {
@@ -16,8 +17,7 @@ export void ApplySphericalConstraints(uniform FVector4f ParticlesPandInvM[],
 
 	foreach(Index = 0 ... ParticleCount)
 	{
-		const int32 ParticleIndex = ParticleOffset + Index;
-		const uniform int32 Offset = extract(ParticleIndex, 0);
+		const uniform int32 Offset = ParticleOffset + extract(Index, 0);
 		const FVector4f PandInvM = VectorLoad(&ParticlesPandInvM[Offset]);
 
 		varying FVector3f Position;
@@ -36,7 +36,7 @@ export void ApplySphericalConstraints(uniform FVector4f ParticlesPandInvM[],
 		const float DistanceSquared = VectorSizeSquared(CenterToParticle);
 		const float Distance = sqrt(DistanceSquared);
 
-		const float Radius = SphereRadii[Index] * SphereRadiiMultiplier;
+		const float Radius = SphereRadiiBase + SphereRadiiRange * SphereRadii[Index];
 		const FVector3f FinalPosition = ((Radius / Distance) * CenterToParticle) + Center;
 
 		if (IM != 0 && DistanceSquared > ((Radius * Radius) + DeadZoneSquareRadius))
@@ -46,12 +46,73 @@ export void ApplySphericalConstraints(uniform FVector4f ParticlesPandInvM[],
 	}
 }
 
-export void ApplyLegacySphericalBackstopConstraints(uniform FVector4f ParticlesPandInvM[],
+export void ApplySphericalConstraintsWithoutMap(uniform FVector4f ParticlesPandInvM[],
+									const uniform FVector3f AnimationPositions[],
+									const uniform float Radius,
+									const uniform int32 ParticleOffset,
+									const uniform int32 ParticleCount)
+{
+	static const float DeadZoneSquareRadius = FLOAT_SMALL_NUMBER; // We will not push the particle away in the dead zone
+
+	foreach(Index = 0 ... ParticleCount)
+	{
+		const uniform int32 Offset = ParticleOffset + extract(Index, 0);
+		const FVector4f PandInvM = VectorLoad(&ParticlesPandInvM[Offset]);
+
+		varying FVector3f Position;
+		varying float IM;
+		UnzipPandInvM(PandInvM, Position, IM);
+
+		if(all(IM == 0))
+		{
+			continue;
+		}
+
+		
+		const FVector3f Center = VectorLoad(&AnimationPositions[Offset]);
+
+		const FVector3f CenterToParticle = Position - Center;
+		const float DistanceSquared = VectorSizeSquared(CenterToParticle);
+		const float Distance = sqrt(DistanceSquared);
+
+		const FVector3f FinalPosition = ((Radius / Distance) * CenterToParticle) + Center;
+
+		if (IM != 0 && DistanceSquared > ((Radius * Radius) + DeadZoneSquareRadius))
+		{
+			VectorStore(&ParticlesPandInvM[Offset], SetVector4(FinalPosition, IM));
+		}
+	}
+}
+
+// Deprecated since 5.3
+export void ApplySphericalConstraints(uniform FVector4f ParticlesPandInvM[],
+									const uniform FVector3f AnimationPositions[],
+									const uniform float SphereRadii[],
+									const uniform float SphereRadiiMultiplier,
+									const uniform int32 ParticleOffset,
+									const uniform int32 ParticleCount)
+{
+	static const uniform float SphereRadiiBase = 0.f;
+
+	ApplySphericalConstraintsWithMap(
+		ParticlesPandInvM,
+		AnimationPositions,
+		SphereRadii,
+		SphereRadiiBase,
+		SphereRadiiMultiplier,
+		ParticleOffset,
+		ParticleCount);
+}
+
+export void ApplyLegacySphericalBackstopConstraintsWithMaps(uniform FVector4f ParticlesPandInvM[],
 													const uniform FVector3f AnimationPositions[],
 													const uniform FVector3f AnimationNormals[],
 													const uniform float SphereOffsetDistances[],
 													const uniform float SphereRadii[],
-													const uniform float SphereRadiiMultiplier,
+													const uniform float SphereOffsetDistancesBase,
+													const uniform float SphereOffsetDistancesRange,
+													const uniform float SphereRadiiBase,
+													const uniform float SphereRadiiRange,
 													const uniform int32 ParticleOffset,
 													const uniform int32 ParticleCount)
 {
@@ -59,8 +120,7 @@ export void ApplyLegacySphericalBackstopConstraints(uniform FVector4f ParticlesP
 
 	foreach(Index = 0 ... ParticleCount)
 	{
-		const int32 ParticleIndex = ParticleOffset + Index;
-		const uniform int32 Offset = extract(ParticleIndex, 0);
+		const uniform int32 Offset = ParticleOffset + extract(Index, 0);
 		const FVector4f PandInvM = VectorLoad(&ParticlesPandInvM[Offset]);
 
 		varying FVector3f Position;
@@ -72,11 +132,11 @@ export void ApplyLegacySphericalBackstopConstraints(uniform FVector4f ParticlesP
 			continue;
 		}
 
-		const FVector3f AnimationPosition = VectorLoad(&AnimationPositions[extract(ParticleIndex, 0)]);
-		const FVector3f AnimationNormal = VectorLoad(&AnimationNormals[extract(ParticleIndex, 0)]);
+		const FVector3f AnimationPosition = VectorLoad(&AnimationPositions[Offset]);
+		const FVector3f AnimationNormal = VectorLoad(&AnimationNormals[Offset]);
 
-		const float SphereOffsetDistance = SphereOffsetDistances[Index] * SphereRadiiMultiplier;
-		const float Radius = SphereRadii[Index] * SphereRadiiMultiplier;
+		const float SphereOffsetDistance = SphereOffsetDistancesBase + SphereOffsetDistancesRange * SphereOffsetDistances[Index];
+		const float Radius = SphereRadiiBase + SphereRadiiRange * SphereRadii[Index];
 
 		const FVector3f Center = AnimationPosition - SphereOffsetDistance * AnimationNormal;  // Legacy version already includes the radius within the distance
 		const FVector3f CenterToParticle = Position - Center;
@@ -95,12 +155,232 @@ export void ApplyLegacySphericalBackstopConstraints(uniform FVector4f ParticlesP
 	}
 }
 
-export void ApplySphericalBackstopConstraints(uniform FVector4f ParticlesPandInvM[],
+export void ApplyLegacySphericalBackstopConstraintsWithDistanceMap(uniform FVector4f ParticlesPandInvM[],
+													const uniform FVector3f AnimationPositions[],
+													const uniform FVector3f AnimationNormals[],
+													const uniform float SphereOffsetDistances[],
+													const uniform float SphereOffsetDistancesBase,
+													const uniform float SphereOffsetDistancesRange,
+													const uniform float Radius,
+													const uniform int32 ParticleOffset,
+													const uniform int32 ParticleCount)
+{
+	static const float DeadZoneSquareRadius = FLOAT_SMALL_NUMBER; // We will not push the particle away in the dead zone
+
+	foreach(Index = 0 ... ParticleCount)
+	{
+		const uniform int32 Offset = ParticleOffset + extract(Index, 0);
+		const FVector4f PandInvM = VectorLoad(&ParticlesPandInvM[Offset]);
+
+		varying FVector3f Position;
+		varying float IM;
+		UnzipPandInvM(PandInvM, Position, IM);
+
+		if(all(IM == 0))
+		{
+			continue;
+		}
+
+		const FVector3f AnimationPosition = VectorLoad(&AnimationPositions[Offset]);
+		const FVector3f AnimationNormal = VectorLoad(&AnimationNormals[Offset]);
+
+		const float SphereOffsetDistance = SphereOffsetDistancesBase + SphereOffsetDistancesRange * SphereOffsetDistances[Index];
+
+		const FVector3f Center = AnimationPosition - SphereOffsetDistance * AnimationNormal;  // Legacy version already includes the radius within the distance
+		const FVector3f CenterToParticle = Position - Center;
+		const float DistanceSquared = VectorSizeSquared(CenterToParticle);
+
+		const FVector3f DistSquaredLessRadiusSquared = ((Radius / sqrt(DistanceSquared)) * CenterToParticle) + Center;
+		const FVector3f DistSquaredLessDeadZoneRadius = AnimationPosition - (SphereOffsetDistance - Radius) * AnimationNormal;  // Legacy version already includes the radius to the distance
+
+		FVector3f NewP = VectorSelect(DistanceSquared < Radius * Radius, DistSquaredLessRadiusSquared, FloatZeroVector);
+		NewP = VectorSelect(DistanceSquared < DeadZoneSquareRadius, DistSquaredLessDeadZoneRadius, NewP);
+
+		if (IM != 0 && NewP.V[0] != 0)
+		{
+			VectorStore(&ParticlesPandInvM[Offset], SetVector4(NewP, IM));
+		}
+	}
+}
+
+export void ApplyLegacySphericalBackstopConstraintsWithRadiusMap(uniform FVector4f ParticlesPandInvM[],
+													const uniform FVector3f AnimationPositions[],
+													const uniform FVector3f AnimationNormals[],
+													const uniform float SphereRadii[],
+													const uniform float SphereOffsetDistance,
+													const uniform float SphereRadiiBase,
+													const uniform float SphereRadiiRange,
+													const uniform int32 ParticleOffset,
+													const uniform int32 ParticleCount)
+{
+	static const float DeadZoneSquareRadius = FLOAT_SMALL_NUMBER; // We will not push the particle away in the dead zone
+
+	foreach(Index = 0 ... ParticleCount)
+	{
+		const uniform int32 Offset = ParticleOffset + extract(Index, 0);
+		const FVector4f PandInvM = VectorLoad(&ParticlesPandInvM[Offset]);
+
+		varying FVector3f Position;
+		varying float IM;
+		UnzipPandInvM(PandInvM, Position, IM);
+
+		if(all(IM == 0))
+		{
+			continue;
+		}
+
+		const FVector3f AnimationPosition = VectorLoad(&AnimationPositions[Offset]);
+		const FVector3f AnimationNormal = VectorLoad(&AnimationNormals[Offset]);
+
+		const float Radius = SphereRadiiBase + SphereRadiiRange * SphereRadii[Index];
+
+		const FVector3f Center = AnimationPosition - SphereOffsetDistance * AnimationNormal;  // Legacy version already includes the radius within the distance
+		const FVector3f CenterToParticle = Position - Center;
+		const float DistanceSquared = VectorSizeSquared(CenterToParticle);
+
+		const FVector3f DistSquaredLessRadiusSquared = ((Radius / sqrt(DistanceSquared)) * CenterToParticle) + Center;
+		const FVector3f DistSquaredLessDeadZoneRadius = AnimationPosition - (SphereOffsetDistance - Radius) * AnimationNormal;  // Legacy version already includes the radius to the distance
+
+		FVector3f NewP = VectorSelect(DistanceSquared < Radius * Radius, DistSquaredLessRadiusSquared, FloatZeroVector);
+		NewP = VectorSelect(DistanceSquared < DeadZoneSquareRadius, DistSquaredLessDeadZoneRadius, NewP);
+
+		if (IM != 0 && NewP.V[0] != 0)
+		{
+			VectorStore(&ParticlesPandInvM[Offset], SetVector4(NewP, IM));
+		}
+	}
+}
+
+export void ApplyLegacySphericalBackstopConstraintsWithoutMaps(uniform FVector4f ParticlesPandInvM[],
+													const uniform FVector3f AnimationPositions[],
+													const uniform FVector3f AnimationNormals[],
+													const uniform float SphereOffsetDistance,
+													const uniform float Radius,
+													const uniform int32 ParticleOffset,
+													const uniform int32 ParticleCount)
+{
+	static const float DeadZoneSquareRadius = FLOAT_SMALL_NUMBER; // We will not push the particle away in the dead zone
+
+	foreach(Index = 0 ... ParticleCount)
+	{
+		const uniform int32 Offset = ParticleOffset + extract(Index, 0);
+		const FVector4f PandInvM = VectorLoad(&ParticlesPandInvM[Offset]);
+
+		varying FVector3f Position;
+		varying float IM;
+		UnzipPandInvM(PandInvM, Position, IM);
+
+		if(all(IM == 0))
+		{
+			continue;
+		}
+
+		const FVector3f AnimationPosition = VectorLoad(&AnimationPositions[Offset]);
+		const FVector3f AnimationNormal = VectorLoad(&AnimationNormals[Offset]);
+
+		const FVector3f Center = AnimationPosition - SphereOffsetDistance * AnimationNormal;  // Legacy version already includes the radius within the distance
+		const FVector3f CenterToParticle = Position - Center;
+		const float DistanceSquared = VectorSizeSquared(CenterToParticle);
+
+		const FVector3f DistSquaredLessRadiusSquared = ((Radius / sqrt(DistanceSquared)) * CenterToParticle) + Center;
+		const FVector3f DistSquaredLessDeadZoneRadius = AnimationPosition - (SphereOffsetDistance - Radius) * AnimationNormal;  // Legacy version already includes the radius to the distance
+
+		FVector3f NewP = VectorSelect(DistanceSquared < Radius * Radius, DistSquaredLessRadiusSquared, FloatZeroVector);
+		NewP = VectorSelect(DistanceSquared < DeadZoneSquareRadius, DistSquaredLessDeadZoneRadius, NewP);
+
+		if (IM != 0 && NewP.V[0] != 0)
+		{
+			VectorStore(&ParticlesPandInvM[Offset], SetVector4(NewP, IM));
+		}
+	}
+}
+
+// Deprecated since 5.3
+export void ApplyLegacySphericalBackstopConstraints(uniform FVector4f ParticlesPandInvM[],
+													const uniform FVector3f AnimationPositions[],
+													const uniform FVector3f AnimationNormals[],
+													const uniform float SphereOffsetDistances[],
+													const uniform float SphereRadii[],
+													const uniform float SphereRadiiMultiplier,
+													const uniform int32 ParticleOffset,
+													const uniform int32 ParticleCount)
+{
+	static const uniform float SphereOffsetDistancesBase = 0.f;
+	static const uniform float SphereOffsetDistancesRange = 1.f;
+	static const uniform float SphereRadiiBase = 0.f;
+
+	ApplyLegacySphericalBackstopConstraintsWithMaps(
+		ParticlesPandInvM,
+		AnimationPositions,
+		AnimationNormals,
+		SphereOffsetDistances,
+		SphereRadii,
+		SphereOffsetDistancesBase,
+		SphereOffsetDistancesRange,
+		SphereRadiiBase,
+		SphereRadiiMultiplier,
+		ParticleOffset,
+		ParticleCount);
+}
+
+export void ApplySphericalBackstopConstraintsWithMaps(uniform FVector4f ParticlesPandInvM[],
 												const uniform FVector3f AnimationPositions[],
 												const uniform FVector3f AnimationNormals[],
 												const uniform float SphereOffsetDistances[],
 												const uniform float SphereRadii[],
-												const uniform float SphereRadiiMultiplier,
+												const uniform float SphereOffsetDistancesBase,
+												const uniform float SphereOffsetDistancesRange,
+												const uniform float SphereRadiiBase,
+												const uniform float SphereRadiiRange,
+												const uniform int32 ParticleOffset,
+												const uniform int32 ParticleCount)
+{
+	static const float DeadZoneSquareRadius = FLOAT_SMALL_NUMBER; // We will not push the particle away in the dead zone
+
+	foreach(Index = 0 ... ParticleCount)
+	{
+		const uniform int32 Offset = ParticleOffset + extract(Index, 0);
+		const FVector4f PandInvM = VectorLoad(&ParticlesPandInvM[Offset]);
+
+		varying FVector3f Position;
+		varying float IM;
+		UnzipPandInvM(PandInvM, Position, IM);
+
+		if(all(IM == 0))
+		{
+			continue;
+		}
+
+		const FVector3f AnimationPosition = VectorLoad(&AnimationPositions[Offset]);
+		const FVector3f AnimationNormal = VectorLoad(&AnimationNormals[Offset]);
+
+		const float SphereOffsetDistance = SphereOffsetDistancesBase + SphereOffsetDistancesRange * SphereOffsetDistances[Index];
+		const float Radius = SphereRadiiBase + SphereRadiiRange * SphereRadii[Index];
+
+		const FVector3f Center = AnimationPosition - (Radius + SphereOffsetDistance) * AnimationNormal;  // Non legacy version adds radius to the distance
+		const FVector3f CenterToParticle = Position - Center;
+		const float DistanceSquared = VectorSizeSquared(CenterToParticle);
+
+		const FVector3f DistSquaredLessRadiusSquared = ((Radius / sqrt(DistanceSquared)) * CenterToParticle) + Center;
+		const FVector3f DistSquaredLessDeadZoneRadius = AnimationPosition - SphereOffsetDistance * AnimationNormal;  // Non legacy version adds radius to the distance
+
+		FVector3f NewP = VectorSelect(DistanceSquared < Radius * Radius, DistSquaredLessRadiusSquared, FloatZeroVector);
+		NewP = VectorSelect(DistanceSquared < DeadZoneSquareRadius, DistSquaredLessDeadZoneRadius, NewP);
+
+		if (IM != 0 && NewP.V[0] != 0)
+		{
+			VectorStore(&ParticlesPandInvM[Offset], SetVector4(NewP, IM));
+		}
+	}
+}
+
+export void ApplySphericalBackstopConstraintsWithDistanceMap(uniform FVector4f ParticlesPandInvM[],
+												const uniform FVector3f AnimationPositions[],
+												const uniform FVector3f AnimationNormals[],
+												const uniform float SphereOffsetDistances[],
+												const uniform float SphereOffsetDistancesBase,
+												const uniform float SphereOffsetDistancesRange,
+												const uniform float Radius,
 												const uniform int32 ParticleOffset,
 												const uniform int32 ParticleCount)
 {
@@ -108,8 +388,7 @@ export void ApplySphericalBackstopConstraints(uniform FVector4f ParticlesPandInv
 
 	foreach(Index = 0 ... ParticleCount)
 	{
-		const int32 ParticleIndex = ParticleOffset + Index;
-		const uniform int32 Offset = extract(ParticleIndex, 0);
+		const uniform int32 Offset = ParticleOffset + extract(Index, 0);
 		const FVector4f PandInvM = VectorLoad(&ParticlesPandInvM[Offset]);
 
 		varying FVector3f Position;
@@ -121,11 +400,10 @@ export void ApplySphericalBackstopConstraints(uniform FVector4f ParticlesPandInv
 			continue;
 		}
 
-		const FVector3f AnimationPosition = VectorLoad(&AnimationPositions[extract(ParticleIndex, 0)]);
-		const FVector3f AnimationNormal = VectorLoad(&AnimationNormals[extract(ParticleIndex, 0)]);
+		const FVector3f AnimationPosition = VectorLoad(&AnimationPositions[Offset]);
+		const FVector3f AnimationNormal = VectorLoad(&AnimationNormals[Offset]);
 
-		const float SphereOffsetDistance = SphereOffsetDistances[Index] * SphereRadiiMultiplier;
-		const float Radius = SphereRadii[Index] * SphereRadiiMultiplier;
+		const float SphereOffsetDistance = SphereOffsetDistancesBase + SphereOffsetDistancesRange * SphereOffsetDistances[Index];
 
 		const FVector3f Center = AnimationPosition - (Radius + SphereOffsetDistance) * AnimationNormal;  // Non legacy version adds radius to the distance
 		const FVector3f CenterToParticle = Position - Center;
@@ -143,3 +421,123 @@ export void ApplySphericalBackstopConstraints(uniform FVector4f ParticlesPandInv
 		}
 	}
 }
+
+export void ApplySphericalBackstopConstraintsWithRadiusMap(uniform FVector4f ParticlesPandInvM[],
+												const uniform FVector3f AnimationPositions[],
+												const uniform FVector3f AnimationNormals[],
+												const uniform float SphereRadii[],
+												const uniform float SphereOffsetDistance,
+												const uniform float SphereRadiiBase,
+												const uniform float SphereRadiiRange,
+												const uniform int32 ParticleOffset,
+												const uniform int32 ParticleCount)
+{
+	static const float DeadZoneSquareRadius = FLOAT_SMALL_NUMBER; // We will not push the particle away in the dead zone
+
+	foreach(Index = 0 ... ParticleCount)
+	{
+		const uniform int32 Offset = ParticleOffset + extract(Index, 0);
+		const FVector4f PandInvM = VectorLoad(&ParticlesPandInvM[Offset]);
+
+		varying FVector3f Position;
+		varying float IM;
+		UnzipPandInvM(PandInvM, Position, IM);
+
+		if(all(IM == 0))
+		{
+			continue;
+		}
+
+		const FVector3f AnimationPosition = VectorLoad(&AnimationPositions[Offset]);
+		const FVector3f AnimationNormal = VectorLoad(&AnimationNormals[Offset]);
+
+		const float Radius = SphereRadiiBase + SphereRadiiRange * SphereRadii[Index];
+
+		const FVector3f Center = AnimationPosition - (Radius + SphereOffsetDistance) * AnimationNormal;  // Non legacy version adds radius to the distance
+		const FVector3f CenterToParticle = Position - Center;
+		const float DistanceSquared = VectorSizeSquared(CenterToParticle);
+
+		const FVector3f DistSquaredLessRadiusSquared = ((Radius / sqrt(DistanceSquared)) * CenterToParticle) + Center;
+		const FVector3f DistSquaredLessDeadZoneRadius = AnimationPosition - SphereOffsetDistance * AnimationNormal;  // Non legacy version adds radius to the distance
+
+		FVector3f NewP = VectorSelect(DistanceSquared < Radius * Radius, DistSquaredLessRadiusSquared, FloatZeroVector);
+		NewP = VectorSelect(DistanceSquared < DeadZoneSquareRadius, DistSquaredLessDeadZoneRadius, NewP);
+
+		if (IM != 0 && NewP.V[0] != 0)
+		{
+			VectorStore(&ParticlesPandInvM[Offset], SetVector4(NewP, IM));
+		}
+	}
+}
+
+export void ApplySphericalBackstopConstraintsWithoutMaps(uniform FVector4f ParticlesPandInvM[],
+												const uniform FVector3f AnimationPositions[],
+												const uniform FVector3f AnimationNormals[],
+												const uniform float SphereOffsetDistance,
+												const uniform float Radius,
+												const uniform int32 ParticleOffset,
+												const uniform int32 ParticleCount)
+{
+	static const float DeadZoneSquareRadius = FLOAT_SMALL_NUMBER; // We will not push the particle away in the dead zone
+
+	foreach(Index = 0 ... ParticleCount)
+	{
+		const uniform int32 Offset = ParticleOffset + extract(Index, 0);
+		const FVector4f PandInvM = VectorLoad(&ParticlesPandInvM[Offset]);
+
+		varying FVector3f Position;
+		varying float IM;
+		UnzipPandInvM(PandInvM, Position, IM);
+
+		if(all(IM == 0))
+		{
+			continue;
+		}
+
+		const FVector3f AnimationPosition = VectorLoad(&AnimationPositions[Offset]);
+		const FVector3f AnimationNormal = VectorLoad(&AnimationNormals[Offset]);
+
+		const FVector3f Center = AnimationPosition - (Radius + SphereOffsetDistance) * AnimationNormal;  // Non legacy version adds radius to the distance
+		const FVector3f CenterToParticle = Position - Center;
+		const float DistanceSquared = VectorSizeSquared(CenterToParticle);
+
+		const FVector3f DistSquaredLessRadiusSquared = ((Radius / sqrt(DistanceSquared)) * CenterToParticle) + Center;
+		const FVector3f DistSquaredLessDeadZoneRadius = AnimationPosition - SphereOffsetDistance * AnimationNormal;  // Non legacy version adds radius to the distance
+
+		FVector3f NewP = VectorSelect(DistanceSquared < Radius * Radius, DistSquaredLessRadiusSquared, FloatZeroVector);
+		NewP = VectorSelect(DistanceSquared < DeadZoneSquareRadius, DistSquaredLessDeadZoneRadius, NewP);
+
+		if (IM != 0 && NewP.V[0] != 0)
+		{
+			VectorStore(&ParticlesPandInvM[Offset], SetVector4(NewP, IM));
+		}
+	}
+}
+
+// Deprecated since 5.3
+export void ApplySphericalBackstopConstraints(uniform FVector4f ParticlesPandInvM[],
+												const uniform FVector3f AnimationPositions[],
+												const uniform FVector3f AnimationNormals[],
+												const uniform float SphereOffsetDistances[],
+												const uniform float SphereRadii[],
+												const uniform float SphereRadiiMultiplier,
+												const uniform int32 ParticleOffset,
+												const uniform int32 ParticleCount)
+{
+	static const uniform float SphereOffsetDistancesBase = 0.f;
+	static const uniform float SphereOffsetDistancesRange = 1.f;
+	static const uniform float SphereRadiiBase = 0.f;
+
+	ApplySphericalBackstopConstraintsWithMaps(
+		ParticlesPandInvM,
+		AnimationPositions,
+		AnimationNormals,
+		SphereOffsetDistances,
+		SphereRadii,
+		SphereOffsetDistancesBase,
+		SphereOffsetDistancesRange,
+		SphereRadiiBase,
+		SphereRadiiMultiplier,
+		ParticleOffset,
+		ParticleCount);
+}
diff --git a/Engine/Source/Runtime/Experimental/Chaos/Private/Chaos/PBDSpringConstraints.cpp b/Engine/Source/Runtime/Experimental/Chaos/Private/Chaos/PBDSpringConstraints.cpp
--- a/Engine/Source/Runtime/Experimental/Chaos/Private/Chaos/PBDSpringConstraints.cpp
+++ b/Engine/Source/Runtime/Experimental/Chaos/Private/Chaos/PBDSpringConstraints.cpp
@@ -30,7 +30,7 @@ static int32 Chaos_Spring_ParallelConstraintCount = 100;
 FAutoConsoleVariableRef CVarChaosSpringParallelConstraintCount(TEXT("p.Chaos.Spring.ParallelConstraintCount"), Chaos_Spring_ParallelConstraintCount, TEXT("If we have more constraints than this, use parallel-for in Apply."));
 #endif
 
-void FPBDSpringConstraints::InitColor(const FSolverParticles& Particles, const int32 ParticleOffset, const int32 ParticleCount)
+void FPBDSpringConstraints::InitColor(const FSolverParticles& Particles)
 {
 	// In dev builds we always color so we can tune the system without restarting. See Apply()
 #if UE_BUILD_SHIPPING || UE_BUILD_TEST
@@ -184,4 +184,42 @@ void FPBDSpringConstraints::Apply(FSolverParticles& Particles, const FSolverReal
 	}
 }
 
+void FPBDEdgeSpringConstraints::SetProperties(
+	const FCollectionPropertyConstFacade& PropertyCollection,
+	const TMap<FString, TConstArrayView<FRealSingle>>& WeightMaps)
+{
+	if (IsEdgeSpringStiffnessMutable(PropertyCollection))
+	{
+		const FSolverVec2 WeightedValue(GetWeightedFloatEdgeSpringStiffness(PropertyCollection));
+		if (IsEdgeSpringStiffnessStringDirty(PropertyCollection))
+		{
+			const FString& WeightMapName = GetEdgeSpringStiffnessString(PropertyCollection);
+			Stiffness = FPBDStiffness(WeightedValue, WeightMaps.FindRef(WeightMapName), ParticleCount);
+		}
+		else
+		{
+			Stiffness.SetWeightedValue(WeightedValue);
+		}
+	}
+}
+
+void FPBDBendingSpringConstraints::SetProperties(
+	const FCollectionPropertyConstFacade& PropertyCollection,
+	const TMap<FString, TConstArrayView<FRealSingle>>& WeightMaps)
+{
+	if (IsBendingSpringStiffnessMutable(PropertyCollection))
+	{
+		const FSolverVec2 WeightedValue(GetWeightedFloatBendingSpringStiffness(PropertyCollection));
+		if (IsBendingSpringStiffnessStringDirty(PropertyCollection))
+		{
+			const FString& WeightMapName = GetBendingSpringStiffnessString(PropertyCollection);
+			Stiffness = FPBDStiffness(WeightedValue, WeightMaps.FindRef(WeightMapName), ParticleCount);
+		}
+		else
+		{
+			Stiffness.SetWeightedValue(WeightedValue);
+		}
+	}
+}
+
 } // End namespace Chaos::Softs
diff --git a/Engine/Source/Runtime/Experimental/Chaos/Private/Chaos/VelocityField.cpp b/Engine/Source/Runtime/Experimental/Chaos/Private/Chaos/VelocityField.cpp
--- a/Engine/Source/Runtime/Experimental/Chaos/Private/Chaos/VelocityField.cpp
+++ b/Engine/Source/Runtime/Experimental/Chaos/Private/Chaos/VelocityField.cpp
@@ -16,6 +16,56 @@ FAutoConsoleVariableRef CVarChaosVelocityFieldISPCEnabled(TEXT("p.Chaos.Velocity
 
 namespace Chaos::Softs {
 
+void FVelocityAndPressureField::SetProperties(
+	const FCollectionPropertyConstFacade& PropertyCollection,
+	const TMap<FString, TConstArrayView<FRealSingle>>& Weightmaps,
+	FSolverReal WorldScale)
+{
+	if (IsFluidDensityMutable(PropertyCollection))
+	{
+		constexpr FSolverReal OneQuarter = (FSolverReal)0.25;
+		QuarterRho = (FSolverReal)GetFluidDensity(PropertyCollection) * OneQuarter;
+	}
+
+	bool bCheckForWeightMapUpdate = false;
+
+	if (IsDragMutable(PropertyCollection))
+	{
+		const FSolverVec2 Drag(GetWeightedFloatDrag(PropertyCollection));
+		DragBase = FMath::Clamp(Drag[0], MinCoefficient, MaxCoefficient);
+		DragRange = FMath::Clamp(Drag[1], MinCoefficient, MaxCoefficient) - DragBase;
+		bCheckForWeightMapUpdate = true;
+	}
+
+	if (IsLiftMutable(PropertyCollection))
+	{
+		const FSolverVec2 Lift(GetWeightedFloatLift(PropertyCollection));
+		LiftBase = FMath::Clamp(Lift[0], MinCoefficient, MaxCoefficient);
+		LiftRange = FMath::Clamp(Lift[1], MinCoefficient, MaxCoefficient) - LiftBase;
+		bCheckForWeightMapUpdate = true;
+	}
+
+	if (IsPressureMutable(PropertyCollection))
+	{
+		const FSolverVec2 Pressure(GetWeightedFloatPressure(PropertyCollection));
+		PressureBase = Pressure[0] / WorldScale;
+		PressureRange = Pressure[1] / WorldScale - PressureBase;
+		bCheckForWeightMapUpdate = true;
+	}
+
+	if (bCheckForWeightMapUpdate && (
+		IsDragStringDirty(PropertyCollection) ||
+		IsLiftStringDirty(PropertyCollection) ||
+		IsPressureStringDirty(PropertyCollection)))
+	{
+		const TConstArrayView<FRealSingle> DragMultipliers = Weightmaps.FindRef(GetDragString(PropertyCollection));
+		const TConstArrayView<FRealSingle> LiftMultipliers = Weightmaps.FindRef(GetLiftString(PropertyCollection));
+		const TConstArrayView<FRealSingle> PressureMultipliers = Weightmaps.FindRef(GetPressureString(PropertyCollection));
+
+		SetMultipliers(DragMultipliers, LiftMultipliers, PressureMultipliers);
+	}
+}
+
 void FVelocityAndPressureField::SetGeometry(const FTriangleMesh* TriangleMesh, const TConstArrayView<FRealSingle>& DragMultipliers, const TConstArrayView<FRealSingle>& LiftMultipliers, const TConstArrayView<FRealSingle>& PressureMultipliers)
 {
 	if (TriangleMesh)
@@ -28,35 +78,6 @@ void FVelocityAndPressureField::SetGeometry(const FTriangleMesh* TriangleMesh, c
 		NumParticles = 1 + Range[1] - Offset;
 
 		Forces.SetNumUninitialized(Elements.Num());
-
-		const bool bHasDragMultipliers = DragMultipliers.Num() == NumParticles;
-		const bool bHasLiftMultipliers = LiftMultipliers.Num() == NumParticles;
-		const bool bHasPressureMultipliers = PressureMultipliers.Num() == NumParticles;
-
-		if (!bHasDragMultipliers && !bHasLiftMultipliers && !bHasPressureMultipliers)
-		{
-			Multipliers.Reset();
-		}
-		else
-		{
-			constexpr FSolverReal OneThird = (FSolverReal)1. / (FSolverReal)3.;
-
-			Multipliers.SetNumUninitialized(Elements.Num());
-
-			for (int32 ElementIndex = 0; ElementIndex < Elements.Num(); ++ElementIndex)
-			{
-				const TVec3<int32>& Element = Elements[ElementIndex];
-				const int32 I0 = Element[0] - Offset;
-				const int32 I1 = Element[1] - Offset;
-				const int32 I2 = Element[2] - Offset;
-
-				const FSolverReal DragMultiplier = bHasDragMultipliers ? (FSolverReal)(DragMultipliers[I0] + DragMultipliers[I1] + DragMultipliers[I2]) * OneThird : (FSolverReal)0.;
-				const FSolverReal LiftMultiplier = bHasLiftMultipliers ? (FSolverReal)(LiftMultipliers[I0] + LiftMultipliers[I1] + LiftMultipliers[I2]) * OneThird : (FSolverReal)0.;
-				const FSolverReal PressureMultiplier = bHasPressureMultipliers ? (FSolverReal)(PressureMultipliers[I0] + PressureMultipliers[I1] + PressureMultipliers[I2]) * OneThird : (FSolverReal)0.;
-
-				Multipliers[ElementIndex] = FSolverVec3(DragMultiplier, LiftMultiplier, PressureMultiplier);
-			}
-		}
 	}
 	else
 	{
@@ -65,12 +86,44 @@ void FVelocityAndPressureField::SetGeometry(const FTriangleMesh* TriangleMesh, c
 		Offset = 0;
 		NumParticles = 0;
 		Forces.Reset();
-		Multipliers.Reset();
 		SetProperties(FSolverVec2(0.), FSolverVec2(0.), (FSolverReal)0., FSolverVec2(0.));
 	}
+
+	SetMultipliers(DragMultipliers, LiftMultipliers, PressureMultipliers);
 }
 
+void FVelocityAndPressureField::SetMultipliers(
+	const TConstArrayView<FRealSingle>& DragMultipliers,
+	const TConstArrayView<FRealSingle>& LiftMultipliers,
+	const TConstArrayView<FRealSingle>& PressureMultipliers)
+{
+	Multipliers.Reset();
+
+	const bool bHasDragMultipliers = DragMultipliers.Num() == NumParticles;
+	const bool bHasLiftMultipliers = LiftMultipliers.Num() == NumParticles;
+	const bool bHasPressureMultipliers = PressureMultipliers.Num() == NumParticles;
+
+	if (bHasDragMultipliers || bHasLiftMultipliers || bHasPressureMultipliers)
+	{
+		constexpr FSolverReal OneThird = (FSolverReal)1. / (FSolverReal)3.;
+
+		Multipliers.SetNumUninitialized(Elements.Num());
 
+		for (int32 ElementIndex = 0; ElementIndex < Elements.Num(); ++ElementIndex)
+		{
+			const TVec3<int32>& Element = Elements[ElementIndex];
+			const int32 I0 = Element[0] - Offset;
+			const int32 I1 = Element[1] - Offset;
+			const int32 I2 = Element[2] - Offset;
+
+			const FSolverReal DragMultiplier = bHasDragMultipliers ? (FSolverReal)(DragMultipliers[I0] + DragMultipliers[I1] + DragMultipliers[I2]) * OneThird : (FSolverReal)0.;
+			const FSolverReal LiftMultiplier = bHasLiftMultipliers ? (FSolverReal)(LiftMultipliers[I0] + LiftMultipliers[I1] + LiftMultipliers[I2]) * OneThird : (FSolverReal)0.;
+			const FSolverReal PressureMultiplier = bHasPressureMultipliers ? (FSolverReal)(PressureMultipliers[I0] + PressureMultipliers[I1] + PressureMultipliers[I2]) * OneThird : (FSolverReal)0.;
+
+			Multipliers[ElementIndex] = FSolverVec3(DragMultiplier, LiftMultiplier, PressureMultiplier);
+		}
+	}
+}
 
 void FVelocityAndPressureField::UpdateForces(const FSolverParticles& InParticles, const FSolverReal /*Dt*/)
 {
diff --git a/Engine/Source/Runtime/Experimental/Chaos/Private/Chaos/XPBDAnisotropicBendingConstraints.cpp b/Engine/Source/Runtime/Experimental/Chaos/Private/Chaos/XPBDAnisotropicBendingConstraints.cpp
--- a/Engine/Source/Runtime/Experimental/Chaos/Private/Chaos/XPBDAnisotropicBendingConstraints.cpp
+++ b/Engine/Source/Runtime/Experimental/Chaos/Private/Chaos/XPBDAnisotropicBendingConstraints.cpp
@@ -23,26 +23,20 @@ namespace Chaos::Softs {
 extern int32 Chaos_XPBDBending_ParallelConstraintCount;
 
 FXPBDAnisotropicBendingConstraints::FXPBDAnisotropicBendingConstraints(const FSolverParticles& InParticles,
-	int32 ParticleOffset,
-	int32 ParticleCount,
+	int32 InParticleOffset,
+	int32 InParticleCount,
 	const FTriangleMesh& TriangleMesh,
 	const TArray<TVec3<FVec2f>>& FaceVertexPatternPositions,
-	const TConstArrayView<FRealSingle>& StiffnessWarpMultipliers,
-	const TConstArrayView<FRealSingle>& StiffnessWeftMultipliers,
-	const TConstArrayView<FRealSingle>& StiffnessBiasMultipliers,
-	const TConstArrayView<FRealSingle>& BucklingStiffnessWarpMultipliers,
-	const TConstArrayView<FRealSingle>& BucklingStiffnessWeftMultipliers,
-	const TConstArrayView<FRealSingle>& BucklingStiffnessBiasMultipliers,
-	const TConstArrayView<FRealSingle>& DampingMultipliers,
+	const TMap<FString, TConstArrayView<FRealSingle>>& WeightMaps,
 	const FCollectionPropertyConstFacade& PropertyCollection,
 	bool bTrimKinematicConstraints)
 	: Base(
 		InParticles,
-		ParticleOffset,
-		ParticleCount,
+		InParticleOffset,
+		InParticleCount,
 		TriangleMesh.GetUniqueAdjacentElements(),
-		StiffnessWarpMultipliers,
-		BucklingStiffnessWarpMultipliers,
+		WeightMaps.FindRef(GetXPBDBendingElementStiffnessWarpString(PropertyCollection, XPBDBendingElementStiffnessWarpName.ToString())),
+		WeightMaps.FindRef(GetXPBDBucklingStiffnessWarpString(PropertyCollection, XPBDBucklingStiffnessWarpName.ToString())),
 		FSolverVec2(GetWeightedFloatXPBDBendingElementStiffnessWarp(PropertyCollection, MaxStiffness)),
 		(FSolverReal)GetXPBDBucklingRatio(PropertyCollection, 0.f),
 		FSolverVec2(GetWeightedFloatXPBDBucklingStiffnessWarp(PropertyCollection, MaxStiffness)),
@@ -50,7 +44,7 @@ FXPBDAnisotropicBendingConstraints::FXPBDAnisotropicBendingConstraints(const FSo
 		MaxStiffness)
 	, StiffnessWeft(
 		FSolverVec2(GetWeightedFloatXPBDBendingElementStiffnessWeft(PropertyCollection, MaxStiffness)),
-		StiffnessWeftMultipliers,
+		WeightMaps.FindRef(GetXPBDBendingElementStiffnessWeftString(PropertyCollection, XPBDBendingElementStiffnessWeftName.ToString())),
 		TConstArrayView<TVec2<int32>>(ConstraintSharedEdges),
 		ParticleOffset,
 		ParticleCount,
@@ -59,7 +53,7 @@ FXPBDAnisotropicBendingConstraints::FXPBDAnisotropicBendingConstraints(const FSo
 		MaxStiffness)
 	, StiffnessBias(
 		FSolverVec2(GetWeightedFloatXPBDBendingElementStiffnessBias(PropertyCollection, MaxStiffness)),
-		StiffnessBiasMultipliers,
+		WeightMaps.FindRef(GetXPBDBendingElementStiffnessBiasString(PropertyCollection, XPBDBendingElementStiffnessBiasName.ToString())),
 		TConstArrayView<TVec2<int32>>(ConstraintSharedEdges),
 		ParticleOffset,
 		ParticleCount,
@@ -68,7 +62,7 @@ FXPBDAnisotropicBendingConstraints::FXPBDAnisotropicBendingConstraints(const FSo
 		MaxStiffness)
 	, BucklingStiffnessWeft(
 		FSolverVec2(GetWeightedFloatXPBDBucklingStiffnessWeft(PropertyCollection, MaxStiffness)),
-		BucklingStiffnessWeftMultipliers,
+		WeightMaps.FindRef(GetXPBDBucklingStiffnessWeftString(PropertyCollection, XPBDBucklingStiffnessWeftName.ToString())),
 		TConstArrayView<TVec2<int32>>(ConstraintSharedEdges),
 		ParticleOffset,
 		ParticleCount,
@@ -77,7 +71,7 @@ FXPBDAnisotropicBendingConstraints::FXPBDAnisotropicBendingConstraints(const FSo
 		MaxStiffness)
 	, BucklingStiffnessBias(
 		FSolverVec2(GetWeightedFloatXPBDBucklingStiffnessBias(PropertyCollection, MaxStiffness)),
-		BucklingStiffnessBiasMultipliers,
+		WeightMaps.FindRef(GetXPBDBucklingStiffnessBiasString(PropertyCollection, XPBDBucklingStiffnessBiasName.ToString())),
 		TConstArrayView<TVec2<int32>>(ConstraintSharedEdges),
 		ParticleOffset,
 		ParticleCount,
@@ -86,14 +80,14 @@ FXPBDAnisotropicBendingConstraints::FXPBDAnisotropicBendingConstraints(const FSo
 		MaxStiffness)
 	, DampingRatio(
 		FSolverVec2(GetWeightedFloatXPBDBendingElementDamping(PropertyCollection, MinDamping)).ClampAxes(MinDamping, MaxDamping),
-		DampingMultipliers,
+		WeightMaps.FindRef(GetXPBDBendingElementDampingString(PropertyCollection, XPBDBendingElementDampingName.ToString())),
 		TConstArrayView<TVec2<int32>>(ConstraintSharedEdges),
 		ParticleOffset,
 		ParticleCount)
 	, WarpWeftBiasBaseMultipliers(GenerateWarpWeftBiasBaseMultipliers(FaceVertexPatternPositions, TriangleMesh))
 {
 	Lambdas.Init((FSolverReal)0., Constraints.Num());
-	InitColor(InParticles, ParticleOffset, ParticleCount);
+	InitColor(InParticles);
 }
 
 FXPBDAnisotropicBendingConstraints::FXPBDAnisotropicBendingConstraints(const FSolverParticles& InParticles,
@@ -174,7 +168,7 @@ FXPBDAnisotropicBendingConstraints::FXPBDAnisotropicBendingConstraints(const FSo
 	, WarpWeftBiasBaseMultipliers(GenerateWarpWeftBiasBaseMultipliers(FaceVertexPatternPositions, TriangleMesh))
 {
 	Lambdas.Init((FSolverReal)0., Constraints.Num());
-	InitColor(InParticles, ParticleOffset, ParticleCount);
+	InitColor(InParticles);
 }
 
 TArray<FSolverVec3> FXPBDAnisotropicBendingConstraints::GenerateWarpWeftBiasBaseMultipliers(const TArray<TVec3<FVec2f>>& FaceVertexPatternPositions, const FTriangleMesh& TriangleMesh) const
@@ -232,7 +226,7 @@ TArray<FSolverVec3> FXPBDAnisotropicBendingConstraints::GenerateWarpWeftBiasBase
 	return WarpWeftBaseMultiplierResult;
 }
 
-void FXPBDAnisotropicBendingConstraints::InitColor(const FSolverParticles& InParticles, const int32 ParticleOffset, const int32 ParticleCount)
+void FXPBDAnisotropicBendingConstraints::InitColor(const FSolverParticles& InParticles)
 {
 	// In dev builds we always color so we can tune the system without restarting. See Apply()
 #if UE_BUILD_SHIPPING || UE_BUILD_TEST
@@ -286,6 +280,160 @@ void FXPBDAnisotropicBendingConstraints::InitColor(const FSolverParticles& InPar
 	}
 }
 
+void FXPBDAnisotropicBendingConstraints::SetProperties(
+	const FCollectionPropertyConstFacade& PropertyCollection,
+	const TMap<FString, TConstArrayView<FRealSingle>>& WeightMaps)
+{
+	if (IsXPBDBendingElementStiffnessWarpMutable(PropertyCollection))
+	{
+		const FSolverVec2 WeightedValue(GetWeightedFloatXPBDBendingElementStiffnessWarp(PropertyCollection));
+		if (IsXPBDBendingElementStiffnessWarpStringDirty(PropertyCollection))
+		{
+			const FString& WeightMapName = GetXPBDBendingElementStiffnessWarpString(PropertyCollection);
+			Stiffness = FPBDStiffness(
+				WeightedValue,
+				WeightMaps.FindRef(WeightMapName),
+				TConstArrayView<TVec2<int32>>(ConstraintSharedEdges),
+				ParticleOffset,
+				ParticleCount,
+				FPBDStiffness::DefaultTableSize,
+				FPBDStiffness::DefaultParameterFitBase,
+				MaxStiffness);
+		}
+		else
+		{
+			Stiffness.SetWeightedValue(WeightedValue, MaxStiffness);
+		}
+	}
+	if (IsXPBDBendingElementStiffnessWeftMutable(PropertyCollection))
+	{
+		const FSolverVec2 WeightedValue(GetWeightedFloatXPBDBendingElementStiffnessWeft(PropertyCollection));
+		if (IsXPBDBendingElementStiffnessWeftStringDirty(PropertyCollection))
+		{
+			const FString& WeightMapName = GetXPBDBendingElementStiffnessWeftString(PropertyCollection);
+			StiffnessWeft = FPBDStiffness(
+				WeightedValue,
+				WeightMaps.FindRef(WeightMapName),
+				TConstArrayView<TVec2<int32>>(ConstraintSharedEdges),
+				ParticleOffset,
+				ParticleCount,
+				FPBDStiffness::DefaultTableSize,
+				FPBDStiffness::DefaultParameterFitBase,
+				MaxStiffness);
+		}
+		else
+		{
+			StiffnessWeft.SetWeightedValue(WeightedValue, MaxStiffness);
+		}
+	}
+	if (IsXPBDBendingElementStiffnessBiasMutable(PropertyCollection))
+	{
+		const FSolverVec2 WeightedValue(GetWeightedFloatXPBDBendingElementStiffnessBias(PropertyCollection));
+		if (IsXPBDBendingElementStiffnessBiasStringDirty(PropertyCollection))
+		{
+			const FString& WeightMapName = GetXPBDBendingElementStiffnessBiasString(PropertyCollection);
+			StiffnessBias = FPBDStiffness(
+				WeightedValue,
+				WeightMaps.FindRef(WeightMapName),
+				TConstArrayView<TVec2<int32>>(ConstraintSharedEdges),
+				ParticleOffset,
+				ParticleCount,
+				FPBDStiffness::DefaultTableSize,
+				FPBDStiffness::DefaultParameterFitBase,
+				MaxStiffness);
+		}
+		else
+		{
+			StiffnessBias.SetWeightedValue(WeightedValue, MaxStiffness);
+		}
+	}
+	if (IsXPBDBucklingRatioMutable(PropertyCollection))
+	{
+		BucklingRatio = FMath::Clamp(GetXPBDBucklingRatio(PropertyCollection), (FSolverReal)0., (FSolverReal)1.);
+	}
+	if (IsXPBDBucklingStiffnessWarpMutable(PropertyCollection))
+	{
+		const FSolverVec2 WeightedValue(GetWeightedFloatXPBDBucklingStiffnessWarp(PropertyCollection));
+		if (IsXPBDBucklingStiffnessWarpStringDirty(PropertyCollection))
+		{
+			const FString& WeightMapName = GetXPBDBucklingStiffnessWarpString(PropertyCollection);
+			BucklingStiffness = FPBDStiffness(
+				WeightedValue,
+				WeightMaps.FindRef(WeightMapName),
+				TConstArrayView<TVec2<int32>>(ConstraintSharedEdges),
+				ParticleOffset,
+				ParticleCount,
+				FPBDStiffness::DefaultTableSize,
+				FPBDStiffness::DefaultParameterFitBase,
+				MaxStiffness);
+		}
+		else
+		{
+			BucklingStiffness.SetWeightedValue(WeightedValue, MaxStiffness);
+		}
+	}
+	if (IsXPBDBucklingStiffnessWeftMutable(PropertyCollection))
+	{
+		const FSolverVec2 WeightedValue(GetWeightedFloatXPBDBucklingStiffnessWeft(PropertyCollection));
+		if (IsXPBDBucklingStiffnessWeftStringDirty(PropertyCollection))
+		{
+			const FString& WeightMapName = GetXPBDBucklingStiffnessWeftString(PropertyCollection);
+			BucklingStiffnessWeft = FPBDStiffness(
+				WeightedValue,
+				WeightMaps.FindRef(WeightMapName),
+				TConstArrayView<TVec2<int32>>(ConstraintSharedEdges),
+				ParticleOffset,
+				ParticleCount,
+				FPBDStiffness::DefaultTableSize,
+				FPBDStiffness::DefaultParameterFitBase,
+				MaxStiffness);
+		}
+		else
+		{
+			BucklingStiffnessWeft.SetWeightedValue(WeightedValue, MaxStiffness);
+		}
+	}
+	if (IsXPBDBucklingStiffnessBiasMutable(PropertyCollection))
+	{
+		const FSolverVec2 WeightedValue(GetWeightedFloatXPBDBucklingStiffnessBias(PropertyCollection));
+		if (IsXPBDBucklingStiffnessBiasStringDirty(PropertyCollection))
+		{
+			const FString& WeightMapName = GetXPBDBucklingStiffnessBiasString(PropertyCollection);
+			BucklingStiffnessBias = FPBDStiffness(
+				WeightedValue,
+				WeightMaps.FindRef(WeightMapName),
+				TConstArrayView<TVec2<int32>>(ConstraintSharedEdges),
+				ParticleOffset,
+				ParticleCount,
+				FPBDStiffness::DefaultTableSize,
+				FPBDStiffness::DefaultParameterFitBase,
+				MaxStiffness);
+		}
+		else
+		{
+			BucklingStiffnessBias.SetWeightedValue(WeightedValue, MaxStiffness);
+		}
+	}
+	if (IsXPBDBendingElementDampingMutable(PropertyCollection))
+	{
+		const FSolverVec2 WeightedValue = FSolverVec2(GetXPBDBendingElementDamping(PropertyCollection)).ClampAxes(MinDamping, MaxDamping);
+		if (IsXPBDBendingElementDampingStringDirty(PropertyCollection))
+		{
+			const FString& WeightMapName = GetXPBDBendingElementDampingString(PropertyCollection);
+			DampingRatio = FPBDWeightMap(
+				WeightedValue,
+				WeightMaps.FindRef(WeightMapName),
+				TConstArrayView<TVec2<int32>>(ConstraintSharedEdges),
+				ParticleOffset,
+				ParticleCount);
+		}
+		else
+		{
+			DampingRatio.SetWeightedValue(WeightedValue);
+		}
+	}
+}
+
 void FXPBDAnisotropicBendingConstraints::ApplyHelper(FSolverParticles& Particles, const FSolverReal Dt, const int32 ConstraintIndex, const FSolverVec3& ExpStiffnessValues, 
 	const FSolverVec3& ExpBucklingStiffnessValues, const FSolverReal DampingRatioValue) const
 {
diff --git a/Engine/Source/Runtime/Experimental/Chaos/Private/Chaos/XPBDBendingConstraints.cpp b/Engine/Source/Runtime/Experimental/Chaos/Private/Chaos/XPBDBendingConstraints.cpp
--- a/Engine/Source/Runtime/Experimental/Chaos/Private/Chaos/XPBDBendingConstraints.cpp
+++ b/Engine/Source/Runtime/Experimental/Chaos/Private/Chaos/XPBDBendingConstraints.cpp
@@ -23,7 +23,7 @@ namespace Chaos::Softs {
 int32 Chaos_XPBDBending_ParallelConstraintCount = 100;
 FAutoConsoleVariableRef CVarChaosXPBDBendingParallelConstraintCount(TEXT("p.Chaos.XPBDBending.ParallelConstraintCount"), Chaos_XPBDBending_ParallelConstraintCount, TEXT("If we have more constraints than this, use parallel-for in Apply."));
 
-void FXPBDBendingConstraints::InitColor(const FSolverParticles& InParticles, const int32 ParticleOffset, const int32 ParticleCount)
+void FXPBDBendingConstraints::InitColor(const FSolverParticles& InParticles)
 {
 	// In dev builds we always color so we can tune the system without restarting. See Apply()
 #if UE_BUILD_SHIPPING || UE_BUILD_TEST
@@ -68,6 +68,76 @@ void FXPBDBendingConstraints::InitColor(const FSolverParticles& InParticles, con
 	}
 }
 
+void FXPBDBendingConstraints::SetProperties(
+	const FCollectionPropertyConstFacade& PropertyCollection,
+	const TMap<FString, TConstArrayView<FRealSingle>>& WeightMaps)
+{
+	if (IsXPBDBendingElementStiffnessMutable(PropertyCollection))
+	{
+		const FSolverVec2 WeightedValue(GetWeightedFloatXPBDBendingElementStiffness(PropertyCollection));
+		if (IsXPBDBendingElementStiffnessStringDirty(PropertyCollection))
+		{
+			const FString& WeightMapName = GetXPBDBendingElementStiffnessString(PropertyCollection);
+			Stiffness = FPBDStiffness(
+				WeightedValue,
+				WeightMaps.FindRef(WeightMapName),
+				TConstArrayView<TVec2<int32>>(ConstraintSharedEdges),
+				ParticleOffset,
+				ParticleCount,
+				FPBDStiffness::DefaultTableSize,
+				FPBDStiffness::DefaultParameterFitBase,
+				MaxStiffness);
+		}
+		else
+		{
+			Stiffness.SetWeightedValue(WeightedValue, MaxStiffness);
+		}
+	}
+	if (IsXPBDBucklingRatioMutable(PropertyCollection))
+	{
+		BucklingRatio = FMath::Clamp(GetXPBDBucklingRatio(PropertyCollection), (FSolverReal)0., (FSolverReal)1.);
+	}
+	if (IsXPBDBucklingStiffnessMutable(PropertyCollection))
+	{
+		const FSolverVec2 WeightedValue(GetWeightedFloatXPBDBucklingStiffness(PropertyCollection));
+		if (IsXPBDBucklingStiffnessStringDirty(PropertyCollection))
+		{
+			const FString& WeightMapName = GetXPBDBucklingStiffnessString(PropertyCollection);
+			BucklingStiffness = FPBDStiffness(
+				WeightedValue,
+				WeightMaps.FindRef(WeightMapName),
+				TConstArrayView<TVec2<int32>>(ConstraintSharedEdges),
+				ParticleOffset,
+				ParticleCount,
+				FPBDStiffness::DefaultTableSize,
+				FPBDStiffness::DefaultParameterFitBase,
+				MaxStiffness);
+		}
+		else
+		{
+			BucklingStiffness.SetWeightedValue(WeightedValue, MaxStiffness);
+		}
+	}
+	if (IsXPBDBendingElementDampingMutable(PropertyCollection))
+	{
+		const FSolverVec2 WeightedValue = FSolverVec2(GetWeightedFloatXPBDBendingElementDamping(PropertyCollection)).ClampAxes(MinDamping, MaxDamping);
+		if (IsXPBDBendingElementDampingStringDirty(PropertyCollection))
+		{
+			const FString& WeightMapName = GetXPBDBendingElementDampingString(PropertyCollection);
+			DampingRatio = FPBDWeightMap(
+				WeightedValue,
+				WeightMaps.FindRef(WeightMapName),
+				TConstArrayView<TVec2<int32>>(ConstraintSharedEdges),
+				ParticleOffset,
+				ParticleCount);
+		}
+		else
+		{
+			DampingRatio.SetWeightedValue(WeightedValue);
+		}
+	}
+}
+
 void FXPBDBendingConstraints::ApplyHelper(FSolverParticles& Particles, const FSolverReal Dt, const int32 ConstraintIndex, const FSolverReal StiffnessValue, const FSolverReal BucklingValue, const FSolverReal DampingRatioValue) const
 {
 	const TVec4<int32>& Constraint = Constraints[ConstraintIndex];
diff --git a/Engine/Source/Runtime/Experimental/Chaos/Private/Chaos/XPBDSpringConstraints.cpp b/Engine/Source/Runtime/Experimental/Chaos/Private/Chaos/XPBDSpringConstraints.cpp
--- a/Engine/Source/Runtime/Experimental/Chaos/Private/Chaos/XPBDSpringConstraints.cpp
+++ b/Engine/Source/Runtime/Experimental/Chaos/Private/Chaos/XPBDSpringConstraints.cpp
@@ -30,7 +30,7 @@ static int32 Chaos_XPBDSpring_ParallelConstraintCount = 100;
 FAutoConsoleVariableRef CVarChaosXPBDSpringParallelConstraintCount(TEXT("p.Chaos.XPBDSpring.ParallelConstraintCount"), Chaos_XPBDSpring_ParallelConstraintCount, TEXT("If we have more constraints than this, use parallel-for in Apply."));
 #endif
 
-void FXPBDSpringConstraints::InitColor(const FSolverParticles& Particles, const int32 ParticleOffset, const int32 ParticleCount)
+void FXPBDSpringConstraints::InitColor(const FSolverParticles& Particles)
 {
 	// In dev builds we always color so we can tune the system without restarting. See Apply()
 #if UE_BUILD_SHIPPING || UE_BUILD_TEST
@@ -257,4 +257,94 @@ void FXPBDSpringConstraints::Apply(FSolverParticles& Particles, const FSolverRea
 	}
 }
 
+void FXPBDEdgeSpringConstraints::SetProperties(
+	const FCollectionPropertyConstFacade& PropertyCollection,
+	const TMap<FString, TConstArrayView<FRealSingle>>& WeightMaps)
+{
+	if (IsXPBDEdgeSpringStiffnessMutable(PropertyCollection))
+	{
+		const FSolverVec2 WeightedValue(GetWeightedFloatXPBDEdgeSpringStiffness(PropertyCollection));
+		if (IsXPBDEdgeSpringStiffnessStringDirty(PropertyCollection))
+		{
+			const FString& WeightMapName = GetXPBDEdgeSpringStiffnessString(PropertyCollection);
+			Stiffness = FPBDStiffness(
+				WeightedValue,
+				WeightMaps.FindRef(WeightMapName),
+				TConstArrayView<TVec2<int32>>(Constraints),
+				ParticleOffset,
+				ParticleCount,
+				FPBDStiffness::DefaultTableSize,
+				FPBDStiffness::DefaultParameterFitBase,
+				MaxStiffness);
+		}
+		else
+		{
+			Stiffness.SetWeightedValue(WeightedValue, MaxStiffness);
+		}
+	}
+	if (IsXPBDEdgeSpringDampingMutable(PropertyCollection))
+	{
+		const FSolverVec2 WeightedValue = FSolverVec2(GetWeightedFloatXPBDEdgeSpringDamping(PropertyCollection)).ClampAxes(MinDampingRatio, MaxDampingRatio);
+		if (IsXPBDEdgeSpringDampingStringDirty(PropertyCollection))
+		{
+			const FString& WeightMapName = GetXPBDEdgeSpringDampingString(PropertyCollection);
+			DampingRatio = FPBDWeightMap(
+				WeightedValue,
+				WeightMaps.FindRef(WeightMapName),
+				TConstArrayView<TVec2<int32>>(Constraints),
+				ParticleOffset,
+				ParticleCount);
+		}
+		else
+		{
+			DampingRatio.SetWeightedValue(WeightedValue);
+		}
+	}
+}
+
+void FXPBDBendingSpringConstraints::SetProperties(
+	const FCollectionPropertyConstFacade& PropertyCollection,
+	const TMap<FString, TConstArrayView<FRealSingle>>& WeightMaps)
+{
+	if (IsXPBDBendingSpringStiffnessMutable(PropertyCollection))
+	{
+		const FSolverVec2 WeightedValue(GetWeightedFloatXPBDBendingSpringStiffness(PropertyCollection));
+		if (IsXPBDBendingSpringStiffnessStringDirty(PropertyCollection))
+		{
+			const FString& WeightMapName = GetXPBDBendingSpringStiffnessString(PropertyCollection);
+			Stiffness = FPBDStiffness(
+				WeightedValue,
+				WeightMaps.FindRef(WeightMapName),
+				TConstArrayView<TVec2<int32>>(Constraints),
+				ParticleOffset,
+				ParticleCount,
+				FPBDStiffness::DefaultTableSize,
+				FPBDStiffness::DefaultParameterFitBase,
+				MaxStiffness);
+		}
+		else
+		{
+			Stiffness.SetWeightedValue(WeightedValue, MaxStiffness);
+		}
+	}
+	if (IsXPBDBendingSpringDampingMutable(PropertyCollection))
+	{
+		const FSolverVec2 WeightedValue = FSolverVec2(GetWeightedFloatXPBDBendingSpringDamping(PropertyCollection)).ClampAxes(MinDampingRatio, MaxDampingRatio);
+		if (IsXPBDBendingSpringDampingStringDirty(PropertyCollection))
+		{
+			const FString& WeightMapName = GetXPBDBendingSpringDampingString(PropertyCollection);
+			DampingRatio = FPBDWeightMap(
+				WeightedValue,
+				WeightMaps.FindRef(WeightMapName),
+				TConstArrayView<TVec2<int32>>(Constraints),
+				ParticleOffset,
+				ParticleCount);
+		}
+		else
+		{
+			DampingRatio.SetWeightedValue(WeightedValue);
+		}
+	}
+}
+
 } // End namespace Chaos::Softs
diff --git a/Engine/Source/Runtime/Experimental/Chaos/Private/Chaos/XPBDStretchBiasElementConstraints.cpp b/Engine/Source/Runtime/Experimental/Chaos/Private/Chaos/XPBDStretchBiasElementConstraints.cpp
--- a/Engine/Source/Runtime/Experimental/Chaos/Private/Chaos/XPBDStretchBiasElementConstraints.cpp
+++ b/Engine/Source/Runtime/Experimental/Chaos/Private/Chaos/XPBDStretchBiasElementConstraints.cpp
@@ -28,21 +28,18 @@ namespace Chaos::Softs {
 static int32 Chaos_XPBDStretchBias_ParallelConstraintCount = 100;
 
 FXPBDStretchBiasElementConstraints::FXPBDStretchBiasElementConstraints(const FSolverParticles& InParticles,
-	int32 ParticleOffset,
-	int32 ParticleCount,
+	int32 InParticleOffset,
+	int32 InParticleCount,
 	const FTriangleMesh& TriangleMesh,
 	const TArray<TVec3<FVec2f>>& FaceVertexUVs,
-	const TConstArrayView<FRealSingle>& StiffnessWarpMultipliers,
-	const TConstArrayView<FRealSingle>& StiffnessWeftMultipliers,
-	const TConstArrayView<FRealSingle>& StiffnessBiasMultipliers,
-	const TConstArrayView<FRealSingle>& DampingMultipliers,
-	const TConstArrayView<FRealSingle>& WarpScaleMultipliers,
-	const TConstArrayView<FRealSingle>& WeftScaleMultipliers,
+	const TMap<FString, TConstArrayView<FRealSingle>>& WeightMaps,
 	const FCollectionPropertyConstFacade& PropertyCollection,
 	bool bTrimKinematicConstraints)
-	: StiffnessWarp(
+	: ParticleOffset(InParticleOffset)
+	, ParticleCount(InParticleCount)
+	, StiffnessWarp(
 		FSolverVec2(GetWeightedFloatXPBDStretchBiasElementStiffnessWarp(PropertyCollection, MaxStiffness)),
-		StiffnessWarpMultipliers,
+		WeightMaps.FindRef(GetXPBDStretchBiasElementStiffnessWarpString(PropertyCollection, XPBDStretchBiasElementStiffnessWarpName.ToString())),
 		TConstArrayView<TVec3<int32>>(Constraints),
 		ParticleOffset,
 		ParticleCount,
@@ -51,7 +48,7 @@ FXPBDStretchBiasElementConstraints::FXPBDStretchBiasElementConstraints(const FSo
 		MaxStiffness)
 	, StiffnessWeft(
 		FSolverVec2(GetWeightedFloatXPBDStretchBiasElementStiffnessWeft(PropertyCollection, MaxStiffness)),
-		StiffnessWeftMultipliers,
+		WeightMaps.FindRef(GetXPBDStretchBiasElementStiffnessWeftString(PropertyCollection, XPBDStretchBiasElementStiffnessWeftName.ToString())),
 		TConstArrayView<TVec3<int32>>(Constraints),
 		ParticleOffset,
 		ParticleCount,
@@ -60,7 +57,7 @@ FXPBDStretchBiasElementConstraints::FXPBDStretchBiasElementConstraints(const FSo
 		MaxStiffness)
 	, StiffnessBias(
 		FSolverVec2(GetWeightedFloatXPBDStretchBiasElementStiffnessBias(PropertyCollection, MaxStiffness)),
-		StiffnessBiasMultipliers,
+		WeightMaps.FindRef(GetXPBDStretchBiasElementStiffnessBiasString(PropertyCollection, XPBDStretchBiasElementStiffnessBiasName.ToString())),
 		TConstArrayView<TVec3<int32>>(Constraints),
 		ParticleOffset,
 		ParticleCount,
@@ -69,31 +66,31 @@ FXPBDStretchBiasElementConstraints::FXPBDStretchBiasElementConstraints(const FSo
 		MaxStiffness)
 	, DampingRatio(
 		FSolverVec2(GetWeightedFloatXPBDStretchBiasElementDamping(PropertyCollection, MinDamping)).ClampAxes(MinDamping, MaxDamping),
-		DampingMultipliers,
+		WeightMaps.FindRef(GetXPBDStretchBiasElementDampingString(PropertyCollection, XPBDStretchBiasElementDampingName.ToString())),
 		TConstArrayView<TVec3<int32>>(Constraints),
 		ParticleOffset,
 		ParticleCount)
 	, WarpScale(
 		FSolverVec2(GetWeightedFloatXPBDStretchBiasElementWarpScale(PropertyCollection, DefaultWarpWeftScale)).ClampAxes(MinWarpWeftScale, MaxWarpWeftScale),
-		WarpScaleMultipliers,
+		WeightMaps.FindRef(GetXPBDStretchBiasElementWarpScaleString(PropertyCollection, XPBDStretchBiasElementWarpScaleName.ToString())),
 		TConstArrayView<TVec3<int32>>(Constraints),
 		ParticleOffset,
 		ParticleCount)
 	, WeftScale(
 		FSolverVec2(GetWeightedFloatXPBDStretchBiasElementWeftScale(PropertyCollection, DefaultWarpWeftScale)).ClampAxes(MinWarpWeftScale, MaxWarpWeftScale),
-		WeftScaleMultipliers,
+		WeightMaps.FindRef(GetXPBDStretchBiasElementWeftScaleString(PropertyCollection, XPBDStretchBiasElementWeftScaleName.ToString())),
 		TConstArrayView<TVec3<int32>>(Constraints),
 		ParticleOffset,
 		ParticleCount)
 {
 	Lambdas.Init(FSolverVec3(0.), Constraints.Num());
 	InitConstraintsAndRestData(InParticles, TriangleMesh, FaceVertexUVs, GetXPBDStretchBiasElementUse3dRestLengths(PropertyCollection, bDefaultUse3dRestLengths), bTrimKinematicConstraints);
-	InitColor(InParticles, ParticleOffset, ParticleCount);
+	InitColor(InParticles);
 }
 
 FXPBDStretchBiasElementConstraints::FXPBDStretchBiasElementConstraints(const FSolverParticles& InParticles,
-	int32 ParticleOffset,
-	int32 ParticleCount,
+	int32 InParticleOffset,
+	int32 InParticleCount,
 	const FTriangleMesh& TriangleMesh,
 	const TArray<TVec3<FVec2f>>& FaceVertexUVs,
 	const TConstArrayView<FRealSingle>& StiffnessWarpMultipliers,
@@ -110,7 +107,9 @@ FXPBDStretchBiasElementConstraints::FXPBDStretchBiasElementConstraints(const FSo
 	const FSolverVec2& InWeftScale,
 	bool bUse3dRestLengths,
 	bool bTrimKinematicConstraints)
-	: StiffnessWarp(
+	: ParticleOffset(InParticleOffset)
+	, ParticleCount(InParticleCount)
+	, StiffnessWarp(
 		InStiffnessWarp,
 		StiffnessWarpMultipliers,
 		TConstArrayView<TVec3<int32>>(Constraints),
@@ -158,7 +157,7 @@ FXPBDStretchBiasElementConstraints::FXPBDStretchBiasElementConstraints(const FSo
 {
 	Lambdas.Init(FSolverVec3(0.), Constraints.Num());
 	InitConstraintsAndRestData(InParticles, TriangleMesh, FaceVertexUVs, bUse3dRestLengths, bTrimKinematicConstraints);
-	InitColor(InParticles, ParticleOffset, ParticleCount);
+	InitColor(InParticles);
 }
 
 void FXPBDStretchBiasElementConstraints::InitConstraintsAndRestData(const FSolverParticles& InParticles, const FTriangleMesh& TriangleMesh,
@@ -236,7 +235,7 @@ void FXPBDStretchBiasElementConstraints::InitConstraintsAndRestData(const FSolve
 	}
 }
 
-void FXPBDStretchBiasElementConstraints::InitColor(const FSolverParticles& InParticles, const int32 ParticleOffset, const int32 ParticleCount)
+void FXPBDStretchBiasElementConstraints::InitColor(const FSolverParticles& InParticles)
 {
 	// In dev builds we always color so we can tune the system without restarting. See Apply()
 #if UE_BUILD_SHIPPING || UE_BUILD_TEST
@@ -288,6 +287,129 @@ void FXPBDStretchBiasElementConstraints::InitColor(const FSolverParticles& InPar
 	}
 }
 
+void FXPBDStretchBiasElementConstraints::SetProperties(
+	const FCollectionPropertyConstFacade& PropertyCollection,
+	const TMap<FString, TConstArrayView<FRealSingle>>& WeightMaps)
+{
+	if (IsXPBDStretchBiasElementStiffnessWarpMutable(PropertyCollection))
+	{
+		const FSolverVec2 WeightedValue(GetWeightedFloatXPBDStretchBiasElementStiffnessWarp(PropertyCollection));
+		if (IsXPBDStretchBiasElementStiffnessWarpStringDirty(PropertyCollection))
+		{
+			const FString& WeightMapName = GetXPBDStretchBiasElementStiffnessWarpString(PropertyCollection);
+			StiffnessWarp = FPBDStiffness(
+				WeightedValue,
+				WeightMaps.FindRef(WeightMapName),
+				TConstArrayView<TVec3<int32>>(Constraints),
+				ParticleOffset,
+				ParticleCount,
+				FPBDStiffness::DefaultTableSize,
+				FPBDStiffness::DefaultParameterFitBase,
+				MaxStiffness);
+		}
+		else
+		{
+			StiffnessWarp.SetWeightedValue(WeightedValue, MaxStiffness);
+		}
+	}
+	if (IsXPBDStretchBiasElementStiffnessWeftMutable(PropertyCollection))
+	{
+		const FSolverVec2 WeightedValue(GetWeightedFloatXPBDStretchBiasElementStiffnessWeft(PropertyCollection));
+		if (IsXPBDStretchBiasElementStiffnessWeftStringDirty(PropertyCollection))
+		{
+			const FString& WeightMapName = GetXPBDStretchBiasElementStiffnessWeftString(PropertyCollection);
+			StiffnessWeft = FPBDStiffness(
+				WeightedValue,
+				WeightMaps.FindRef(WeightMapName),
+				TConstArrayView<TVec3<int32>>(Constraints),
+				ParticleOffset,
+				ParticleCount,
+				FPBDStiffness::DefaultTableSize,
+				FPBDStiffness::DefaultParameterFitBase,
+				MaxStiffness);
+		}
+		else
+		{
+			StiffnessWeft.SetWeightedValue(WeightedValue, MaxStiffness);
+		}
+	}
+	if (IsXPBDStretchBiasElementStiffnessBiasMutable(PropertyCollection))
+	{
+		const FSolverVec2 WeightedValue(GetWeightedFloatXPBDStretchBiasElementStiffnessBias(PropertyCollection));
+		if (IsXPBDStretchBiasElementStiffnessBiasStringDirty(PropertyCollection))
+		{
+			const FString& WeightMapName = GetXPBDStretchBiasElementStiffnessBiasString(PropertyCollection);
+			StiffnessBias = FPBDStiffness(
+				WeightedValue,
+				WeightMaps.FindRef(WeightMapName),
+				TConstArrayView<TVec3<int32>>(Constraints),
+				ParticleOffset,
+				ParticleCount,
+				FPBDStiffness::DefaultTableSize,
+				FPBDStiffness::DefaultParameterFitBase,
+				MaxStiffness);
+		}
+		else
+		{
+			StiffnessBias.SetWeightedValue(WeightedValue, MaxStiffness);
+		}
+	}
+	if (IsXPBDStretchBiasElementDampingMutable(PropertyCollection))
+	{
+		const FSolverVec2 WeightedValue = FSolverVec2(GetWeightedFloatXPBDStretchBiasElementDamping(PropertyCollection)).ClampAxes(MinDamping, MaxDamping);
+		if (IsXPBDStretchBiasElementDampingStringDirty(PropertyCollection))
+		{
+			const FString& WeightMapName = GetXPBDStretchBiasElementDampingString(PropertyCollection);
+			DampingRatio = FPBDWeightMap(
+				WeightedValue,
+				WeightMaps.FindRef(WeightMapName),
+				TConstArrayView<TVec3<int32>>(Constraints),
+				ParticleOffset,
+				ParticleCount);
+		}
+		else
+		{
+			DampingRatio.SetWeightedValue(WeightedValue);
+		}
+	}
+	if (IsXPBDStretchBiasElementWarpScaleMutable(PropertyCollection))
+	{
+		const FSolverVec2 WeightedValue = FSolverVec2(GetWeightedFloatXPBDStretchBiasElementWarpScale(PropertyCollection)).ClampAxes(MinWarpWeftScale, MaxWarpWeftScale);
+		if (IsXPBDStretchBiasElementWarpScaleStringDirty(PropertyCollection))
+		{
+			const FString& WeightMapName = GetXPBDStretchBiasElementWarpScaleString(PropertyCollection);
+			WarpScale = FPBDWeightMap(
+				WeightedValue,
+				WeightMaps.FindRef(WeightMapName),
+				TConstArrayView<TVec3<int32>>(Constraints),
+				ParticleOffset,
+				ParticleCount);
+		}
+		else
+		{
+			WarpScale.SetWeightedValue(WeightedValue);
+		}
+	}
+	if (IsXPBDStretchBiasElementWeftScaleMutable(PropertyCollection))
+	{
+		const FSolverVec2 WeightedValue = FSolverVec2(GetWeightedFloatXPBDStretchBiasElementWeftScale(PropertyCollection)).ClampAxes(MinWarpWeftScale, MaxWarpWeftScale);
+		if (IsXPBDStretchBiasElementWeftScaleStringDirty(PropertyCollection))
+		{
+			const FString& WeightMapName = GetXPBDStretchBiasElementWeftScaleString(PropertyCollection);
+			WeftScale = FPBDWeightMap(
+				WeightedValue,
+				WeightMaps.FindRef(WeightMapName),
+				TConstArrayView<TVec3<int32>>(Constraints),
+				ParticleOffset,
+				ParticleCount);
+		}
+		else
+		{
+			WeftScale.SetWeightedValue(WeightedValue);
+		}
+	}
+}
+
 void FXPBDStretchBiasElementConstraints::Apply(FSolverParticles& Particles, const FSolverReal Dt) const
 {
 	TRACE_CPUPROFILER_EVENT_SCOPE(FXPBDStretchBiasElementConstraints_Apply);
diff --git a/Engine/Source/Runtime/Experimental/Chaos/Public/Chaos/CollectionPropertyFacade.h b/Engine/Source/Runtime/Experimental/Chaos/Public/Chaos/CollectionPropertyFacade.h
--- a/Engine/Source/Runtime/Experimental/Chaos/Public/Chaos/CollectionPropertyFacade.h
+++ b/Engine/Source/Runtime/Experimental/Chaos/Public/Chaos/CollectionPropertyFacade.h
@@ -10,16 +10,28 @@ struct FManagedArrayCollection;
 
 namespace Chaos::Softs
 {
-	/** Property flags, whether properties are enabled, animatable, ...etc. */
-	enum class ECollectionPropertyFlag : uint8
+	enum class UE_DEPRECATED(5.3, "Use ECollectionPropertyFlags instead.") ECollectionPropertyFlag : uint8
 	{
 		None,
+		Enabled = 1 << 0,
+		Animatable = 1 << 1,
+		Dirty = 1 << 7
+	};
+	PRAGMA_DISABLE_DEPRECATION_WARNINGS
+	ENUM_CLASS_FLAGS(ECollectionPropertyFlag)
+	PRAGMA_ENABLE_DEPRECATION_WARNINGS
+
+	/** Property flags, whether properties are enabled, animatable, ...etc. */
+	enum class ECollectionPropertyFlags : uint8
+	{
+		None = 0,
 		Enabled = 1 << 0,  /** Whether this property is enabled(so that it doesn't have to be removed from the collection when not needed). */
 		Animatable = 1 << 1,  /** Whether this property needs to be set at every frame. */
 		//~ Add new flags above this line
-		Dirty = 1 << 7  /** Whether this property has changed and needs to be updated at the next frame. */
+		StringDirty = 1 << 6,  /** Whether this property's string has changed and needs to be updated at the next frame. */
+		Dirty = 1 << 7  /** Whether this property's value has changed and needs to be updated at the next frame. */
 	};
-	ENUM_CLASS_FLAGS(ECollectionPropertyFlag)
+	ENUM_CLASS_FLAGS(ECollectionPropertyFlags)
 
 	/** Weighted types are all property types that can have a pair of low and high values to be associated with a weight map. */
 	template<typename T> struct TIsWeightedType { static constexpr bool Value = false; };
@@ -70,34 +82,80 @@ namespace Chaos::Softs
 		template<typename T, TEMPLATE_REQUIRES(TIsWeightedType<T>::Value)>
 		T GetValue(int32 KeyIndex) const { return GetLowValue<T>(KeyIndex); }
 
-		const FString& GetStringValue(int32 KeyIndex, const FString& Default = "") const { return GetValue<const FString&>(KeyIndex, StringValueArray); }
+		const FString& GetStringValue(int32 KeyIndex) const { return GetValue<const FString&>(KeyIndex, StringValueArray); }
+		UE_DEPRECATED(5.3, "Use GetStringValue(int32) or GetStringValue(const FString&, const FString&, int32*) instead.")
+		const FString& GetStringValue(int32 KeyIndex, const FString& Default) const { return GetValue<const FString&>(KeyIndex, StringValueArray); }
 
+		UE_DEPRECATED(5.3, "GetFlags is being phased out to promote correct dirtying operations.")
 		uint8 GetFlags(int32 KeyIndex) const { return FlagsArray[KeyIndex]; }
-		bool IsEnabled(int32 KeyIndex) const { return EnumHasAnyFlags((ECollectionPropertyFlag)FlagsArray[KeyIndex], ECollectionPropertyFlag::Enabled); }
-		bool IsAnimatable(int32 KeyIndex) const { return EnumHasAnyFlags((ECollectionPropertyFlag)FlagsArray[KeyIndex], ECollectionPropertyFlag::Animatable); }
-		bool IsDirty(int32 KeyIndex) const { return EnumHasAnyFlags((ECollectionPropertyFlag)FlagsArray[KeyIndex], ECollectionPropertyFlag::Dirty); }
+
+		bool IsEnabled(int32 KeyIndex) const { return HasAnyFlags(KeyIndex, ECollectionPropertyFlags::Enabled); }
+		bool IsAnimatable(int32 KeyIndex) const { return HasAnyFlags(KeyIndex, ECollectionPropertyFlags::Animatable); }
+		bool IsStringDirty(int32 KeyIndex) const { return HasAnyFlags(KeyIndex, ECollectionPropertyFlags::StringDirty); }
+		bool IsDirty(int32 KeyIndex) const { return HasAnyFlags(KeyIndex, ECollectionPropertyFlags::Dirty); }
 
 		//~ Values access per key
 		template<typename T, TEMPLATE_REQUIRES(TIsWeightedType<T>::Value)>
-		T GetLowValue(const FString& Key, const T& Default = T(0), int32* OutKeyIndex = nullptr) const { return GetValue(Key, LowValueArray, Default, OutKeyIndex); }
+		T GetLowValue(const FString& Key, const T& Default = T(0), int32* OutKeyIndex = nullptr) const
+		{
+			return SafeGet(Key, [this](int32 KeyIndex)->T { return GetLowValue<T>(KeyIndex); }, Default, OutKeyIndex);
+		}
 
 		template<typename T, TEMPLATE_REQUIRES(TIsWeightedType<T>::Value)>
-		T GetHighValue(const FString& Key, const T& Default = T(0), int32* OutKeyIndex = nullptr) const { return GetValue(Key, HighValueArray, Default, OutKeyIndex); }
+		T GetHighValue(const FString& Key, const T& Default = T(0), int32* OutKeyIndex = nullptr) const
+		{
+			return SafeGet(Key, [this](int32 KeyIndex)->T { return GetHighValue<T>(KeyIndex); }, Default, OutKeyIndex);
+		}
 
 		template<typename T, TEMPLATE_REQUIRES(TIsWeightedType<T>::Value)>
-		inline TPair<T, T> GetWeightedValue(const FString& Key, const T& Default = T(0), int32* OutKeyIndex = nullptr) const;
+		TPair<T, T> GetWeightedValue(const FString& Key, const T& Default = T(0), int32* OutKeyIndex = nullptr) const
+		{
+			return SafeGet(Key, [this](int32 KeyIndex)->TPair<T, T> { return GetWeightedValue<T>(KeyIndex); }, MakeTuple(Default, Default), OutKeyIndex);
+		}
 
-		FVector2f GetWeightedFloatValue(const FString& Key, const float& Default = 0.f, int32* OutKeyIndex = nullptr) const;
+		FVector2f GetWeightedFloatValue(const FString& Key, const float& Default = 0.f, int32* OutKeyIndex = nullptr) const
+		{
+			return SafeGet(Key, [this](int32 KeyIndex)->FVector2f { return GetWeightedFloatValue(KeyIndex); }, FVector2f(Default), OutKeyIndex);
+		}
 
 		template<typename T, TEMPLATE_REQUIRES(TIsWeightedType<T>::Value)>
-		T GetValue(const FString& Key, const T& Default = T(0), int32* OutKeyIndex = nullptr) const { return GetValue(Key, LowValueArray, Default, OutKeyIndex); }
+		T GetValue(const FString& Key, const T& Default = T(0), int32* OutKeyIndex = nullptr) const
+		{
+			return SafeGet(Key, [this](int32 KeyIndex)->T { return GetValue<T>(KeyIndex); }, Default, OutKeyIndex);
+		}
+
+		FString GetStringValue(const FString& Key, const FString& Default = "", int32* OutKeyIndex = nullptr) const
+		{
+			return SafeGet(Key, [this](int32 KeyIndex)->FString { return GetStringValue(KeyIndex); }, Default, OutKeyIndex);
+		}
+
+		UE_DEPRECATED(5.3, "GetFlags is being phased out to promote correct dirtying operations.")
+		uint8 GetFlags(const FString& Key, uint8 Default = 0, int32* OutKeyIndex = nullptr) const
+		{
+PRAGMA_DISABLE_DEPRECATION_WARNINGS
+			return SafeGet(Key, [this](int32 KeyIndex)->uint8 { return GetFlags(KeyIndex); }, Default, OutKeyIndex);
+PRAGMA_ENABLE_DEPRECATION_WARNINGS
+		}
+
+		bool IsEnabled(const FString& Key, bool bDefault = false, int32* OutKeyIndex = nullptr) const
+		{
+			return SafeGet(Key, [this](int32 KeyIndex)->bool { return IsEnabled(KeyIndex); }, bDefault, OutKeyIndex);
+		}
 
-		FString GetStringValue(const FString& Key, const FString& Default = "", int32* OutKeyIndex = nullptr) const { return GetValue(Key, StringValueArray, Default, OutKeyIndex); }
+		bool IsAnimatable(const FString& Key, bool bDefault = false, int32* OutKeyIndex = nullptr) const
+		{
+			return SafeGet(Key, [this](int32 KeyIndex)->bool { return IsAnimatable(KeyIndex); }, bDefault, OutKeyIndex);
+		}
 
-		uint8 GetFlags(const FString& Key, uint8 Default = 0, int32* OutKeyIndex = nullptr) const { return GetValue(Key, FlagsArray, (uint8)Default, OutKeyIndex); }
-		bool IsEnabled(const FString& Key, bool bDefault = false, int32* OutKeyIndex = nullptr) const { return EnumHasAnyFlags((ECollectionPropertyFlag)GetValue(Key, FlagsArray, (uint8)bDefault, OutKeyIndex), ECollectionPropertyFlag::Enabled); }
-		bool IsAnimatable(const FString& Key, bool bDefault = false, int32* OutKeyIndex = nullptr) const { return EnumHasAnyFlags((ECollectionPropertyFlag)GetValue(Key, FlagsArray, (uint8)bDefault, OutKeyIndex), ECollectionPropertyFlag::Animatable); }
-		bool IsDirty(const FString& Key, bool bDefault = false, int32* OutKeyIndex = nullptr) const { return EnumHasAnyFlags((ECollectionPropertyFlag)GetValue(Key, FlagsArray, (uint8)bDefault, OutKeyIndex), ECollectionPropertyFlag::Dirty); }
+		bool IsStringDirty(const FString& Key, bool bDefault = false, int32* OutKeyIndex = nullptr) const
+		{
+			return SafeGet(Key, [this](int32 KeyIndex)->bool { return IsStringDirty(KeyIndex); }, bDefault, OutKeyIndex);
+		}
+
+		bool IsDirty(const FString& Key, bool bDefault = false, int32* OutKeyIndex = nullptr) const
+		{
+			return SafeGet(Key, [this](int32 KeyIndex)->bool { return IsDirty(KeyIndex); }, bDefault, OutKeyIndex);
+		}
 
 	protected:
 		// No init constructor for FCollectionPropertyFacade
@@ -109,15 +167,25 @@ namespace Chaos::Softs
 		// Rebuild the search map
 		void RebuildKeyIndices();
 
-		template<typename T, typename ElementType>
-		T GetValue(int32 KeyIndex, const TConstArrayView<ElementType>& ValueArray) const;
+		template<typename CallableType, typename ReturnType>
+		ReturnType SafeGet(const FString& Key, CallableType&& Callable, ReturnType Default, int32* OutKeyIndex) const
+		{
+			const int32 KeyIndex = GetKeyIndex(Key);
+			if (OutKeyIndex)
+			{
+				*OutKeyIndex = KeyIndex;
+			}
+			return KeyIndex != INDEX_NONE ? Callable(KeyIndex) : Default;
+		}
 
 		template<typename T, typename ElementType>
-		inline T GetValue(const FString& Key, const TConstArrayView<ElementType>& ValueArray, const T& Default, int32* OutKeyIndex) const;
+		T GetValue(int32 KeyIndex, const TConstArrayView<ElementType>& ValueArray) const;
 
 		template <typename T>
 		TConstArrayView<T> GetArray(const FName& Name) const;
 
+		bool HasAnyFlags(int32 KeyIndex, ECollectionPropertyFlags Flags) const { return EnumHasAnyFlags(GetValue<ECollectionPropertyFlags, uint8>(KeyIndex, FlagsArray), Flags); }
+
 		// Attribute groups, predefined data member of the collection
 		static const FName PropertyGroup;
 		static const FName KeyName;  // Property key, name to look for
@@ -162,10 +230,10 @@ namespace Chaos::Softs
 
 		//~ Values set per index
 		template<typename T, TEMPLATE_REQUIRES(TIsWeightedType<T>::Value)>
-		void SetLowValue(int32 KeyIndex, const T& Value) { GetLowValueArray()[KeyIndex] = FVector3f(Value); }
+		void SetLowValue(int32 KeyIndex, const T& Value) { SetValue(KeyIndex, GetLowValueArray(), FVector3f(Value)); }
 
 		template<typename T, TEMPLATE_REQUIRES(TIsWeightedType<T>::Value)>
-		void SetHighValue(int32 KeyIndex, const T& Value) { GetHighValueArray()[KeyIndex] = FVector3f(Value); }
+		void SetHighValue(int32 KeyIndex, const T& Value) { SetValue(KeyIndex, GetHighValueArray(), FVector3f(Value)); }
 
 		template<typename T, TEMPLATE_REQUIRES(TIsWeightedType<T>::Value)>
 		void SetWeightedValue(int32 KeyIndex, const T& LowValue, const T& HighValue) { SetLowValue(KeyIndex, LowValue); SetHighValue(KeyIndex, HighValue); }
@@ -175,34 +243,90 @@ namespace Chaos::Softs
 		template<typename T, TEMPLATE_REQUIRES(TIsWeightedType<T>::Value)>
 		void SetValue(int32 KeyIndex, const T& Value) { SetWeightedValue(KeyIndex, Value, Value); }
 
-		void SetStringValue(int32 KeyIndex, const FString& Value) { GetStringValueArray()[KeyIndex] = Value; }
+		void SetStringValue(int32 KeyIndex, const FString& Value) { GetStringValueArray()[KeyIndex] = Value; SetStringDirty(KeyIndex); }
+
+		UE_DEPRECATED(5.3, "SetFlags is being phased out to promote correct dirtying operations.")
+		void SetFlags(int32 KeyIndex, uint8 Flags) { SetValue(KeyIndex, GetFlagsArray(), (ECollectionPropertyFlags)Flags); }
 
-		void SetFlags(int32 KeyIndex, uint8 Flags) { GetFlagsArray()[KeyIndex] = Flags; }
-		void SetEnabled(int32 KeyIndex, bool bEnabled) { EnableFlag(KeyIndex, ECollectionPropertyFlag::Enabled, bEnabled); }
-		void SetAnimatable(int32 KeyIndex, bool bAnimatable) { EnableFlag(KeyIndex, ECollectionPropertyFlag::Animatable, bAnimatable); }
-		void SetDirty(int32 KeyIndex, bool bDirty) { EnableFlag(KeyIndex, ECollectionPropertyFlag::Animatable, bDirty); }
+		void SetEnabled(int32 KeyIndex, bool bEnabled) { EnableFlags(KeyIndex, ECollectionPropertyFlags::Enabled, bEnabled); }
+		void SetAnimatable(int32 KeyIndex, bool bAnimatable) { EnableFlags(KeyIndex, ECollectionPropertyFlags::Animatable, bAnimatable); }
+		void SetDirty(int32 KeyIndex) { EnableFlags(KeyIndex, ECollectionPropertyFlags::Dirty, true); }
+		UE_DEPRECATED(5.3, "SetDirty can only be set, to unset use ClearDirtyFlags instead.")
+		void SetDirty(int32 KeyIndex, bool bDirty) { EnableFlags(KeyIndex, ECollectionPropertyFlags::Dirty, bDirty); }
+		void SetStringDirty(int32 KeyIndex) { EnableFlags(KeyIndex, ECollectionPropertyFlags::StringDirty | ECollectionPropertyFlags::Dirty, true); }
 
 		//~ Values set per key
 		template<typename T, TEMPLATE_REQUIRES(TIsWeightedType<T>::Value)>
-		int32 SetLowValue(const FString& Key, const T& Value) { return SetValue(Key, GetLowValueArray(), FVector3f(Value)); }
+		int32 SetLowValue(const FString& Key, const T& Value)
+		{
+			return SafeSet(Key, [this, &Value](int32 KeyIndex) { SetLowValue(KeyIndex, Value); });
+		}
 
 		template<typename T, TEMPLATE_REQUIRES(TIsWeightedType<T>::Value)>
-		int32 SetHighValue(const FString& Key, const T& Value) { return SetValue(Key, GetHighValueArray(), FVector3f(Value)); }
+		int32 SetHighValue(const FString& Key, const T& Value)
+		{
+			return SafeSet(Key, [this, &Value](int32 KeyIndex) { SetHighValue(KeyIndex, Value); });
+		}
 
 		template<typename T, TEMPLATE_REQUIRES(TIsWeightedType<T>::Value)>
-		inline int32 SetWeightedValue(const FString& Key, const T& LowValue, const T& HighValue);
+		int32 SetWeightedValue(const FString& Key, const T& LowValue, const T& HighValue)
+		{
+			return SafeSet(Key, [this, &LowValue, &HighValue](int32 KeyIndex) { SetWeightedValue(KeyIndex, LowValue, HighValue); });
+		}
 
-		int32 SetWeightedFloatValue(const FString& Key, const FVector2f& Value);
+		int32 SetWeightedFloatValue(const FString& Key, const FVector2f& Value)
+		{
+			return SafeSet(Key, [this, &Value](int32 KeyIndex) { SetWeightedFloatValue(KeyIndex, Value); });
+		}
 
 		template<typename T, TEMPLATE_REQUIRES(TIsWeightedType<T>::Value)>
-		int32 SetValue(const FString& Key, const T& Value) { return SetWeightedValue(Key, Value, Value); }
+		int32 SetValue(const FString& Key, const T& Value)
+		{
+			return SafeSet(Key, [this, &Value](int32 KeyIndex) { SetValue(KeyIndex, Value); });
+		}
 
-		int32 SetStringValue(const FString& Key, const FString& Value) { return SetValue(Key, GetStringValueArray(), Value); }
+		int32 SetStringValue(const FString& Key, const FString& Value)
+		{
+			return SafeSet(Key, [this, &Value](int32 KeyIndex) { SetStringValue(KeyIndex, Value); });
+		}
+
+		UE_DEPRECATED(5.3, "SetFlags is being phased out to promote correct dirtying operations.")
+		int32 SetFlags(const FString& Key, uint8 Flags)
+		{
+PRAGMA_DISABLE_DEPRECATION_WARNINGS
+			return SafeSet(Key, [this, Flags](int32 KeyIndex) { SetFlags(KeyIndex, Flags); });
+PRAGMA_ENABLE_DEPRECATION_WARNINGS
+		}
+
+		int32 SetEnabled(const FString& Key, bool bEnabled)
+		{
+			return SafeSet(Key, [this, bEnabled](int32 KeyIndex) { SetEnabled(KeyIndex, bEnabled); });
+		}
+
+		int32 SetAnimatable(const FString& Key, bool bAnimatable)
+		{
+			return SafeSet(Key, [this, bAnimatable](int32 KeyIndex) { SetAnimatable(KeyIndex, bAnimatable); });
+		}
+
+		UE_DEPRECATED(5.3, "SetDirty can only be set, to unset use ClearDirtyFlags instead.")
+		int32 SetDirty(const FString& Key, bool bDirty)
+		{
+PRAGMA_DISABLE_DEPRECATION_WARNINGS
+			return SafeSet(Key, [this, bDirty](int32 KeyIndex) { SetDirty(KeyIndex, bDirty); });
+PRAGMA_ENABLE_DEPRECATION_WARNINGS
+		}
 
-		int32 SetFlags(const FString& Key, uint8 Flags) { return SetValue(Key, GetFlagsArray(), Flags); }
-		int32 SetEnabled(const FString& Key, bool bEnabled) { return EnableFlag(Key, ECollectionPropertyFlag::Enabled, bEnabled); }
-		int32 SetAnimatable(const FString& Key, bool bAnimatable) { return EnableFlag(Key, ECollectionPropertyFlag::Animatable, bAnimatable); }
-		int32 SetDirty(const FString& Key, bool bDirty) { return EnableFlag(Key, ECollectionPropertyFlag::Animatable, bDirty); }
+		int32 SetDirty(const FString& Key)
+		{
+			return SafeSet(Key, [this](int32 KeyIndex) { SetDirty(KeyIndex); });
+		}
+
+		int32 SetStringDirty(const FString& Key)
+		{
+			return SafeSet(Key, [this](int32 KeyIndex) { SetStringDirty(KeyIndex); });
+		}
+
+		void ClearDirtyFlags();
 
 	protected:
 		// No init constructor for FCollectionPropertyMutableFacade
@@ -215,14 +339,24 @@ namespace Chaos::Softs
 		const TArrayView<FVector3f>& GetLowValueArray() { return reinterpret_cast<TArrayView<FVector3f>&>(LowValueArray); }
 		const TArrayView<FVector3f>& GetHighValueArray() { return reinterpret_cast<TArrayView<FVector3f>&>(HighValueArray); }
 		const TArrayView<FString>& GetStringValueArray() { return reinterpret_cast<TArrayView<FString>&>(StringValueArray); }
-		const TArrayView<uint8>& GetFlagsArray() { return reinterpret_cast<const TArrayView<uint8>&>(FlagsArray); }
+		const TArrayView<ECollectionPropertyFlags>& GetFlagsArray() { return reinterpret_cast<const TArrayView<ECollectionPropertyFlags>&>(FlagsArray); }
 
 	private:
+		template<typename CallableType>
+		int32 SafeSet(const FString& Key, CallableType&& Callable)
+		{
+			const int32 KeyIndex = GetKeyIndex(Key);
+			if (KeyIndex != INDEX_NONE)
+			{
+				Callable(KeyIndex);
+			}
+			return KeyIndex;
+		}
+
 		template<typename T>
-		inline int32 SetValue(const FString& Key, const TArrayView<T>& ValueArray, const T& Value);
+		inline void SetValue(int32 KeyIndex, const TArrayView<T>& ValueArray, const T& Value);
 
-		void EnableFlag(int32 KeyIndex, ECollectionPropertyFlag Flag, bool bEnable);
-		int32 EnableFlag(const FString& Key, ECollectionPropertyFlag Flag, bool bEnable);
+		void EnableFlags(int32 KeyIndex, ECollectionPropertyFlags Flags, bool bEnable);
 	};
 
 	/**
@@ -278,49 +412,11 @@ namespace Chaos::Softs
 		int32 AddStringValue(const FString& Key, const FString& Value, bool bEnabled = true, bool bAnimatable = false);
 	};
 
-	template<typename T, typename TEnableIf<TIsWeightedType<T>::Value, int>::type>
-	inline TPair<T, T> FCollectionPropertyConstFacade::GetWeightedValue(const FString& Key, const T& Default, int32* OutKeyIndex) const
-	{
-		const int32 KeyIndex = GetKeyIndex(Key);
-		if (OutKeyIndex)
-		{
-			*OutKeyIndex = KeyIndex;
-		}
-		return KeyIndex != INDEX_NONE ? GetWeightedValue<T>(KeyIndex) : MakeTuple(Default, Default);
-	}
-
-	template<typename T, typename ElementType>
-	inline T FCollectionPropertyConstFacade::GetValue(const FString& Key, const TConstArrayView<ElementType>& ValueArray, const T& Default, int32* OutKeyIndex) const
-	{
-		const int32 KeyIndex = GetKeyIndex(Key);
-		if (OutKeyIndex)
-		{
-			*OutKeyIndex = KeyIndex;
-		}
-		return KeyIndex != INDEX_NONE ? GetValue<T>(KeyIndex, ValueArray) : Default;
-	}
-
-	template<typename T, typename TEnableIf<TIsWeightedType<T>::Value, int>::type>
-	inline int32 FCollectionPropertyFacade::SetWeightedValue(const FString& Key, const T& LowValue, const T& HighValue)
-	{
-		const int32 KeyIndex = GetKeyIndex(Key);
-		if (KeyIndex != INDEX_NONE)
-		{
-			SetLowValue(KeyIndex, LowValue);
-			SetHighValue(KeyIndex, HighValue);
-		}
-		return KeyIndex;
-	}
-
 	template<typename T>
-	inline int32 FCollectionPropertyFacade::SetValue(const FString& Key, const TArrayView<T>& ValueArray, const T& Value)
+	inline void FCollectionPropertyFacade::SetValue(int32 KeyIndex, const TArrayView<T>& ValueArray, const T& Value)
 	{
-		const int32 KeyIndex = GetKeyIndex(Key);
-		if (KeyIndex != INDEX_NONE)
-		{
-			ValueArray[KeyIndex] = Value;
-		}
-		return KeyIndex;
+		ValueArray[KeyIndex] = Value;
+		SetDirty(KeyIndex);
 	}
 
 	template<typename T, typename TEnableIf<TIsWeightedType<T>::Value, int>::type>
@@ -335,96 +431,119 @@ namespace Chaos::Softs
 // Use this macro to add shorthands for property getters and direct access through the declared key index
 #define UE_CHAOS_DECLARE_PROPERTYCOLLECTION_NAME(PropertyName, Type) \
 	inline static const FName PropertyName##Name = TEXT(#PropertyName); \
+	UE_DEPRECATED(5.3, "PropertyName##String is to be removed as to not be confused with GetPropertyName##String().") \
 	static FString PropertyName##String() { return PropertyName##Name.ToString(); } \
 	static bool Is##PropertyName##Enabled(const FCollectionPropertyConstFacade& PropertyCollection, bool bDefault) \
 	{ \
-		return PropertyCollection.IsEnabled(PropertyName##String(), bDefault); \
+		return PropertyCollection.IsEnabled(PropertyName##Name.ToString(), bDefault); \
 	} \
 	static bool Is##PropertyName##Animatable(const FCollectionPropertyConstFacade& PropertyCollection, bool bDefault) \
 	{ \
-		return PropertyCollection.IsAnimatable(PropertyName##String(), bDefault); \
+		return PropertyCollection.IsAnimatable(PropertyName##Name.ToString(), bDefault); \
 	} \
 	Type GetLow##PropertyName(const FCollectionPropertyConstFacade& PropertyCollection, const Type& Default) \
 	{ \
-		return PropertyCollection.GetLowValue<Type>(PropertyName##String(), Default, &PropertyName##Index); \
+		return PropertyCollection.GetLowValue<Type>(PropertyName##Name.ToString(), Default, &PropertyName##Index); \
 	} \
 	Type GetHigh##PropertyName(const FCollectionPropertyConstFacade& PropertyCollection, const Type& Default) \
 	{ \
-		return PropertyCollection.GetHighValue<Type>(PropertyName##String(), Default, &PropertyName##Index); \
+		return PropertyCollection.GetHighValue<Type>(PropertyName##Name.ToString(), Default, &PropertyName##Index); \
 	} \
 	TPair<Type, Type> GetWeighted##PropertyName(const FCollectionPropertyConstFacade& PropertyCollection, const Type& Default) \
 	{ \
-		return PropertyCollection.GetWeightedValue<Type>(PropertyName##String(), Default, &PropertyName##Index); \
+		return PropertyCollection.GetWeightedValue<Type>(PropertyName##Name.ToString(), Default, &PropertyName##Index); \
 	} \
 	FVector2f GetWeightedFloat##PropertyName(const FCollectionPropertyConstFacade& PropertyCollection, const float& Default) \
 	{ \
-		return PropertyCollection.GetWeightedFloatValue(PropertyName##String(), Default, &PropertyName##Index); \
+		return PropertyCollection.GetWeightedFloatValue(PropertyName##Name.ToString(), Default, &PropertyName##Index); \
 	} \
 	Type Get##PropertyName(const FCollectionPropertyConstFacade& PropertyCollection, const Type& Default) \
 	{ \
-		return PropertyCollection.GetValue<Type>(PropertyName##String(), Default, &PropertyName##Index); \
+		return PropertyCollection.GetValue<Type>(PropertyName##Name.ToString(), Default, &PropertyName##Index); \
+	} \
+	FString Get##PropertyName##String(const FCollectionPropertyConstFacade& PropertyCollection, const FString& Default) \
+	{ \
+		return PropertyCollection.GetStringValue(PropertyName##Name.ToString(), Default, &PropertyName##Index); \
 	} \
+	UE_DEPRECATED(5.3, "GetFlags is being phased out to promote correct dirtying operations.") \
 	uint8 Get##PropertyName##Flags(const FCollectionPropertyConstFacade& PropertyCollection, uint8 Default) \
 	{ \
-		return PropertyCollection.GetFlags(PropertyName##String(), Default, &PropertyName##Index); \
+PRAGMA_DISABLE_DEPRECATION_WARNINGS \
+	return PropertyCollection.GetFlags(PropertyName##Name.ToString(), Default, &PropertyName##Index); \
+PRAGMA_ENABLE_DEPRECATION_WARNINGS \
 	} \
 	Type GetLow##PropertyName(const FCollectionPropertyConstFacade& PropertyCollection) const \
 	{ \
-		checkSlow(PropertyName##Index == PropertyCollection.GetKeyIndex(PropertyName##String())); \
+		checkSlow(PropertyName##Index == PropertyCollection.GetKeyIndex(PropertyName##Name.ToString())); \
 		checkf(PropertyName##Index != INDEX_NONE, TEXT("The default value getter that sets the property index must be called once prior to calling this function.")); \
 		return PropertyCollection.GetLowValue<Type>(PropertyName##Index); \
 	} \
 	Type GetHigh##PropertyName(const FCollectionPropertyConstFacade& PropertyCollection) const \
 	{ \
-		checkSlow(PropertyName##Index == PropertyCollection.GetKeyIndex(PropertyName##String())); \
+		checkSlow(PropertyName##Index == PropertyCollection.GetKeyIndex(PropertyName##Name.ToString())); \
 		checkf(PropertyName##Index != INDEX_NONE, TEXT("The default value getter that sets the property index must be called once prior to calling this function.")); \
 		return PropertyCollection.GetHighValue<Type>(PropertyName##Index); \
 	} \
 	TPair<Type, Type> GetWeighted##PropertyName(const FCollectionPropertyConstFacade& PropertyCollection) const \
 	{ \
-		checkSlow(PropertyName##Index == PropertyCollection.GetKeyIndex(PropertyName##String())); \
+		checkSlow(PropertyName##Index == PropertyCollection.GetKeyIndex(PropertyName##Name.ToString())); \
 		checkf(PropertyName##Index != INDEX_NONE, TEXT("The default value getter that sets the property index must be called once prior to calling this function.")); \
 		return PropertyCollection.GetWeightedValue<Type>(PropertyName##Index); \
 	} \
 	FVector2f GetWeightedFloat##PropertyName(const FCollectionPropertyConstFacade& PropertyCollection) const \
 	{ \
-		checkSlow(PropertyName##Index == PropertyCollection.GetKeyIndex(PropertyName##String())); \
+		checkSlow(PropertyName##Index == PropertyCollection.GetKeyIndex(PropertyName##Name.ToString())); \
 		checkf(PropertyName##Index != INDEX_NONE, TEXT("The default value getter that sets the property index must be called once prior to calling this function.")); \
 		return PropertyCollection.GetWeightedFloatValue(PropertyName##Index); \
 	} \
 	Type Get##PropertyName(const FCollectionPropertyConstFacade& PropertyCollection) const \
 	{ \
-		checkSlow(PropertyName##Index == PropertyCollection.GetKeyIndex(PropertyName##String())); \
+		checkSlow(PropertyName##Index == PropertyCollection.GetKeyIndex(PropertyName##Name.ToString())); \
 		checkf(PropertyName##Index != INDEX_NONE, TEXT("The default value getter that sets the property index must be called once prior to calling this function.")); \
 		return PropertyCollection.GetValue<Type>(PropertyName##Index); \
 	} \
+	const FString& Get##PropertyName##String(const FCollectionPropertyConstFacade& PropertyCollection) \
+	{ \
+		checkSlow(PropertyName##Index == PropertyCollection.GetKeyIndex(PropertyName##Name.ToString())); \
+		checkf(PropertyName##Index != INDEX_NONE, TEXT("The default value getter that sets the property index must be called once prior to calling this function.")); \
+		return PropertyCollection.GetStringValue(PropertyName##Index); \
+	} \
+	UE_DEPRECATED(5.3, "GetFlags is being phased out to promote correct dirtying operations.") \
 	uint8 Get##PropertyName##Flags(const FCollectionPropertyConstFacade& PropertyCollection) const \
 	{ \
-		checkSlow(PropertyName##Index == PropertyCollection.GetKeyIndex(PropertyName##String())); \
+		checkSlow(PropertyName##Index == PropertyCollection.GetKeyIndex(PropertyName##Name.ToString())); \
 		checkf(PropertyName##Index != INDEX_NONE, TEXT("The default value getter that sets the property index must be called once prior to calling this function.")); \
+PRAGMA_DISABLE_DEPRECATION_WARNINGS \
 		return PropertyCollection.GetFlags(PropertyName##Index); \
+PRAGMA_ENABLE_DEPRECATION_WARNINGS \
 	} \
 	bool Is##PropertyName##Enabled(const FCollectionPropertyConstFacade& PropertyCollection) const \
 	{ \
-		checkSlow(PropertyName##Index == PropertyCollection.GetKeyIndex(PropertyName##String())); \
+		checkSlow(PropertyName##Index == PropertyCollection.GetKeyIndex(PropertyName##Name.ToString())); \
 		checkf(PropertyName##Index != INDEX_NONE, TEXT("The default value getter that sets the property index must be called once prior to calling this function.")); \
 		return PropertyCollection.IsEnabled(PropertyName##Index); \
 	} \
 	bool Is##PropertyName##Animatable(const FCollectionPropertyConstFacade& PropertyCollection) const\
 	{ \
-		checkSlow(PropertyName##Index == PropertyCollection.GetKeyIndex(PropertyName##String())); \
+		checkSlow(PropertyName##Index == PropertyCollection.GetKeyIndex(PropertyName##Name.ToString())); \
 		checkf(PropertyName##Index != INDEX_NONE, TEXT("The default value getter that sets the property index must be called once prior to calling this function.")); \
 		return PropertyCollection.IsAnimatable(PropertyName##Index); \
 	} \
 	bool Is##PropertyName##Dirty(const FCollectionPropertyConstFacade& PropertyCollection) const \
 	{ \
-		checkSlow(PropertyName##Index == PropertyCollection.GetKeyIndex(PropertyName##String())); \
+		checkSlow(PropertyName##Index == PropertyCollection.GetKeyIndex(PropertyName##Name.ToString())); \
 		checkf(PropertyName##Index != INDEX_NONE, TEXT("The default value getter that sets the property index must be called once prior to calling this function.")); \
 		return PropertyCollection.IsDirty(PropertyName##Index); \
 	} \
+	bool Is##PropertyName##StringDirty(const FCollectionPropertyConstFacade& PropertyCollection) const \
+	{ \
+		checkSlow(PropertyName##Index == PropertyCollection.GetKeyIndex(PropertyName##Name.ToString())); \
+		checkf(PropertyName##Index != INDEX_NONE, TEXT("The default value getter that sets the property index must be called once prior to calling this function.")); \
+		return PropertyCollection.IsStringDirty(PropertyName##Index); \
+	} \
 	bool Is##PropertyName##Mutable(const FCollectionPropertyConstFacade& PropertyCollection) const \
 	{ \
 		return PropertyName##Index != INDEX_NONE && \
-			EnumHasAllFlags((ECollectionPropertyFlag)PropertyCollection.GetFlags(PropertyName##Index), ECollectionPropertyFlag::Enabled | ECollectionPropertyFlag::Animatable | ECollectionPropertyFlag::Dirty); \
+			PropertyCollection.IsEnabled(PropertyName##Index) && PropertyCollection.IsAnimatable(PropertyName##Index) && PropertyCollection.IsDirty(PropertyName##Index); \
 	} \
 	int32 PropertyName##Index = INDEX_NONE;
diff --git a/Engine/Source/Runtime/Experimental/Chaos/Public/Chaos/PBDAnimDriveConstraint.h b/Engine/Source/Runtime/Experimental/Chaos/Public/Chaos/PBDAnimDriveConstraint.h
--- a/Engine/Source/Runtime/Experimental/Chaos/Public/Chaos/PBDAnimDriveConstraint.h
+++ b/Engine/Source/Runtime/Experimental/Chaos/Public/Chaos/PBDAnimDriveConstraint.h
@@ -5,6 +5,7 @@
 
 #include "Chaos/Framework/Parallel.h"
 #include "Chaos/PBDSoftsEvolutionFwd.h"
+#include "Chaos/PBDSoftsSolverParticles.h"
 #include "Chaos/PBDStiffness.h"
 #include "Chaos/CollectionPropertyFacade.h"
 #include "ChaosStats.h"
@@ -42,6 +43,32 @@ namespace Chaos::Softs
 		{
 		}
 
+		FPBDAnimDriveConstraint(
+			const int32 InParticleOffset,
+			const int32 InParticleCount,
+			const TArray<FSolverVec3>& InAnimationPositions,  // Use global indexation (will need adding ParticleOffset)
+			const TArray<FSolverVec3>& InAnimationVelocities,  // Use global indexation (will need adding ParticleOffset)
+			const TMap<FString, TConstArrayView<FRealSingle>>& WeightMaps,
+			const FCollectionPropertyConstFacade& PropertyCollection
+		)
+			: AnimationPositions(InAnimationPositions)
+			, OldAnimationPositions_deprecated(InAnimationVelocities) // Unused when not using apply
+			, AnimationVelocities(InAnimationVelocities)
+			, ParticleOffset(InParticleOffset)
+			, ParticleCount(InParticleCount)
+			, UseDeprecatedApply(false)
+			, Stiffness(
+				FSolverVec2(GetWeightedFloatAnimDriveStiffness(PropertyCollection, 1.f)),
+				WeightMaps.FindRef(GetAnimDriveStiffnessString(PropertyCollection, AnimDriveStiffnessName.ToString())),
+				InParticleCount)
+			, Damping(
+				FSolverVec2(GetWeightedFloatAnimDriveDamping(PropertyCollection, 1.f)),
+				WeightMaps.FindRef(GetAnimDriveDampingString(PropertyCollection, AnimDriveDampingName.ToString())),
+				InParticleCount)
+		{
+		}
+
+		UE_DEPRECATED(5.3, "Use weight map constructor instead.")
 		FPBDAnimDriveConstraint(
 			const int32 InParticleOffset,
 			const int32 InParticleCount,
@@ -96,18 +123,44 @@ namespace Chaos::Softs
 		// Return the damping input values used by the constraint
 		FSolverVec2 GetDamping() const { return Damping.GetWeightedValue(); }
 
-		void SetProperties(const FCollectionPropertyConstFacade& PropertyCollection)
+		void SetProperties(
+			const FCollectionPropertyConstFacade& PropertyCollection,
+			const TMap<FString, TConstArrayView<FRealSingle>>& WeightMaps)
 		{
 			if (IsAnimDriveStiffnessMutable(PropertyCollection))
 			{
-				Stiffness.SetWeightedValue(FSolverVec2(GetWeightedFloatAnimDriveStiffness(PropertyCollection)));
+				const FSolverVec2 WeightedValue(GetWeightedFloatAnimDriveStiffness(PropertyCollection));
+				if (IsAnimDriveStiffnessStringDirty(PropertyCollection))
+				{
+					const FString& WeightMapName = GetAnimDriveStiffnessString(PropertyCollection);
+					Stiffness = FPBDStiffness(WeightedValue, WeightMaps.FindRef(WeightMapName), ParticleCount);
+				}
+				else
+				{
+					Stiffness.SetWeightedValue(WeightedValue);
+				}
 			}
 			if (IsAnimDriveDampingMutable(PropertyCollection))
 			{
-				Damping.SetWeightedValue(FSolverVec2(GetWeightedFloatAnimDriveDamping(PropertyCollection)));
+				const FSolverVec2 WeightedValue(GetWeightedFloatAnimDriveDamping(PropertyCollection));
+				if (IsAnimDriveDampingStringDirty(PropertyCollection))
+				{
+					const FString& WeightMapName = GetAnimDriveDampingString(PropertyCollection);
+					Damping = FPBDStiffness(WeightedValue, WeightMaps.FindRef(WeightMapName), ParticleCount);
+				}
+				else
+				{
+					Damping.SetWeightedValue(WeightedValue);
+				}
 			}
 		}
 
+		UE_DEPRECATED(5.3, "Use SetProperties(const FCollectionPropertyConstFacade&, const TMap<FString, TConstArrayView<FRealSingle>>&, FSolverReal) instead.")
+		void SetProperties(const FCollectionPropertyConstFacade& PropertyCollection)
+		{
+			SetProperties(PropertyCollection, TMap<FString, TConstArrayView<FRealSingle>>());
+		}
+
 		inline void SetProperties(const FSolverVec2& InStiffness, const FSolverVec2& InDamping)
 		{
 			Stiffness.SetWeightedValue(InStiffness);
diff --git a/Engine/Source/Runtime/Experimental/Chaos/Public/Chaos/PBDAxialSpringConstraints.h b/Engine/Source/Runtime/Experimental/Chaos/Public/Chaos/PBDAxialSpringConstraints.h
--- a/Engine/Source/Runtime/Experimental/Chaos/Public/Chaos/PBDAxialSpringConstraints.h
+++ b/Engine/Source/Runtime/Experimental/Chaos/Public/Chaos/PBDAxialSpringConstraints.h
@@ -13,20 +13,26 @@ class CHAOS_API FPBDAxialSpringConstraints : public FPBDAxialSpringConstraintsBa
 {
 	typedef FPBDAxialSpringConstraintsBase Base;
 	using Base::Barys;
-	using Base::Constraints;
 
 public:
 	FPBDAxialSpringConstraints(
 		const FSolverParticles& Particles,
-		int32 ParticleOffset,
-		int32 ParticleCount,
+		int32 InParticleOffset,
+		int32 InParticleCount,
 		const TArray<TVec3<int32>>& InConstraints,
 		const TConstArrayView<FRealSingle>& StiffnessMultipliers,
 		const FSolverVec2& InStiffness,
 		bool bTrimKinematicConstraints)
-		: Base(Particles, ParticleOffset, ParticleCount, InConstraints, StiffnessMultipliers, InStiffness, bTrimKinematicConstraints)
+		: Base(
+			Particles,
+			InParticleOffset,
+			InParticleCount,
+			InConstraints,
+			StiffnessMultipliers,
+			InStiffness,
+			bTrimKinematicConstraints)
 	{
-		InitColor(Particles, ParticleOffset, ParticleCount);
+		InitColor(Particles);
 	}
 
 	virtual ~FPBDAxialSpringConstraints() override {}
@@ -34,10 +40,13 @@ public:
 	void Apply(FSolverParticles& InParticles, const FSolverReal Dt) const;
 
 protected:
+	using Base::Constraints;
 	using Base::Stiffness;
+	using Base::ParticleOffset;
+	using Base::ParticleCount;
 
 private:
-	void InitColor(const FSolverParticles& InParticles, const int32 ParticleOffset, const int32 ParticleCount);
+	void InitColor(const FSolverParticles& InParticles);
 	void ApplyHelper(FSolverParticles& InParticles, const FSolverReal Dt, const int32 ConstraintIndex, const FSolverReal ExpStiffnessValue) const;
 
 	TArray<int32> ConstraintsPerColorStartIndex; // Constraints are ordered so each batch is contiguous. This is ColorNum + 1 length so it can be used as start and end.
@@ -53,16 +62,35 @@ public:
 
 	FPBDAreaSpringConstraints(
 		const FSolverParticles& Particles,
-		int32 ParticleOffset,
-		int32 ParticleCount,
+		int32 InParticleOffset,
+		int32 InParticleCount,
+		const TArray<TVec3<int32>>& InConstraints,
+		const TMap<FString, TConstArrayView<FRealSingle>>& WeightMaps,
+		const FCollectionPropertyConstFacade& PropertyCollection,
+		bool bTrimKinematicConstraints)
+		: FPBDAxialSpringConstraints(
+			Particles,
+			InParticleOffset,
+			InParticleCount,
+			InConstraints,
+			WeightMaps.FindRef(GetAreaSpringStiffnessString(PropertyCollection, AreaSpringStiffnessName.ToString())),
+			FSolverVec2(GetWeightedFloatAreaSpringStiffness(PropertyCollection, 1.f)),
+			bTrimKinematicConstraints)
+	{}
+
+	UE_DEPRECATED(5.3, "Use weight map constructor instead.")
+	FPBDAreaSpringConstraints(
+		const FSolverParticles& Particles,
+		int32 InParticleOffset,
+		int32 InParticleCount,
 		const TArray<TVec3<int32>>& InConstraints,
 		const TConstArrayView<FRealSingle>& StiffnessMultipliers,
 		const FCollectionPropertyConstFacade& PropertyCollection,
 		bool bTrimKinematicConstraints)
 		: FPBDAxialSpringConstraints(
 			Particles,
-			ParticleOffset,
-			ParticleCount,
+			InParticleOffset,
+			InParticleCount,
 			InConstraints,
 			StiffnessMultipliers,
 			FSolverVec2(GetWeightedFloatAreaSpringStiffness(PropertyCollection, 1.f)),
@@ -71,16 +99,21 @@ public:
 
 	virtual ~FPBDAreaSpringConstraints() override = default;
 
+	void SetProperties(
+		const FCollectionPropertyConstFacade& PropertyCollection,
+		const TMap<FString, TConstArrayView<FRealSingle>>& WeightMaps);
+
+	UE_DEPRECATED(5.3, "Use SetProperties(const FCollectionPropertyConstFacade&, const TMap<FString, TConstArrayView<FRealSingle>>&, FSolverReal) instead.")
 	void SetProperties(const FCollectionPropertyConstFacade& PropertyCollection)
 	{
-		if (IsAreaSpringStiffnessMutable(PropertyCollection))
-		{
-			Stiffness.SetWeightedValue(FSolverVec2(GetWeightedFloatAreaSpringStiffness(PropertyCollection)));
-		}
+		SetProperties(PropertyCollection, TMap<FString, TConstArrayView<FRealSingle>>());
 	}
 
 private:
+	using FPBDAxialSpringConstraints::Constraints;
 	using FPBDAxialSpringConstraints::Stiffness;
+	using FPBDAxialSpringConstraints::ParticleOffset;
+	using FPBDAxialSpringConstraints::ParticleCount;
 
 	UE_CHAOS_DECLARE_PROPERTYCOLLECTION_NAME(AreaSpringStiffness, float);
 };
diff --git a/Engine/Source/Runtime/Experimental/Chaos/Public/Chaos/PBDAxialSpringConstraintsBase.h b/Engine/Source/Runtime/Experimental/Chaos/Public/Chaos/PBDAxialSpringConstraintsBase.h
--- a/Engine/Source/Runtime/Experimental/Chaos/Public/Chaos/PBDAxialSpringConstraintsBase.h
+++ b/Engine/Source/Runtime/Experimental/Chaos/Public/Chaos/PBDAxialSpringConstraintsBase.h
@@ -16,8 +16,8 @@ class FPBDAxialSpringConstraintsBase
 public:
 	FPBDAxialSpringConstraintsBase(
 		const FSolverParticles& Particles,
-		int32 ParticleOffset,
-		int32 ParticleCount,
+		int32 InParticleOffset,
+		int32 InParticleCount,
 		const TArray<TVec3<int32>>& InConstraints,
 		const TConstArrayView<FRealSingle>& StiffnessMultipliers,
 		const FSolverVec2& InStiffness,
@@ -28,6 +28,8 @@ public:
 			{
 				return bTrimKinematicConstraints && Particles.InvM(Index0) == (FSolverReal)0. && Particles.InvM(Index1) == (FSolverReal)0. && Particles.InvM(Index2) == (FSolverReal)0.;
 			}))
+		, ParticleOffset(InParticleOffset)
+		, ParticleCount(InParticleCount)
 		, Stiffness(
 			InStiffness,
 			StiffnessMultipliers,
@@ -162,6 +164,8 @@ private:
 
 protected:
 	TArray<TVec3<int32>> Constraints;
+	const int32 ParticleOffset;
+	const int32 ParticleCount;
 	TArray<FSolverReal> Barys;
 	TArray<FSolverReal> Dists;
 	FPBDStiffness Stiffness;
diff --git a/Engine/Source/Runtime/Experimental/Chaos/Public/Chaos/PBDBendingConstraints.h b/Engine/Source/Runtime/Experimental/Chaos/Public/Chaos/PBDBendingConstraints.h
--- a/Engine/Source/Runtime/Experimental/Chaos/Public/Chaos/PBDBendingConstraints.h
+++ b/Engine/Source/Runtime/Experimental/Chaos/Public/Chaos/PBDBendingConstraints.h
@@ -18,8 +18,31 @@ public:
 	}
 
 	FPBDBendingConstraints(const FSolverParticles& InParticles,
-		int32 ParticleOffset,
-		int32 ParticleCount,
+		int32 InParticleOffset,
+		int32 InParticleCount,
+		TArray<TVec4<int32>>&& InConstraints,
+		const TMap<FString, TConstArrayView<FRealSingle>>& WeightMaps,
+		const FCollectionPropertyConstFacade& PropertyCollection,
+		bool bTrimKinematicConstraints = false)
+		: Base(
+			InParticles,
+			InParticleOffset,
+			InParticleCount,
+			MoveTemp(InConstraints),
+			WeightMaps.FindRef(GetBendingElementStiffnessString(PropertyCollection, BendingElementStiffnessName.ToString())),
+			WeightMaps.FindRef(GetBucklingStiffnessString(PropertyCollection, BucklingStiffnessName.ToString())),
+			FSolverVec2(GetWeightedFloatBendingElementStiffness(PropertyCollection, 1.f)),
+			(FSolverReal)GetBucklingRatio(PropertyCollection, 0.f),  // BucklingRatio is clamped in base class
+			FSolverVec2(GetWeightedFloatBucklingStiffness(PropertyCollection, 1.f)),
+			bTrimKinematicConstraints)
+	{
+		InitColor(InParticles);
+	}
+
+	UE_DEPRECATED(5.3, "Use weight map constructor instead.")
+	FPBDBendingConstraints(const FSolverParticles& InParticles,
+		int32 InParticleOffset,
+		int32 InParticleCount,
 		TArray<TVec4<int32>>&& InConstraints,
 		const TConstArrayView<FRealSingle>& StiffnessMultipliers,
 		const TConstArrayView<FRealSingle>& BucklingStiffnessMultipliers,
@@ -27,8 +50,8 @@ public:
 		bool bTrimKinematicConstraints = false)
 		: Base(
 			InParticles,
-			ParticleOffset,
-			ParticleCount,
+			InParticleOffset,
+			InParticleCount,
 			MoveTemp(InConstraints),
 			StiffnessMultipliers,
 			BucklingStiffnessMultipliers,
@@ -37,7 +60,7 @@ public:
 			FSolverVec2(GetWeightedFloatBucklingStiffness(PropertyCollection, 1.f)),
 			bTrimKinematicConstraints) 
 	{
-		InitColor(InParticles, ParticleOffset, ParticleCount);
+		InitColor(InParticles);
 	}
 
 	FPBDBendingConstraints(const FSolverParticles& InParticles,
@@ -62,7 +85,7 @@ public:
 			InBucklingStiffness,
 			bTrimKinematicConstraints) 
 	{
-		InitColor(InParticles, ParticleOffset, ParticleCount);
+		InitColor(InParticles);
 	}
 
 	UE_DEPRECATED(5.2, "Use one of the other constructors instead.")
@@ -78,29 +101,25 @@ PRAGMA_ENABLE_DEPRECATION_WARNINGS
 
 	using Base::SetProperties;
 
+	void SetProperties(
+		const FCollectionPropertyConstFacade& PropertyCollection,
+		const TMap<FString, TConstArrayView<FRealSingle>>& WeightMaps);
+
+	UE_DEPRECATED(5.3, "Use SetProperties(const FCollectionPropertyConstFacade&, const TMap<FString, TConstArrayView<FRealSingle>>&, FSolverReal) instead.")
 	void SetProperties(const FCollectionPropertyConstFacade& PropertyCollection)
 	{
-		if (IsBendingElementStiffnessMutable(PropertyCollection))
-		{
-			Stiffness.SetWeightedValue(FSolverVec2(GetWeightedFloatBendingElementStiffness(PropertyCollection)));
-		}
-		if (IsBucklingRatioMutable(PropertyCollection))
-		{
-			BucklingRatio = (FSolverReal)FMath::Clamp(GetBucklingRatio(PropertyCollection), 0.f, 1.);
-		}
-		if (IsBucklingStiffnessMutable(PropertyCollection))
-		{
-			BucklingStiffness.SetWeightedValue(FSolverVec2(GetWeightedFloatBucklingStiffness(PropertyCollection)));
-		}
+		SetProperties(PropertyCollection, TMap<FString, TConstArrayView<FRealSingle>>());
 	}
 
 	void Apply(FSolverParticles& InParticles, const FSolverReal Dt) const;
 
 private:
-	void InitColor(const FSolverParticles& InParticles, const int32 ParticleOffset, const int32 ParticleCount);
+	void InitColor(const FSolverParticles& InParticles);
 	void ApplyHelper(FSolverParticles& InParticles, const FSolverReal Dt, const int32 ConstraintIndex, const FSolverReal ExpStiffnessValue, const FSolverReal ExpBucklingValue) const;
 
 	using Base::Constraints;
+	using Base::ParticleOffset;
+	using Base::ParticleCount;
 	using Base::Stiffness;
 	using Base::BucklingRatio;
 	using Base::BucklingStiffness;
diff --git a/Engine/Source/Runtime/Experimental/Chaos/Public/Chaos/PBDBendingConstraintsBase.h b/Engine/Source/Runtime/Experimental/Chaos/Public/Chaos/PBDBendingConstraintsBase.h
--- a/Engine/Source/Runtime/Experimental/Chaos/Public/Chaos/PBDBendingConstraintsBase.h
+++ b/Engine/Source/Runtime/Experimental/Chaos/Public/Chaos/PBDBendingConstraintsBase.h
@@ -14,8 +14,8 @@ class FPBDBendingConstraintsBase
 {
 public:
 	FPBDBendingConstraintsBase(const FSolverParticles& InParticles,
-		int32 ParticleOffset,
-		int32 ParticleCount, 
+		int32 InParticleOffset,
+		int32 InParticleCount, 
 		TArray<TVec4<int32>>&& InConstraints,
 		const TConstArrayView<FRealSingle>& StiffnessMultipliers,
 		const TConstArrayView<FRealSingle>& BucklingStiffnessMultipliers,
@@ -26,6 +26,8 @@ public:
 		FSolverReal MaxStiffness = FPBDStiffness::DefaultPBDMaxStiffness)
 		: Constraints(bTrimKinematicConstraints ? TrimKinematicConstraints(InConstraints, InParticles): MoveTemp(InConstraints))
 		, ConstraintSharedEdges(ExtractConstraintSharedEdges(Constraints))
+		, ParticleOffset(InParticleOffset)
+		, ParticleCount(InParticleCount)
 		, Stiffness(
 			InStiffness,
 			StiffnessMultipliers,
@@ -58,8 +60,10 @@ public:
 
 	UE_DEPRECATED(5.2, "Use one of the other constructors instead.")
 	FPBDBendingConstraintsBase(const FSolverParticles& InParticles, TArray<TVec4<int32>>&& InConstraints, const FSolverReal InStiffness = (FSolverReal)1.)
-	    : Constraints(MoveTemp(InConstraints))
+		: Constraints(MoveTemp(InConstraints))
 		, ConstraintSharedEdges(ExtractConstraintSharedEdges(Constraints))
+		, ParticleOffset(0)
+		, ParticleCount(InParticles.Size())
 		, Stiffness(FSolverVec2(InStiffness))
 		, BucklingRatio(0.f)
 		, BucklingStiffness(FSolverVec2(InStiffness))
@@ -257,6 +261,9 @@ protected:
 	TArray<TVec4<int32>> Constraints;
 	TArray<TVec2<int32>> ConstraintSharedEdges; // Only shared edges are used for calculating weighted stiffnesses.
 
+	const int32 ParticleOffset;
+	const int32 ParticleCount;
+
 	FPBDStiffness Stiffness;
 	FSolverReal BucklingRatio;
 	FPBDStiffness BucklingStiffness;
diff --git a/Engine/Source/Runtime/Experimental/Chaos/Public/Chaos/PBDLongRangeConstraints.h b/Engine/Source/Runtime/Experimental/Chaos/Public/Chaos/PBDLongRangeConstraints.h
--- a/Engine/Source/Runtime/Experimental/Chaos/Public/Chaos/PBDLongRangeConstraints.h
+++ b/Engine/Source/Runtime/Experimental/Chaos/Public/Chaos/PBDLongRangeConstraints.h
@@ -21,6 +21,28 @@ public:
 		return IsTetherStiffnessEnabled(PropertyCollection, false);
 	}
 
+	FPBDLongRangeConstraints(
+		const FSolverParticles& Particles,
+		const int32 InParticleOffset,
+		const int32 InParticleCount,
+		const TArray<TConstArrayView<TTuple<int32, int32, FRealSingle>>>& InTethers,
+		const TMap<FString, TConstArrayView<FRealSingle>>& WeightMaps,
+		const FCollectionPropertyConstFacade& PropertyCollection,
+		FSolverReal MeshScale)
+		: FPBDLongRangeConstraintsBase(
+			Particles,
+			InParticleOffset,
+			InParticleCount,
+			InTethers,
+			WeightMaps.FindRef(GetTetherStiffnessString(PropertyCollection, TetherStiffnessName.ToString())),
+			WeightMaps.FindRef(GetTetherScaleString(PropertyCollection, TetherScaleName.ToString())),
+			FSolverVec2(GetWeightedFloatTetherStiffness(PropertyCollection, 1.f)),
+			FSolverVec2(GetWeightedFloatTetherScale(PropertyCollection, 1.f)),  // Scale clamping done in constructor
+			FPBDStiffness::DefaultPBDMaxStiffness,
+			MeshScale)
+	{}
+
+	UE_DEPRECATED(5.3, "Use weight map constructor instead.")
 	FPBDLongRangeConstraints(
 		const FSolverParticles& Particles,
 		const int32 InParticleOffset,
@@ -69,25 +91,27 @@ public:
 
 	using Base::SetProperties;
 
-	void SetProperties(const FCollectionPropertyConstFacade& PropertyCollection, FSolverReal MeshScale)
+	void SetProperties(
+		const FCollectionPropertyConstFacade& PropertyCollection,
+		const TMap<FString, TConstArrayView<FRealSingle>>& WeightMaps,
+		FSolverReal MeshScale);
+
+	UE_DEPRECATED(5.3, "Use SetProperties(const FCollectionPropertyConstFacade&, const TMap<FString, TConstArrayView<FRealSingle>>&, FSolverReal) instead.")
+	void SetProperties(const FCollectionPropertyConstFacade& PropertyCollection)
 	{
-		if (IsTetherStiffnessMutable(PropertyCollection))
-		{
-			Stiffness.SetWeightedValue(FSolverVec2(GetWeightedFloatTetherStiffness(PropertyCollection)));
-		}
-		if (IsTetherScaleMutable(PropertyCollection))
-		{
-			TetherScale.SetWeightedValue(FSolverVec2(GetWeightedFloatTetherScale(PropertyCollection)).ClampAxes(MinTetherScale, MaxTetherScale) * MeshScale);
-		}
+		SetProperties(PropertyCollection, TMap<FString, TConstArrayView<FRealSingle>>(), (FSolverReal)1.);
 	}
 
 	void Apply(FSolverParticles& Particles, const FSolverReal Dt) const;
 
 private:
+	using Base::MinTetherScale;
+	using Base::MaxTetherScale;
 	using Base::Tethers;
 	using Base::Stiffness;
 	using Base::TetherScale;
 	using Base::ParticleOffset;
+	using Base::ParticleCount;
 
 	UE_CHAOS_DECLARE_PROPERTYCOLLECTION_NAME(TetherStiffness, float);
 	UE_CHAOS_DECLARE_PROPERTYCOLLECTION_NAME(TetherScale, float);
diff --git a/Engine/Source/Runtime/Experimental/Chaos/Public/Chaos/PBDLongRangeConstraintsBase.h b/Engine/Source/Runtime/Experimental/Chaos/Public/Chaos/PBDLongRangeConstraintsBase.h
--- a/Engine/Source/Runtime/Experimental/Chaos/Public/Chaos/PBDLongRangeConstraintsBase.h
+++ b/Engine/Source/Runtime/Experimental/Chaos/Public/Chaos/PBDLongRangeConstraintsBase.h
@@ -140,9 +140,10 @@ protected:
 protected:
 	static constexpr int32 TableSize = 16;  // The size of the weightmaps lookup table
 	const TArray<TConstArrayView<FTether>>& Tethers;  // Array view on the tether provided to this constraint
+	const int32 ParticleOffset;  // Index of the first usable particle
+	const int32 ParticleCount;
 	FPBDStiffness Stiffness;  // Stiffness weightmap lookup table
 	FPBDWeightMap TetherScale;  // Scale weightmap lookup table
-	const int32 ParticleOffset;  // Index of the first usable particle
 
 	UE_DEPRECATED(5.2, "Use TetherScale instead")
 	TArray<uint8> ScaleIndices;
@@ -150,7 +151,5 @@ protected:
 	TArray<FSolverReal> ScaleTable;
 	UE_DEPRECATED(5.2, "Use TetherScale instead")
 	FSolverVec2 Scale;
-	UE_DEPRECATED(5.2, "Use TetherScale.Num() or Stiffness.Num(), instead")
-	const int32 ParticleCount = 0;
 };
 }  // End namespace Chaos::Softs
diff --git a/Engine/Source/Runtime/Experimental/Chaos/Public/Chaos/PBDSphericalConstraint.h b/Engine/Source/Runtime/Experimental/Chaos/Public/Chaos/PBDSphericalConstraint.h
--- a/Engine/Source/Runtime/Experimental/Chaos/Public/Chaos/PBDSphericalConstraint.h
+++ b/Engine/Source/Runtime/Experimental/Chaos/Public/Chaos/PBDSphericalConstraint.h
@@ -34,6 +34,25 @@ public:
 		return IsMaxDistanceEnabled(PropertyCollection, false);
 	}
 
+	FPBDSphericalConstraint(
+		const uint32 InParticleOffset,
+		const uint32 InParticleCount,
+		const TArray<FSolverVec3>& InAnimationPositions,  // Use global indexation (will need adding ParticleOffset)
+		const TMap<FString, TConstArrayView<FRealSingle>>& WeightMaps,  // Use local indexation
+		const FCollectionPropertyConstFacade& PropertyCollection,
+		FSolverReal MeshScale
+	)
+		: AnimationPositions(InAnimationPositions)
+		, SphereRadii(WeightMaps.FindRef(GetMaxDistanceString(PropertyCollection, MaxDistanceName.ToString())))
+		, ParticleOffset(InParticleOffset)
+		, ParticleCount(InParticleCount)
+		, Scale(MeshScale)
+		, MaxDistanceBase((FSolverReal)GetLowMaxDistance(PropertyCollection, 0.f))
+		, MaxDistanceRange((FSolverReal)GetHighMaxDistance(PropertyCollection, 1.f) - MaxDistanceBase)
+	{
+	}
+
+	UE_DEPRECATED(5.3, "Use weight map constructor instead.")
 	FPBDSphericalConstraint(
 		const uint32 InParticleOffset,
 		const uint32 InParticleCount,
@@ -45,9 +64,9 @@ public:
 		: AnimationPositions(InAnimationPositions)
 		, SphereRadii(InSphereRadii)
 		, ParticleOffset(InParticleOffset)
+		, ParticleCount(InParticleCount)
 		, Scale(MeshScale)
 	{
-		check(InSphereRadii.Num() == InParticleCount);
 	}
 
 	FPBDSphericalConstraint(
@@ -59,16 +78,22 @@ public:
 		: AnimationPositions(InAnimationPositions)
 		, SphereRadii(InSphereRadii)
 		, ParticleOffset(InParticleOffset)
+		, ParticleCount(InParticleCount)
 		, Scale((FSolverReal)1.)
 	{
-		check(InSphereRadii.Num() == InParticleCount);
 	}
+
 	~FPBDSphericalConstraint() {}
 
+	void SetProperties(
+		const FCollectionPropertyConstFacade& PropertyCollection,
+		const TMap<FString, TConstArrayView<FRealSingle>>& WeightMaps,
+		FSolverReal MeshScale);
+
+	UE_DEPRECATED(5.3, "Use SetProperties(const FCollectionPropertyConstFacade&, const TMap<FString, TConstArrayView<FRealSingle>>&, FSolverReal) instead.")
 	void SetProperties(const FCollectionPropertyConstFacade& PropertyCollection, FSolverReal MeshScale)
 	{
-		SetScale((FSolverReal)1., MeshScale);
-		// TODO: MaxDistance
+		SetProperties(PropertyCollection, TMap<FString, TConstArrayView<FRealSingle>>(), MeshScale);
 	}
 
 	void Apply(FSolverParticles& Particles, const FSolverReal Dt) const
@@ -81,10 +106,23 @@ public:
 		}
 		else
 		{
-			ApplyHelper(Particles, Dt);
+			if (SphereRadii.Num() == ParticleCount)
+			{
+				constexpr bool bHasMaxDistance = true;
+				ApplyHelper<bHasMaxDistance>(Particles, Dt);
+			}
+			else
+			{
+				constexpr bool bHasMaxDistance = false;
+				ApplyHelper<bHasMaxDistance>(Particles, Dt);
+			}
 		}
 	}
 
+	// Set a new mesh scale
+	void SetScale(FSolverReal InScale) { Scale = InScale; }
+
+	UE_DEPRECATED(5.3, "Use SetScale(FSolverReal) instead.")
 	void SetScale(FSolverReal MaxDistanceScale, FSolverReal MeshScale)
 	{
 		Scale = FMath::Max(MaxDistanceScale, (FSolverReal)0.) * MeshScale;
@@ -95,14 +133,13 @@ public:
 	UE_DEPRECATED(5.2, "Use SetScale instead.")
 	void SetSphereRadiiMultiplier(FSolverReal InSphereRadiiMultiplier, FSolverReal MeshScale)
 	{
-		SetScale(InSphereRadiiMultiplier, MeshScale);
+		SetScale(InSphereRadiiMultiplier * MeshScale);
 	}
 
 private:
+	template<bool bHasMaxDistance>
 	void ApplyHelper(FSolverParticles& Particles, const FSolverReal Dt) const
 	{
-		const int32 ParticleCount = SphereRadii.Num();
-
 		PhysicsParallelFor(ParticleCount, [this, &Particles, Dt](int32 Index)  // TODO: profile need for parallel loop based on particle count
 		{
 			const int32 ParticleIndex = ParticleOffset + Index;
@@ -112,7 +149,7 @@ private:
 				return;
 			}
 
-			const FSolverReal Radius = SphereRadii[Index] * Scale;
+			const FSolverReal Radius = (bHasMaxDistance ? MaxDistanceBase + MaxDistanceRange * SphereRadii[Index] : MaxDistanceBase) * Scale;
 			const FSolverVec3& Center = AnimationPositions[ParticleIndex];
 
 			const FSolverVec3 CenterToParticle = Particles.P(ParticleIndex) - Center;
@@ -132,13 +169,16 @@ private:
 
 protected:
 	const TArray<FSolverVec3>& AnimationPositions;  // Use global indexation (will need adding ParticleOffset)
-	const TConstArrayView<FRealSingle> SphereRadii;  // Use local indexation
+	TConstArrayView<FRealSingle> SphereRadii;  // Use local indexation
 	const int32 ParticleOffset;
+	const int32 ParticleCount;
 	UE_DEPRECATED(5.2, "Use Scale instead.")
 	FSolverReal SphereRadiiMultiplier = 1.f;
 
 private:
-	FSolverReal Scale;
+	FSolverReal Scale = (FSolverReal)1.;
+	FSolverReal MaxDistanceBase = (FSolverReal)0.;
+	FSolverReal MaxDistanceRange = (FSolverReal)1.;
 
 	UE_CHAOS_DECLARE_PROPERTYCOLLECTION_NAME(MaxDistance, float);
 };
@@ -152,6 +192,31 @@ public:
 			IsBackstopDistanceAnimatable(PropertyCollection, false);  // Backstop can be re-enabled if animated
 	}
 
+	FPBDSphericalBackstopConstraint(
+		const int32 InParticleOffset,
+		const int32 InParticleCount,
+		const TArray<FSolverVec3>& InAnimationPositions,  // Use global indexation (will need adding ParticleOffset)
+		const TArray<FSolverVec3>& InAnimationNormals,  // Use global indexation (will need adding ParticleOffset)
+		const TMap<FString, TConstArrayView<FRealSingle>>& WeightMaps,  // Use local indexation
+		const FCollectionPropertyConstFacade& PropertyCollection,
+		FSolverReal MeshScale
+	)
+		: AnimationPositions(InAnimationPositions)
+		, AnimationNormals(InAnimationNormals)
+		, SphereRadii(WeightMaps.FindRef(GetBackstopRadiusString(PropertyCollection, BackstopRadiusName.ToString())))
+		, SphereOffsetDistances(WeightMaps.FindRef(GetBackstopDistanceString(PropertyCollection, BackstopDistanceName.ToString())))
+		, ParticleOffset(InParticleOffset)
+		, ParticleCount(InParticleCount)
+		, Scale(MeshScale)
+		, BackstopRadiusBase((FSolverReal)FMath::Max(GetLowBackstopRadius(PropertyCollection, 0.f), 0.f))
+		, BackstopRadiusRange((FSolverReal)FMath::Max(GetHighBackstopRadius(PropertyCollection, 1.f), 0.f) - BackstopRadiusBase)
+		, BackstopDistanceBase((FSolverReal)GetLowBackstopDistance(PropertyCollection, 0.f))
+		, BackstopDistanceRange((FSolverReal)GetHighBackstopDistance(PropertyCollection, 1.f) - BackstopDistanceBase)
+		, bUseLegacyBackstop(GetUseLegacyBackstop(PropertyCollection, false))  // Only set the legacy backstop in constructor
+	{
+	}
+
+	UE_DEPRECATED(5.3, "Use weight map constructor instead.")
 	FPBDSphericalBackstopConstraint(
 		const int32 InParticleOffset,
 		const int32 InParticleCount,
@@ -167,12 +232,11 @@ public:
 		, SphereRadii(InSphereRadii)
 		, SphereOffsetDistances(InSphereOffsetDistances)
 		, ParticleOffset(InParticleOffset)
+		, ParticleCount(InParticleCount)
 		, Scale(MeshScale)
 		, bEnabled(true)
 		, bUseLegacyBackstop(GetUseLegacyBackstop(PropertyCollection, false))  // Only set the legacy backstop in constructor
 	{
-		check(InSphereRadii.Num() == InParticleCount);
-		check(InSphereOffsetDistances.Num() == InParticleCount);
 	}
 
 	FPBDSphericalBackstopConstraint(
@@ -189,19 +253,23 @@ public:
 		, SphereRadii(InSphereRadii)
 		, SphereOffsetDistances(InSphereOffsetDistances)
 		, ParticleOffset(InParticleOffset)
+		, ParticleCount(InParticleCount)
 		, Scale((FSolverReal)1.)
 		, bEnabled(true)
 		, bUseLegacyBackstop(bInUseLegacyBackstop)
 	{
-		check(InSphereRadii.Num() == InParticleCount);
-		check(InSphereOffsetDistances.Num() == InParticleCount);
 	}
 	~FPBDSphericalBackstopConstraint() {}
 
+	void SetProperties(
+		const FCollectionPropertyConstFacade& PropertyCollection,
+		const TMap<FString, TConstArrayView<FRealSingle>>& WeightMaps,
+		FSolverReal MeshScale);
+
+	UE_DEPRECATED(5.3, "Use SetProperties(const FCollectionPropertyConstFacade&, const TMap<FString, TConstArrayView<FRealSingle>>&, FSolverReal) instead.")
 	void SetProperties(const FCollectionPropertyConstFacade& PropertyCollection, FSolverReal MeshScale)
 	{
-		SetScale((FSolverReal)1., MeshScale);
-		// TODO: BackstopDistance and BackstopRadius
+		SetProperties(PropertyCollection, TMap<FString, TConstArrayView<FRealSingle>>(), MeshScale);
 	}
 
 	void SetEnabled(bool bInEnabled) { bEnabled = bInEnabled; }
@@ -222,9 +290,32 @@ public:
 				{
 					ApplyLegacyHelperISPC(Particles, Dt);
 				}
+				else if (SphereRadii.Num() == ParticleCount)
+				{
+					if (SphereOffsetDistances.Num() == ParticleCount)
+					{
+						constexpr bool bHasBackstopDistance = true;
+						constexpr bool bHasBackstopRadius = true;
+						ApplyLegacyHelper<bHasBackstopDistance, bHasBackstopRadius>(Particles, Dt);
+					}
+					else
+					{
+						constexpr bool bHasBackstopDistance = false;
+						constexpr bool bHasBackstopRadius = true;
+						ApplyLegacyHelper<bHasBackstopDistance, bHasBackstopRadius>(Particles, Dt);
+					}
+				}
+				else if (SphereOffsetDistances.Num() == ParticleCount)
+				{
+					constexpr bool bHasBackstopDistance = true;
+					constexpr bool bHasBackstopRadius = false;
+					ApplyLegacyHelper<bHasBackstopDistance, bHasBackstopRadius>(Particles, Dt);
+				}
 				else
 				{
-					ApplyLegacyHelper(Particles, Dt);
+					constexpr bool bHasBackstopDistance = false;
+					constexpr bool bHasBackstopRadius = false;
+					ApplyLegacyHelper<bHasBackstopDistance, bHasBackstopRadius>(Particles, Dt);
 				}
 			}
 			else
@@ -234,14 +325,41 @@ public:
 				{
 					ApplyHelperISPC(Particles, Dt);
 				}
+				else if (SphereRadii.Num() == ParticleCount)
+				{
+					if (SphereOffsetDistances.Num() == ParticleCount)
+					{
+						constexpr bool bHasBackstopDistance = true;
+						constexpr bool bHasBackstopRadius = true;
+						ApplyHelper<bHasBackstopDistance, bHasBackstopRadius>(Particles, Dt);
+					}
+					else
+					{
+						constexpr bool bHasBackstopDistance = false;
+						constexpr bool bHasBackstopRadius = true;
+						ApplyHelper<bHasBackstopDistance, bHasBackstopRadius>(Particles, Dt);
+					}
+				}
+				else if (SphereOffsetDistances.Num() == ParticleCount)
+				{
+					constexpr bool bHasBackstopDistance = true;
+					constexpr bool bHasBackstopRadius = false;
+					ApplyHelper<bHasBackstopDistance, bHasBackstopRadius>(Particles, Dt);
+				}
 				else
 				{
-					ApplyHelper(Particles, Dt);
+					constexpr bool bHasBackstopDistance = false;
+					constexpr bool bHasBackstopRadius = false;
+					ApplyHelper<bHasBackstopDistance, bHasBackstopRadius>(Particles, Dt);
 				}
 			}
 		}
 	}
 
+	// Set a new mesh scale
+	void SetScale(FSolverReal InScale) { Scale = InScale; }
+
+	UE_DEPRECATED(5.3, "Use SetScale(FSolverReal) instead.")
 	void SetScale(FSolverReal BackstopScale, FSolverReal MeshScale)
 	{
 		Scale = FMath::Max(BackstopScale, (FSolverReal)0.) * MeshScale;
@@ -252,7 +370,9 @@ public:
 	UE_DEPRECATED(5.2, "Use SetScale instead.")
 	void SetSphereRadiiMultiplier(FSolverReal InSphereRadiiMultiplier, FSolverReal MeshScale = (FSolverReal)1.)
 	{
+PRAGMA_DISABLE_DEPRECATION_WARNINGS
 		SetScale(InSphereRadiiMultiplier, MeshScale);
+PRAGMA_ENABLE_DEPRECATION_WARNINGS
 	}
 
 	UE_DEPRECATED(5.2, "Use GetScale() instead.")
@@ -264,10 +384,9 @@ public:
 	}
 
 private:
+	template<bool bHasBackstopDistance, bool bHasBackstopRadius>
 	void ApplyHelper(FSolverParticles& Particles, const FSolverReal Dt) const
 	{
-		const int32 ParticleCount = SphereRadii.Num();
-
 		PhysicsParallelFor(ParticleCount, [this, &Particles, Dt](int32 Index)  // TODO: profile need for parallel loop based on particle count
 		{
 			const int32 ParticleIndex = ParticleOffset + Index;
@@ -280,8 +399,8 @@ private:
 			const FSolverVec3& AnimationPosition = AnimationPositions[ParticleIndex];
 			const FSolverVec3& AnimationNormal = AnimationNormals[ParticleIndex];
 
-			const FSolverReal SphereOffsetDistance = SphereOffsetDistances[Index] * Scale;
-			const FSolverReal Radius = SphereRadii[Index] * Scale;
+			const FSolverReal SphereOffsetDistance = (bHasBackstopDistance ? BackstopDistanceBase + BackstopDistanceRange * SphereOffsetDistances[Index] : BackstopDistanceBase) * Scale;
+			const FSolverReal Radius = (bHasBackstopRadius ? BackstopRadiusBase + BackstopRadiusRange * SphereRadii[Index] : BackstopRadiusBase) * Scale;
 
 			const FSolverVec3 Center = AnimationPosition - (Radius + SphereOffsetDistance) * AnimationNormal;  // Non legacy version adds radius to the distance
 			const FSolverVec3 CenterToParticle = Particles.P(ParticleIndex) - Center;
@@ -301,10 +420,9 @@ private:
 		});
 	}
 
+	template<bool bHasBackstopDistance, bool bHasBackstopRadius>
 	void ApplyLegacyHelper(FSolverParticles& Particles, const FSolverReal Dt) const
 	{
-		const int32 ParticleCount = SphereRadii.Num();
-
 		PhysicsParallelFor(ParticleCount, [this, &Particles, Dt](int32 Index)  // TODO: profile need for parallel loop based on particle count
 		{
 			const int32 ParticleIndex = ParticleOffset + Index;
@@ -317,8 +435,8 @@ private:
 			const FSolverVec3& AnimationPosition = AnimationPositions[ParticleIndex];
 			const FSolverVec3& AnimationNormal = AnimationNormals[ParticleIndex];
 
-			const FSolverReal SphereOffsetDistance = SphereOffsetDistances[Index] * Scale;
-			const FSolverReal Radius = SphereRadii[Index] * Scale;
+			const FSolverReal SphereOffsetDistance = (bHasBackstopDistance ? BackstopDistanceBase + BackstopDistanceRange * SphereOffsetDistances[Index] : BackstopDistanceBase) * Scale;
+			const FSolverReal Radius = (bHasBackstopRadius ? BackstopRadiusBase + BackstopRadiusRange * SphereRadii[Index] : BackstopRadiusBase) * Scale;
 
 			const FSolverVec3 Center = AnimationPosition - SphereOffsetDistance * AnimationNormal;  // Legacy version already includes the radius within the distance
 			const FSolverVec3 CenterToParticle = Particles.P(ParticleIndex) - Center;
@@ -344,12 +462,17 @@ private:
 private:
 	const TArray<FSolverVec3>& AnimationPositions;  // Positions of spheres, use global indexation (will need adding ParticleOffset)
 	const TArray<FSolverVec3>& AnimationNormals; // Sphere offset directions, use global indexation (will need adding ParticleOffset)
-	const TConstArrayView<FRealSingle> SphereRadii; // Start at index 0, use local indexation
-	const TConstArrayView<FRealSingle> SphereOffsetDistances;  // Sphere position offsets, use local indexation
+	TConstArrayView<FRealSingle> SphereRadii; // Start at index 0, use local indexation
+	TConstArrayView<FRealSingle> SphereOffsetDistances;  // Sphere position offsets, use local indexation
 	const int32 ParticleOffset;
-	FSolverReal Scale;
-	bool bEnabled;
-	bool bUseLegacyBackstop;
+	const int32 ParticleCount;
+	FSolverReal Scale = (FSolverReal)1.;
+	FSolverReal BackstopRadiusBase = (FSolverReal)0.;
+	FSolverReal BackstopRadiusRange = (FSolverReal)1.;
+	FSolverReal BackstopDistanceBase = (FSolverReal)0.;
+	FSolverReal BackstopDistanceRange = (FSolverReal)1.;
+	bool bEnabled = true;
+	bool bUseLegacyBackstop = false;
 
 	UE_CHAOS_DECLARE_PROPERTYCOLLECTION_NAME(BackstopDistance, float);
 	UE_CHAOS_DECLARE_PROPERTYCOLLECTION_NAME(BackstopRadius, float);
diff --git a/Engine/Source/Runtime/Experimental/Chaos/Public/Chaos/PBDSpringConstraints.h b/Engine/Source/Runtime/Experimental/Chaos/Public/Chaos/PBDSpringConstraints.h
--- a/Engine/Source/Runtime/Experimental/Chaos/Public/Chaos/PBDSpringConstraints.h
+++ b/Engine/Source/Runtime/Experimental/Chaos/Public/Chaos/PBDSpringConstraints.h
@@ -13,16 +13,23 @@ public:
 	template<int32 Valence>
 	FPBDSpringConstraints(
 		const FSolverParticles& Particles,
-		int32 ParticleOffset,
-		int32 ParticleCount,
+		int32 InParticleOffset,
+		int32 InParticleCount,
 		const TArray<TVector<int32, Valence>>& InConstraints,
 		const TConstArrayView<FRealSingle>& StiffnessMultipliers,
 		const FSolverVec2& InStiffness,
 		bool bTrimKinematicConstraints = false,
 		typename TEnableIf<Valence >= 2 && Valence <= 4>::Type* = nullptr)
-		: Base(Particles, ParticleOffset, ParticleCount, InConstraints, StiffnessMultipliers, InStiffness, bTrimKinematicConstraints)
+		: Base(
+			Particles,
+			InParticleOffset,
+			InParticleCount,
+			InConstraints,
+			StiffnessMultipliers,
+			InStiffness,
+			bTrimKinematicConstraints)
 	{
-		InitColor(Particles, ParticleOffset, ParticleCount);
+		InitColor(Particles);
 	}
 
 	virtual ~FPBDSpringConstraints() override {}
@@ -35,9 +42,11 @@ protected:
 	typedef FPBDSpringConstraintsBase Base;
 	using Base::Constraints;
 	using Base::Stiffness;
+	using Base::ParticleOffset;
+	using Base::ParticleCount;
 
 private:
-	void InitColor(const FSolverParticles& InParticles, const int32 ParticleOffset, const int32 ParticleCount);
+	void InitColor(const FSolverParticles& InParticles);
 	void ApplyHelper(FSolverParticles& Particles, const FSolverReal Dt, const int32 ConstraintIndex, const FSolverReal ExpStiffnessValue) const;
 
 private:
@@ -52,6 +61,25 @@ public:
 		return IsEdgeSpringStiffnessEnabled(PropertyCollection, false);
 	}
 
+	FPBDEdgeSpringConstraints(
+		const FSolverParticles& Particles,
+		int32 ParticleOffset,
+		int32 ParticleCount,
+		const TArray<TVec3<int32>>& InConstraints,
+		const TMap<FString, TConstArrayView<FRealSingle>>& WeightMaps,
+		const FCollectionPropertyConstFacade& PropertyCollection,
+		bool bTrimKinematicConstraints = false)
+		: FPBDSpringConstraints(
+			Particles,
+			ParticleOffset,
+			ParticleCount,
+			InConstraints,
+			WeightMaps.FindRef(GetEdgeSpringStiffnessString(PropertyCollection, EdgeSpringStiffnessName.ToString())),
+			FSolverVec2(GetWeightedFloatEdgeSpringStiffness(PropertyCollection, 1.f)),
+			bTrimKinematicConstraints)
+	{}
+
+	UE_DEPRECATED(5.3, "Use weight map constructor instead.")
 	FPBDEdgeSpringConstraints(
 		const FSolverParticles& Particles,
 		int32 ParticleOffset,
@@ -72,15 +100,19 @@ public:
 
 	virtual ~FPBDEdgeSpringConstraints() override = default;
 
+	void SetProperties(
+		const FCollectionPropertyConstFacade& PropertyCollection,
+		const TMap<FString, TConstArrayView<FRealSingle>>& WeightMaps);
+
+	UE_DEPRECATED(5.3, "Use SetProperties(const FCollectionPropertyConstFacade&, const TMap<FString, TConstArrayView<FRealSingle>>&, FSolverReal) instead.")
 	void SetProperties(const FCollectionPropertyConstFacade& PropertyCollection)
 	{
-		if (IsEdgeSpringStiffnessMutable(PropertyCollection))
-		{
-			Stiffness.SetWeightedValue(FSolverVec2(GetWeightedFloatEdgeSpringStiffness(PropertyCollection)));
-		}
+		SetProperties(PropertyCollection, TMap<FString, TConstArrayView<FRealSingle>>());
 	}
 
 private:
+	using FPBDSpringConstraints::ParticleCount;
+
 	UE_CHAOS_DECLARE_PROPERTYCOLLECTION_NAME(EdgeSpringStiffness, float);
 };
 
@@ -92,6 +124,25 @@ public:
 		return IsBendingSpringStiffnessEnabled(PropertyCollection, false);
 	}
 
+	FPBDBendingSpringConstraints(
+		const FSolverParticles& Particles,
+		int32 ParticleOffset,
+		int32 ParticleCount,
+		const TArray<TVec2<int32>>& InConstraints,
+		const TMap<FString, TConstArrayView<FRealSingle>>& WeightMaps,
+		const FCollectionPropertyConstFacade& PropertyCollection,
+		bool bTrimKinematicConstraints = false)
+		: FPBDSpringConstraints(
+			Particles,
+			ParticleOffset,
+			ParticleCount,
+			InConstraints,
+			WeightMaps.FindRef(GetBendingSpringStiffnessString(PropertyCollection, BendingSpringStiffnessName.ToString())),
+			FSolverVec2(GetWeightedFloatBendingSpringStiffness(PropertyCollection, 1.f)),
+			bTrimKinematicConstraints)
+	{}
+
+	UE_DEPRECATED(5.3, "Use weight map constructor instead.")
 	FPBDBendingSpringConstraints(
 		const FSolverParticles& Particles,
 		int32 ParticleOffset,
@@ -112,15 +163,19 @@ public:
 
 	virtual ~FPBDBendingSpringConstraints() override = default;
 
+	void SetProperties(
+		const FCollectionPropertyConstFacade& PropertyCollection,
+		const TMap<FString, TConstArrayView<FRealSingle>>& WeightMaps);
+
+	UE_DEPRECATED(5.3, "Use SetProperties(const FCollectionPropertyConstFacade&, const TMap<FString, TConstArrayView<FRealSingle>>&, FSolverReal) instead.")
 	void SetProperties(const FCollectionPropertyConstFacade& PropertyCollection)
 	{
-		if (IsBendingSpringStiffnessMutable(PropertyCollection))
-		{
-			Stiffness.SetWeightedValue(FSolverVec2(GetWeightedFloatBendingSpringStiffness(PropertyCollection)));
-		}
+		SetProperties(PropertyCollection, TMap<FString, TConstArrayView<FRealSingle>>());
 	}
 
 private:
+	using FPBDSpringConstraints::ParticleCount;
+
 	UE_CHAOS_DECLARE_PROPERTYCOLLECTION_NAME(BendingSpringStiffness, float);
 };
 
diff --git a/Engine/Source/Runtime/Experimental/Chaos/Public/Chaos/PBDSpringConstraintsBase.h b/Engine/Source/Runtime/Experimental/Chaos/Public/Chaos/PBDSpringConstraintsBase.h
--- a/Engine/Source/Runtime/Experimental/Chaos/Public/Chaos/PBDSpringConstraintsBase.h
+++ b/Engine/Source/Runtime/Experimental/Chaos/Public/Chaos/PBDSpringConstraintsBase.h
@@ -18,8 +18,8 @@ public:
 	template<int32 Valence, TEMPLATE_REQUIRES(Valence >= 2 && Valence <= 4)>
 	FPBDSpringConstraintsBase(
 		const FSolverParticles& Particles,
-		int32 ParticleOffset,
-		int32 ParticleCount,
+		int32 InParticleOffset,
+		int32 InParticleCount,
 		const TArray<TVector<int32, Valence>>& InConstraints,
 		const TConstArrayView<FRealSingle>& StiffnessMultipliers,
 		const FSolverVec2& InStiffness,
@@ -30,6 +30,8 @@ public:
 			{
 				return bTrimKinematicConstraints && Particles.InvM(Index0) == (FSolverReal)0. && Particles.InvM(Index1) == (FSolverReal)0.;
 			}))
+		, ParticleOffset(InParticleOffset)
+		, ParticleCount(InParticleCount)
 		, Stiffness(
 			InStiffness,
 			StiffnessMultipliers,
@@ -113,6 +115,8 @@ private:
 protected:
 	TArray<TVec2<int32>> Constraints;
 	TArray<FSolverReal> Dists;
+	const int32 ParticleOffset;
+	const int32 ParticleCount;
 	FPBDStiffness Stiffness;
 };
 
diff --git a/Engine/Source/Runtime/Experimental/Chaos/Public/Chaos/PBDStiffness.h b/Engine/Source/Runtime/Experimental/Chaos/Public/Chaos/PBDStiffness.h
--- a/Engine/Source/Runtime/Experimental/Chaos/Public/Chaos/PBDStiffness.h
+++ b/Engine/Source/Runtime/Experimental/Chaos/Public/Chaos/PBDStiffness.h
@@ -50,7 +50,12 @@ public:
 		FSolverReal InParameterFitBase = DefaultParameterFitBase,  // Logarithm base to use in the PBD stiffness parameter fit function
 		FSolverReal MaxStiffness = DefaultPBDMaxStiffness);
 
-	~FPBDStiffness() {}
+	virtual ~FPBDStiffness() override = default;
+
+	FPBDStiffness(const FPBDStiffness&) = default;
+	FPBDStiffness(FPBDStiffness&&) = default;
+	FPBDStiffness& operator=(const FPBDStiffness&) = default;
+	FPBDStiffness& operator=(FPBDStiffness&&) = default;
 
 	/**
 	 * Set the low and high values of the weight map.
@@ -73,8 +78,8 @@ public:
 	inline void ApplyXPBDValues(const FSolverReal MaxStiffnesss);
 
 private:
-	const FSolverReal ParameterFitBase;
-	const FSolverReal ParameterFitLogBase;
+	FSolverReal ParameterFitBase;
+	FSolverReal ParameterFitLogBase;
 	FSolverReal PrevDtOrMaxStiffness = (FSolverReal)0.;
 	int32 PrevNumIterations = 0;
 	using FPBDWeightMap::bIsDirty;
diff --git a/Engine/Source/Runtime/Experimental/Chaos/Public/Chaos/PBDWeightMap.h b/Engine/Source/Runtime/Experimental/Chaos/Public/Chaos/PBDWeightMap.h
--- a/Engine/Source/Runtime/Experimental/Chaos/Public/Chaos/PBDWeightMap.h
+++ b/Engine/Source/Runtime/Experimental/Chaos/Public/Chaos/PBDWeightMap.h
@@ -41,7 +41,12 @@ public:
 		int32 TableSize = 16,  // Size of the lookup table, can't be more than 256 values, the larger the table the longer it takes to apply changes to the stiffness values
 		typename TEnableIf<Valence >= 2 && Valence <= 4>::Type* = nullptr);  // Prevents incorrect valence, the value is actually unused
 
-	~FPBDWeightMap() {}
+	virtual ~FPBDWeightMap() = default;
+
+	FPBDWeightMap(const FPBDWeightMap&) = default;
+	FPBDWeightMap(FPBDWeightMap&&) = default;
+	FPBDWeightMap& operator=(const FPBDWeightMap&) = default;
+	FPBDWeightMap& operator=(FPBDWeightMap&&) = default;
 
 	/** Return the number of values stored in the weight map. */
 	int32 Num() const { return Indices.Num(); }
diff --git a/Engine/Source/Runtime/Experimental/Chaos/Public/Chaos/VelocityField.h b/Engine/Source/Runtime/Experimental/Chaos/Public/Chaos/VelocityField.h
--- a/Engine/Source/Runtime/Experimental/Chaos/Public/Chaos/VelocityField.h
+++ b/Engine/Source/Runtime/Experimental/Chaos/Public/Chaos/VelocityField.h
@@ -61,34 +61,15 @@ public:
 		}
 	}
 
+	void SetProperties(
+		const FCollectionPropertyConstFacade& PropertyCollection,
+		const TMap<FString, TConstArrayView<FRealSingle>>& Weightmaps,
+		FSolverReal WorldScale);
+
+	UE_DEPRECATED(5.3, "Use SetProperties(const FCollectionPropertyConstFacade&, const TMap<FString, TConstArrayView<FRealSingle>>&, FSolverReal) instead.")
 	void SetProperties(const FCollectionPropertyConstFacade& PropertyCollection, FSolverReal WorldScale)
 	{
-		if (IsFluidDensityMutable(PropertyCollection))
-		{
-			constexpr FSolverReal OneQuarter = (FSolverReal)0.25;
-			QuarterRho = (FSolverReal)GetFluidDensity(PropertyCollection) * OneQuarter;
-		}
-
-		if (IsDragMutable(PropertyCollection))
-		{
-			const FSolverVec2 Drag(GetWeightedFloatDrag(PropertyCollection));
-			DragBase = FMath::Clamp(Drag[0], MinCoefficient, MaxCoefficient);
-			DragRange = FMath::Clamp(Drag[1], MinCoefficient, MaxCoefficient) - DragBase;
-		}
-
-		if (IsLiftMutable(PropertyCollection))
-		{
-			const FSolverVec2 Lift(GetWeightedFloatLift(PropertyCollection));
-			LiftBase = FMath::Clamp(Lift[0], MinCoefficient, MaxCoefficient);
-			LiftRange = FMath::Clamp(Lift[1], MinCoefficient, MaxCoefficient) - LiftBase;
-		}
-
-		if (IsPressureMutable(PropertyCollection))
-		{
-			const FSolverVec2 Pressure(GetWeightedFloatPressure(PropertyCollection));
-			PressureBase = Pressure[0] / WorldScale;
-			PressureRange = Pressure[1] / WorldScale - PressureBase;
-		}
+		SetProperties(PropertyCollection, TMap<FString, TConstArrayView<FRealSingle>>(), WorldScale);
 	}
 
 	void SetProperties(const FSolverVec2& Drag, const FSolverVec2& Lift, const FSolverReal FluidDensity, const FSolverVec2& Pressure = FSolverVec2::ZeroVector)
@@ -111,6 +92,10 @@ public:
 	}
 
 	void SetGeometry(const FTriangleMesh* TriangleMesh, const TConstArrayView<FRealSingle>& DragMultipliers, const TConstArrayView<FRealSingle>& LiftMultipliers, const TConstArrayView<FRealSingle>& PressureMultipliers = TConstArrayView<FRealSingle>());
+	void SetMultipliers(
+		const TConstArrayView<FRealSingle>& DragMultipliers,
+		const TConstArrayView<FRealSingle>& LiftMultipliers,
+		const TConstArrayView<FRealSingle>& PressureMultipliers);
 
 	void SetVelocity(const FSolverVec3& InVelocity) { Velocity = InVelocity; }
 
diff --git a/Engine/Source/Runtime/Experimental/Chaos/Public/Chaos/XPBDAnisotropicBendingConstraints.h b/Engine/Source/Runtime/Experimental/Chaos/Public/Chaos/XPBDAnisotropicBendingConstraints.h
--- a/Engine/Source/Runtime/Experimental/Chaos/Public/Chaos/XPBDAnisotropicBendingConstraints.h
+++ b/Engine/Source/Runtime/Experimental/Chaos/Public/Chaos/XPBDAnisotropicBendingConstraints.h
@@ -30,23 +30,17 @@ public:
 	}
 
 	FXPBDAnisotropicBendingConstraints(const FSolverParticles& InParticles,
-		int32 ParticleOffset,
-		int32 ParticleCount,
+		int32 InParticleOffset,
+		int32 InParticleCount,
 		const FTriangleMesh& TriangleMesh,
 		const TArray<TVec3<FVec2f>>& FaceVertexPatternPositions,
-		const TConstArrayView<FRealSingle>& StiffnessWarpMultipliers,
-		const TConstArrayView<FRealSingle>& StiffnessWeftMultipliers,
-		const TConstArrayView<FRealSingle>& StiffnessBiasMultipliers,
-		const TConstArrayView<FRealSingle>& BucklingStiffnessWarpMultipliers,
-		const TConstArrayView<FRealSingle>& BucklingStiffnessWeftMultipliers,
-		const TConstArrayView<FRealSingle>& BucklingStiffnessBiasMultipliers,
-		const TConstArrayView<FRealSingle>& DampingMultipliers,
+		const TMap<FString, TConstArrayView<FRealSingle>>& WeightMaps,
 		const FCollectionPropertyConstFacade& PropertyCollection,
 		bool bTrimKinematicConstraints = false);
 
 	FXPBDAnisotropicBendingConstraints(const FSolverParticles& InParticles,
-		int32 ParticleOffset,
-		int32 ParticleCount,
+		int32 InParticleOffset,
+		int32 InParticleCount,
 		const FTriangleMesh& TriangleMesh,
 		const TArray<TVec3<FVec2f>>& FaceVertexPatternPositions,
 		const TConstArrayView<FRealSingle>& StiffnessWarpMultipliers,
@@ -75,40 +69,14 @@ public:
 		FPBDBendingConstraintsBase::Init(InParticles);
 	}
 
+	void SetProperties(
+		const FCollectionPropertyConstFacade& PropertyCollection,
+		const TMap<FString, TConstArrayView<FRealSingle>>& WeightMaps);
+
+	UE_DEPRECATED(5.3, "Use SetProperties(const FCollectionPropertyConstFacade&, const TMap<FString, TConstArrayView<FRealSingle>>&, FSolverReal) instead.")
 	void SetProperties(const FCollectionPropertyConstFacade& PropertyCollection)
 	{
-		if (IsXPBDBendingElementStiffnessWarpMutable(PropertyCollection))
-		{
-			Stiffness.SetWeightedValue(FSolverVec2(GetWeightedFloatXPBDBendingElementStiffnessWarp(PropertyCollection)), MaxStiffness);
-		}
-		if (IsXPBDBendingElementStiffnessWeftMutable(PropertyCollection))
-		{
-			Stiffness.SetWeightedValue(FSolverVec2(GetWeightedFloatXPBDBendingElementStiffnessWeft(PropertyCollection)), MaxStiffness);
-		}
-		if (IsXPBDBendingElementStiffnessBiasMutable(PropertyCollection))
-		{
-			Stiffness.SetWeightedValue(FSolverVec2(GetWeightedFloatXPBDBendingElementStiffnessBias(PropertyCollection)), MaxStiffness);
-		}
-		if (IsXPBDBucklingRatioMutable(PropertyCollection))
-		{
-			BucklingRatio = FMath::Clamp(GetXPBDBucklingRatio(PropertyCollection), (FSolverReal)0., (FSolverReal)1.);
-		}
-		if (IsXPBDBucklingStiffnessWarpMutable(PropertyCollection))
-		{
-			BucklingStiffness.SetWeightedValue(FSolverVec2(GetWeightedFloatXPBDBucklingStiffnessWarp(PropertyCollection)), MaxStiffness);
-		}
-		if (IsXPBDBucklingStiffnessWeftMutable(PropertyCollection))
-		{
-			BucklingStiffness.SetWeightedValue(FSolverVec2(GetWeightedFloatXPBDBucklingStiffnessWeft(PropertyCollection)), MaxStiffness);
-		}
-		if (IsXPBDBucklingStiffnessBiasMutable(PropertyCollection))
-		{
-			BucklingStiffness.SetWeightedValue(FSolverVec2(GetWeightedFloatXPBDBucklingStiffnessBias(PropertyCollection)), MaxStiffness);
-		}
-		if (IsXPBDBendingElementDampingMutable(PropertyCollection))
-		{
-			DampingRatio.SetWeightedValue(FSolverVec2(GetWeightedFloatXPBDBendingElementDamping(PropertyCollection)).ClampAxes(MinDamping, MaxDamping));
-		}
+		SetProperties(PropertyCollection, TMap<FString, TConstArrayView<FRealSingle>>());
 	}
 
 	// Update stiffness table, as well as the simulation stiffness exponent
@@ -126,13 +94,15 @@ public:
 	void Apply(FSolverParticles& Particles, const FSolverReal Dt) const;
 
 private:
-	void InitColor(const FSolverParticles& InParticles, const int32 ParticleOffset, const int32 ParticleCount);
+	void InitColor(const FSolverParticles& InParticles);
 	void ApplyHelper(FSolverParticles& Particles, const FSolverReal Dt, const int32 ConstraintIndex, const FSolverVec3& ExpStiffnessValues, 
 		const FSolverVec3& ExpBucklingStiffnessValues, const FSolverReal DampingRatioValue) const;
 
 	TArray<FSolverVec3> GenerateWarpWeftBiasBaseMultipliers(const TArray<TVec3<FVec2f>>& FaceVertexPatternPositions, const FTriangleMesh& TriangleMesh) const;
 
 	using Base::Constraints;
+	using Base::ParticleOffset;
+	using Base::ParticleCount;
 	using Base::RestAngles;
 	using Base::Stiffness; // Warp
 	using Base::BucklingStiffness; // Warp
diff --git a/Engine/Source/Runtime/Experimental/Chaos/Public/Chaos/XPBDAxialSpringConstraints.h b/Engine/Source/Runtime/Experimental/Chaos/Public/Chaos/XPBDAxialSpringConstraints.h
--- a/Engine/Source/Runtime/Experimental/Chaos/Public/Chaos/XPBDAxialSpringConstraints.h
+++ b/Engine/Source/Runtime/Experimental/Chaos/Public/Chaos/XPBDAxialSpringConstraints.h
@@ -155,11 +155,13 @@ private:
 	}
 
 protected:
+	using Base::Constraints;
+	using Base::ParticleOffset;
+	using Base::ParticleCount;
 	using Base::Stiffness;
 
 private:
 	using Base::Barys;
-	using Base::Constraints;
 	using Base::Dists;
 
 	mutable TArray<FSolverReal> Lambdas;
@@ -173,6 +175,25 @@ public:
 		return IsXPBDAreaSpringStiffnessEnabled(PropertyCollection, false);
 	}
 
+	FXPBDAreaSpringConstraints(
+		const FSolverParticles& Particles,
+		int32 ParticleOffset,
+		int32 ParticleCount,
+		const TArray<TVec3<int32>>& InConstraints,
+		const TMap<FString, TConstArrayView<FRealSingle>>& WeightMaps,
+		const FCollectionPropertyConstFacade& PropertyCollection,
+		bool bTrimKinematicConstraints)
+		: FXPBDAxialSpringConstraints(
+			Particles,
+			ParticleOffset,
+			ParticleCount,
+			InConstraints,
+			WeightMaps.FindRef(GetXPBDAreaSpringStiffnessString(PropertyCollection, XPBDAreaSpringStiffnessName.ToString())),
+			FSolverVec2(GetWeightedFloatXPBDAreaSpringStiffness(PropertyCollection, MaxStiffness)),
+			bTrimKinematicConstraints)
+	{}
+
+	UE_DEPRECATED(5.3, "Use weight map constructor instead.")
 	FXPBDAreaSpringConstraints(
 		const FSolverParticles& Particles,
 		int32 ParticleOffset,
@@ -193,15 +214,43 @@ public:
 
 	virtual ~FXPBDAreaSpringConstraints() override = default;
 
-	void SetProperties(const FCollectionPropertyConstFacade& PropertyCollection)
+	void SetProperties(
+		const FCollectionPropertyConstFacade& PropertyCollection,
+		const TMap<FString, TConstArrayView<FRealSingle>>& WeightMaps)
 	{
 		if (IsXPBDAreaSpringStiffnessMutable(PropertyCollection))
 		{
-			Stiffness.SetWeightedValue(FSolverVec2(GetWeightedFloatXPBDAreaSpringStiffness(PropertyCollection)), MaxStiffness);
+			const FSolverVec2 WeightedValue(GetWeightedFloatXPBDAreaSpringStiffness(PropertyCollection));
+			if (IsXPBDAreaSpringStiffnessStringDirty(PropertyCollection))
+			{
+				const FString& WeightMapName = GetXPBDAreaSpringStiffnessString(PropertyCollection);
+				Stiffness = FPBDStiffness(
+					WeightedValue,
+					WeightMaps.FindRef(WeightMapName),
+					TConstArrayView<TVec3<int32>>(Constraints),
+					ParticleOffset,
+					ParticleCount,
+					FPBDStiffness::DefaultTableSize,
+					FPBDStiffness::DefaultParameterFitBase,
+					MaxStiffness);
+			}
+			else
+			{
+				Stiffness.SetWeightedValue(WeightedValue, MaxStiffness);
+			}
 		}
 	}
 
+	UE_DEPRECATED(5.3, "Use SetProperties(const FCollectionPropertyConstFacade&, const TMap<FString, TConstArrayView<FRealSingle>>&, FSolverReal) instead.")
+	void SetProperties(const FCollectionPropertyConstFacade& PropertyCollection)
+	{
+		SetProperties(PropertyCollection, TMap<FString, TConstArrayView<FRealSingle>>());
+	}
+
 private:
+	using FXPBDAxialSpringConstraints::Constraints;
+	using FXPBDAxialSpringConstraints::ParticleOffset;
+	using FXPBDAxialSpringConstraints::ParticleCount;
 	using FXPBDAxialSpringConstraints::Stiffness;
 
 	UE_CHAOS_DECLARE_PROPERTYCOLLECTION_NAME(XPBDAreaSpringStiffness, float);
diff --git a/Engine/Source/Runtime/Experimental/Chaos/Public/Chaos/XPBDBendingConstraints.h b/Engine/Source/Runtime/Experimental/Chaos/Public/Chaos/XPBDBendingConstraints.h
--- a/Engine/Source/Runtime/Experimental/Chaos/Public/Chaos/XPBDBendingConstraints.h
+++ b/Engine/Source/Runtime/Experimental/Chaos/Public/Chaos/XPBDBendingConstraints.h
@@ -31,8 +31,39 @@ public:
 	}
 
 	FXPBDBendingConstraints(const FSolverParticles& InParticles,
-		int32 ParticleOffset,
-		int32 ParticleCount,
+		int32 InParticleOffset,
+		int32 InParticleCount,
+		TArray<TVec4<int32>>&& InConstraints,
+		const TMap<FString, TConstArrayView<FRealSingle>>& WeightMaps,
+		const FCollectionPropertyConstFacade& PropertyCollection,
+		bool bTrimKinematicConstraints = false)
+		: Base(
+			InParticles,
+			InParticleOffset,
+			InParticleCount,
+			MoveTemp(InConstraints),
+			WeightMaps.FindRef(GetXPBDBendingElementStiffnessString(PropertyCollection, XPBDBendingElementStiffnessName.ToString())),
+			WeightMaps.FindRef(GetXPBDBucklingStiffnessString(PropertyCollection, XPBDBucklingStiffnessName.ToString())),
+			FSolverVec2(GetWeightedFloatXPBDBendingElementStiffness(PropertyCollection, MaxStiffness)),
+			(FSolverReal)GetXPBDBucklingRatio(PropertyCollection, 0.f),
+			FSolverVec2(GetWeightedFloatXPBDBucklingStiffness(PropertyCollection, MaxStiffness)),
+			bTrimKinematicConstraints,
+			MaxStiffness)
+		, DampingRatio(
+			FSolverVec2(GetWeightedFloatXPBDBendingElementDamping(PropertyCollection, MinDamping)).ClampAxes(MinDamping, MaxDamping),
+			WeightMaps.FindRef(GetXPBDBendingElementDampingString(PropertyCollection, XPBDBendingElementDampingName.ToString())),
+			TConstArrayView<TVec2<int32>>(ConstraintSharedEdges),
+			ParticleOffset,
+			ParticleCount)
+	{
+		Lambdas.Init((FSolverReal)0., Constraints.Num());
+		InitColor(InParticles);
+	}
+
+	UE_DEPRECATED(5.3, "Use weight map constructor instead.")
+	FXPBDBendingConstraints(const FSolverParticles& InParticles,
+		int32 InParticleOffset,
+		int32 InParticleCount,
 		TArray<TVec4<int32>>&& InConstraints,
 		const TConstArrayView<FRealSingle>& StiffnessMultipliers,
 		const TConstArrayView<FRealSingle>& BucklingStiffnessMultipliers,
@@ -41,8 +72,8 @@ public:
 		bool bTrimKinematicConstraints = false)
 		: Base(
 			InParticles,
-			ParticleOffset,
-			ParticleCount,
+			InParticleOffset,
+			InParticleCount,
 			MoveTemp(InConstraints),
 			StiffnessMultipliers,
 			BucklingStiffnessMultipliers,
@@ -59,13 +90,13 @@ public:
 			ParticleCount)
 	{
 		Lambdas.Init((FSolverReal)0., Constraints.Num());
-		InitColor(InParticles, ParticleOffset, ParticleCount);
+		InitColor(InParticles);
 	}
 
 	UE_DEPRECATED(5.2, "Use one the other constructors instead.")
 	FXPBDBendingConstraints(const FSolverParticles& InParticles,
-		int32 ParticleOffset,
-		int32 ParticleCount,
+		int32 InParticleOffset,
+		int32 InParticleCount,
 		TArray<TVec4<int32>>&& InConstraints,
 		const TConstArrayView<FRealSingle>& StiffnessMultipliers,
 		const TConstArrayView<FRealSingle>& BucklingStiffnessMultipliers,
@@ -75,8 +106,8 @@ public:
 		bool bTrimKinematicConstraints = false)
 		: Base(
 			InParticles,
-			ParticleOffset,
-			ParticleCount,
+			InParticleOffset,
+			InParticleCount,
 			MoveTemp(InConstraints),
 			StiffnessMultipliers,
 			BucklingStiffnessMultipliers,
@@ -88,7 +119,7 @@ public:
 		, DampingRatio(FSolverVec2::ZeroVector)
 	{
 		Lambdas.Init((FSolverReal)0., Constraints.Num());
-		InitColor(InParticles, ParticleOffset, ParticleCount);
+		InitColor(InParticles);
 	}
 
 	FXPBDBendingConstraints(const FSolverParticles& InParticles,
@@ -97,7 +128,7 @@ public:
 		TArray<TVec4<int32>>&& InConstraints,
 		const TConstArrayView<FRealSingle>& StiffnessMultipliers,
 		const TConstArrayView<FRealSingle>& BucklingStiffnessMultipliers,
-		const TConstArrayView<FRealSingle>& DampingMultipliers,
+		const TConstArrayView<FRealSingle>& InDampingMultipliers,
 		const FSolverVec2& InStiffness,
 		const FSolverReal InBucklingRatio,
 		const FSolverVec2& InBucklingStiffness,
@@ -117,13 +148,13 @@ public:
 			MaxStiffness)
 		, DampingRatio(
 			InDampingRatio.ClampAxes((FSolverReal)0., (FSolverReal)1.),
-			DampingMultipliers,
+			InDampingMultipliers,
 			TConstArrayView<TVec2<int32>>(ConstraintSharedEdges),
 			ParticleOffset,
 			ParticleCount)
 	{
 		Lambdas.Init((FSolverReal)0., Constraints.Num());
-		InitColor(InParticles, ParticleOffset, ParticleCount);
+		InitColor(InParticles);
 	}
 
 	virtual ~FXPBDBendingConstraints() override {}
@@ -135,24 +166,14 @@ public:
 		FPBDBendingConstraintsBase::Init(InParticles);
 	}
 
+	void SetProperties(
+		const FCollectionPropertyConstFacade& PropertyCollection,
+		const TMap<FString, TConstArrayView<FRealSingle>>& WeightMaps);
+
+	UE_DEPRECATED(5.3, "Use SetProperties(const FCollectionPropertyConstFacade&, const TMap<FString, TConstArrayView<FRealSingle>>&, FSolverReal) instead.")
 	void SetProperties(const FCollectionPropertyConstFacade& PropertyCollection)
 	{
-		if (IsXPBDBendingElementStiffnessMutable(PropertyCollection))
-		{
-			Stiffness.SetWeightedValue(FSolverVec2(GetWeightedFloatXPBDBendingElementStiffness(PropertyCollection)), MaxStiffness);
-		}
-		if (IsXPBDBucklingRatioMutable(PropertyCollection))
-		{
-			BucklingRatio = FMath::Clamp(GetXPBDBucklingRatio(PropertyCollection), (FSolverReal)0., (FSolverReal)1.);
-		}
-		if (IsXPBDBucklingStiffnessMutable(PropertyCollection))
-		{
-			BucklingStiffness.SetWeightedValue(FSolverVec2(GetWeightedFloatXPBDBucklingStiffness(PropertyCollection)), MaxStiffness);
-		}
-		if (IsXPBDBendingElementDampingMutable(PropertyCollection))
-		{
-			DampingRatio.SetWeightedValue(FSolverVec2(GetWeightedFloatXPBDBendingElementDamping(PropertyCollection)).ClampAxes(MinDamping, MaxDamping));
-		}
+		SetProperties(PropertyCollection, TMap<FString, TConstArrayView<FRealSingle>>());
 	}
 
 	void SetProperties(const FSolverVec2& InStiffness, const FSolverReal InBucklingRatio, const FSolverVec2& InBucklingStiffness, const FSolverVec2& InDampingRatio = FSolverVec2::ZeroVector)
@@ -174,16 +195,17 @@ public:
 	void Apply(FSolverParticles& Particles, const FSolverReal Dt) const;
 
 private:
-	void InitColor(const FSolverParticles& InParticles, const int32 ParticleOffset, const int32 ParticleCount);
+	void InitColor(const FSolverParticles& InParticles);
 	void ApplyHelper(FSolverParticles& Particles, const FSolverReal Dt, const int32 ConstraintIndex, const FSolverReal ExpStiffnessValue, const FSolverReal ExpBucklingValue, const FSolverReal DampingRatioValue) const;
 
 	using Base::Constraints;
+	using Base::ParticleOffset;
+	using Base::ParticleCount;
 	using Base::RestAngles;
 	using Base::Stiffness;
 	using Base::BucklingStiffness;
 
 	FPBDWeightMap DampingRatio;
-	TArray<FSolverReal> DampingMultipliers;
 	mutable TArray<FSolverReal> Lambdas;
 	TArray<int32> ConstraintsPerColorStartIndex; // Constraints are ordered so each batch is contiguous. This is ColorNum + 1 length so it can be used as start and end.
 
diff --git a/Engine/Source/Runtime/Experimental/Chaos/Public/Chaos/XPBDLongRangeConstraints.h b/Engine/Source/Runtime/Experimental/Chaos/Public/Chaos/XPBDLongRangeConstraints.h
--- a/Engine/Source/Runtime/Experimental/Chaos/Public/Chaos/XPBDLongRangeConstraints.h
+++ b/Engine/Source/Runtime/Experimental/Chaos/Public/Chaos/XPBDLongRangeConstraints.h
@@ -25,6 +25,35 @@ public:
 	static constexpr FSolverReal MinStiffness = (FSolverReal)1e-1;
 	static constexpr FSolverReal MaxStiffness = (FSolverReal)1e7;
 
+	FXPBDLongRangeConstraints(
+		const FSolverParticles& Particles,
+		const int32 InParticleOffset,
+		const int32 InParticleCount,
+		const TArray<TConstArrayView<TTuple<int32, int32, FRealSingle>>>& InTethers,
+		const TMap<FString, TConstArrayView<FRealSingle>>& WeightMaps,
+		const FCollectionPropertyConstFacade& PropertyCollection,
+		FSolverReal MeshScale)
+		: FPBDLongRangeConstraintsBase(
+			Particles,
+			InParticleOffset,
+			InParticleCount,
+			InTethers,
+			WeightMaps.FindRef(GetXPBDTetherStiffnessString(PropertyCollection, XPBDTetherStiffnessName.ToString())),
+			WeightMaps.FindRef(GetXPBDTetherScaleString(PropertyCollection, XPBDTetherScaleName.ToString())),
+			FSolverVec2(GetWeightedFloatXPBDTetherStiffness(PropertyCollection, MaxStiffness)),
+			FSolverVec2(GetWeightedFloatXPBDTetherScale(PropertyCollection, 1.f)),  // Scale clamping done in constructor
+			MaxStiffness,
+			MeshScale)
+	{
+		NumTethers = 0;
+		for (const TConstArrayView<FTether>& TetherBatch : Tethers)
+		{
+			NumTethers += TetherBatch.Num();
+		}
+		Lambdas.Reserve(NumTethers);
+	}
+
+	UE_DEPRECATED(5.3, "Use weight map constructor instead.")
 	FXPBDLongRangeConstraints(
 		const FSolverParticles& Particles,
 		const int32 InParticleOffset,
@@ -82,23 +111,56 @@ public:
 
 	virtual ~FXPBDLongRangeConstraints() override {}
 
-	void SetProperties(const FCollectionPropertyConstFacade& PropertyCollection, FSolverReal MeshScale)
+	void SetProperties(
+		const FCollectionPropertyConstFacade& PropertyCollection,
+		const TMap<FString, TConstArrayView<FRealSingle>>& WeightMaps,
+		FSolverReal MeshScale)
 	{
 		if (IsXPBDTetherStiffnessMutable(PropertyCollection))
 		{
-			Stiffness.SetWeightedValue(FSolverVec2(GetWeightedFloatXPBDTetherStiffness(PropertyCollection)));
+			const FSolverVec2 WeightedValue(GetWeightedFloatXPBDTetherStiffness(PropertyCollection));
+			if (IsXPBDTetherStiffnessStringDirty(PropertyCollection))
+			{
+				const FString& WeightMapName = GetXPBDTetherStiffnessString(PropertyCollection);
+				Stiffness = FPBDStiffness(
+					WeightedValue,
+					WeightMaps.FindRef(WeightMapName),
+					ParticleCount,
+					FPBDStiffness::DefaultTableSize,
+					FPBDStiffness::DefaultParameterFitBase,
+					MaxStiffness);
+			}
+			else
+			{
+				Stiffness.SetWeightedValue(WeightedValue, MaxStiffness);
+			}
 		}
 		if (IsXPBDTetherScaleMutable(PropertyCollection))
 		{
-			TetherScale.SetWeightedValue(FSolverVec2(GetWeightedFloatXPBDTetherScale(PropertyCollection)).ClampAxes(Base::MinTetherScale, Base::MaxTetherScale) * MeshScale);
+			const FSolverVec2 WeightedValue = FSolverVec2(GetWeightedFloatXPBDTetherScale(PropertyCollection)).ClampAxes(MinTetherScale, MaxTetherScale) * MeshScale;
+			if (IsXPBDTetherScaleStringDirty(PropertyCollection))
+			{
+				const FString& WeightMapName = GetXPBDTetherScaleString(PropertyCollection);
+				TetherScale = FPBDWeightMap(WeightedValue, WeightMaps.FindRef(WeightMapName), ParticleCount);
+			}
+			else
+			{
+				TetherScale.SetWeightedValue(WeightedValue);
+			}
 		}
 	}
 
+	UE_DEPRECATED(5.3, "Use SetProperties(const FCollectionPropertyConstFacade&, const TMap<FString, TConstArrayView<FRealSingle>>&, FSolverReal) instead.")
+	void SetProperties(const FCollectionPropertyConstFacade& PropertyCollection, FSolverReal MeshScale)
+	{
+		SetProperties(PropertyCollection, TMap<FString, TConstArrayView<FRealSingle>>(), MeshScale);
+	}
+
 	// Set the stiffness and scale values used by the constraint
 	void SetProperties(const FSolverVec2& InStiffness, const FSolverVec2& InTetherScale, FSolverReal MeshScale = (FSolverReal)1)
 	{
 		Stiffness.SetWeightedValue(InStiffness, MaxStiffness);
-		TetherScale.SetWeightedValue(InTetherScale.ClampAxes(Base::MinTetherScale, Base::MaxTetherScale) * MeshScale);
+		TetherScale.SetWeightedValue(InTetherScale.ClampAxes(MinTetherScale, MaxTetherScale) * MeshScale);
 	}
 
 	// Set stiffness offset and range, as well as the simulation stiffness exponent
@@ -210,7 +272,11 @@ private:
 		Lambda += DLambda;
 	}
 
+	using Base::MinTetherScale;
+	using Base::MaxTetherScale;
 	using Base::Tethers;
+	using Base::ParticleOffset;
+	using Base::ParticleCount;
 	using Base::Stiffness;
 	using Base::TetherScale;
 
diff --git a/Engine/Source/Runtime/Experimental/Chaos/Public/Chaos/XPBDSpringConstraints.h b/Engine/Source/Runtime/Experimental/Chaos/Public/Chaos/XPBDSpringConstraints.h
--- a/Engine/Source/Runtime/Experimental/Chaos/Public/Chaos/XPBDSpringConstraints.h
+++ b/Engine/Source/Runtime/Experimental/Chaos/Public/Chaos/XPBDSpringConstraints.h
@@ -47,7 +47,7 @@ public:
 		, DampingRatio(FSolverVec2::ZeroVector)
 	{
 		Lambdas.Init((FSolverReal)0., Constraints.Num());
-		InitColor(Particles, ParticleOffset, ParticleCount);
+		InitColor(Particles);
 	}
 
 	template<int32 Valence, TEMPLATE_REQUIRES(Valence >= 2 && Valence <= 4)>
@@ -78,7 +78,7 @@ public:
 			ParticleCount)
 	{
 		Lambdas.Init((FSolverReal)0., Constraints.Num());
-		InitColor(Particles, ParticleOffset, ParticleCount);
+		InitColor(Particles);
 	}
 
 	virtual ~FXPBDSpringConstraints() override {}
@@ -102,7 +102,7 @@ public:
 	void Apply(FSolverParticles& Particles, const FSolverReal Dt) const;
 
 private:
-	void InitColor(const FSolverParticles& InParticles, const int32 ParticleOffset, const int32 ParticleCount);
+	void InitColor(const FSolverParticles& InParticles);
 	void ApplyHelper(FSolverParticles& Particles, const FSolverReal Dt, const int32 ConstraintIndex, const FSolverReal StiffnessValue, const FSolverReal DampingRatioValue) const;
 
 	FSolverVec3 GetDelta(const FSolverParticles& Particles, const FSolverReal Dt, const int32 ConstraintIndex, const FSolverReal StiffnessValue, const FSolverReal DampingRatioValue) const
@@ -148,11 +148,13 @@ private:
 	}
 
 protected:
+	using Base::Constraints;
+	using Base::ParticleOffset;
+	using Base::ParticleCount;
 	using Base::Stiffness;
 	FPBDWeightMap DampingRatio;
 
 private:
-	using Base::Constraints;
 	using Base::Dists;
 	mutable TArray<FSolverReal> Lambdas;
 	TArray<int32> ConstraintsPerColorStartIndex; // Constraints are ordered so each batch is contiguous. This is ColorNum + 1 length so it can be used as start and end.
@@ -166,6 +168,27 @@ public:
 		return IsXPBDEdgeSpringStiffnessEnabled(PropertyCollection, false);
 	}
 
+	FXPBDEdgeSpringConstraints(
+		const FSolverParticles& Particles,
+		int32 ParticleOffset,
+		int32 ParticleCount,
+		const TArray<TVec3<int32>>& InConstraints,
+		const TMap<FString, TConstArrayView<FRealSingle>>& WeightMaps,
+		const FCollectionPropertyConstFacade& PropertyCollection,
+		bool bTrimKinematicConstraints = false)
+		: FXPBDSpringConstraints(
+			Particles,
+			ParticleOffset,
+			ParticleCount,
+			InConstraints,
+			WeightMaps.FindRef(GetXPBDEdgeSpringStiffnessString(PropertyCollection, XPBDEdgeSpringStiffnessName.ToString())),
+			WeightMaps.FindRef(GetXPBDEdgeSpringDampingString(PropertyCollection, XPBDEdgeSpringDampingName.ToString())),
+			FSolverVec2(GetWeightedFloatXPBDEdgeSpringStiffness(PropertyCollection, MaxStiffness)),
+			FSolverVec2(GetWeightedFloatXPBDEdgeSpringDamping(PropertyCollection, MinDampingRatio)),
+			bTrimKinematicConstraints)
+	{}
+
+	UE_DEPRECATED(5.3, "Use weight map constructor instead.")
 	FXPBDEdgeSpringConstraints(
 		const FSolverParticles& Particles,
 		int32 ParticleOffset,
@@ -189,19 +212,20 @@ public:
 
 	virtual ~FXPBDEdgeSpringConstraints() override = default;
 
+	void SetProperties(
+		const FCollectionPropertyConstFacade& PropertyCollection,
+		const TMap<FString, TConstArrayView<FRealSingle>>& WeightMaps);
+
+	UE_DEPRECATED(5.3, "Use SetProperties(const FCollectionPropertyConstFacade&, const TMap<FString, TConstArrayView<FRealSingle>>&, FSolverReal) instead.")
 	void SetProperties(const FCollectionPropertyConstFacade& PropertyCollection)
 	{
-		if (IsXPBDEdgeSpringStiffnessMutable(PropertyCollection))
-		{
-			Stiffness.SetWeightedValue(FSolverVec2(GetWeightedFloatXPBDEdgeSpringStiffness(PropertyCollection)), MaxStiffness);
-		}
-		if (IsXPBDEdgeSpringDampingMutable(PropertyCollection))
-		{
-			DampingRatio.SetWeightedValue(FSolverVec2(GetWeightedFloatXPBDEdgeSpringDamping(PropertyCollection)).ClampAxes(MinDampingRatio, MaxDampingRatio));
-		}
+		SetProperties(PropertyCollection, TMap<FString, TConstArrayView<FRealSingle>>());
 	}
 
 private:
+	using FXPBDSpringConstraints::Constraints;
+	using FXPBDSpringConstraints::ParticleOffset;
+	using FXPBDSpringConstraints::ParticleCount;
 	using FXPBDSpringConstraints::Stiffness;
 	using FXPBDSpringConstraints::DampingRatio;
 
@@ -217,6 +241,27 @@ public:
 		return IsXPBDBendingSpringStiffnessEnabled(PropertyCollection, false);
 	}
 
+	FXPBDBendingSpringConstraints(
+		const FSolverParticles& Particles,
+		int32 ParticleOffset,
+		int32 ParticleCount,
+		const TArray<TVec2<int32>>& InConstraints,
+		const TMap<FString, TConstArrayView<FRealSingle>>& WeightMaps,
+		const FCollectionPropertyConstFacade& PropertyCollection,
+		bool bTrimKinematicConstraints = false)
+		: FXPBDSpringConstraints(
+			Particles,
+			ParticleOffset,
+			ParticleCount,
+			InConstraints,
+			WeightMaps.FindRef(GetXPBDBendingSpringStiffnessString(PropertyCollection, XPBDBendingSpringStiffnessName.ToString())),
+			WeightMaps.FindRef(GetXPBDBendingSpringDampingString(PropertyCollection, XPBDBendingSpringDampingName.ToString())),
+			FSolverVec2(GetWeightedFloatXPBDBendingSpringStiffness(PropertyCollection, MaxStiffness)),
+			FSolverVec2(GetWeightedFloatXPBDBendingSpringDamping(PropertyCollection, MinDampingRatio)),
+			bTrimKinematicConstraints)
+	{}
+
+	UE_DEPRECATED(5.3, "Use weight map constructor instead.")
 	FXPBDBendingSpringConstraints(
 		const FSolverParticles& Particles,
 		int32 ParticleOffset,
@@ -240,19 +285,20 @@ public:
 
 	virtual ~FXPBDBendingSpringConstraints() override = default;
 
+	void SetProperties(
+		const FCollectionPropertyConstFacade& PropertyCollection,
+		const TMap<FString, TConstArrayView<FRealSingle>>& WeightMaps);
+
+	UE_DEPRECATED(5.3, "Use SetProperties(const FCollectionPropertyConstFacade&, const TMap<FString, TConstArrayView<FRealSingle>>&, FSolverReal) instead.")
 	void SetProperties(const FCollectionPropertyConstFacade& PropertyCollection)
 	{
-		if (IsXPBDBendingSpringStiffnessMutable(PropertyCollection))
-		{
-			Stiffness.SetWeightedValue(FSolverVec2(GetWeightedFloatXPBDBendingSpringStiffness(PropertyCollection)), MaxStiffness);
-		}
-		if (IsXPBDBendingSpringDampingMutable(PropertyCollection))
-		{
-			DampingRatio.SetWeightedValue(FSolverVec2(GetWeightedFloatXPBDBendingSpringDamping(PropertyCollection)).ClampAxes(MinDampingRatio, MaxDampingRatio));
-		}
+		SetProperties(PropertyCollection, TMap<FString, TConstArrayView<FRealSingle>>());
 	}
 
 private:
+	using FXPBDSpringConstraints::Constraints;
+	using FXPBDSpringConstraints::ParticleOffset;
+	using FXPBDSpringConstraints::ParticleCount;
 	using FXPBDSpringConstraints::Stiffness;
 	using FXPBDSpringConstraints::DampingRatio;
 
diff --git a/Engine/Source/Runtime/Experimental/Chaos/Public/Chaos/XPBDStretchBiasElementConstraints.h b/Engine/Source/Runtime/Experimental/Chaos/Public/Chaos/XPBDStretchBiasElementConstraints.h
--- a/Engine/Source/Runtime/Experimental/Chaos/Public/Chaos/XPBDStretchBiasElementConstraints.h
+++ b/Engine/Source/Runtime/Experimental/Chaos/Public/Chaos/XPBDStretchBiasElementConstraints.h
@@ -39,12 +39,7 @@ public:
 		int32 ParticleCount,
 		const FTriangleMesh& TriangleMesh,
 		const TArray<TVec3<FVec2f>>& FaceVertexUVs,
-		const TConstArrayView<FRealSingle>& StiffnessWarpMultipliers,
-		const TConstArrayView<FRealSingle>& StiffnessWeftMultipliers,
-		const TConstArrayView<FRealSingle>& StiffnessBiasMultipliers,
-		const TConstArrayView<FRealSingle>& DampingMultipliers,
-		const TConstArrayView<FRealSingle>& WarpScaleMultipliers,
-		const TConstArrayView<FRealSingle>& WeftScaleMultipliers,
+		const TMap<FString, TConstArrayView<FRealSingle>>& WeightMaps,
 		const FCollectionPropertyConstFacade& PropertyCollection,
 		bool bTrimKinematicConstraints = false);
 
@@ -76,33 +71,9 @@ public:
 		Lambdas.AddZeroed(Constraints.Num());
 	}
 
-	void SetProperties(const FCollectionPropertyConstFacade& PropertyCollection)
-	{
-		if (IsXPBDStretchBiasElementStiffnessWarpMutable(PropertyCollection))
-		{
-			StiffnessWarp.SetWeightedValue(FSolverVec2(GetWeightedFloatXPBDStretchBiasElementStiffnessWarp(PropertyCollection)), MaxStiffness);
-		}
-		if (IsXPBDStretchBiasElementStiffnessWeftMutable(PropertyCollection))
-		{
-			StiffnessWeft.SetWeightedValue(FSolverVec2(GetWeightedFloatXPBDStretchBiasElementStiffnessWeft(PropertyCollection)), MaxStiffness);
-		}
-		if (IsXPBDStretchBiasElementStiffnessBiasMutable(PropertyCollection))
-		{
-			StiffnessBias.SetWeightedValue(FSolverVec2(GetWeightedFloatXPBDStretchBiasElementStiffnessBias(PropertyCollection)), MaxStiffness);
-		}
-		if (IsXPBDStretchBiasElementDampingMutable(PropertyCollection))
-		{
-			DampingRatio.SetWeightedValue(FSolverVec2(GetWeightedFloatXPBDStretchBiasElementDamping(PropertyCollection)).ClampAxes(MinDamping, MaxDamping));
-		}
-		if (IsXPBDStretchBiasElementWarpScaleMutable(PropertyCollection))
-		{
-			WarpScale.SetWeightedValue(FSolverVec2(GetWeightedFloatXPBDStretchBiasElementWarpScale(PropertyCollection)).ClampAxes(MinWarpWeftScale, MaxWarpWeftScale));
-		}
-		if (IsXPBDStretchBiasElementWeftScaleMutable(PropertyCollection))
-		{
-			WeftScale.SetWeightedValue(FSolverVec2(GetWeightedFloatXPBDStretchBiasElementWeftScale(PropertyCollection)).ClampAxes(MinWarpWeftScale, MaxWarpWeftScale));
-		}
-	}
+	void SetProperties(
+		const FCollectionPropertyConstFacade& PropertyCollection,
+		const TMap<FString, TConstArrayView<FRealSingle>>& WeightMaps);
 
 	void SetProperties(const FSolverVec2& InStiffnessWarp, const FSolverVec2& InStiffnessWeft, const FSolverVec2& InStiffnessBias, const FSolverVec2& InDampingRatio, const FSolverVec2& InWarpScale, const FSolverVec2& InWeftScale)
 	{
@@ -130,11 +101,12 @@ public:
 private:
 	void InitConstraintsAndRestData(const FSolverParticles& InParticles, const FTriangleMesh& TriangleMesh,
 		const TArray<TVec3<FVec2f>>& FaceVertexUVs, const bool bUse3dRestLengths, const bool bTrimKinematicConstraints);
-	void InitColor(const FSolverParticles& InParticles, const int32 ParticleOffset, const int32 ParticleCount);
+	void InitColor(const FSolverParticles& InParticles);
 	void ApplyHelper(FSolverParticles& Particles, const FSolverReal Dt, const int32 ConstraintIndex, const FSolverVec3& ExpStiffnessValue, const FSolverReal DampingRatioValue, const FSolverReal WarpScaleValue, const FSolverReal WeftScaleValue) const;
 
 	TArray<TVec3<int32>> Constraints;
-
+	const int32 ParticleOffset;
+	const int32 ParticleCount;
 	FPBDStiffness StiffnessWarp;
 	FPBDStiffness StiffnessWeft;
 	FPBDStiffness StiffnessBias;
