Fix trailing events

diff --git a/Engine/Source/Runtime/Experimental/Chaos/Public/EventsData.h b/Engine/Source/Runtime/Experimental/Chaos/Public/EventsData.h
--- a/Engine/Source/Runtime/Experimental/Chaos/Public/EventsData.h
+++ b/Engine/Source/Runtime/Experimental/Chaos/Public/EventsData.h
@@ -254,7 +254,7 @@ namespace Chaos
 		}
 		else if constexpr (std::is_same_v<PayloadType, FTrailingEventData>)
 		{
-			return &Buffer->PhysicsProxyToTrailingIndices.PhysicsProxyToIndicesMap;
+			return nullptr; //&Buffer->PhysicsProxyToTrailingIndices.PhysicsProxyToIndicesMap;
 		}
 		else if constexpr (std::is_same_v<PayloadType, FRemovalEventData>)
 		{
