Add root breaking event to geometry collection
This bypass the normal event system and allow to detect when the top root cluster gets broken without the need to get notified of the subsequent breaks

diff --git a/Engine/Source/Runtime/Experimental/Chaos/Private/PhysicsProxy/ClusterUnionPhysicsProxy.cpp b/Engine/Source/Runtime/Experimental/Chaos/Private/PhysicsProxy/ClusterUnionPhysicsProxy.cpp
--- a/Engine/Source/Runtime/Experimental/Chaos/Private/PhysicsProxy/ClusterUnionPhysicsProxy.cpp
+++ b/Engine/Source/Runtime/Experimental/Chaos/Private/PhysicsProxy/ClusterUnionPhysicsProxy.cpp
@@ -190,18 +190,21 @@ namespace Chaos
 		Solver->EnqueueCommandImmediate(
 			[this, bIsAnchored]() mutable
 			{
-				if (Particle_Internal->IsAnchored() != bIsAnchored)
+				if (Particle_Internal)
 				{
-					if (FPBDRigidsEvolutionGBF* Evolution = GetEvolution(this))
+					if (Particle_Internal->IsAnchored() != bIsAnchored)
 					{
-						Particle_Internal->SetIsAnchored(bIsAnchored);
-						if (!bIsAnchored && !Particle_Internal->IsDynamic())
+						if (FPBDRigidsEvolutionGBF* Evolution = GetEvolution(this))
 						{
-							FKinematicTarget NoKinematicTarget;
-							Evolution->SetParticleObjectState(Particle_Internal, EObjectStateType::Dynamic);
-							Evolution->SetParticleKinematicTarget(Particle_Internal, NoKinematicTarget);
+							Particle_Internal->SetIsAnchored(bIsAnchored);
+							if (!bIsAnchored && !Particle_Internal->IsDynamic())
+							{
+								FKinematicTarget NoKinematicTarget;
+								Evolution->SetParticleObjectState(Particle_Internal, EObjectStateType::Dynamic);
+								Evolution->SetParticleKinematicTarget(Particle_Internal, NoKinematicTarget);
+							}
+							Evolution->GetParticles().MarkTransientDirtyParticle(Particle_Internal);
 						}
-						Evolution->GetParticles().MarkTransientDirtyParticle(Particle_Internal);
 					}
 				}
 			}
@@ -228,11 +231,14 @@ namespace Chaos
 		Solver->EnqueueCommandImmediate(
 			[this, State]() mutable
 			{
-				if (Particle_Internal->ObjectState() != State)
+				if (Particle_Internal)
 				{
-					if (FPBDRigidsEvolutionGBF* Evolution = GetEvolution(this))
+					if (Particle_Internal->ObjectState() != State)
 					{
-						Evolution->SetParticleObjectState(Particle_Internal, State);
+						if (FPBDRigidsEvolutionGBF* Evolution = GetEvolution(this))
+						{
+							Evolution->SetParticleObjectState(Particle_Internal, State);
+						}
 					}
 				}
 			}
diff --git a/Engine/Source/Runtime/Experimental/GeometryCollectionEngine/Private/GeometryCollection/GeometryCollectionComponent.cpp b/Engine/Source/Runtime/Experimental/GeometryCollectionEngine/Private/GeometryCollection/GeometryCollectionComponent.cpp
--- a/Engine/Source/Runtime/Experimental/GeometryCollectionEngine/Private/GeometryCollection/GeometryCollectionComponent.cpp
+++ b/Engine/Source/Runtime/Experimental/GeometryCollectionEngine/Private/GeometryCollection/GeometryCollectionComponent.cpp
@@ -121,10 +121,29 @@ FAutoConsoleVariableRef CVarMaxGeometryCollectionAsyncPhysicsTickIdleTimeMs(TEXT
 float GeometryCollectionRemovalMultiplier = 1.0f;
 FAutoConsoleVariableRef CVarGeometryCollectionRemovalTimerMultiplier(TEXT("p.Chaos.GC.RemovalTimerMultiplier"), GeometryCollectionRemovalMultiplier, TEXT("Multiplier for the removal time evaluation ( > 1 : faster removal , > 1 slower"));
 
+// temporary cvar , shoul dbe removed when the root event is no longer no ecessary
+bool GeometryCollectionEmitRootBreakingEvent = true;
+FAutoConsoleVariableRef CVarGeometryCollectionEmitRootBreakingEvent(TEXT("p.Chaos.GC.EmitRootBreakingEvent"), GeometryCollectionEmitRootBreakingEvent, TEXT("When true send a breaking event when root is breaking"));
+
 DEFINE_LOG_CATEGORY_STATIC(UGCC_LOG, Error, All);
 
 extern FGeometryCollectionDynamicDataPool GDynamicDataPool;
 
+static void GetGeometryCollectionComponentDebugDebugName(const UGeometryCollectionComponent& Comp, FString& DebugName)
+{
+	// Setup names
+	// Make the debug name for this geometry...
+	DebugName.Reset();
+
+#if WITH_EDITOR
+	if (Comp.GetOwner())
+	{
+		DebugName += FString::Printf(TEXT("Actor: '%s' "), *Comp.GetOwner()->GetActorLabel(false));
+	}
+#endif
+	DebugName += FString::Printf(TEXT("Component: '%s' "), *Comp.GetPathName());
+}
+
 FString NetModeToString(ENetMode InMode)
 {
 	switch(InMode)
@@ -2687,7 +2706,7 @@ void UGeometryCollectionComponent::RegisterAndInitializePhysicsProxy()
 	FSimulationParameters SimulationParameters;
 	{
 #if !(UE_BUILD_SHIPPING || UE_BUILD_TEST)
-		SimulationParameters.Name = GetPathName();
+		GetGeometryCollectionComponentDebugDebugName(*this, SimulationParameters.Name);
 #endif
 		EClusterConnectionTypeEnum ClusterCollectionType = ClusterConnectionType_DEPRECATED;
 		float ConnectionGraphBoundsFilteringMargin = 0;
@@ -2836,6 +2855,16 @@ void UGeometryCollectionComponent::OnPostPhysicsSync()
 	{
 		if (!PrimaryComponentTick.IsTickFunctionEnabled())
 		{ 
+			if (GeometryCollectionEmitRootBreakingEvent)
+			{
+				if (OnRootBreakEvent.IsBound())
+				{
+					FChaosBreakEvent Event;
+					Event.Index = GetRootIndex();
+					OnRootBreakEvent.Broadcast(Event);
+				}
+			}
+
 			PrimaryComponentTick.SetTickFunctionEnable(true);
 		}
 	}
diff --git a/Engine/Source/Runtime/Experimental/GeometryCollectionEngine/Public/GeometryCollection/GeometryCollectionComponent.h b/Engine/Source/Runtime/Experimental/GeometryCollectionEngine/Public/GeometryCollection/GeometryCollectionComponent.h
--- a/Engine/Source/Runtime/Experimental/GeometryCollectionEngine/Public/GeometryCollection/GeometryCollectionComponent.h
+++ b/Engine/Source/Runtime/Experimental/GeometryCollectionEngine/Public/GeometryCollection/GeometryCollectionComponent.h
@@ -973,6 +973,9 @@ public:
 	UPROPERTY(BlueprintAssignable, Category = "Chaos")
 	FOnChaosCrumblingEvent OnChaosCrumblingEvent;
 
+	// todo(chaos) remove when no longer necessary
+	FOnChaosBreakEvent OnRootBreakEvent;
+
 	void DispatchBreakEvent(const FChaosBreakEvent& Event);
 
 	void DispatchRemovalEvent(const FChaosRemovalEvent& Event);
