Support preallocation of ISMs in the ISMPool

diff --git a/Engine/Source/Runtime/Experimental/GeometryCollectionEngine/Private/GeometryCollection/GeometryCollectionISMPoolComponent.cpp b/Engine/Source/Runtime/Experimental/GeometryCollectionEngine/Private/GeometryCollection/GeometryCollectionISMPoolComponent.cpp
--- a/Engine/Source/Runtime/Experimental/GeometryCollectionEngine/Private/GeometryCollection/GeometryCollectionISMPoolComponent.cpp
+++ b/Engine/Source/Runtime/Experimental/GeometryCollectionEngine/Private/GeometryCollection/GeometryCollectionISMPoolComponent.cpp
@@ -52,25 +52,25 @@ void FGeometryCollectionMeshGroup::RemoveAllMeshes(FGeometryCollectionISMPool& I
 	Meshes.Empty();
 }
 
-FGeometryCollectionISM::FGeometryCollectionISM(AActor* OwmingActor, const FGeometryCollectionStaticMeshInstance& InMeshInstance)
+FGeometryCollectionISM::FGeometryCollectionISM(AActor* InOwningActor, const FGeometryCollectionStaticMeshInstance& InMeshInstance)
 {
 	MeshInstance = InMeshInstance;
 
 	check(MeshInstance.StaticMesh);
-	check(OwmingActor);
+	check(InOwningActor);
 
 	UHierarchicalInstancedStaticMeshComponent* HISMC = nullptr;
 	UInstancedStaticMeshComponent* ISMC = nullptr;
 	
 	if (MeshInstance.Desc.bUseHISM)
 	{
-		const FName ISMName = MakeUniqueObjectName(OwmingActor, UHierarchicalInstancedStaticMeshComponent::StaticClass(), MeshInstance.StaticMesh->GetFName());
-		ISMC = HISMC = NewObject<UHierarchicalInstancedStaticMeshComponent>(OwmingActor, ISMName, RF_Transient | RF_DuplicateTransient);
+		const FName ISMName = MakeUniqueObjectName(InOwningActor, UHierarchicalInstancedStaticMeshComponent::StaticClass(), MeshInstance.StaticMesh->GetFName());
+		ISMC = HISMC = NewObject<UHierarchicalInstancedStaticMeshComponent>(InOwningActor, ISMName, RF_Transient | RF_DuplicateTransient);
 	}
 	else
 	{
-		const FName ISMName = MakeUniqueObjectName(OwmingActor, UInstancedStaticMeshComponent::StaticClass(), MeshInstance.StaticMesh->GetFName());
-		ISMC = NewObject<UInstancedStaticMeshComponent>(OwmingActor, ISMName, RF_Transient | RF_DuplicateTransient);
+		const FName ISMName = MakeUniqueObjectName(InOwningActor, UInstancedStaticMeshComponent::StaticClass(), MeshInstance.StaticMesh->GetFName());
+		ISMC = NewObject<UInstancedStaticMeshComponent>(InOwningActor, ISMName, RF_Transient | RF_DuplicateTransient);
 	}
 
 	if (!ensure(ISMC != nullptr))
@@ -102,7 +102,7 @@ FGeometryCollectionISM::FGeometryCollectionISM(AActor* OwmingActor, const FGeome
 		HISMC->SetLODDistanceScale(MeshInstance.Desc.LodScale);
 	}
 	
-	OwmingActor->AddInstanceComponent(ISMC);
+	InOwningActor->AddInstanceComponent(ISMC);
 	ISMC->RegisterComponent();
 	ISMComponent = ISMC;
 }
@@ -148,32 +148,35 @@ FInstanceGroups::FInstanceGroupId FGeometryCollectionISM::AddInstanceGroup(int32
 	return InstanceGroupIndex;
 }
 
-FGeometryCollectionMeshInfo FGeometryCollectionISMPool::AddISM(UGeometryCollectionISMPoolComponent* OwningComponent, const FGeometryCollectionStaticMeshInstance& MeshInstance, int32 InstanceCount, TArrayView<const float> CustomDataFloats)
+FGeometryCollectionISMPool::FISMIndex FGeometryCollectionISMPool::AddISM(UGeometryCollectionISMPoolComponent* OwningComponent, const FGeometryCollectionStaticMeshInstance& MeshInstance)
 {
-	FGeometryCollectionMeshInfo Info;
+	FISMIndex* ISMIndexPtr = MeshToISMIndex.Find(MeshInstance);
+	if (ISMIndexPtr != nullptr)
+	{
+		return *ISMIndexPtr;
+	}
 
-	FISMIndex* ISMIndex = MeshToISMIndex.Find(MeshInstance);
-	if (!ISMIndex)
+	FISMIndex ISMIndex = INDEX_NONE;
+	if (FreeList.Num())
 	{
-		if (FreeList.Num())
-		{
-			// Take an ISM from the current FreeList instead of allocating a new slot.
-			Info.ISMIndex = FreeList.Last();
-			FreeList.RemoveAt(FreeList.Num() - 1);
-			ISMs[Info.ISMIndex] = FGeometryCollectionISM(OwningComponent->GetOwner(), MeshInstance);
-		}
-		else
-		{
-			Info.ISMIndex = ISMs.Emplace(OwningComponent->GetOwner(), MeshInstance);
-		}
-		MeshToISMIndex.Add(MeshInstance, Info.ISMIndex);
-		ISMComponentToISMIndex.Add(ISMs[Info.ISMIndex].ISMComponent, Info.ISMIndex);
+		// Take an ISM from the current FreeList instead of allocating a new slot.
+		ISMIndex = FreeList.Last();
+		FreeList.RemoveAt(FreeList.Num() - 1);
+		ISMs[ISMIndex] = FGeometryCollectionISM(OwningComponent->GetOwner(), MeshInstance);
 	}
 	else
 	{
-		Info.ISMIndex = *ISMIndex;
+		ISMIndex = ISMs.Emplace(OwningComponent->GetOwner(), MeshInstance);
 	}
-	// add to the ISM 
+	
+	MeshToISMIndex.Add(MeshInstance, ISMIndex);
+	return ISMIndex;
+}
+
+FGeometryCollectionMeshInfo FGeometryCollectionISMPool::AddISM(UGeometryCollectionISMPoolComponent* OwningComponent, const FGeometryCollectionStaticMeshInstance& MeshInstance, int32 InstanceCount, TArrayView<const float> CustomDataFloats)
+{
+	FGeometryCollectionMeshInfo Info;
+	Info.ISMIndex = AddISM(OwningComponent, MeshInstance);
 	Info.InstanceGroupIndex = ISMs[Info.ISMIndex].AddInstanceGroup(InstanceCount, CustomDataFloats);
 	return Info;
 }
@@ -247,25 +250,25 @@ void FGeometryCollectionISMPool::RemoveISM(const FGeometryCollectionMeshInfo& Me
 
 		ISM.InstanceGroups.RemoveGroup(MeshInfo.InstanceGroupIndex);
 	
-		if (ISM.InstanceGroups.IsEmpty() && ISM.ISMComponent->PerInstanceSMData.Num() == 0)
+		if (ISM.InstanceGroups.IsEmpty())
 		{
+			// No live instances, so take opportunity to reset indexing.
 			ISM.InstanceGroups.Reset();
 			ISM.InstanceIndexToRenderIndex.Reset();
 			ISM.RenderIndexToInstanceIndex.Reset();
+		}
+
+		if (GUseComponentFreeList && ISM.ISMComponent->PerInstanceSMData.Num() == 0)
+		{
+			// Remove component and push this ISM slot to the free list.
+			// todo: profile if it is better to push component into a free pool and recycle it.
+			ISM.ISMComponent->GetOwner()->RemoveInstanceComponent(ISM.ISMComponent);
+			ISM.ISMComponent->UnregisterComponent();
+			ISM.ISMComponent->DestroyComponent();
 			
-			if (GUseComponentFreeList)
-			{
-				// Remove component and push this ISM slot to the free list.
-				// todo: profile if it is better to push component into a free pool and recycle it.
-				ISM.ISMComponent->GetOwner()->RemoveInstanceComponent(ISM.ISMComponent);
-				ISM.ISMComponent->UnregisterComponent();
-				ISM.ISMComponent->DestroyComponent();
-			
-				MeshToISMIndex.Remove(ISM.MeshInstance);
-				ISMComponentToISMIndex.Remove(ISM.ISMComponent);
-				FreeList.Add(MeshInfo.ISMIndex);
-				ISM.ISMComponent = nullptr;
-			}
+			MeshToISMIndex.Remove(ISM.MeshInstance);
+			FreeList.Add(MeshInfo.ISMIndex);
+			ISM.ISMComponent = nullptr;
 		}
 	}
 }
@@ -273,7 +276,6 @@ void FGeometryCollectionISMPool::RemoveISM(const FGeometryCollectionMeshInfo& Me
 void FGeometryCollectionISMPool::Clear()
 {
 	MeshToISMIndex.Reset();
-	ISMComponentToISMIndex.Reset();
 	FreeList.Reset();
 	if (ISMs.Num() > 0)
 	{
@@ -331,6 +333,16 @@ bool UGeometryCollectionISMPoolComponent::BatchUpdateInstancesTransforms(FMeshGr
 	return false;
 }
 
+void UGeometryCollectionISMPoolComponent::PreallocateMeshInstance(const FGeometryCollectionStaticMeshInstance& MeshInstance)
+{
+	// If we are recycling components with a free list then we don't expect to have zero instance components.
+	// So don't do preallocation of components either in that case.
+	if (!GUseComponentFreeList)
+	{
+		Pool.AddISM(this, MeshInstance);
+	}
+}
+
 void UGeometryCollectionISMPoolComponent::GetResourceSizeEx(FResourceSizeEx& CumulativeResourceSize)
 {
 	Super::GetResourceSizeEx(CumulativeResourceSize);
@@ -338,7 +350,6 @@ void UGeometryCollectionISMPoolComponent::GetResourceSizeEx(FResourceSizeEx& Cum
 	int32 SizeBytes =
 		MeshGroups.GetAllocatedSize()
 		+ Pool.MeshToISMIndex.GetAllocatedSize()
-		+ Pool.ISMComponentToISMIndex.GetAllocatedSize()
 		+ Pool.ISMs.GetAllocatedSize()
 		+ Pool.FreeList.GetAllocatedSize();
 	
diff --git a/Engine/Source/Runtime/Experimental/GeometryCollectionEngine/Public/GeometryCollection/GeometryCollectionISMPoolComponent.h b/Engine/Source/Runtime/Experimental/GeometryCollectionEngine/Public/GeometryCollection/GeometryCollectionISMPoolComponent.h
--- a/Engine/Source/Runtime/Experimental/GeometryCollectionEngine/Public/GeometryCollection/GeometryCollectionISMPoolComponent.h
+++ b/Engine/Source/Runtime/Experimental/GeometryCollectionEngine/Public/GeometryCollection/GeometryCollectionISMPoolComponent.h
@@ -10,8 +10,8 @@
 #include "GeometryCollectionISMPoolComponent.generated.h"
 
 class AActor;
-class UInstancedStaticMeshComponent;
 class UGeometryCollectionISMPoolComponent;
+class UInstancedStaticMeshComponent;
 
 /** 
  * Structure containing a set of allocated instance ranges in an FGeometryCollectionISM which is the manager for a single ISM component.
@@ -236,6 +236,7 @@ struct FGeometryCollectionISMPool
 {
 	using FISMIndex = int32;
 
+	FISMIndex AddISM(UGeometryCollectionISMPoolComponent* OwningComponent, const FGeometryCollectionStaticMeshInstance& MeshInstance);
 	FGeometryCollectionMeshInfo AddISM(UGeometryCollectionISMPoolComponent* OwningComponent, const FGeometryCollectionStaticMeshInstance& MeshInstance, int32 InstanceCount, TArrayView<const float> CustomDataFloats);
 	bool BatchUpdateInstancesTransforms(FGeometryCollectionMeshInfo& MeshInfo, int32 StartInstanceIndex, const TArray<FTransform>& NewInstancesTransforms, bool bWorldSpace, bool bMarkRenderStateDirty, bool bTeleport);
 	void RemoveISM(const FGeometryCollectionMeshInfo& MeshInfo);
@@ -244,7 +245,6 @@ struct FGeometryCollectionISMPool
 	void Clear();
 
 	TMap<FGeometryCollectionStaticMeshInstance, FISMIndex> MeshToISMIndex;
-	TMap<UInstancedStaticMeshComponent*, FISMIndex> ISMComponentToISMIndex;
 	TArray<FGeometryCollectionISM> ISMs;
 	TArray<int32> FreeList;
 };
@@ -277,12 +277,18 @@ public:
 	/** destroy  a mesh group and its associated resources */
 	void DestroyMeshGroup(FMeshGroupId MeshGroupId);
 
-	/** Add a static mesh for a nmesh group */
+	/** Add a static mesh for a mesh group */
 	FMeshId AddMeshToGroup(FMeshGroupId MeshGroupId, const FGeometryCollectionStaticMeshInstance& MeshInstance, int32 InstanceCount, TArrayView<const float> CustomDataFloats);
 
 	/** Add a static mesh for a mesh group */
 	bool BatchUpdateInstancesTransforms(FMeshGroupId MeshGroupId, FMeshId MeshId, int32 StartInstanceIndex, const TArray<FTransform>& NewInstancesTransforms, bool bWorldSpace = false, bool bMarkRenderStateDirty = false, bool bTeleport = false);
 
+	/** 
+	 * Preallocate an ISM in the pool. 
+	 * Doing this early for known mesh instance descriptions can reduce the component registration cost of AddMeshToGroup() for newly discovered mesh descriptions.
+	 */
+	void PreallocateMeshInstance(const FGeometryCollectionStaticMeshInstance& MeshInstance);
+
 private:
 	uint32 NextMeshGroupId = 0;
 	TMap<FMeshGroupId, FGeometryCollectionMeshGroup> MeshGroups;
