Iris
* Add cull distance to world locations storage
* Now only UNetObjectGridFragmentLocFilter reads the cull distance from the fragment data.
* UNetObjectGridWorldLocFilter wil instead read it from FWorldLocations

#rnx

diff --git a/Engine/Plugins/Runtime/ReplicationSystemTestPlugin/Source/Private/Tests/ReplicationSystem/TestObjectReferences.cpp b/Engine/Plugins/Runtime/ReplicationSystemTestPlugin/Source/Private/Tests/ReplicationSystem/TestObjectReferences.cpp
--- a/Engine/Plugins/Runtime/ReplicationSystemTestPlugin/Source/Private/Tests/ReplicationSystem/TestObjectReferences.cpp
+++ b/Engine/Plugins/Runtime/ReplicationSystemTestPlugin/Source/Private/Tests/ReplicationSystem/TestObjectReferences.cpp
@@ -951,7 +951,9 @@ UE_NET_TEST_FIXTURE(FReplicationSystemServerClientTestFixture, TestOutOfOrderRes
 	Item.ObjectRef = ServerReferencedObjectC;
 	ServerFastArray.Items.Add(Item);
 	ServerFastArray.MarkItemDirty(ServerFastArray.Items[2]);
+#if !UE_BUILD_SHIPPING
 	UE::Net::IrisDebugHelper::DebugOutputNetObjectState(ServerObject->NetRefHandle.GetId(), ServerObject->NetRefHandle.GetReplicationSystemId());
+#endif
 
 	// Make sure that references are not replicated to client
 	Server->ReplicationSystem->AddToGroup(NotReplicatedNetObjectGroupHandle, ServerReferencedObjectA->NetRefHandle);
@@ -985,7 +987,9 @@ UE_NET_TEST_FIXTURE(FReplicationSystemServerClientTestFixture, TestOutOfOrderRes
 	Server->SendAndDeliverTo(Client, true);
 	Server->PostSendUpdate();
 
+#if !UE_BUILD_SHIPPING
 	UE::Net::IrisDebugHelper::DebugOutputNetObjectState(ClientObject->NetRefHandle.GetId(), ClientObject->NetRefHandle.GetReplicationSystemId());
+#endif
 
 	// Verify that we managed to resolve objectA even though we have other unresolved objects in the array
 	UE_NET_ASSERT_TRUE(ClientFastArray.Items[0].ObjectRef != nullptr);
diff --git a/Engine/Source/Runtime/Engine/Private/Net/Iris/ReplicationSystem/ActorReplicationBridge.cpp b/Engine/Source/Runtime/Engine/Private/Net/Iris/ReplicationSystem/ActorReplicationBridge.cpp
--- a/Engine/Source/Runtime/Engine/Private/Net/Iris/ReplicationSystem/ActorReplicationBridge.cpp
+++ b/Engine/Source/Runtime/Engine/Private/Net/Iris/ReplicationSystem/ActorReplicationBridge.cpp
@@ -85,14 +85,13 @@ void ActorReplicationBridgePreUpdateFunction(FNetRefHandle Handle, UObject* Inst
 	}
 }
 
-FVector ActorReplicationBridgeGetActorWorldLocation(FNetRefHandle Handle, const UObject* Instance)
+void ActorReplicationBridgeGetActorWorldObjectInfo(FNetRefHandle Handle, const UObject* Instance, FVector& OutWorldLocation, float& OutCullDistance)
 {
 	if (const AActor* Actor = Cast<AActor>(Instance))
 	{
-		return Actor->GetActorLocation();
+		OutWorldLocation = Actor->GetActorLocation();
+		OutCullDistance = Actor->NetCullDistanceSquared > 0.0f ? FMath::Sqrt(Actor->NetCullDistanceSquared) : 0.0f;
 	}
-
-	return FVector::Zero();
 }
 
 }
@@ -105,7 +104,7 @@ UActorReplicationBridge::UActorReplicationBridge()
 , SpawnInfoFlags(0U)
 {
 	SetInstancePreUpdateFunction(UE::Net::Private::ActorReplicationBridgePreUpdateFunction);
-	SetInstanceGetWorldLocationFunction(UE::Net::Private::ActorReplicationBridgeGetActorWorldLocation);
+	SetInstanceGetWorldObjectInfoFunction(UE::Net::Private::ActorReplicationBridgeGetActorWorldObjectInfo);
 }
 
 void UActorReplicationBridge::Initialize(UReplicationSystem* InReplicationSystem)
diff --git a/Engine/Source/Runtime/Experimental/Iris/Core/Private/Iris/ReplicationSystem/Filtering/NetObjectGridFilter.cpp b/Engine/Source/Runtime/Experimental/Iris/Core/Private/Iris/ReplicationSystem/Filtering/NetObjectGridFilter.cpp
--- a/Engine/Source/Runtime/Experimental/Iris/Core/Private/Iris/ReplicationSystem/Filtering/NetObjectGridFilter.cpp
+++ b/Engine/Source/Runtime/Experimental/Iris/Core/Private/Iris/ReplicationSystem/Filtering/NetObjectGridFilter.cpp
@@ -33,10 +33,8 @@ void UNetObjectGridFilter::RemoveConnection(uint32 ConnectionId)
 
 bool UNetObjectGridFilter::AddObject(uint32 ObjectIndex, FNetObjectFilterAddObjectParams& Params)
 {
-	
 	// We support either a world location in the state, tagged with RepTag_WorldLocation, or via the WorldLocations instance.
-	UE::Net::FRepTagFindInfo WorldLocationTagInfo;
-	const bool bHasLocation = BuildLocationTagInfo(WorldLocationTagInfo, ObjectIndex, Params);
+	const bool bHasLocation = BuildObjectInfo(ObjectIndex, Params);
 
 	if (!bHasLocation)
 	{
@@ -44,28 +42,13 @@ bool UNetObjectGridFilter::AddObject(uint32 ObjectIndex, FNetObjectFilterAddObje
 	}
 
 	FObjectLocationInfo& ObjectLocationInfo = static_cast<FObjectLocationInfo&>(Params.OutInfo);
-	ObjectLocationInfo.SetLocationStateOffset(WorldLocationTagInfo.ExternalStateOffset);
-	ObjectLocationInfo.SetLocationStateIndex(WorldLocationTagInfo.StateIndex);
-
+	
 	const uint32 InfoIndex = AllocObjectInfo();
 	ObjectLocationInfo.SetInfoIndex(InfoIndex);
 
 	FPerObjectInfo& PerObjectInfo = ObjectInfos[InfoIndex];
 	PerObjectInfo.ObjectIndex = ObjectIndex;
 
-	// NetCullDistanceSqr is optional. 
-	{
-		UE::Net::FRepTagFindInfo NetCullDistanceSqrTagInfo;
-		if (UE::Net::FindRepTag(Params.Protocol, UE::Net::RepTag_CullDistanceSqr, NetCullDistanceSqrTagInfo))
-		{
-			if ((NetCullDistanceSqrTagInfo.ExternalStateOffset < MAX_uint16) && (NetCullDistanceSqrTagInfo.StateIndex < MAX_uint16))
-			{
-				PerObjectInfo.CullDistanceSqrStateIndex = NetCullDistanceSqrTagInfo.StateIndex;
-				PerObjectInfo.CullDistanceSqrStateOffset = NetCullDistanceSqrTagInfo.ExternalStateOffset;
-			}
-		}
-	}
-
 	AddCellInfoForObject(ObjectLocationInfo, Params.InstanceProtocol);
 	
 	if (PerObjectInfo.CullDistance > Config->MaxCullDistance)
@@ -84,6 +67,8 @@ void UNetObjectGridFilter::RemoveObject(uint32 ObjectIndex, const FNetObjectFilt
 	const FObjectLocationInfo& ObjectLocationInfo = static_cast<const FObjectLocationInfo&>(Info);
 	RemoveCellInfoForObject(ObjectLocationInfo);
 
+	OnObjectRemoved(ObjectIndex);
+
 	const uint32 InfoIndex = ObjectLocationInfo.GetInfoIndex();
 	FreeObjectInfo(InfoIndex);
 }
@@ -97,7 +82,7 @@ void UNetObjectGridFilter::UpdateObjects(FNetObjectFilterUpdateParams& Params)
 		const uint32 ObjectIndex = Params.ObjectIndices[ObjectIt];
 
 		const FObjectLocationInfo& ObjectLocationInfo = static_cast<const FObjectLocationInfo&>(Params.FilteringInfos[ObjectIndex]);
-		const UE::Net::FReplicationInstanceProtocol* InstanceProtocol = Params.InstanceProtocols[ObjectIt];
+		const UE::Net::FReplicationInstanceProtocol* InstanceProtocol = Params.InstanceProtocols ? Params.InstanceProtocols[ObjectIt] : nullptr;
 		UpdateCellInfoForObject(ObjectLocationInfo, InstanceProtocol);
 	}
 }
@@ -362,9 +347,9 @@ void UNetObjectGridFilter::UpdateCellInfoForObject(const FObjectLocationInfo& Ob
 
 void UNetObjectGridFilter::UpdatePositionAndCullDistance(const UNetObjectGridFilter::FObjectLocationInfo& ObjectLocationInfo, UNetObjectGridFilter::FPerObjectInfo& PerObjectInfo, const UE::Net::FReplicationInstanceProtocol* InstanceProtocol)
 {
-	// Update position
+	// Update position & culldistance
 	{
-		UpdateObjectPosition(PerObjectInfo, ObjectLocationInfo, InstanceProtocol);
+		UpdateObjectInfo(PerObjectInfo, ObjectLocationInfo, InstanceProtocol);
 	}
 
 	// Optionally update cull distance
@@ -373,15 +358,6 @@ void UNetObjectGridFilter::UpdatePositionAndCullDistance(const UNetObjectGridFil
 		const float CullDistanceSqr = NetCullDistanceOverrides->GetCullDistanceSqr(PerObjectInfo.ObjectIndex);
 		PerObjectInfo.CullDistance = FPlatformMath::Sqrt(CullDistanceSqr);
 	}
-	else if (PerObjectInfo.CullDistanceSqrStateOffset != InvalidStateOffset)
-	{
-		TArrayView<const UE::Net::FReplicationInstanceProtocol::FFragmentData> FragmentDatas = MakeArrayView(InstanceProtocol->FragmentData, InstanceProtocol->FragmentCount);
-
-		const UE::Net::FReplicationInstanceProtocol::FFragmentData& FragmentData = FragmentDatas[PerObjectInfo.CullDistanceSqrStateIndex];
-		const uint8* CullDistanceSqrAddress = FragmentData.ExternalSrcBuffer + PerObjectInfo.CullDistanceSqrStateOffset;
-		const float CullDistanceSqr = *reinterpret_cast<const float*>(CullDistanceSqrAddress);
-		PerObjectInfo.CullDistance = FPlatformMath::Sqrt(CullDistanceSqr);
-	}
 }
 
 void UNetObjectGridFilter::CalculateCellBox(const UNetObjectGridFilter::FPerObjectInfo& PerObjectInfo, UNetObjectGridFilter::FCellBox& OutCellBox)
@@ -422,7 +398,9 @@ bool UNetObjectGridFilter::DoesCellContainCoord(const UNetObjectGridFilter::FCel
 	return (Coord.X >= Cell.MinX) & (Coord.X <= Cell.MaxX) & (Coord.Y >= Cell.MinY) & (Coord.Y <= Cell.MaxY);
 }
 
+//*************************************************************************************************
 // UNetObjectGridWorldLocFilter
+//*************************************************************************************************
 
 void UNetObjectGridWorldLocFilter::Init(FNetObjectFilterInitParams& Params)
 {
@@ -431,39 +409,95 @@ void UNetObjectGridWorldLocFilter::Init(FNetObjectFilterInitParams& Params)
 	WorldLocations = &Params.ReplicationSystem->GetWorldLocations();
 }
 
-void UNetObjectGridWorldLocFilter::UpdateObjectPosition(UNetObjectGridFilter::FPerObjectInfo& PerObjectInfo, const UNetObjectGridFilter::FObjectLocationInfo& ObjectLocationInfo, const UE::Net::FReplicationInstanceProtocol* InstanceProtocol)
+void UNetObjectGridWorldLocFilter::UpdateObjectInfo(UNetObjectGridFilter::FPerObjectInfo& PerObjectInfo, const UNetObjectGridFilter::FObjectLocationInfo& ObjectLocationInfo, const UE::Net::FReplicationInstanceProtocol* InstanceProtocol)
 {
 	check(ObjectLocationInfo.IsUsingWorldLocations());
+
+	const UE::Net::FWorldLocations::FObjectInfo& CachedObjectInfo = WorldLocations->GetObjectInfo(PerObjectInfo.ObjectIndex);
 	
-	PerObjectInfo.Position = WorldLocations->GetWorldLocation(PerObjectInfo.ObjectIndex);
+	PerObjectInfo.Position = CachedObjectInfo.WorldLocation;
+	PerObjectInfo.CullDistance = CachedObjectInfo.CullDistance;
 }
 
-bool UNetObjectGridWorldLocFilter::BuildLocationTagInfo(UE::Net::FRepTagFindInfo& OutWorldLocationTagInfo, uint32 ObjectIndex, FNetObjectFilterAddObjectParams&)
+bool UNetObjectGridWorldLocFilter::BuildObjectInfo(uint32 ObjectIndex, FNetObjectFilterAddObjectParams& Params)
 {
-	check(WorldLocations->HasWorldLocation(ObjectIndex));
+	checkf(WorldLocations->HasInfoForObject(ObjectIndex), TEXT("UNetObjectGridWorldLocFilter only supports objects with valid world location data."));
 
 	// Craft tag info that will let us know we need to retrieve the location from WorldLocations
-	OutWorldLocationTagInfo.StateIndex = InvalidStateIndex;
-	OutWorldLocationTagInfo.ExternalStateOffset = InvalidStateOffset;
-	
+	FObjectLocationInfo& ObjectLocationInfo = static_cast<FObjectLocationInfo&>(Params.OutInfo);
+	ObjectLocationInfo.SetLocationStateOffset(InvalidStateOffset);
+	ObjectLocationInfo.SetLocationStateIndex(InvalidStateIndex);
+
 	return true;
 }
 
+//*************************************************************************************************
 // UNetObjectGridFragmentLocFilter
+//*************************************************************************************************
 
-void UNetObjectGridFragmentLocFilter::UpdateObjectPosition(UNetObjectGridFilter::FPerObjectInfo& PerObjectInfo, const UNetObjectGridFilter::FObjectLocationInfo& ObjectLocationInfo, const UE::Net::FReplicationInstanceProtocol* InstanceProtocol)
+void UNetObjectGridFragmentLocFilter::Init(FNetObjectFilterInitParams& InitParams)
+{
+	Super::Init(InitParams);
+}
+
+bool UNetObjectGridFragmentLocFilter::BuildObjectInfo(uint32 ObjectIndex, FNetObjectFilterAddObjectParams& Params)
+{
+	UE::Net::FRepTagFindInfo WorldLocationTagInfo;
+
+	if (UE::Net::FindRepTag(Params.Protocol, UE::Net::RepTag_WorldLocation, WorldLocationTagInfo))
+	{
+		FObjectLocationInfo& ObjectLocationInfo = static_cast<FObjectLocationInfo&>(Params.OutInfo);
+		ObjectLocationInfo.SetLocationStateOffset(WorldLocationTagInfo.ExternalStateOffset);
+		ObjectLocationInfo.SetLocationStateIndex(WorldLocationTagInfo.StateIndex);
+
+
+		// NetCullDistanceSqr is optional. 
+		UE::Net::FRepTagFindInfo NetCullDistanceSqrTagInfo;
+		if (UE::Net::FindRepTag(Params.Protocol, UE::Net::RepTag_CullDistanceSqr, NetCullDistanceSqrTagInfo))
+		{
+			if ((NetCullDistanceSqrTagInfo.ExternalStateOffset < MAX_uint16) && (NetCullDistanceSqrTagInfo.StateIndex < MAX_uint16))
+			{
+				FCullDistanceFragmentInfo FragmentInfo;
+				FragmentInfo.CullDistanceSqrStateIndex = NetCullDistanceSqrTagInfo.StateIndex;
+				FragmentInfo.CullDistanceSqrStateOffset = NetCullDistanceSqrTagInfo.ExternalStateOffset;
+
+				CullDistanceFragments.Add(ObjectIndex, MoveTemp(FragmentInfo));
+			}
+		}
+
+		return true;
+	}
+
+	return false;
+}
+
+void UNetObjectGridFragmentLocFilter::OnObjectRemoved(uint32 ObjectIndex)
+{
+	CullDistanceFragments.Remove(ObjectIndex);
+}
+
+void UNetObjectGridFragmentLocFilter::UpdateObjectInfo(UNetObjectGridFilter::FPerObjectInfo& PerObjectInfo, const UNetObjectGridFilter::FObjectLocationInfo& ObjectLocationInfo, const UE::Net::FReplicationInstanceProtocol* InstanceProtocol)
 {
 	check(ObjectLocationInfo.IsUsingWorldLocations() == false);
+	check(InstanceProtocol);
 
 	TArrayView<const UE::Net::FReplicationInstanceProtocol::FFragmentData> FragmentDatas = MakeArrayView(InstanceProtocol->FragmentData, InstanceProtocol->FragmentCount);
 
-	const UE::Net::FReplicationInstanceProtocol::FFragmentData& FragmentData = FragmentDatas[ObjectLocationInfo.GetLocationStateIndex()];
-	const uint8* LocationAddress = FragmentData.ExternalSrcBuffer + ObjectLocationInfo.GetLocationStateOffset();
-	const float* LocationComponents = reinterpret_cast<const float*>(LocationAddress);
-	PerObjectInfo.Position = FVector(LocationComponents[0], LocationComponents[1], LocationComponents[2]);
+	// Update the location
+	{
+		const UE::Net::FReplicationInstanceProtocol::FFragmentData& FragmentData = FragmentDatas[ObjectLocationInfo.GetLocationStateIndex()];
+		const uint8* LocationAddress = FragmentData.ExternalSrcBuffer + ObjectLocationInfo.GetLocationStateOffset();
+		const float* LocationComponents = reinterpret_cast<const float*>(LocationAddress);
+		PerObjectInfo.Position = FVector(LocationComponents[0], LocationComponents[1], LocationComponents[2]);
+	}
+
+	// Update the culldistance
+	if( FCullDistanceFragmentInfo* CullDistanceFragmentInfo = CullDistanceFragments.Find(PerObjectInfo.ObjectIndex) )
+	{
+		const UE::Net::FReplicationInstanceProtocol::FFragmentData& FragmentData = FragmentDatas[CullDistanceFragmentInfo->CullDistanceSqrStateIndex];
+		const uint8* CullDistanceSqrAddress = FragmentData.ExternalSrcBuffer + CullDistanceFragmentInfo->CullDistanceSqrStateOffset;
+		const float CullDistanceSqr = *reinterpret_cast<const float*>(CullDistanceSqrAddress);
+		PerObjectInfo.CullDistance = FPlatformMath::Sqrt(CullDistanceSqr);
+	}
 }
 
-bool UNetObjectGridFragmentLocFilter::BuildLocationTagInfo(UE::Net::FRepTagFindInfo& OutWorldLocationTagInfo, uint32 ObjectIndex, FNetObjectFilterAddObjectParams& Params)
-{
-	return UE::Net::FindRepTag(Params.Protocol, UE::Net::RepTag_WorldLocation, OutWorldLocationTagInfo);
-}
\ No newline at end of file
diff --git a/Engine/Source/Runtime/Experimental/Iris/Core/Private/Iris/ReplicationSystem/ObjectReplicationBridge.cpp b/Engine/Source/Runtime/Experimental/Iris/Core/Private/Iris/ReplicationSystem/ObjectReplicationBridge.cpp
--- a/Engine/Source/Runtime/Experimental/Iris/Core/Private/Iris/ReplicationSystem/ObjectReplicationBridge.cpp
+++ b/Engine/Source/Runtime/Experimental/Iris/Core/Private/Iris/ReplicationSystem/ObjectReplicationBridge.cpp
@@ -64,10 +64,10 @@ static FAutoConsoleVariableRef CVarEnableFilterMappings(
 
 UObjectReplicationBridge::FCreateNetRefHandleParams UObjectReplicationBridge::DefaultCreateNetRefHandleParams =
 {
-	0U, // bCanReceive
-	0U, // bNeedsPreUpdate
-	0U, // bNeedsWorldLocationUpdate
-	1U, // bAllowDynamicFilter
+	false, // bCanReceive
+	false, // bNeedsPreUpdate
+	false, // bNeedsWorldLocationUpdate
+	true, // bAllowDynamicFilter
 	0.0f, // StaticPriority
 	0U, // PollFramePeriod
 };
@@ -301,7 +301,11 @@ UE::Net::FNetRefHandle UObjectReplicationBridge::BeginReplication(UObject* Insta
 
 			{
 				FWorldLocations& WorldLocations = ReplicationSystem->GetReplicationSystemInternal()->GetWorldLocations();
-				WorldLocations.SetHasWorldLocation(InternalReplicationIndex, Params.bNeedsWorldLocationUpdate);
+
+				if (Params.bNeedsWorldLocationUpdate)
+				{
+					WorldLocations.InitObjectInfoCache(InternalReplicationIndex);
+				}
 			}
 
 			// Set prioritizer
@@ -687,9 +691,9 @@ void UObjectReplicationBridge::SetInstancePreUpdateFunction(FInstancePreUpdateFu
 	PreUpdateInstanceFunction = InPreUpdateFunction;
 }
 
-void UObjectReplicationBridge::SetInstanceGetWorldLocationFunction(FInstanceGetWorldLocationFunction InGetWorldLocationFunction)
+void UObjectReplicationBridge::SetInstanceGetWorldObjectInfoFunction(FInstanceGetWorldObjectInfoFunction InGetWorldObjectInfoFunction)
 {
-	GetInstanceWorldLocationFunction = InGetWorldLocationFunction;
+	GetInstanceWorldObjectInfoFunction = InGetWorldObjectInfoFunction;
 }
 
 void UObjectReplicationBridge::PreUpdateAndPollImpl(FNetRefHandle Handle)
@@ -972,7 +976,7 @@ void UObjectReplicationBridge::UpdateInstancesWorldLocation()
 	using namespace UE::Net;
 	using namespace UE::Net::Private;
 
-	if (!GetInstanceWorldLocationFunction)
+	if (!GetInstanceWorldObjectInfoFunction)
 	{
 		return;
 	}
@@ -989,8 +993,10 @@ void UObjectReplicationBridge::UpdateInstancesWorldLocation()
 		const FNetRefHandleManager::FReplicatedObjectData& ObjectData = LocalNetRefHandleManager.GetReplicatedObjectDataNoCheck(InternalObjectIndex);
 		if (ObjectData.InstanceProtocol && EnumHasAnyFlags(ObjectData.InstanceProtocol->InstanceTraits, EReplicationInstanceProtocolTraits::NeedsWorldLocationUpdate))
 		{
-			FVector WorldLocation = (*GetInstanceWorldLocationFunction)(ObjectData.RefHandle, ReplicatedInstances[InternalObjectIndex]);
-			WorldLocations.SetWorldLocation(InternalObjectIndex, WorldLocation);
+			FWorldLocations::FObjectInfo CachedObjectInfo;
+			
+			GetInstanceWorldObjectInfoFunction(ObjectData.RefHandle, ReplicatedInstances[InternalObjectIndex], CachedObjectInfo.WorldLocation, CachedObjectInfo.CullDistance);
+			WorldLocations.SetObjectInfo(InternalObjectIndex, CachedObjectInfo);
 		}
 	};
 
diff --git a/Engine/Source/Runtime/Experimental/Iris/Core/Private/Iris/ReplicationSystem/Prioritization/LocationBasedNetObjectPrioritizer.cpp b/Engine/Source/Runtime/Experimental/Iris/Core/Private/Iris/ReplicationSystem/Prioritization/LocationBasedNetObjectPrioritizer.cpp
--- a/Engine/Source/Runtime/Experimental/Iris/Core/Private/Iris/ReplicationSystem/Prioritization/LocationBasedNetObjectPrioritizer.cpp
+++ b/Engine/Source/Runtime/Experimental/Iris/Core/Private/Iris/ReplicationSystem/Prioritization/LocationBasedNetObjectPrioritizer.cpp
@@ -24,7 +24,7 @@ bool ULocationBasedNetObjectPrioritizer::AddObject(uint32 ObjectIndex, FNetObjec
 	// We support either a world location in the state, tagged with RepTag_WorldLocation, or via the WorldLocations instance.
 	UE::Net::FRepTagFindInfo TagInfo;
 	bool bHasWorldLocation = false;
-	if (WorldLocations->HasWorldLocation(ObjectIndex))
+	if (WorldLocations->HasInfoForObject(ObjectIndex))
 	{
 		bHasWorldLocation = true;
 		// Craft tag info that will let us know we need to retrieve the location from WorldLocations
diff --git a/Engine/Source/Runtime/Experimental/Iris/Core/Private/Iris/ReplicationSystem/ReplicationBridge.cpp b/Engine/Source/Runtime/Experimental/Iris/Core/Private/Iris/ReplicationSystem/ReplicationBridge.cpp
--- a/Engine/Source/Runtime/Experimental/Iris/Core/Private/Iris/ReplicationSystem/ReplicationBridge.cpp
+++ b/Engine/Source/Runtime/Experimental/Iris/Core/Private/Iris/ReplicationSystem/ReplicationBridge.cpp
@@ -579,8 +579,7 @@ void UReplicationBridge::RemoveDestructionInfosForGroup(UE::Net::FNetObjectGroup
 
 		for (FInternalNetRefIndex InternalReplicationIndex : MakeArrayView(ObjectIndicesToRemove))
 		{
-			constexpr bool bHasWorldLocation = false;
-			WorldLocations.SetHasWorldLocation(InternalReplicationIndex, bHasWorldLocation);
+			WorldLocations.RemoveObjectInfoCache(InternalReplicationIndex);
 		}
 	}
 	else
@@ -594,8 +593,7 @@ void UReplicationBridge::RemoveDestructionInfosForGroup(UE::Net::FNetObjectGroup
 		// Remove from WorldLocations
 		for (const auto& It : StaticObjectsPendingDestroy)
 		{
-			constexpr bool bHasWorldLocation = false;
-			WorldLocations.SetHasWorldLocation(It.Value.InternalReplicationIndex, bHasWorldLocation);
+			WorldLocations.RemoveObjectInfoCache(It.Value.InternalReplicationIndex);
 		}
 
 		StaticObjectsPendingDestroy.Empty();
@@ -828,8 +826,8 @@ UE::Net::FNetRefHandle UReplicationBridge::InternalAddDestructionInfo(FNetRefHan
 	{
 		// Use WorldLocations to feed the location of the destruction info so that it can be prioritized properly.
 		FWorldLocations& WorldLocations = GetReplicationSystem()->GetReplicationSystemInternal()->GetWorldLocations();
-		WorldLocations.SetHasWorldLocation(InternalReplicationIndex, true);
-		WorldLocations.SetWorldLocation(InternalReplicationIndex, Parameters.Location);
+		WorldLocations.InitObjectInfoCache(InternalReplicationIndex);
+		WorldLocations.UpdateWorldLocation(InternalReplicationIndex, Parameters.Location);
 
 		GetReplicationSystem()->SetPrioritizer(DestructionInfoHandle, DefaultSpatialNetObjectPrioritizerHandle);
 	}
diff --git a/Engine/Source/Runtime/Experimental/Iris/Core/Private/Iris/ReplicationSystem/WorldLocations.cpp b/Engine/Source/Runtime/Experimental/Iris/Core/Private/Iris/ReplicationSystem/WorldLocations.cpp
--- a/Engine/Source/Runtime/Experimental/Iris/Core/Private/Iris/ReplicationSystem/WorldLocations.cpp
+++ b/Engine/Source/Runtime/Experimental/Iris/Core/Private/Iris/ReplicationSystem/WorldLocations.cpp
@@ -8,27 +8,44 @@ namespace UE::Net
 
 void FWorldLocations::Init(const FWorldLocationsInitParams& InitParams)
 {
-	ValidWorldLocations.Init(InitParams.MaxObjectCount);
+	ValidInfoIndexes.Init(InitParams.MaxObjectCount);
 }
 
-void FWorldLocations::SetHasWorldLocation(uint32 ObjectIndex, bool bHasWorldLocation)
+void FWorldLocations::InitObjectInfoCache(uint32 ObjectIndex)
 {
-	ValidWorldLocations.SetBitValue(ObjectIndex, bHasWorldLocation);
-	if (bHasWorldLocation)
+	if (ValidInfoIndexes.IsBitSet(ObjectIndex))
 	{
-		if (ObjectIndex >= uint32(WorldLocations.Num()))
-		{
-			LLM_SCOPE_BYTAG(Iris);
-			WorldLocations.Add(ObjectIndex + 1U - WorldLocations.Num());
-		}
-		WorldLocations[ObjectIndex] = FVector::Zero();
+		// Only init on first assignment
+		return;
 	}
+
+	ValidInfoIndexes.SetBit(ObjectIndex);
+	
+	if (ObjectIndex >= uint32(StoredObjectInfo.Num()))
+	{
+		LLM_SCOPE_BYTAG(Iris);
+		StoredObjectInfo.Add(ObjectIndex + 1U - StoredObjectInfo.Num());
+	}
+
+	StoredObjectInfo[ObjectIndex].WorldLocation = FVector::Zero();
+	StoredObjectInfo[ObjectIndex].CullDistance = 0.0f;
+}
+
+void FWorldLocations::RemoveObjectInfoCache(uint32 ObjectIndex)
+{
+	ValidInfoIndexes.ClearBit(ObjectIndex);
+}
+
+void FWorldLocations::SetObjectInfo(uint32 ObjectIndex, const FWorldLocations::FObjectInfo& ObjectInfo)
+{
+	checkSlow(ValidInfoIndexes.IsBitSet(ObjectIndex));
+	StoredObjectInfo[ObjectIndex] = ObjectInfo;
 }
 
-void FWorldLocations::SetWorldLocation(uint32 ObjectIndex, const FVector& WorldLocation)
+void FWorldLocations::UpdateWorldLocation(uint32 ObjectIndex, const FVector& WorldLocation)
 {
-	checkSlow(ValidWorldLocations.GetBit(ObjectIndex));
-	WorldLocations[ObjectIndex] = WorldLocation;
+	checkSlow(ValidInfoIndexes.GetBit(ObjectIndex));
+	StoredObjectInfo[ObjectIndex].WorldLocation = WorldLocation;
 }
 
 }
diff --git a/Engine/Source/Runtime/Experimental/Iris/Core/Public/Iris/ReplicationSystem/Filtering/NetObjectGridFilter.h b/Engine/Source/Runtime/Experimental/Iris/Core/Public/Iris/ReplicationSystem/Filtering/NetObjectGridFilter.h
--- a/Engine/Source/Runtime/Experimental/Iris/Core/Public/Iris/ReplicationSystem/Filtering/NetObjectGridFilter.h
+++ b/Engine/Source/Runtime/Experimental/Iris/Core/Public/Iris/ReplicationSystem/Filtering/NetObjectGridFilter.h
@@ -108,18 +108,20 @@ protected:
 	// We can't fit all info we need in 4x16bits.
 	struct FPerObjectInfo
 	{
-		FVector Position = { 0,0,0 };
+		FVector Position = { 0.f,0.f,0.f };
 		FCellBox CellBox = {};
-		float CullDistance = 0;
-		uint32 ObjectIndex = 0;
-		uint16 CullDistanceSqrStateIndex = InvalidStateIndex;
-		uint16 CullDistanceSqrStateOffset = InvalidStateOffset;
+		float CullDistance = 0.0f;
+		uint32 ObjectIndex = 0U;
 	};
 
-	/** Sets the current position of the object based on how we access it's given location */
-	virtual void UpdateObjectPosition(FPerObjectInfo& PerObjectInfo, const FObjectLocationInfo& ObjectLocationInfo, const UE::Net::FReplicationInstanceProtocol* InstanceProtocol) {}
+	/** Sets the current position of the object based on how we access it's given location. */
+	virtual void UpdateObjectInfo(FPerObjectInfo& PerObjectInfo, const FObjectLocationInfo& ObjectLocationInfo, const UE::Net::FReplicationInstanceProtocol* InstanceProtocol) {}
+
+	/** Build the data needed to properly filter this object. */
+	virtual bool BuildObjectInfo(uint32 ObjectIndex, FNetObjectFilterAddObjectParams& Params) { return false; }
 
-	virtual bool BuildLocationTagInfo(UE::Net::FRepTagFindInfo& OutWorldLocationTagInfo, uint32 ObjectIndex, FNetObjectFilterAddObjectParams&) { return false; }
+	/** Callback when an object is removed from the grid. Useful to cleanup data tied to the object. */
+	virtual void OnObjectRemoved(uint32 ObjectIndex) {}
 
 private:
 
@@ -198,10 +200,8 @@ class UNetObjectGridWorldLocFilter : public UNetObjectGridFilter
 protected:
 
 	virtual void Init(FNetObjectFilterInitParams&) override;
-
-	virtual void UpdateObjectPosition(FPerObjectInfo& PerObjectInfo, const UNetObjectGridFilter::FObjectLocationInfo& ObjectLocationInfo, const UE::Net::FReplicationInstanceProtocol* InstanceProtocol) override;
-
-	virtual bool BuildLocationTagInfo(UE::Net::FRepTagFindInfo& OutWorldLocationTagInfo, uint32 ObjectIndex, FNetObjectFilterAddObjectParams&) override;
+	virtual void UpdateObjectInfo(FPerObjectInfo& PerObjectInfo, const UNetObjectGridFilter::FObjectLocationInfo& ObjectLocationInfo, const UE::Net::FReplicationInstanceProtocol* InstanceProtocol) override;
+	virtual bool BuildObjectInfo(uint32 ObjectIndex, FNetObjectFilterAddObjectParams& Params) override;
 
 private:
 
@@ -220,9 +220,19 @@ class UNetObjectGridFragmentLocFilter : public UNetObjectGridFilter
 
 protected:
 
-	virtual void UpdateObjectPosition(FPerObjectInfo& PerObjectInfo, const UNetObjectGridFilter::FObjectLocationInfo& ObjectLocationInfo, const UE::Net::FReplicationInstanceProtocol* InstanceProtocol) override;
+	virtual void Init(FNetObjectFilterInitParams&) override;
+	virtual void UpdateObjectInfo(FPerObjectInfo& PerObjectInfo, const UNetObjectGridFilter::FObjectLocationInfo& ObjectLocationInfo, const UE::Net::FReplicationInstanceProtocol* InstanceProtocol) override;
+	virtual bool BuildObjectInfo(uint32 ObjectIndex, FNetObjectFilterAddObjectParams& Params) override;
+	virtual void OnObjectRemoved(uint32 ObjectIndex) override;
+
+private:
+	struct FCullDistanceFragmentInfo
+	{
+		uint16 CullDistanceSqrStateIndex = InvalidStateIndex;
+		uint16 CullDistanceSqrStateOffset = InvalidStateOffset;
+	};
 
-	virtual bool BuildLocationTagInfo(UE::Net::FRepTagFindInfo& OutWorldLocationTagInfo, uint32 ObjectIndex, FNetObjectFilterAddObjectParams&) override;
+	TMap<uint32, FCullDistanceFragmentInfo> CullDistanceFragments;
 };
 
 //
diff --git a/Engine/Source/Runtime/Experimental/Iris/Core/Public/Iris/ReplicationSystem/ObjectReplicationBridge.h b/Engine/Source/Runtime/Experimental/Iris/Core/Public/Iris/ReplicationSystem/ObjectReplicationBridge.h
--- a/Engine/Source/Runtime/Experimental/Iris/Core/Public/Iris/ReplicationSystem/ObjectReplicationBridge.h
+++ b/Engine/Source/Runtime/Experimental/Iris/Core/Public/Iris/ReplicationSystem/ObjectReplicationBridge.h
@@ -45,23 +45,24 @@ public:
 
 	struct FCreateNetRefHandleParams
 	{
-		uint32 bCanReceive : 1;
-		uint32 bNeedsPreUpdate : 1;
-		uint32 bNeedsWorldLocationUpdate : 1;
-		uint32 bAllowDynamicFilter : 1;
+		bool bCanReceive = false;
+		bool bNeedsPreUpdate = false;
+		bool bNeedsWorldLocationUpdate = false;
+		bool bAllowDynamicFilter = false;
 
 		/**
 		  * If StaticPriority is > 0 the ReplicationSystem will use that as priority when scheduling objects.
 		  * If it's <= 0.0f one will look for a world location support and then use the default spatial prioritizer.
 		  */
-		float StaticPriority;
+		float StaticPriority = 0.0f;
 
 		/**
 		  * How often the object should be polled for dirtiness, including calling the InstancePreUpdate function.
 		  * The period is zero based so 0 means poll every frame and 255 means poll every 256th frame.
 		  */
-		uint8 PollFramePeriod;
+		uint8 PollFramePeriod = 0;
 	};
+
 	IRISCORE_API static FCreateNetRefHandleParams DefaultCreateNetRefHandleParams;
 
 	IRISCORE_API UObjectReplicationBridge();
@@ -196,11 +197,11 @@ protected:
 	/** Set the function that we should call before copying state data. */
 	IRISCORE_API void SetInstancePreUpdateFunction(FInstancePreUpdateFunction InPreUpdateFunction);
 	
-	/** Helper method to get the world location for replicated instances with the HasWorldLocation trait. */
-	using FInstanceGetWorldLocationFunction = FVector(*)(FNetRefHandle, const UObject*);
+	/** Helper method to get the world location & cull distance for replicated instances with the HasWorldLocation trait. */
+	using FInstanceGetWorldObjectInfoFunction = TFunction<void(UE::Net::FNetRefHandle, const UObject*, FVector&, float&)>;
 
 	/** Set the function that we should call to get the world location of an object. */
-	IRISCORE_API void SetInstanceGetWorldLocationFunction(FInstanceGetWorldLocationFunction InGetWorldLocationFunction);
+	IRISCORE_API void SetInstanceGetWorldObjectInfoFunction(FInstanceGetWorldObjectInfoFunction InGetWorldObjectInfoFunction);
 	
 	// Poll frequency support
 
@@ -249,7 +250,8 @@ private:
 	bool ShouldClassBeDeltaCompressed(const UClass* Class);
 
 	FInstancePreUpdateFunction PreUpdateInstanceFunction;
-	FInstanceGetWorldLocationFunction GetInstanceWorldLocationFunction;
+	FInstanceGetWorldObjectInfoFunction GetInstanceWorldObjectInfoFunction;
+	
 
 	FName GetConfigClassPathName(const UClass* Class);
 
diff --git a/Engine/Source/Runtime/Experimental/Iris/Core/Public/Iris/ReplicationSystem/WorldLocations.h b/Engine/Source/Runtime/Experimental/Iris/Core/Public/Iris/ReplicationSystem/WorldLocations.h
--- a/Engine/Source/Runtime/Experimental/Iris/Core/Public/Iris/ReplicationSystem/WorldLocations.h
+++ b/Engine/Source/Runtime/Experimental/Iris/Core/Public/Iris/ReplicationSystem/WorldLocations.h
@@ -20,16 +20,36 @@ class FWorldLocations
 public:
 	void Init(const FWorldLocationsInitParams& InitParams);
 
-	/** Returns whether the object has a valid world location or not. */
-	bool HasWorldLocation(uint32 ObjectIndex) const;
+	/** Returns whether the object has a valid cached data or not. */
+	bool HasInfoForObject(uint32 ObjectIndex) const;
 
 	/** Returns the object's world location if it's valid or a zero vector if it's not. */
 	FVector GetWorldLocation(uint32 ObjectIndex) const;
 
+	/** Get the object's last cull distance we stored. */
+	float GetCullDistance(uint32 ObjectIndex) const;
+
+	/** Contains the cached object data we are storing. */
+	struct FObjectInfo
+	{
+		/** Absolute coordinate of the object */
+		FVector WorldLocation;
+		/** Network cull distance of the object */
+		float CullDistance;
+
+		FObjectInfo() : WorldLocation(EForceInit::ForceInitToZero), CullDistance(0.0f) {}
+	};
+
 public:
+
 	// Internal API
-	void SetHasWorldLocation(uint32 ObjectIndex, bool bHasWorldLocation);
-	void SetWorldLocation(uint32 ObjectIndex, const FVector& WorldLocation);
+	void UpdateWorldLocation(uint32 ObjectIndex, const FVector& WorldLocation);
+
+	void InitObjectInfoCache(uint32 ObjectIndex);
+	void RemoveObjectInfoCache(uint32 ObjectIndex);
+
+	void SetObjectInfo(uint32 ObjectIndex, const FObjectInfo& ObjectInfo);
+    const FObjectInfo& GetObjectInfo(uint32 ObjectIndex) const;
 
 private:
 	enum : uint32
@@ -37,20 +57,33 @@ private:
 		BytesPerLocationChunk = 65536U,
 	};
 
-	FNetBitArray ValidWorldLocations;
-	TChunkedArray<FVector, BytesPerLocationChunk> WorldLocations;
+	/** Set bits indicate that we have stored information for this internal object index */
+	FNetBitArray ValidInfoIndexes;
+
+	TChunkedArray<FObjectInfo, BytesPerLocationChunk> StoredObjectInfo;
 };
 
-inline bool FWorldLocations::HasWorldLocation(uint32 ObjectIndex) const
+inline bool FWorldLocations::HasInfoForObject(uint32 ObjectIndex) const
 {
-	return ValidWorldLocations.GetBit(ObjectIndex);
+	return ValidInfoIndexes.IsBitSet(ObjectIndex);
+}
+
+inline const FWorldLocations::FObjectInfo& FWorldLocations::GetObjectInfo(uint32 ObjectIndex) const
+{
+	if (ValidInfoIndexes.IsBitSet(ObjectIndex))
+	{
+		return StoredObjectInfo[ObjectIndex];
+	}
+
+	static FObjectInfo EmptyInfo;
+	return EmptyInfo;
 }
 
 inline FVector FWorldLocations::GetWorldLocation(uint32 ObjectIndex) const
 {
-	if (ValidWorldLocations.GetBit(ObjectIndex))
+	if (ValidInfoIndexes.IsBitSet(ObjectIndex))
 	{
-		return WorldLocations[ObjectIndex];
+		return StoredObjectInfo[ObjectIndex].WorldLocation;
 	}
 	else
 	{
@@ -58,5 +91,17 @@ inline FVector FWorldLocations::GetWorldLocation(uint32 ObjectIndex) const
 	}
 }
 
+inline float FWorldLocations::GetCullDistance(uint32 ObjectIndex) const
+{
+	if (ValidInfoIndexes.IsBitSet(ObjectIndex))
+	{
+		return StoredObjectInfo[ObjectIndex].CullDistance;
+	}
+	else
+	{
+		return 0.0f;
+	}
+}
+
 }
 
