Iris
- Fix for issue with unresolvable client->server rpc targets introduced with CL 24988126

diff --git a/Engine/Source/Runtime/Experimental/Iris/Core/Private/Iris/ReplicationSystem/ObjectReferenceCache.cpp b/Engine/Source/Runtime/Experimental/Iris/Core/Private/Iris/ReplicationSystem/ObjectReferenceCache.cpp
--- a/Engine/Source/Runtime/Experimental/Iris/Core/Private/Iris/ReplicationSystem/ObjectReferenceCache.cpp
+++ b/Engine/Source/Runtime/Experimental/Iris/Core/Private/Iris/ReplicationSystem/ObjectReferenceCache.cpp
@@ -1236,23 +1236,16 @@ void FObjectReferenceCache::ReadFullReferenceInternal(FNetSerializationContext&
 
 void FObjectReferenceCache::ReadFullReference(FNetSerializationContext& Context, FNetObjectReference& OutRef)
 {
-	// Normally all references are imported, unless we are doing this when reading exports
-	if (Context.GetInternalContext()->bInlineObjectReferenceExports == 0U)
-	{
-		ReadReference(Context, OutRef);
-		return;
-	}
-
-	FNetObjectReference ObjectRef;
 	FNetBitStreamReader* Reader = Context.GetBitStreamReader();
 
 	UE_NET_TRACE_SCOPE(FullNetObjectReference, *Reader, Context.GetTraceCollector(), ENetTraceVerbosity::Trace);
-	UE_NET_TRACE_NAMED_OBJECT_SCOPE(ReferenceScope, FNetRefHandle(), *Reader, Context.GetTraceCollector(), ENetTraceVerbosity::Verbose);
 	LLM_SCOPE_BYTAG(Iris);
 
-	// Handle client assigned reference
+	// Handle client assigned reference as they never should take the export path.
 	if (const bool bIsClientAssignedReference = Reader->ReadBool())
 	{
+		UE_NET_TRACE_SCOPE(ClientAssigned, *Reader, Context.GetTraceCollector(), ENetTraceVerbosity::Verbose);
+
 		const FNetRefHandle NetRefHandle = ReadNetRefHandle(Context);
 		FNetToken RelativePath = ReadNetToken(Reader);
 		ConditionalReadNetTokenData(Context, RelativePath);
@@ -1263,6 +1256,16 @@ void FObjectReferenceCache::ReadFullReference(FNetSerializationContext& Context,
 		return;
 	}
 
+	// Normally all references are imported, unless we are doing this when reading exports
+	if (Context.GetInternalContext()->bInlineObjectReferenceExports == 0U)
+	{
+		ReadReference(Context, OutRef);
+		return;
+	}
+
+	FNetObjectReference ObjectRef;
+
+	UE_NET_TRACE_NAMED_OBJECT_SCOPE(ReferenceScope, FNetRefHandle(), *Reader, Context.GetTraceCollector(), ENetTraceVerbosity::Verbose);
 	ReadFullReferenceInternal(Context, ObjectRef, 0U);
 
 	UE_NET_TRACE_SET_SCOPE_OBJECTID(ReferenceScope, ObjectRef.GetRefHandle());
@@ -1272,7 +1275,21 @@ void FObjectReferenceCache::ReadFullReference(FNetSerializationContext& Context,
 
 void FObjectReferenceCache::WriteFullReference(FNetSerializationContext& Context, FNetObjectReference Ref) const
 {
-	UE_NET_TRACE_SCOPE(FullNetObjectReference, *Context.GetBitStreamWriter(), Context.GetTraceCollector(), ENetTraceVerbosity::Trace);
+	FNetBitStreamWriter* Writer = Context.GetBitStreamWriter();
+
+	UE_NET_TRACE_SCOPE(FullNetObjectReference, *Writer, Context.GetTraceCollector(), ENetTraceVerbosity::Trace);
+
+	// Client assigned references are always written outside of the export path
+	const bool bIsClientAssignedReference = Ref.PathToken.IsValid();
+	if (Writer->WriteBool(bIsClientAssignedReference))
+	{
+		UE_NET_TRACE_SCOPE(ClientAssigned, *Writer, Context.GetTraceCollector(), ENetTraceVerbosity::Verbose);
+
+		WriteNetRefHandle(Context, Ref.GetRefHandle());
+		WriteNetToken(Writer, Ref.PathToken);
+		ConditionalWriteNetTokenData(Context, Context.GetExportContext(), Ref.PathToken);
+		return;
+	}
 
 	// If we do not inline reference exports we just write the reference
 	if (Context.GetInternalContext()->bInlineObjectReferenceExports == 0U)
@@ -1283,17 +1300,6 @@ void FObjectReferenceCache::WriteFullReference(FNetSerializationContext& Context
 		}
 		WriteReference(Context, Ref);
 		return;
-	}		
-
-	// Handle client assigned reference 
-	FNetBitStreamWriter* Writer = Context.GetBitStreamWriter();
-	const bool bIsClientAssignedReference = Ref.PathToken.IsValid();
-	if (Writer->WriteBool(bIsClientAssignedReference))
-	{
-		WriteNetRefHandle(Context, Ref.GetRefHandle());
-		WriteNetToken(Writer, Ref.PathToken);
-		ConditionalWriteNetTokenData(Context, Context.GetExportContext(), Ref.PathToken);
-		return;
 	}
 
 	WriteFullReferenceInternal(Context, Ref);
