Removed virtual assignment operator from ISpatialAcceleration and derived classes, replacing them with a virtual DeepAssign member function, in order to fix mismatched copy constructor/assignment warnings.

diff --git a/Engine/Source/Runtime/Experimental/Chaos/Private/Chaos/PBDRigidsEvolution.cpp b/Engine/Source/Runtime/Experimental/Chaos/Private/Chaos/PBDRigidsEvolution.cpp
--- a/Engine/Source/Runtime/Experimental/Chaos/Private/Chaos/PBDRigidsEvolution.cpp
+++ b/Engine/Source/Runtime/Experimental/Chaos/Private/Chaos/PBDRigidsEvolution.cpp
@@ -491,7 +491,7 @@ namespace Chaos
 					ISpatialAcceleration<FAccelerationStructureHandle, FReal, 3>* AccelerationSUBStructureCopy = GetSubStructureCopy(SpatialAccIdx);
 					if (AccelerationSUBStructureCopy)
 					{
-						*AccelerationSUBStructureCopy = *NewStruct; // Normal copy
+						AccelerationSUBStructureCopy->DeepAssign(*NewStruct); // Normal copy
 					}
 					if (AccelerationStructure->IsBucketActive(SpatialAccIdx.Bucket))
 					{
@@ -836,7 +836,7 @@ namespace Chaos
 				}
 
 				check(ExternalAccelerationSubStructure->IsTreeDynamic());
-				*ExternalAccelerationSubStructure = *AsyncInternalAccelerationSubStructure;
+				ExternalAccelerationSubStructure->DeepAssign(*AsyncInternalAccelerationSubStructure);
 			}
 			
 			InternalAcceleration->SwapSubstructure(*AsyncInternalAcceleration, SpatialIdx);
@@ -864,11 +864,11 @@ namespace Chaos
 
 			if (!CheckPristine || (ToSubStructure->ShouldRebuild() && !FromSubStructure->ShouldRebuild()))
 			{
-				*ToSubStructure = *FromSubStructure;
+				ToSubStructure->DeepAssign(*FromSubStructure);
 			}
 		}
 		// This is a pretty shallow copy: It won't copy anything that we copied above
-		*ToStructure = *FromStructure;
+		ToStructure->DeepAssign(*FromStructure);
 	}
 
 	void FPBDRigidsEvolutionBase::UpdateExternalAccelerationStructure_External(
diff --git a/Engine/Source/Runtime/Experimental/Chaos/Public/Chaos/AABBTree.h b/Engine/Source/Runtime/Experimental/Chaos/Public/Chaos/AABBTree.h
--- a/Engine/Source/Runtime/Experimental/Chaos/Public/Chaos/AABBTree.h
+++ b/Engine/Source/Runtime/Experimental/Chaos/Public/Chaos/AABBTree.h
@@ -1998,7 +1998,7 @@ public:
 
 		// Copy all the small objects first
 
-		ISpatialAcceleration<TPayloadType, T, 3>::operator=(From);
+		ISpatialAcceleration<TPayloadType, T, 3>::DeepAssign(From);
 
 		DirtyElementGridCellSize = From.DirtyElementGridCellSize;
 		DirtyElementGridCellSizeInv = From.DirtyElementGridCellSizeInv;
@@ -3544,16 +3544,15 @@ private:
 		}
 	}
 
-	virtual ISpatialAcceleration<TPayloadType, T, 3>& operator=(const ISpatialAcceleration<TPayloadType, T, 3>& Other) override
+	virtual void DeepAssign(const ISpatialAcceleration<TPayloadType, T, 3>& Other) override
 	{
-		
 		check(Other.GetType() == ESpatialAcceleration::AABBTree);
-		return operator=(static_cast<const TAABBTree<TPayloadType, TLeafType, bMutable, T>&>(Other));
+		*this = static_cast<const TAABBTree<TPayloadType, TLeafType, bMutable, T>&>(Other);
 	}
 
 	TAABBTree<TPayloadType,TLeafType, bMutable, T>& operator=(const TAABBTree<TPayloadType,TLeafType,bMutable, T>& Rhs)
 	{
-		ISpatialAcceleration<TPayloadType, T, 3>::operator=(Rhs);
+		ISpatialAcceleration<TPayloadType, T, 3>::DeepAssign(Rhs);
 		ensure(Rhs.WorkStack.Num() == 0);
 		//ensure(Rhs.WorkPool.Num() == 0);
 		//ensure(Rhs.WorkPoolFreeList.Num() == 0);
diff --git a/Engine/Source/Runtime/Experimental/Chaos/Public/Chaos/BoundingVolume.h b/Engine/Source/Runtime/Experimental/Chaos/Public/Chaos/BoundingVolume.h
--- a/Engine/Source/Runtime/Experimental/Chaos/Public/Chaos/BoundingVolume.h
+++ b/Engine/Source/Runtime/Experimental/Chaos/Public/Chaos/BoundingVolume.h
@@ -164,16 +164,16 @@ class TBoundingVolume final : public ISpatialAcceleration<InPayloadType, T,d>
 
 public:
 
-	virtual ISpatialAcceleration<TPayloadType, T, d>& operator=(const ISpatialAcceleration<TPayloadType, T, d>& Other) override
+	virtual void DeepAssign(const ISpatialAcceleration<TPayloadType, T, d>& Other) override
 	{
 
 		check(Other.GetType() == ESpatialAcceleration::BoundingVolume);
-		return operator=(static_cast<const TBoundingVolume<TPayloadType, T, d>&>(Other));
+		*this = static_cast<const TBoundingVolume<TPayloadType, T, d>&>(Other);
 	}
 
 	TBoundingVolume<TPayloadType, T, d>& operator=(const TBoundingVolume<TPayloadType, T, d>& Other)
 	{
-		ISpatialAcceleration<TPayloadType, FReal, 3>::operator=(Other);
+		ISpatialAcceleration<TPayloadType, FReal, 3>::DeepAssign(Other);
 		MGlobalPayloads = Other.MGlobalPayloads;
 		MGrid = Other.MGrid;
 		MElements = Other.MElements;
diff --git a/Engine/Source/Runtime/Experimental/Chaos/Public/Chaos/BoundingVolumeHierarchy.h b/Engine/Source/Runtime/Experimental/Chaos/Public/Chaos/BoundingVolumeHierarchy.h
--- a/Engine/Source/Runtime/Experimental/Chaos/Public/Chaos/BoundingVolumeHierarchy.h
+++ b/Engine/Source/Runtime/Experimental/Chaos/Public/Chaos/BoundingVolumeHierarchy.h
@@ -79,16 +79,16 @@ class TBoundingVolumeHierarchy final : public ISpatialAcceleration<int32, T,d>
 	{
 	}
 
-	CHAOS_API virtual ISpatialAcceleration<int32, T, d>& operator=(const ISpatialAcceleration<int32, T, d>& Other) override
+	CHAOS_API virtual void DeepAssign(const ISpatialAcceleration<int32, T, d>& Other) override
 	{
 
 		check(Other.GetType() == ESpatialAcceleration::AABBTreeBV);
-		return operator=(static_cast<const TBoundingVolumeHierarchy<OBJECT_ARRAY, LEAF_TYPE, T, d>&>(Other));
+		*this = static_cast<const TBoundingVolumeHierarchy<OBJECT_ARRAY, LEAF_TYPE, T, d>&>(Other);
 	}
 
 	CHAOS_API TBoundingVolumeHierarchy& operator=(const TBoundingVolumeHierarchy<OBJECT_ARRAY, LEAF_TYPE, T, d>& Other)
 	{
-		ISpatialAcceleration<int32, T, d>::operator=(Other);
+		ISpatialAcceleration<int32, T, d>::DeepAssign(Other);
 		MObjects = Other.MObjects;
 		MGlobalObjects = Other.MGlobalObjects;
 		MWorldSpaceBoxes = Other.MWorldSpaceBoxes;
diff --git a/Engine/Source/Runtime/Experimental/Chaos/Public/Chaos/ISpatialAcceleration.h b/Engine/Source/Runtime/Experimental/Chaos/Public/Chaos/ISpatialAcceleration.h
--- a/Engine/Source/Runtime/Experimental/Chaos/Public/Chaos/ISpatialAcceleration.h
+++ b/Engine/Source/Runtime/Experimental/Chaos/Public/Chaos/ISpatialAcceleration.h
@@ -255,6 +255,7 @@ template <typename TPayloadType, typename T, int d>
 class CHAOS_API ISpatialAcceleration
 {
 public:
+	UE_NONCOPYABLE(ISpatialAcceleration)
 
 	ISpatialAcceleration(SpatialAccelerationType InType = static_cast<SpatialAccelerationType>(ESpatialAcceleration::Unknown))
 		: Type(InType), SyncTimestamp(0), AsyncTimeSlicingComplete(true)
@@ -317,12 +318,11 @@ public:
 		return nullptr;
 	}
 
-	virtual ISpatialAcceleration<TPayloadType, T, d>& operator=(const ISpatialAcceleration<TPayloadType, T, d>& Other)
+	virtual void DeepAssign(const ISpatialAcceleration<TPayloadType, T, d>& Other)
 	{
 		Type = Other.Type;
 		SyncTimestamp = Other.SyncTimestamp;
 		AsyncTimeSlicingComplete = Other.AsyncTimeSlicingComplete;
-		return *this;
 	}
 
 #if !UE_BUILD_SHIPPING
diff --git a/Engine/Source/Runtime/Experimental/Chaos/Public/Chaos/ISpatialAccelerationCollection.h b/Engine/Source/Runtime/Experimental/Chaos/Public/Chaos/ISpatialAccelerationCollection.h
--- a/Engine/Source/Runtime/Experimental/Chaos/Public/Chaos/ISpatialAccelerationCollection.h
+++ b/Engine/Source/Runtime/Experimental/Chaos/Public/Chaos/ISpatialAccelerationCollection.h
@@ -22,6 +22,8 @@ template <typename TPayloadType, typename T, int d>
 class CHAOS_API ISpatialAccelerationCollection : public ISpatialAcceleration<TPayloadType, T, d>
 {
 public:
+	UE_NONCOPYABLE(ISpatialAccelerationCollection)
+
 	ISpatialAccelerationCollection()
 	: ISpatialAcceleration<TPayloadType, T, d>(StaticType)
 	, ActiveBucketsMask(0)
@@ -45,18 +47,16 @@ public:
 	bool IsAllAsyncTasksComplete() const { return AllAsyncTasksComplete; }
 	void SetAllAsyncTasksComplete(bool State) { AllAsyncTasksComplete = State; }
 
-	ISpatialAccelerationCollection<TPayloadType, FReal, 3>& operator=(const ISpatialAccelerationCollection<TPayloadType, FReal, 3>& Other)
+	void DeepAssign(const ISpatialAccelerationCollection<TPayloadType, FReal, 3>& Other)
 	{
-		ISpatialAcceleration<TPayloadType, FReal, 3>::operator=(Other);
+		ISpatialAcceleration<TPayloadType, FReal, 3>::DeepAssign(Other);
 		check(ActiveBucketsMask == Other.ActiveBucketsMask);
 		AllAsyncTasksComplete = Other.AllAsyncTasksComplete;
-		return *this;
 	}
 
-	virtual ISpatialAcceleration<TPayloadType, FReal, 3>& operator=(const ISpatialAcceleration<TPayloadType, FReal, 3>& Other) override
+	virtual void DeepAssign(const ISpatialAcceleration<TPayloadType, FReal, 3>& Other) override
 	{
 		check(false);	//not implemented
-		return *this;
 	}
 
 #if !UE_BUILD_SHIPPING
diff --git a/Engine/Source/Runtime/Experimental/Chaos/Public/Chaos/SpatialAccelerationCollection.h b/Engine/Source/Runtime/Experimental/Chaos/Public/Chaos/SpatialAccelerationCollection.h
--- a/Engine/Source/Runtime/Experimental/Chaos/Public/Chaos/SpatialAccelerationCollection.h
+++ b/Engine/Source/Runtime/Experimental/Chaos/Public/Chaos/SpatialAccelerationCollection.h
@@ -693,10 +693,9 @@ private:
 	}
 
 	TSpatialAccelerationCollection<TSpatialAccelerationTypes...>& operator=(const TSpatialAccelerationCollection<TSpatialAccelerationTypes...>& Other) = delete;
-	virtual ISpatialAcceleration<TPayloadType, FReal, 3>& operator=(const ISpatialAcceleration<TPayloadType, FReal, 3>& Other) override
+	virtual void DeepAssign(const ISpatialAcceleration<TPayloadType, FReal, 3>& Other) override
 	{
 		check(false);	//not implemented
-		return *this;
 	}
 	
 	static constexpr uint32 ClampedIdx(uint32 Idx)
