Chaos Physics: Distinguish between updating and adding thing in the acceleration structures

diff --git a/Engine/Source/Runtime/Engine/Private/PhysicsEngine/Experimental/PhysScene_Chaos.cpp b/Engine/Source/Runtime/Engine/Private/PhysicsEngine/Experimental/PhysScene_Chaos.cpp
--- a/Engine/Source/Runtime/Engine/Private/PhysicsEngine/Experimental/PhysScene_Chaos.cpp
+++ b/Engine/Source/Runtime/Engine/Private/PhysicsEngine/Experimental/PhysScene_Chaos.cpp
@@ -21,6 +21,7 @@
 #include "PhysicsProxy/ClusterUnionPhysicsProxy.h"
 #include "PhysicsProxy/SkeletalMeshPhysicsProxy.h"
 #include "PhysicsProxy/StaticMeshPhysicsProxy.h"
+#include "Chaos/PendingSpatialData.h"
 #include "EventsData.h"
 #include "Chaos/PhysicsSolverBaseImpl.h"
 
@@ -700,7 +701,7 @@ void FPhysScene_Chaos::RemoveObject(FGeometryCollectionPhysicsProxy* InObject)
 			RemoveActorFromAccelerationStructureImp(GTParticle);
 			if (Solver)
 			{
-				Solver->UpdateParticleInAccelerationStructure_External(GTParticle, /*bDelete=*/true); // ensures deletion will be applied to structures passed from physics thread.
+				Solver->UpdateParticleInAccelerationStructure_External(GTParticle,Chaos::EPendingSpatialDataOperation::Delete); // ensures deletion will be applied to structures passed from physics thread.
 			}
 		}
 	}
@@ -721,7 +722,7 @@ void FPhysScene_Chaos::RemoveObject(Chaos::FClusterUnionPhysicsProxy* InObject)
 
 	if (Solver)
 	{
-		Solver->UpdateParticleInAccelerationStructure_External(InObject->GetParticle_External(), true);
+		Solver->UpdateParticleInAccelerationStructure_External(InObject->GetParticle_External(), Chaos::EPendingSpatialDataOperation::Delete);
 		Solver->UnregisterObject(InObject);
 	}
 
diff --git a/Engine/Source/Runtime/Experimental/Chaos/Private/Chaos/Framework/PhysicsSolverBase.cpp b/Engine/Source/Runtime/Experimental/Chaos/Private/Chaos/Framework/PhysicsSolverBase.cpp
--- a/Engine/Source/Runtime/Experimental/Chaos/Private/Chaos/Framework/PhysicsSolverBase.cpp
+++ b/Engine/Source/Runtime/Experimental/Chaos/Private/Chaos/Framework/PhysicsSolverBase.cpp
@@ -284,16 +284,16 @@ namespace Chaos
 		delete &InSolver;
 	}
 
-	void FPhysicsSolverBase::UpdateParticleInAccelerationStructure_External(FGeometryParticle* Particle,bool bDelete)
+	void FPhysicsSolverBase::UpdateParticleInAccelerationStructure_External(FGeometryParticle* Particle, EPendingSpatialDataOperation InOperation)
 	{
 		//mark it as pending for async structure being built
 		FAccelerationStructureHandle AccelerationHandle(Particle);
-		FPendingSpatialData& SpatialData = PendingSpatialOperations_External->FindOrAdd(Particle->UniqueIdx());
+		FPendingSpatialData& SpatialData = PendingSpatialOperations_External->FindOrAdd(Particle->UniqueIdx(), EPendingSpatialDataOperation::Invalid);
 
 		//make sure any new operations (i.e not currently being consumed by sim) are not acting on a deleted object
-		ensure(SpatialData.SyncTimestamp < MarshallingManager.GetExternalTimestamp_External() || !SpatialData.bDelete);
+		ensure(SpatialData.SyncTimestamp < MarshallingManager.GetExternalTimestamp_External() || SpatialData.Operation != EPendingSpatialDataOperation::Delete);
 
-		SpatialData.bDelete = bDelete;
+		SpatialData.Operation = InOperation;
 		SpatialData.SpatialIdx = Particle->SpatialIdx();
 		SpatialData.AccelerationHandle = AccelerationHandle;
 		SpatialData.SyncTimestamp = MarshallingManager.GetExternalTimestamp_External();
diff --git a/Engine/Source/Runtime/Experimental/Chaos/Private/Chaos/PBDRigidsEvolution.cpp b/Engine/Source/Runtime/Experimental/Chaos/Private/Chaos/PBDRigidsEvolution.cpp
--- a/Engine/Source/Runtime/Experimental/Chaos/Private/Chaos/PBDRigidsEvolution.cpp
+++ b/Engine/Source/Runtime/Experimental/Chaos/Private/Chaos/PBDRigidsEvolution.cpp
@@ -21,6 +21,9 @@ FAutoConsoleVariableRef CVarChaosNumPushOutIterationsOverride(TEXT("p.ChaosNumPu
 int32 ChaosNumContactIterationsOverride = -1;
 FAutoConsoleVariableRef CVarChaosNumContactIterationsOverride(TEXT("p.ChaosNumContactIterationsOverride"), ChaosNumContactIterationsOverride, TEXT("Override for num contact iterations if >= 0. [def:-1]"));
 
+int32 ChaosEnableAccelerationStructureConsistencyFix = 1;
+FAutoConsoleVariableRef CVarChaosEnableAccelerationStructureConsistencyFix(TEXT("p.ChaosEnableAccelerationStructureConsistencyFix"), ChaosEnableAccelerationStructureConsistencyFix, TEXT("Temporary CVAR used for debugging"));
+
 namespace Chaos
 {
 	CHAOS_API int32 FixBadAccelerationStructureRemoval = 1;
@@ -524,11 +527,12 @@ namespace Chaos
 
 	void FPBDRigidsEvolutionBase::ApplyParticlePendingData(const FPendingSpatialData& SpatialData, FAccelerationStructure& AccelerationStructure, bool bUpdateCache, bool bUpdateDynamicTrees)
 	{
-		if (SpatialData.bDelete)
+		if (SpatialData.Operation ==  EPendingSpatialDataOperation::Delete)
 		{
 			if (bUpdateDynamicTrees || !FDefaultCollectionFactory::IsDynamicTree(SpatialData.SpatialIdx))
 			{
-				AccelerationStructure.RemoveElementFrom(SpatialData.AccelerationHandle, SpatialData.SpatialIdx);
+				const bool Existed = AccelerationStructure.RemoveElementFrom(SpatialData.AccelerationHandle, SpatialData.SpatialIdx);
+				//ensure(Existed);
 			}
 
 			if (bUpdateCache)
@@ -555,7 +559,19 @@ namespace Chaos
 			if (bUpdateDynamicTrees || !FDefaultCollectionFactory::IsDynamicTree(SpatialData.SpatialIdx))
 			{
 				//CSV_SCOPED_TIMING_STAT(ChaosPhysicsTimers, AABBTreeUpdate)
-				AccelerationStructure.UpdateElementIn(UpdateParticle, UpdateParticle->WorldSpaceInflatedBounds(), UpdateParticle->HasBounds(), SpatialData.SpatialIdx);
+				const bool Existed = AccelerationStructure.UpdateElementIn(UpdateParticle, UpdateParticle->WorldSpaceInflatedBounds(), UpdateParticle->HasBounds(), SpatialData.SpatialIdx);
+				if (SpatialData.Operation == EPendingSpatialDataOperation::Add)
+				{
+					//ensure(Existed == false);
+				}
+				else
+				{
+					//ensure(Existed == true);
+					if (Existed == false && ChaosEnableAccelerationStructureConsistencyFix == 1)
+					{
+						AccelerationStructure.RemoveElementFrom(SpatialData.AccelerationHandle, SpatialData.SpatialIdx);  // This element was deleted! And we accidentally added it with the update. Fix our mistake
+					}
+				}
 			}
 
 			if (bUpdateCache)
@@ -634,9 +650,10 @@ namespace Chaos
 			{
 				//operation still pending so update structure
 				//note: do we care about roll over? if game ticks at 60fps we'd get 385+ days
-				if(SpatialData.bDelete)
+				if(SpatialData.Operation == EPendingSpatialDataOperation::Delete)
 				{
-					Acceleration.RemoveElementFrom(SpatialData.AccelerationHandle,SpatialData.SpatialIdx);
+					const bool Existed = Acceleration.RemoveElementFrom(SpatialData.AccelerationHandle,SpatialData.SpatialIdx);
+					//ensure(Existed);
 				}
 				else
 				{
@@ -648,7 +665,19 @@ namespace Chaos
 						TRigidTransform<FReal,3> WorldTM(UpdateParticle->X(),UpdateParticle->R());
 						WorldBounds = UpdateParticle->Geometry()->BoundingBox().TransformedAABB(WorldTM);
 					}
-					Acceleration.UpdateElementIn(UpdateParticle,WorldBounds,bHasBounds,SpatialData.SpatialIdx);
+					const bool Existed = Acceleration.UpdateElementIn(UpdateParticle,WorldBounds,bHasBounds,SpatialData.SpatialIdx);
+					if (SpatialData.Operation == EPendingSpatialDataOperation::Add)
+					{
+						//ensure(Existed == false);
+					}
+					else
+					{
+						//ensure(Existed == true);
+						if (Existed == false && ChaosEnableAccelerationStructureConsistencyFix == 1)
+						{
+							Acceleration.RemoveElementFrom(SpatialData.AccelerationHandle, SpatialData.SpatialIdx);  // This element was deleted! And we accidentally added it with the update. Fix our mistake
+						}
+					}
 				}
 			}
 			else
diff --git a/Engine/Source/Runtime/Experimental/Chaos/Private/PBDRigidsSolver.cpp b/Engine/Source/Runtime/Experimental/Chaos/Private/PBDRigidsSolver.cpp
--- a/Engine/Source/Runtime/Experimental/Chaos/Private/PBDRigidsSolver.cpp
+++ b/Engine/Source/Runtime/Experimental/Chaos/Private/PBDRigidsSolver.cpp
@@ -630,7 +630,7 @@ namespace Chaos
 		Proxy->GetParticle_LowLevel()->SetProxy(Proxy);
 		AddDirtyProxy(Proxy);
 
-		UpdateParticleInAccelerationStructure_External(Proxy->GetParticle_LowLevel(), /*bDelete=*/false);
+		UpdateParticleInAccelerationStructure_External(Proxy->GetParticle_LowLevel(), EPendingSpatialDataOperation::Add);
 	}
 
 	int32 LogCorruptMap = 0;
@@ -645,7 +645,7 @@ namespace Chaos
 
 		ClearGTParticle_External(*Proxy->GetParticle_LowLevel());	//todo: remove this
 
-		UpdateParticleInAccelerationStructure_External(Proxy->GetParticle_LowLevel(), /*bDelete=*/true);
+		UpdateParticleInAccelerationStructure_External(Proxy->GetParticle_LowLevel(), EPendingSpatialDataOperation::Delete);
 
 		// remove the proxy from the invalidation list
 		RemoveDirtyProxy(Proxy);
diff --git a/Engine/Source/Runtime/Experimental/Chaos/Public/Chaos/Framework/PhysicsSolverBase.h b/Engine/Source/Runtime/Experimental/Chaos/Public/Chaos/Framework/PhysicsSolverBase.h
--- a/Engine/Source/Runtime/Experimental/Chaos/Public/Chaos/Framework/PhysicsSolverBase.h
+++ b/Engine/Source/Runtime/Experimental/Chaos/Public/Chaos/Framework/PhysicsSolverBase.h
@@ -31,6 +31,7 @@ namespace Chaos
 	class FChaosResultsManager;
 	class FRewindData;
 	class IRewindCallback;
+	enum EPendingSpatialDataOperation : uint8;
 
 	extern CHAOS_API int32 GSingleThreadedPhysics;
 	extern CHAOS_API int32 UseAsyncInterpolation;
@@ -495,7 +496,7 @@ namespace Chaos
 			}
 		}
 
-		void UpdateParticleInAccelerationStructure_External(FGeometryParticle* Particle,bool bDelete);
+		void UpdateParticleInAccelerationStructure_External(FGeometryParticle* Particle, EPendingSpatialDataOperation InOperation);
 
 		bool IsPaused_External() const
 		{
diff --git a/Engine/Source/Runtime/Experimental/Chaos/Public/Chaos/ISpatialAcceleration.h b/Engine/Source/Runtime/Experimental/Chaos/Public/Chaos/ISpatialAcceleration.h
--- a/Engine/Source/Runtime/Experimental/Chaos/Public/Chaos/ISpatialAcceleration.h
+++ b/Engine/Source/Runtime/Experimental/Chaos/Public/Chaos/ISpatialAcceleration.h
@@ -302,14 +302,14 @@ public:
 		return true;
 	}
 
-	virtual void RemoveElementFrom(const TPayloadType& Payload, FSpatialAccelerationIdx Idx)
+	virtual bool RemoveElementFrom(const TPayloadType& Payload, FSpatialAccelerationIdx Idx)
 	{
-		RemoveElement(Payload);
+		return RemoveElement(Payload);
 	}
 
-	virtual void UpdateElementIn(const TPayloadType& Payload, const TAABB<T, d>& NewBounds, bool bHasBounds, FSpatialAccelerationIdx Idx)
+	virtual bool UpdateElementIn(const TPayloadType& Payload, const TAABB<T, d>& NewBounds, bool bHasBounds, FSpatialAccelerationIdx Idx)
 	{
-		UpdateElement(Payload, NewBounds, bHasBounds);
+		return UpdateElement(Payload, NewBounds, bHasBounds);
 	}
 
 	virtual TUniquePtr<ISpatialAcceleration<TPayloadType, T, d>> Copy() const
diff --git a/Engine/Source/Runtime/Experimental/Chaos/Public/Chaos/PBDRigidsEvolution.h b/Engine/Source/Runtime/Experimental/Chaos/Public/Chaos/PBDRigidsEvolution.h
--- a/Engine/Source/Runtime/Experimental/Chaos/Public/Chaos/PBDRigidsEvolution.h
+++ b/Engine/Source/Runtime/Experimental/Chaos/Public/Chaos/PBDRigidsEvolution.h
@@ -558,14 +558,17 @@ public:
 		{
 			//TODO: distinguish between new particles and dirty particles
 			const FUniqueIdx UniqueIdx = Particle.UniqueIdx();
-			FPendingSpatialData& SpatialData = InternalAccelerationQueue.FindOrAdd(UniqueIdx);
-			ensure(SpatialData.bDelete == false);
+			FPendingSpatialData& SpatialData = InternalAccelerationQueue.FindOrAdd(UniqueIdx, EPendingSpatialDataOperation::Update);
+			ensure(SpatialData.Operation != EPendingSpatialDataOperation::Delete);
+			
 			SpatialData.AccelerationHandle = FAccelerationStructureHandle(Particle);
 			SpatialData.SpatialIdx = Particle.SpatialIdx();
 
-			auto& AsyncSpatialData = AsyncAccelerationQueue.FindOrAdd(UniqueIdx);
-			ensure(SpatialData.bDelete == false);
-			AsyncSpatialData = SpatialData;
+			auto& AsyncSpatialData = AsyncAccelerationQueue.FindOrAdd(UniqueIdx, EPendingSpatialDataOperation::Update);
+			// ensure(AsyncSpatialData.Operation != EPendingSpatialDataOperation::Delete); // TODO: This may be hit: check for potential UniqueIdx reuse bug
+			
+			AsyncSpatialData.AccelerationHandle = FAccelerationStructureHandle(Particle);
+			AsyncSpatialData.SpatialIdx = Particle.SpatialIdx();
 		}
 	}
 
@@ -952,9 +955,9 @@ public:
 		//TODO: at the moment we don't distinguish between the first time a particle is created and when it's just moved
 		// If we had this distinction we could simply remove the entry for the async queue
 		const FUniqueIdx UniqueIdx = ParticleHandle.UniqueIdx();
-		FPendingSpatialData& SpatialData = AsyncAccelerationQueue.FindOrAdd(UniqueIdx);
+		FPendingSpatialData& SpatialData = AsyncAccelerationQueue.FindOrAdd(UniqueIdx, EPendingSpatialDataOperation::Delete);
 
-		SpatialData.bDelete = true;
+		SpatialData.Operation = EPendingSpatialDataOperation::Delete;
 		SpatialData.SpatialIdx = ParticleHandle.SpatialIdx();
 		SpatialData.AccelerationHandle = FAccelerationStructureHandle(ParticleHandle);
 
diff --git a/Engine/Source/Runtime/Experimental/Chaos/Public/Chaos/PendingSpatialData.h b/Engine/Source/Runtime/Experimental/Chaos/Public/Chaos/PendingSpatialData.h
--- a/Engine/Source/Runtime/Experimental/Chaos/Public/Chaos/PendingSpatialData.h
+++ b/Engine/Source/Runtime/Experimental/Chaos/Public/Chaos/PendingSpatialData.h
@@ -8,6 +8,13 @@
 
 namespace Chaos
 { 
+enum EPendingSpatialDataOperation : uint8 
+{	
+	Delete,
+	Add, // Use this if the element does not exist in the acceleration structure
+	Update, // Use this when it is known that the element already exists in the acceleration structure
+	Invalid
+};
 	
 /** Used for updating intermediate spatial structures when they are finished */
 struct FPendingSpatialData
@@ -15,11 +22,11 @@ struct FPendingSpatialData
 	FAccelerationStructureHandle AccelerationHandle;
 	FSpatialAccelerationIdx SpatialIdx;
 	int32 SyncTimestamp;	//indicates the inputs timestamp associated with latest change. Only relevant for external queue
-	bool bDelete;
+	EPendingSpatialDataOperation Operation;
 
 	FPendingSpatialData()
 	: SyncTimestamp(0)
-	, bDelete(false)
+	, Operation(Add)
 	{}
 
 	void Serialize(FChaosArchive& Ar)
@@ -66,7 +73,7 @@ struct FPendingSpatialDataQueue
 		return PendingData.Num();
 	}
 
-	FPendingSpatialData& FindOrAdd(const FUniqueIdx UniqueIdx)
+	FPendingSpatialData& FindOrAdd(const FUniqueIdx UniqueIdx, EPendingSpatialDataOperation Operation)
 	{
 		if(int32* Existing = ParticleToPendingData.Find(UniqueIdx))
 		{
@@ -75,6 +82,7 @@ struct FPendingSpatialDataQueue
 		{
 			const int32 NewIdx = PendingData.AddDefaulted(1);
 			ParticleToPendingData.Add(UniqueIdx,NewIdx);
+			PendingData[NewIdx].Operation = Operation;
 			return PendingData[NewIdx];
 		}
 	}
diff --git a/Engine/Source/Runtime/Experimental/Chaos/Public/Chaos/SpatialAccelerationCollection.h b/Engine/Source/Runtime/Experimental/Chaos/Public/Chaos/SpatialAccelerationCollection.h
--- a/Engine/Source/Runtime/Experimental/Chaos/Public/Chaos/SpatialAccelerationCollection.h
+++ b/Engine/Source/Runtime/Experimental/Chaos/Public/Chaos/SpatialAccelerationCollection.h
@@ -570,10 +570,10 @@ public:
 		return Indices;
 	}
 
-	virtual void RemoveElementFrom(const TPayloadType& Payload, FSpatialAccelerationIdx SpatialIdx) override
+	virtual bool RemoveElementFrom(const TPayloadType& Payload, FSpatialAccelerationIdx SpatialIdx) override
 	{
 		const uint16 UseBucket = ((1 << SpatialIdx.Bucket) & this->ActiveBucketsMask) ? SpatialIdx.Bucket : 0;
-		const bool Success = Buckets[UseBucket].Objects[SpatialIdx.InnerIdx].Acceleration->RemoveElement(Payload);
+		bool Success = Buckets[UseBucket].Objects[SpatialIdx.InnerIdx].Acceleration->RemoveElement(Payload);
 		//ensure(Success); // Debug check to see if SpatialIdx was what the caller expected it to be
 		// Make sure that we remove this Payload even if the SpatialIdx is wrong
 		for (FSpatialAccelerationIdx Idx : GetAllSpatialIndices())
@@ -581,15 +581,17 @@ public:
 			if (!(Idx == SpatialIdx))
 			{
 				const uint16 Buckt = ((1 << Idx.Bucket) & this->ActiveBucketsMask) ? Idx.Bucket : 0;
-				const bool Removed = Buckets[Buckt].Objects[Idx.InnerIdx].Acceleration->RemoveElement(Payload);					
+				const bool Removed = Buckets[Buckt].Objects[Idx.InnerIdx].Acceleration->RemoveElement(Payload);
+				Success = Success || Removed;
 			}
-		}		
+		}
+		return Success;
 	}
 
-	virtual void UpdateElementIn(const TPayloadType& Payload, const TAABB<T, d>& NewBounds, bool bHasBounds, FSpatialAccelerationIdx SpatialIdx)
+	virtual bool UpdateElementIn(const TPayloadType& Payload, const TAABB<T, d>& NewBounds, bool bHasBounds, FSpatialAccelerationIdx SpatialIdx)
 	{
 		const uint16 UseBucket = ((1 << SpatialIdx.Bucket) & this->ActiveBucketsMask) ? SpatialIdx.Bucket : 0;
-		const bool bElementExisted = Buckets[UseBucket].Objects[SpatialIdx.InnerIdx].Acceleration->UpdateElement(Payload, NewBounds, bHasBounds);
+		bool bElementExisted = Buckets[UseBucket].Objects[SpatialIdx.InnerIdx].Acceleration->UpdateElement(Payload, NewBounds, bHasBounds);
 		// In case spatial index changed, remove this element in all other substructures
 		if (!bElementExisted)
 		{
@@ -598,10 +600,12 @@ public:
 				if (!(Idx == SpatialIdx))
 				{
 					const uint16 Buckt = ((1 << Idx.Bucket) & this->ActiveBucketsMask) ? Idx.Bucket : 0;
-					Buckets[Buckt].Objects[Idx.InnerIdx].Acceleration->RemoveElement(Payload);
+					const bool Removed = Buckets[Buckt].Objects[Idx.InnerIdx].Acceleration->RemoveElement(Payload);
+					bElementExisted = bElementExisted || Removed;
 				}
 			}
 		}
+		return bElementExisted;
 	}
 
 	virtual TUniquePtr<ISpatialAcceleration<TPayloadType, T, d>> Copy() const
diff --git a/Engine/Source/Runtime/PhysicsCore/Private/ChaosScene.cpp b/Engine/Source/Runtime/PhysicsCore/Private/ChaosScene.cpp
--- a/Engine/Source/Runtime/PhysicsCore/Private/ChaosScene.cpp
+++ b/Engine/Source/Runtime/PhysicsCore/Private/ChaosScene.cpp
@@ -233,7 +233,7 @@ void FChaosScene::UpdateActorInAccelerationStructure(const FPhysicsActorHandle&
 			SpatialAcceleration->UpdateElementIn(AccelerationHandle,WorldBounds,bHasBounds, Body_External.SpatialIdx());
 		}
 
-		GetSolver()->UpdateParticleInAccelerationStructure_External(Actor->GetParticle_LowLevel(),/*bDelete=*/false);
+		GetSolver()->UpdateParticleInAccelerationStructure_External(Actor->GetParticle_LowLevel(), EPendingSpatialDataOperation::Update);
 	}
 }
 
@@ -275,7 +275,7 @@ void FChaosScene::UpdateActorsInAccelerationStructure(const TArrayView<FPhysicsA
 			const FPhysicsActorHandle& Actor = Actors[ActorIndex];
 			if(Actor)
 			{
-				GetSolver()->UpdateParticleInAccelerationStructure_External(Actor->GetParticle_LowLevel(),/*bDelete=*/false);
+				GetSolver()->UpdateParticleInAccelerationStructure_External(Actor->GetParticle_LowLevel(), EPendingSpatialDataOperation::Update);
 			}
 		}
 	}
