Fixed incomplete copy constructor in Chaos::FMidPhaseModifierParticleIterator, which was also mismatched with its assignment operator.

diff --git a/Engine/Source/Runtime/Experimental/Chaos/Public/Chaos/MidPhaseModification.h b/Engine/Source/Runtime/Experimental/Chaos/Public/Chaos/MidPhaseModification.h
--- a/Engine/Source/Runtime/Experimental/Chaos/Public/Chaos/MidPhaseModification.h
+++ b/Engine/Source/Runtime/Experimental/Chaos/Public/Chaos/MidPhaseModification.h
@@ -70,12 +70,6 @@ namespace Chaos
 	class CHAOS_API FMidPhaseModifierParticleIterator
 	{
 	public:
-		FMidPhaseModifierParticleIterator(const FMidPhaseModifierParticleIterator& Other)
-			: Particle(Other.Particle)
-			, Accessor(Other.Accessor)
-			, MidPhaseIndex(Other.MidPhaseIndex)
-		{ }
-
 		FMidPhaseModifierParticleIterator(
 			FMidPhaseModifierAccessor* InAccessor,
 			TGeometryParticleHandle<FReal, 3>* InParticle,
