Iris - Fixed typo in WritePackedInt64. GetBitsNeeded will always return at least 1 for signed integer types.

#rnx

diff --git a/Engine/Source/Runtime/Experimental/Iris/Core/Private/Iris/Serialization/NetBitStreamUtil.cpp b/Engine/Source/Runtime/Experimental/Iris/Core/Private/Iris/Serialization/NetBitStreamUtil.cpp
--- a/Engine/Source/Runtime/Experimental/Iris/Core/Private/Iris/Serialization/NetBitStreamUtil.cpp
+++ b/Engine/Source/Runtime/Experimental/Iris/Core/Private/Iris/Serialization/NetBitStreamUtil.cpp
@@ -12,7 +12,7 @@ namespace UE::Net
 
 void WritePackedUint64(FNetBitStreamWriter* Writer, uint64 Value)
 {
-	// As we represent the number of bytes to write with two bits we want bits needed to be >= 1
+	// As we represent the number of bytes to write with three bits we want bits needed to be >= 1 such that the number of bytes ends up in the range [1, 8].
 	const uint32 BitCountNeeded = GetBitsNeeded(Value | 1U);
 	const uint32 ByteCountNeeded = (BitCountNeeded + 7U) / 8U;
 	const uint32 BitCountToWrite = ByteCountNeeded * 8U;
@@ -49,7 +49,7 @@ uint64 ReadPackedUint64(FNetBitStreamReader* Reader)
 
 void WritePackedInt64(FNetBitStreamWriter* Writer, int64 Value)
 {
-	const uint32 BitCountNeeded = GetBitsNeeded(Value | 1U);
+	const uint32 BitCountNeeded = GetBitsNeeded(Value);
 	const uint32 ByteCountNeeded = (BitCountNeeded + 7U) / 8U;
 	const uint32 BitCountToWrite = ByteCountNeeded * 8U;
 
