Fixed mismatched copy constructor/assignment operator in Chaos::Softs::FDeformableSolverProperties.

diff --git a/Engine/Source/Runtime/Experimental/Chaos/Public/Chaos/Deformable/ChaosDeformableSolverTypes.h b/Engine/Source/Runtime/Experimental/Chaos/Public/Chaos/Deformable/ChaosDeformableSolverTypes.h
--- a/Engine/Source/Runtime/Experimental/Chaos/Public/Chaos/Deformable/ChaosDeformableSolverTypes.h
+++ b/Engine/Source/Runtime/Experimental/Chaos/Public/Chaos/Deformable/ChaosDeformableSolverTypes.h
@@ -54,27 +54,6 @@ namespace Chaos::Softs
 			, bEnablePositionTargets(InbEnablePositionTargets)
 		{}
 
-		FDeformableSolverProperties(const FDeformableSolverProperties& InProp)
-			: NumSolverSubSteps(InProp.NumSolverSubSteps)
-			, NumSolverIterations(InProp.NumSolverIterations)
-			, FixTimeStep(InProp.FixTimeStep)
-			, TimeStepSize(InProp.TimeStepSize)
-			, CacheToFile(InProp.CacheToFile)
-			, bEnableKinematics(InProp.bEnableKinematics)
-			, bUseFloor(InProp.bUseFloor)
-			, bDoSelfCollision(InProp.bDoSelfCollision)
-			, bUseGridBasedConstraints(InProp.bUseGridBasedConstraints)
-			, GridDx(InProp.GridDx)
-			, bDoQuasistatics(InProp.bDoQuasistatics)
-			, EMesh(InProp.EMesh)
-			, bDoBlended(InProp.bDoBlended)
-			, BlendedZeta(InProp.BlendedZeta)
-			, Damping(InProp.Damping)
-			, bEnableGravity(InProp.bEnableGravity)
-			, bEnableCorotatedConstraints(InProp.bEnableCorotatedConstraints)
-			, bEnablePositionTargets(InProp.bEnablePositionTargets)
-		{}
-
 		int32 NumSolverSubSteps = 5;
 		int32 NumSolverIterations = 5;
 		bool FixTimeStep = false;
