Fixed mismatched copy constructor/assignment operator in Chaos::FSimplex.

diff --git a/Engine/Source/Runtime/Experimental/Chaos/Public/Chaos/Simplex.h b/Engine/Source/Runtime/Experimental/Chaos/Public/Chaos/Simplex.h
--- a/Engine/Source/Runtime/Experimental/Chaos/Public/Chaos/Simplex.h
+++ b/Engine/Source/Runtime/Experimental/Chaos/Public/Chaos/Simplex.h
@@ -87,16 +87,6 @@ namespace Chaos
 		int32 operator[](int32 Idx) const { return Idxs[Idx]; }
 		int32& operator[](int32 Idx) { return Idxs[Idx]; }
 
-		FSimplex& operator=(const FSimplex& Other)
-		{
-			NumVerts = Other.NumVerts;
-			for (int i = 0; i < 4; ++i)
-			{
-				Idxs[i] = Other.Idxs[i];
-			}
-			return *this;
-		}
-
 		FSimplex(std::initializer_list<int32> InIdxs = {})
 			: NumVerts(InIdxs.size())
 		{
