Iris - Clearly state when NetRefHandles are incomplete when printing them. Fixed which replication system ID to display for StringBuilder streaming operators.

#rnx

diff --git a/Engine/Source/Runtime/Experimental/Iris/Core/Private/Iris/ReplicationSystem/NetRefHandle.cpp b/Engine/Source/Runtime/Experimental/Iris/Core/Private/Iris/ReplicationSystem/NetRefHandle.cpp
--- a/Engine/Source/Runtime/Experimental/Iris/Core/Private/Iris/ReplicationSystem/NetRefHandle.cpp
+++ b/Engine/Source/Runtime/Experimental/Iris/Core/Private/Iris/ReplicationSystem/NetRefHandle.cpp
@@ -11,8 +11,15 @@ FString FNetRefHandle::ToString() const
 {
 	FString Result;
 #if UE_NET_ALLOW_MULTIPLE_REPLICATION_SYSTEMS
-	const uint64 ReplicationSystemIdToDisplay = ReplicationSystemId - 1;
-	Result = FString::Printf(TEXT("NetRefHandle (Id=%" UINT64_FMT "):(RepSystemId=%" UINT64_FMT ")"), GetId(), ReplicationSystemIdToDisplay);
+	if (IsCompleteHandle())
+	{
+		const uint32 ReplicationSystemIdToDisplay = GetReplicationSystemId();
+		Result = FString::Printf(TEXT("NetRefHandle (Id=%" UINT64_FMT "):(RepSystemId=%u)"), GetId(), ReplicationSystemIdToDisplay);
+	}
+	else
+	{
+		Result = FString::Printf(TEXT("NetRefHandle (Id=%" UINT64_FMT "):(Incomplete)"), GetId());
+	}
 #else
 	Result = FString::Printf(TEXT("NetRefHandle (Id=%" UINT64_FMT ")"), GetId());
 #endif
@@ -24,8 +31,15 @@ FString FNetRefHandle::ToString() const
 FStringBuilderBase& operator<<(FStringBuilderBase& Builder, const UE::Net::FNetRefHandle& NetRefHandle) 
 { 	
 #if UE_NET_ALLOW_MULTIPLE_REPLICATION_SYSTEMS
-	const uint32 ReplicationSystemIdToDisplay = NetRefHandle.GetReplicationSystemId() - 1U;
-	Builder.Appendf(TEXT("NetRefHandle (Id=%" UINT64_FMT "):(RepSystemId=%u)"), NetRefHandle.GetId(), ReplicationSystemIdToDisplay);
+	if (NetRefHandle.IsCompleteHandle())
+	{
+		const uint32 ReplicationSystemIdToDisplay = NetRefHandle.GetReplicationSystemId();
+		Builder.Appendf(TEXT("NetRefHandle (Id=%" UINT64_FMT "):(RepSystemId=%u)"), NetRefHandle.GetId(), ReplicationSystemIdToDisplay);
+	}
+	else
+	{
+		Builder.Appendf(TEXT("NetRefHandle (Id=%" UINT64_FMT "):(Incomplete)"), NetRefHandle.GetId());
+	}
 #else
 	Builder.Appendf(TEXT("NetRefHandle (Id=%" UINT64_FMT ")"), NetRefHandle.GetId());
 #endif
@@ -35,8 +49,15 @@ FStringBuilderBase& operator<<(FStringBuilderBase& Builder, const UE::Net::FNetR
 FAnsiStringBuilderBase& operator<<(FAnsiStringBuilderBase& Builder, const UE::Net::FNetRefHandle& NetRefHandle)
 {
 #if UE_NET_ALLOW_MULTIPLE_REPLICATION_SYSTEMS
-	const uint32 ReplicationSystemIdToDisplay = NetRefHandle.GetReplicationSystemId() - 1U;
-	Builder.Appendf("NetRefHandle (Id=%" UINT64_FMT "):(RepSystemId=%u)", NetRefHandle.GetId(), ReplicationSystemIdToDisplay);
+	if (NetRefHandle.IsCompleteHandle())
+	{
+		const uint32 ReplicationSystemIdToDisplay = NetRefHandle.GetReplicationSystemId();
+		Builder.Appendf("NetRefHandle (Id=%" UINT64_FMT "):(RepSystemId=%u)", NetRefHandle.GetId(), ReplicationSystemIdToDisplay);
+	}
+	else
+	{
+		Builder.Appendf("NetRefHandle (Id=%" UINT64_FMT "):(Incomplete)", NetRefHandle.GetId());
+	}
 #else
 	Builder.Appendf("NetRefHandle (Id=%" UINT64_FMT ")", NetRefHandle.GetId());
 #endif
