Iris - Prevent parent objects from being destroyed on the authoritative side before torn off subobjects have flushed. Handle subobject tear off packet drop appropriately. ReplicationReader robustness fixes, mmaking some checks ensures, early out on error conditions etc.

diff --git a/Engine/Source/Runtime/Experimental/Iris/Core/Private/Iris/ReplicationSystem/NetRefHandleManager.cpp b/Engine/Source/Runtime/Experimental/Iris/Core/Private/Iris/ReplicationSystem/NetRefHandleManager.cpp
--- a/Engine/Source/Runtime/Experimental/Iris/Core/Private/Iris/ReplicationSystem/NetRefHandleManager.cpp
+++ b/Engine/Source/Runtime/Experimental/Iris/Core/Private/Iris/ReplicationSystem/NetRefHandleManager.cpp
@@ -5,6 +5,7 @@
 #include "ReplicationProtocol.h"
 #include "ReplicationProtocolManager.h"
 #include "Iris/Core/IrisLog.h"
+#include "Iris/Core/IrisProfiler.h"
 #include "Iris/Serialization/NetSerializationContext.h"
 #include "Iris/Serialization/InternalNetSerializationContext.h"
 #include "Net/Core/Trace/NetTrace.h"
@@ -385,11 +386,14 @@ void FNetRefHandleManager::DestroyNetObject(FNetRefHandle RefHandle)
 
 void FNetRefHandleManager::DestroyObjectsPendingDestroy()
 {
+	IRIS_PROFILER_SCOPE(FNetRefHandleManager_DestroyObjectsPendingDestroy);
+
 	// Destroy Objects pending destroy
 	for (int32 It = PendingDestroyInternalIndices.Num() - 1; It >= 0; --It)
 	{
 		const FInternalNetRefIndex InternalIndex = PendingDestroyInternalIndices[It];
-		if (ReplicatedObjectRefCount[InternalIndex] == 0)
+		// If we have subobjects pending tear off and such then wait before destroying the parent.
+		if (ReplicatedObjectRefCount[InternalIndex] == 0 && GetSubObjects(InternalIndex).Num() <= 0)
 		{
 			InternalDestroyNetObject(InternalIndex);
 			PendingDestroyInternalIndices.RemoveAtSwap(It);
@@ -500,6 +504,7 @@ void FNetRefHandleManager::InternalRemoveSubObject(FInternalNetRefIndex OwnerInt
 		}
 
 		SubObjectData.SubObjectRootIndex = InvalidInternalIndex;
+		SubObjectData.SubObjectParentIndex = InvalidInternalIndex;
 		SubObjectData.bDestroySubObjectWithOwner = false;
 
 		SetIsSubObject(SubObjectInternalIndex, false);
@@ -765,7 +770,8 @@ FNetRefHandle FNetRefHandleManager::MakeNetRefHandle(uint64 Id, uint32 Replicati
 
 FNetRefHandle FNetRefHandleManager::MakeNetRefHandleFromId(uint64 Id)
 {
-	check((Id & FNetRefHandle::IdMask) == Id);
+	// This is called on the receiving end when deserializing replicated objects. We don't want to crash on bit stream errors leading to invalid handle IDs being read.
+	ensureAlways((Id & FNetRefHandle::IdMask) == Id);
 
 	FNetRefHandle Handle;
 
diff --git a/Engine/Source/Runtime/Experimental/Iris/Core/Private/Iris/ReplicationSystem/ReplicationReader.cpp b/Engine/Source/Runtime/Experimental/Iris/Core/Private/Iris/ReplicationSystem/ReplicationReader.cpp
--- a/Engine/Source/Runtime/Experimental/Iris/Core/Private/Iris/ReplicationSystem/ReplicationReader.cpp
+++ b/Engine/Source/Runtime/Experimental/Iris/Core/Private/Iris/ReplicationSystem/ReplicationReader.cpp
@@ -179,10 +179,17 @@ void FReplicationReader::Deinit()
 }
 
 // Read incomplete handle
-FNetRefHandle FReplicationReader::ReadNetRefHandleId(FNetBitStreamReader& Reader) const
+FNetRefHandle FReplicationReader::ReadNetRefHandleId(FNetSerializationContext& Context, FNetBitStreamReader& Reader) const
 {
 	const uint64 NetId = ReadPackedUint64(&Reader);
-	return FNetRefHandleManager::MakeNetRefHandleFromId(NetId);
+	FNetRefHandle RefHandle = FNetRefHandleManager::MakeNetRefHandleFromId(NetId);
+	if (RefHandle.GetId() != NetId)
+	{
+		Context.SetError(GNetError_InvalidNetHandle);
+		return FNetRefHandle();
+	}
+
+	return RefHandle;
 }
 	
 uint16 FReplicationReader::ReadObjectsPendingDestroy(FNetSerializationContext& Context)
@@ -194,19 +201,23 @@ uint16 FReplicationReader::ReadObjectsPendingDestroy(FNetSerializationContext& C
 	// Read how many destroyed objects we have
 	const uint16 ObjectsToRead = Reader.ReadBits(16);
 	
-	if (!Reader.IsOverflown())
+	if (!Context.HasErrorOrOverflow())
 	{
 		for (uint32 It = 0; It < ObjectsToRead; ++It)
 		{
 			UE_NET_TRACE_NAMED_OBJECT_SCOPE(DestroyedObjectScope, FNetRefHandle(), Reader, Context.GetTraceCollector(), ENetTraceVerbosity::Trace);
 
-			FNetRefHandle IncompleteHandle = ReadNetRefHandleId(Reader);
+			FNetRefHandle IncompleteHandle = ReadNetRefHandleId(Context, Reader);
 			const bool bShouldDestroyInstance = Reader.ReadBool();
-			if (!Reader.IsOverflown())
+			if (Context.HasErrorOrOverflow())
+			{
+				break;
+			}
+
 			{
 				// Resolve handle and destroy using bridge
 				const uint32 InternalIndex = NetRefHandleManager->GetInternalIndex(IncompleteHandle);
-				if (InternalIndex)
+				if (InternalIndex != FNetRefHandleManager::InvalidInternalIndex)
 				{
 					UE_NET_TRACE_SET_SCOPE_OBJECTID(DestroyedObjectScope, IncompleteHandle);
 
@@ -231,18 +242,17 @@ uint16 FReplicationReader::ReadObjectsPendingDestroy(FNetSerializationContext& C
 					{
 						EndReplication(InternalIndex, false, bShouldDestroyInstance);
 					}
-
-					continue;
 				}
+				else
+				{
 
-				// If we did not find the object or associated bridge, the packet that would have created the object may have been lost.
-				UE_LOG(LogIris, Log, TEXT("FReplicationReader::Read Tried to destroy object with %s (This can occur if the server sends destroy for an object that has not yet been confirmed as created)"), *IncompleteHandle.ToString());
+					// If we did not find the object or associated bridge, the packet that would have created the object may have been lost.
+					UE_LOG(LogIris, Log, TEXT("FReplicationReader::Read Tried to destroy object with %s (This can occur if the server sends destroy for an object that has not yet been confirmed as created)"), *IncompleteHandle.ToString());
+				}
 			}
 		}
 	}
 
-	check(!Reader.IsOverflown());
-
 	return ObjectsToRead;
 }
 
@@ -370,7 +380,7 @@ uint32 FReplicationReader::ReadObjectBatch(FNetSerializationContext& Context)
 {
 	FNetBitStreamReader& Reader = *Context.GetBitStreamReader();
 
-	// Special handling for desrtuction infos
+	// Special handling for destruction infos
 	if (const bool bIsDestructionInfo = Reader.ReadBool())
 	{
 		FReplicationBridgeSerializationContext BridgeContext(Context, Parameters.ConnectionId, true);
@@ -394,7 +404,7 @@ uint32 FReplicationReader::ReadObjectBatch(FNetSerializationContext& Context)
 	
 	// A batch starts with (RefHandleId | BatchSize | bHasBatchObjectData | bHasExports)
 	// If the batch has exports we must to seek to the end of the batch to read and process exports before reading/processing batch data
-	const FNetRefHandle IncompleteHandle = ReadNetRefHandleId(Reader);
+	const FNetRefHandle IncompleteHandle = ReadNetRefHandleId(Context, Reader);
 
 	uint32 BatchSize = 0U;
 	// Read Batch size
@@ -480,7 +490,7 @@ void FReplicationReader::ReadObjectInBatch(FNetSerializationContext& Context, FN
 	FNetRefHandle IncompleteHandle = BatchHandle;
 	if (bIsSubObject)
 	{
-		IncompleteHandle = ReadNetRefHandleId(Reader);
+		IncompleteHandle = ReadNetRefHandleId(Context, Reader);
 	}
 	
 	// Read replicated destroy header if necessary
@@ -1536,7 +1546,7 @@ void FReplicationReader::Read(FNetSerializationContext& Context)
 	ObjectCountToRead -= DestroyedObjectCount;
 
 	// Nothing more to do or we failed and should disconnect
-	if (Reader.IsOverflown() || (ObjectCountToRead == 0 && ObjectsToDispatchCount == 0))
+	if (Context.HasErrorOrOverflow() || (ObjectCountToRead == 0 && ObjectsToDispatchCount == 0))
 	{
 		return;
 	}
diff --git a/Engine/Source/Runtime/Experimental/Iris/Core/Private/Iris/ReplicationSystem/ReplicationReader.h b/Engine/Source/Runtime/Experimental/Iris/Core/Private/Iris/ReplicationSystem/ReplicationReader.h
--- a/Engine/Source/Runtime/Experimental/Iris/Core/Private/Iris/ReplicationSystem/ReplicationReader.h
+++ b/Engine/Source/Runtime/Experimental/Iris/Core/Private/Iris/ReplicationSystem/ReplicationReader.h
@@ -52,9 +52,6 @@ public:
 	// Read incoming replication data
 	void Read(FNetSerializationContext& Context);
 
-	// Read index part of handle
-	FNetRefHandle ReadNetRefHandleId(FNetBitStreamReader& Reader) const;
-
 	void SetRemoteNetTokenStoreState(FNetTokenStoreState* RemoteTokenStoreState);
 	
 	// Mark objects pending destroy as unresolvable.
@@ -125,8 +122,8 @@ private:
 
 	bool IsObjectIndexForOOBAttachment(uint32 InternalIndex) const { return InternalIndex == ObjectIndexForOOBAttachment; }
 
-	// Read index reference to replicated object
-	uint32 ReadInternalIndex(FNetBitStreamReader& Reader) const;
+	// Read index part of handle
+	FNetRefHandle ReadNetRefHandleId(FNetSerializationContext& Context, FNetBitStreamReader& Reader) const;
 
 	// Read a new or updated object
 	uint32 ReadObjectBatch(FNetSerializationContext& Context);
diff --git a/Engine/Source/Runtime/Experimental/Iris/Core/Private/Iris/ReplicationSystem/ReplicationWriter.cpp b/Engine/Source/Runtime/Experimental/Iris/Core/Private/Iris/ReplicationSystem/ReplicationWriter.cpp
--- a/Engine/Source/Runtime/Experimental/Iris/Core/Private/Iris/ReplicationSystem/ReplicationWriter.cpp
+++ b/Engine/Source/Runtime/Experimental/Iris/Core/Private/Iris/ReplicationSystem/ReplicationWriter.cpp
@@ -1276,7 +1276,6 @@ void FReplicationWriter::HandleDroppedRecord<FReplicationWriter::EReplicatedObje
 			FNetRefHandleManager::FReplicatedObjectData& ObjectData = NetRefHandleManager->GetReplicatedObjectDataNoCheck(InternalIndex);
 			check(ObjectData.IsSubObject());
 
-			// If owner is not pending destroy we mark it as dirty as appropriate.
 			FNetBitArrayView ChangeMask(Info.GetChangeMaskStoragePointer(), Info.ChangeMaskBitCount);
 			Info.HasDirtyChangeMask |= ChangeMask.IsAnyBitSet();
 			Info.SubObjectPendingDestroy = 0U;
@@ -1285,6 +1284,7 @@ void FReplicationWriter::HandleDroppedRecord<FReplicationWriter::EReplicatedObje
 
 			ObjectsWithDirtyChanges.SetBitValue(InternalIndex, Info.HasDirtyChangeMask);
 
+			// If owner is not pending destroy we mark it as dirty as appropriate.
 			FReplicationInfo& OwnerInfo = ReplicatedObjects[ObjectData.SubObjectRootIndex];
 			if (OwnerInfo.GetState() < EReplicatedObjectState::PendingDestroy)
 			{
@@ -1310,7 +1310,7 @@ void FReplicationWriter::HandleDroppedRecord<FReplicationWriter::EReplicatedObje
 		// We dropped a packet with tear-off data, that is a destroy with state data so we need to resend that state
 		if (RecordInfo.WroteTearOff)
 		{
-			check(Info.TearOff);
+			ensureAlways(Info.TearOff);
 
 			SetState(InternalIndex, EReplicatedObjectState::PendingTearOff);
 
@@ -1325,17 +1325,33 @@ void FReplicationWriter::HandleDroppedRecord<FReplicationWriter::EReplicatedObje
 				Info.HasDirtyChangeMask = 1U;
 			}
 
-			// Indicate that we have dirty subobjects
-			Info.HasDirtySubObjects = 1U;
-
 			// Mark attachments as dirty
 			Info.HasAttachments |= RecordInfo.HasAttachments;
 
 			// Mark object as having dirty changes
 			ObjectsWithDirtyChanges.SetBit(InternalIndex);
 
+			// Mark parent as dirty
+			uint32 ParentInternalIndex = InternalIndex;
+			if (Info.IsSubObject)
+			{
+				const FNetRefHandleManager::FReplicatedObjectData& ObjectData = NetRefHandleManager->GetReplicatedObjectDataNoCheck(InternalIndex);
+				if (ensure(ObjectData.SubObjectRootIndex != FNetRefHandleManager::InvalidInternalIndex))
+				{
+					ParentInternalIndex = ObjectData.SubObjectRootIndex;
+					ObjectsWithDirtyChanges.SetBit(ParentInternalIndex);
+
+					FReplicationInfo& OwnerInfo = ReplicatedObjects[ObjectData.SubObjectRootIndex];
+					if (ensure(OwnerInfo.GetState() < EReplicatedObjectState::PendingDestroy))
+					{
+						// Indicate that we have dirty subobjects
+						OwnerInfo.HasDirtySubObjects = 1U;
+					}
+				}
+			}
+
 			// Bump prio
-			SchedulingPriorities[InternalIndex] += FReplicationWriter::TearOffPriority;
+			SchedulingPriorities[ParentInternalIndex] += FReplicationWriter::TearOffPriority;
 		}
 		else if (RecordInfo.WroteDestroySubObject && Info.SubObjectPendingDestroy)
 		{
@@ -1769,7 +1785,7 @@ void FReplicationWriter::SerializeObjectStateDelta(FNetSerializationContext& Con
 
 FReplicationWriter::EWriteObjectStatus FReplicationWriter::WriteObjectAndSubObjects(FNetSerializationContext& Context, uint32 InternalIndex, uint32 WriteObjectFlags, FBatchInfo& OutBatchInfo)
 {
-	IRIS_PROFILER_SCOPE(FReplicationWriter_WriteObjectInBatch);
+	IRIS_PROFILER_SCOPE(FReplicationWriter_WriteObjectAndSubObjects);
 
 	FNetBitStreamWriter& Writer = *Context.GetBitStreamWriter();
 	
@@ -2055,7 +2071,7 @@ FReplicationWriter::EWriteObjectStatus FReplicationWriter::WriteObjectAndSubObje
 						UE_LOG(LogIris, Verbose, TEXT("Failed to write huge attachment for object %s ( InternalIndex: %u ), forcing fallback on hugeobject for attachments"), *NetRefHandle.ToString(), InternalIndex);
 						Writer.DoOverflow();
 					}
-					else if (!(BatchEntry.bSentState || bSentTearOff || Info.SubObjectPendingDestroy || Info.HasDirtySubObjects))
+					else if (!(BatchEntry.bSentState || bSentTearOff || Info.SubObjectPendingDestroy || Info.HasDirtySubObjects || bWriteBatchInfo))
 					{
 						// If we didn't send state and didn't send any attachments let's rollback
 						ObjectRollbackScope.Rollback();
@@ -2184,6 +2200,7 @@ FReplicationWriter::EWriteObjectStatus FReplicationWriter::WriteObjectAndSubObje
 			const bool bWroteExports = WriteExportResult == FObjectReferenceCache::EWriteExportsResult::WroteExports;
 
 			// Update header
+			if (ensure(bWriteBatchInfo))
 			{
 				FNetBitStreamWriteScope SizeScope(Writer, InitialStateHeaderPos);
 				Writer.WriteBits(WrittenBitsInBatch, NumBitsUsedForBatchSize);
@@ -2213,7 +2230,7 @@ FReplicationWriter::EWriteObjectStatus FReplicationWriter::WriteObjectInBatch(FN
 		ExportContext->ClearPendingExports();
 	}
 
-	// Write bit indicating if we are a destruction info or not
+	// Write parent object and subobjects
 	const EWriteObjectStatus WriteObjectStatus = WriteObjectAndSubObjects(Context, InternalIndex, WriteObjectFlags, OutBatchInfo);
 	if (!IsWriteObjectSuccess(WriteObjectStatus))
 	{
