Chaos FImplicitObject: Remove any implicit padding and add compile errors if any new padding is added.

The ispc code in FPerParticlePBDCollisionConstraint expects that the first field in derived classes of FImplicitObject begin at sizeof(FImplicitObject). However, depending on the compiler, the tail padding in FImplicitObject may be reused, causing static_asserts in FPerParticlePBDCollisionConstraint to fire. This removes all tail padding, so this compiler-depending behavior cannot occur. Additionally, I rearranged the order of the fields in FImplicitObject to remove some alignment padding, which drops its size from 24bytes -> 16bytes in the current configuration.

Also the PerParticlePBDCollisionConstraint static asserts are now enabled in shipping.

On top of the preflight to verify that everything is fine now, I also ran a preflight with just the errors added (no padding fixed) to verify that the errors did trigger on all testable platforms.

diff --git a/Engine/Source/Runtime/Experimental/Chaos/Private/Chaos/ImplicitObject.cpp b/Engine/Source/Runtime/Experimental/Chaos/Private/Chaos/ImplicitObject.cpp
--- a/Engine/Source/Runtime/Experimental/Chaos/Private/Chaos/ImplicitObject.cpp
+++ b/Engine/Source/Runtime/Experimental/Chaos/Private/Chaos/ImplicitObject.cpp
@@ -22,15 +22,15 @@
 namespace Chaos
 {
 	FImplicitObject::FImplicitObject(int32 Flags, EImplicitObjectType InType)
-		: Type(InType)
-		, CollisionType(InType)
-		, Margin(0.0f)
+		: Margin(0.0f)
 		, bIsConvex(!!(Flags& EImplicitObject::IsConvex))
 		, bDoCollide(!(Flags& EImplicitObject::DisableCollisions))
 		, bHasBoundingBox(!!(Flags& EImplicitObject::HasBoundingBox))
 #if TRACK_CHAOS_GEOMETRY
 		, bIsTracked(false)
 #endif
+		, Type(InType)
+		, CollisionType(InType)
 	{
 	}
 
diff --git a/Engine/Source/Runtime/Experimental/Chaos/Private/Chaos/PerParticlePBDCollisionConstraint.cpp b/Engine/Source/Runtime/Experimental/Chaos/Private/Chaos/PerParticlePBDCollisionConstraint.cpp
--- a/Engine/Source/Runtime/Experimental/Chaos/Private/Chaos/PerParticlePBDCollisionConstraint.cpp
+++ b/Engine/Source/Runtime/Experimental/Chaos/Private/Chaos/PerParticlePBDCollisionConstraint.cpp
@@ -6,7 +6,6 @@
 #if INTEL_ISPC
 #include "PerParticlePBDCollisionConstraint.ispc.generated.h"
 
-#if !UE_BUILD_SHIPPING
 // These includes are used for the static_asserts below
 #include "Chaos/Capsule.h"
 #include "Chaos/Convex.h"
@@ -14,10 +13,9 @@
 #include "Chaos/Sphere.h"
 #include "Chaos/TaperedCapsule.h"
 #include "Chaos/TaperedCylinder.h"
-#endif 
 #endif
 
-#if INTEL_ISPC && !UE_BUILD_SHIPPING
+#if INTEL_ISPC
 static_assert(sizeof(ispc::FVector3f) == sizeof(Chaos::Softs::FSolverVec3), "sizeof(ispc::FVector3f) != sizeof(Chaos::Softs::FSolverVec3)");
 static_assert(sizeof(ispc::FVector4f) == sizeof(Chaos::Softs::FSolverRotation3), "sizeof(ispc::FVector4f) != sizeof(Chaos::Softs::FSolverRotation3)");
 static_assert(ispc::ImplicitObjectType::Sphere == Chaos::ImplicitObjectType::Sphere);
@@ -107,9 +105,11 @@ static_assert(sizeof(ispc::VerticesS) == sizeof(Chaos::FConvexStructureData::FCo
 static_assert(sizeof(ispc::VerticesM) == sizeof(Chaos::FConvexStructureData::FConvexStructureDataMedium::FVertexData));
 static_assert(sizeof(ispc::VerticesL) == sizeof(Chaos::FConvexStructureData::FConvexStructureDataLarge::FVertexData));
 
+#if !UE_BUILD_SHIPPING
 bool bChaos_PerParticleCollision_ISPC_Enabled = CHAOS_PER_PARTICLE_COLLISION_ISPC_ENABLED_DEFAULT;
 static FAutoConsoleVariableRef CVarChaosPerParticleCollisionISPCEnabled(TEXT("p.Chaos.PerParticleCollision.ISPC"), bChaos_PerParticleCollision_ISPC_Enabled, TEXT("Whether to use ISPC optimizations in per particle collisions"));
 #endif
+#endif
 
 static int32 Chaos_PerParticleCollision_ISPC_ParallelBatchSize = 128;
 #if !UE_BUILD_SHIPPING
diff --git a/Engine/Source/Runtime/Experimental/Chaos/Public/Chaos/Capsule.h b/Engine/Source/Runtime/Experimental/Chaos/Public/Chaos/Capsule.h
--- a/Engine/Source/Runtime/Experimental/Chaos/Public/Chaos/Capsule.h
+++ b/Engine/Source/Runtime/Experimental/Chaos/Public/Chaos/Capsule.h
@@ -582,7 +582,7 @@ namespace Chaos
 		// Used for manifold generation
 		int32 NumPlanes() const { return 0; }
 
-#if INTEL_ISPC && !UE_BUILD_SHIPPING
+#if INTEL_ISPC
 		// See PerParticlePBDCollisionConstraint.cpp
 		// ISPC code has matching structs for interpreting FImplicitObjects.
 		// This is used to verify that the structs stay the same.
@@ -592,7 +592,7 @@ namespace Chaos
 			static constexpr int32 SizeOfMSegment() { return sizeof(FCapsule::MSegment); }
 		};
 		friend FISPCDataVerifier;
-#endif // #if INTEL_ISPC && !UE_BUILD_SHIPPING
+#endif // #if INTEL_ISPC
 
 	private:
 		void SetRadius(FReal InRadius) { SetMargin(InRadius); }
diff --git a/Engine/Source/Runtime/Experimental/Chaos/Public/Chaos/Convex.h b/Engine/Source/Runtime/Experimental/Chaos/Public/Chaos/Convex.h
--- a/Engine/Source/Runtime/Experimental/Chaos/Public/Chaos/Convex.h
+++ b/Engine/Source/Runtime/Experimental/Chaos/Public/Chaos/Convex.h
@@ -1028,7 +1028,7 @@ namespace Chaos
 			return FVec3(0);
 		}
 
-#if INTEL_ISPC && !UE_BUILD_SHIPPING
+#if INTEL_ISPC
 		// See PerParticlePBDCollisionConstraint.cpp
 		// ISPC code has matching structs for interpreting FImplicitObjects.
 		// This is used to verify that the structs stay the same.
@@ -1042,7 +1042,7 @@ namespace Chaos
 			static constexpr int32 SizeOfStructureData() { return sizeof(FConvex::StructureData); }
 		};
 		friend FISPCDataVerifier;
-#endif // #if INTEL_ISPC && !UE_BUILD_SHIPPING
+#endif // #if INTEL_ISPC
 
 	private:
 		void ComputeUnitMassInertiaTensorAndRotationOfMass(const FReal InVolume);
diff --git a/Engine/Source/Runtime/Experimental/Chaos/Public/Chaos/ConvexHalfEdgeStructureData.h b/Engine/Source/Runtime/Experimental/Chaos/Public/Chaos/ConvexHalfEdgeStructureData.h
--- a/Engine/Source/Runtime/Experimental/Chaos/Public/Chaos/ConvexHalfEdgeStructureData.h
+++ b/Engine/Source/Runtime/Experimental/Chaos/Public/Chaos/ConvexHalfEdgeStructureData.h
@@ -475,7 +475,7 @@ namespace Chaos
 			return Ar;
 		}
 
-#if INTEL_ISPC && !UE_BUILD_SHIPPING
+#if INTEL_ISPC
 		// See PerParticlePBDCollisionConstraint.cpp
 		// ISPC code has matching structs for interpreting FImplicitObjects.
 		// This is used to verify that the structs stay the same.
@@ -493,7 +493,7 @@ namespace Chaos
 			static constexpr int32 SizeOfVertexPlanes() { return sizeof(TConvexHalfEdgeStructureData::VertexPlanes); }
 		};
 		friend FISPCDataVerifier;
-#endif // #if INTEL_ISPC && !UE_BUILD_SHIPPING
+#endif // #if INTEL_ISPC
 
 	private:
 
diff --git a/Engine/Source/Runtime/Experimental/Chaos/Public/Chaos/ConvexStructureData.h b/Engine/Source/Runtime/Experimental/Chaos/Public/Chaos/ConvexStructureData.h
--- a/Engine/Source/Runtime/Experimental/Chaos/Public/Chaos/ConvexStructureData.h
+++ b/Engine/Source/Runtime/Experimental/Chaos/Public/Chaos/ConvexStructureData.h
@@ -286,7 +286,7 @@ namespace Chaos
 			return Ar;
 		}
 
-#if INTEL_ISPC && !UE_BUILD_SHIPPING
+#if INTEL_ISPC
 		// See PerParticlePBDCollisionConstraint.cpp
 		// ISPC code has matching structs for interpreting FImplicitObjects.
 		// This is used to verify that the structs stay the same.
@@ -298,7 +298,7 @@ namespace Chaos
 			static constexpr int32 SizeOfIndexType() { return sizeof(FConvexStructureData::IndexType); }
 		};
 		friend FISPCDataVerifier;
-#endif // #if INTEL_ISPC && !UE_BUILD_SHIPPING
+#endif // #if INTEL_ISPC
 
 	private:
 
diff --git a/Engine/Source/Runtime/Experimental/Chaos/Public/Chaos/ImplicitObject.h b/Engine/Source/Runtime/Experimental/Chaos/Public/Chaos/ImplicitObject.h
--- a/Engine/Source/Runtime/Experimental/Chaos/Public/Chaos/ImplicitObject.h
+++ b/Engine/Source/Runtime/Experimental/Chaos/Public/Chaos/ImplicitObject.h
@@ -83,6 +83,21 @@ using FImplicitHierarchyVisitor = TFunctionRef<void(const FImplicitObject* Impli
  * so we do not need a margin on triangles. We would need a margin on every type
  * that can be tested against a triangle though.
  */
+// Some collision ispc code requires that no tail padding in FImplicitObject is reused by derived class members. 
+// This is a compiler-dependent behavior, so if you are not seeing any other compile time errors about sizeof(FImplicitObject) + offsetof(...) with this disabled,
+// you should be OK.
+#define DISALLOW_FIMPLICIT_OBJECT_TAIL_PADDING INTEL_ISPC
+
+#if DISALLOW_FIMPLICIT_OBJECT_TAIL_PADDING
+// This enables errors if any padding is added by the compiler. You can fix the errors by rearranging fields and/or adding explicit padding.
+#if defined(__clang__) || defined(__GNUC__)
+#pragma GCC diagnostic push
+#pragma GCC diagnostic error "-Wpadded"
+#elif defined(_MSC_VER)
+#pragma warning(push)
+#pragma warning(error : 4820)
+#endif
+#endif // #if DISALLOW_FIMPLICIT_OBJECT_TAIL_PADDING
 class CHAOS_API FImplicitObject
 {
 public:
@@ -437,8 +452,7 @@ protected:
 	// in a derived class if at all (and it may want to change the size of the core shape as well)
 	void SetMargin(FReal InMargin) { Margin = InMargin; }
 
-	EImplicitObjectType Type;
-	EImplicitObjectType CollisionType;
+	// Note: if you change the size of FImplicitObject or add fields, you will likely need to update the calculation of PadBytes below.
 	FReal Margin;
 	bool bIsConvex;
 	bool bDoCollide;
@@ -446,11 +460,38 @@ protected:
 
 #if TRACK_CHAOS_GEOMETRY
 	bool bIsTracked;
+#else
+#if DISALLOW_FIMPLICIT_OBJECT_TAIL_PADDING
+	// The purpose of this padding is just to make it easier to calculate PadBytes below.
+	bool bPad;
+#endif // DISALLOW_FIMPLICIT_OBJECT_TAIL_PADDING
 #endif
 
+	EImplicitObjectType Type;
+	EImplicitObjectType CollisionType;
+
+#if DISALLOW_FIMPLICIT_OBJECT_TAIL_PADDING
+	// the following assumptions are made when calculating PadBytes and inserting any explicit padding between fields:
+	static_assert(alignof(bool) <= sizeof(FReal)); // Otherwise, padding would be added between Margin the first bool field
+	static_assert(alignof(EImplicitObjectType) <= 4 * sizeof(bool));  // Otherwise, padding would be added between last bool field and first EImplicitObjectType field
+
+	static constexpr int AlignOfFImplicitObject = (int)FMath::Max(alignof(FReal), FMath::Max(alignof(bool), alignof(EImplicitObjectType)));
+	static constexpr int PadBytes =
+		AlignOfFImplicitObject - (sizeof(FReal) + 4 * sizeof(bool) + 2 * sizeof(EImplicitObjectType)) % AlignOfFImplicitObject;
+	static_assert(PadBytes > 0);
+	char Pad[PadBytes];
+#endif // DISALLOW_FIMPLICIT_OBJECT_TAIL_PADDING
+
 private:
 	virtual Pair<FVec3, bool> FindClosestIntersectionImp(const FVec3& StartPoint, const FVec3& EndPoint, const FReal Thickness) const;
 };
+#if DISALLOW_FIMPLICIT_OBJECT_TAIL_PADDING
+#if defined(__clang__) || defined(__GNUC__)
+#pragma GCC diagnostic pop
+#elif defined(_MSC_VER)
+#pragma warning(pop)
+#endif
+#endif // DISALLOW_FIMPLICIT_OBJECT_TAIL_PADDING
 
 FORCEINLINE FChaosArchive& operator<<(FChaosArchive& Ar, FImplicitObject& Value)
 {
diff --git a/Engine/Source/Runtime/Experimental/Chaos/Public/Chaos/ImplicitObjectUnion.h b/Engine/Source/Runtime/Experimental/Chaos/Public/Chaos/ImplicitObjectUnion.h
--- a/Engine/Source/Runtime/Experimental/Chaos/Public/Chaos/ImplicitObjectUnion.h
+++ b/Engine/Source/Runtime/Experimental/Chaos/Public/Chaos/ImplicitObjectUnion.h
@@ -208,7 +208,7 @@ class CHAOS_API FImplicitObjectUnion : public FImplicitObject
 		return Union;
 	}
 
-#if INTEL_ISPC && !UE_BUILD_SHIPPING
+#if INTEL_ISPC
 	// See PerParticlePBDCollisionConstraint.cpp
 	// ISPC code has matching structs for interpreting FImplicitObjects.
 	// This is used to verify that the structs stay the same.
@@ -218,7 +218,7 @@ class CHAOS_API FImplicitObjectUnion : public FImplicitObject
 		static constexpr int32 SizeOfMObjects() { return sizeof(FImplicitObjectUnion::MObjects); }
 	};
 	friend FISPCDataVerifier;
-#endif // #if INTEL_ISPC && !UE_BUILD_SHIPPING
+#endif // #if INTEL_ISPC
 
 protected:
 	virtual Pair<FVec3, bool> FindClosestIntersectionImp(const FVec3& StartPoint, const FVec3& EndPoint, const FReal Thickness) const override
diff --git a/Engine/Source/Runtime/Experimental/Chaos/Public/Chaos/Sphere.h b/Engine/Source/Runtime/Experimental/Chaos/Public/Chaos/Sphere.h
--- a/Engine/Source/Runtime/Experimental/Chaos/Public/Chaos/Sphere.h
+++ b/Engine/Source/Runtime/Experimental/Chaos/Public/Chaos/Sphere.h
@@ -372,7 +372,7 @@ namespace Chaos
 			return  TUniquePtr<FImplicitObject>(new TSphere<T, d>(Center * Scale, GetRadius() * Scale.Min()));
 		}
 
-#if INTEL_ISPC && !UE_BUILD_SHIPPING
+#if INTEL_ISPC
 		// See PerParticlePBDCollisionConstraint.cpp
 		// ISPC code has matching structs for interpreting FImplicitObjects.
 		// This is used to verify that the structs stay the same.
@@ -382,7 +382,7 @@ namespace Chaos
 			static constexpr int32 SizeOfCenter() { return sizeof(TSphere::Center); }
 		};
 		friend FISPCDataVerifier;
-#endif // #if INTEL_ISPC && !UE_BUILD_SHIPPING
+#endif // #if INTEL_ISPC
 
 	private:
 		void SetRadius(FReal InRadius) { SetMargin(InRadius); }
diff --git a/Engine/Source/Runtime/Experimental/Chaos/Public/Chaos/TaperedCapsule.h b/Engine/Source/Runtime/Experimental/Chaos/Public/Chaos/TaperedCapsule.h
--- a/Engine/Source/Runtime/Experimental/Chaos/Public/Chaos/TaperedCapsule.h
+++ b/Engine/Source/Runtime/Experimental/Chaos/Public/Chaos/TaperedCapsule.h
@@ -218,7 +218,7 @@ namespace Chaos
 			return HashCombine(OriginAxisHash, PropertyHash);
 		}
 
-#if INTEL_ISPC && !UE_BUILD_SHIPPING
+#if INTEL_ISPC
 		// See PerParticlePBDCollisionConstraint.cpp
 		// ISPC code has matching structs for interpreting FImplicitObjects.
 		// This is used to verify that the structs stay the same.
@@ -236,7 +236,7 @@ namespace Chaos
 			static constexpr int32 SizeOfRadius2() { return sizeof(FTaperedCapsule::Radius2); }
 		};
 		friend FISPCDataVerifier;
-#endif // #if INTEL_ISPC && !UE_BUILD_SHIPPING
+#endif // #if INTEL_ISPC
 
 	private:
 		//Phi is distance from closest point on plane1
diff --git a/Engine/Source/Runtime/Experimental/Chaos/Public/Chaos/TaperedCylinder.h b/Engine/Source/Runtime/Experimental/Chaos/Public/Chaos/TaperedCylinder.h
--- a/Engine/Source/Runtime/Experimental/Chaos/Public/Chaos/TaperedCylinder.h
+++ b/Engine/Source/Runtime/Experimental/Chaos/Public/Chaos/TaperedCylinder.h
@@ -330,7 +330,7 @@ namespace Chaos
 			return HashCombine(PlaneHashes, PropertyHash);
 		}
 
-#if INTEL_ISPC && !UE_BUILD_SHIPPING
+#if INTEL_ISPC
 		// See PerParticlePBDCollisionConstraint.cpp
 		// ISPC code has matching structs for interpreting FImplicitObjects.
 		// This is used to verify that the structs stay the same.
@@ -348,7 +348,7 @@ namespace Chaos
 			static constexpr int32 SizeOfMRadius2() { return sizeof(FTaperedCylinder::MRadius2); }
 		};
 		friend FISPCDataVerifier;
-#endif // #if INTEL_ISPC && !UE_BUILD_SHIPPING
+#endif // #if INTEL_ISPC
 
 	private:
 		//Phi is distance from closest point on plane1
