Chaos - fix missing collisions on cluster unions

diff --git a/Engine/Source/Runtime/Experimental/Chaos/Public/Chaos/PBDRigidsEvolution.h b/Engine/Source/Runtime/Experimental/Chaos/Public/Chaos/PBDRigidsEvolution.h
--- a/Engine/Source/Runtime/Experimental/Chaos/Public/Chaos/PBDRigidsEvolution.h
+++ b/Engine/Source/Runtime/Experimental/Chaos/Public/Chaos/PBDRigidsEvolution.h
@@ -566,9 +566,7 @@ public:
 
 			auto& AsyncSpatialData = AsyncAccelerationQueue.FindOrAdd(UniqueIdx, EPendingSpatialDataOperation::Update);
 			// ensure(AsyncSpatialData.Operation != EPendingSpatialDataOperation::Delete); // TODO: This may be hit: check for potential UniqueIdx reuse bug
-			
-			AsyncSpatialData.AccelerationHandle = FAccelerationStructureHandle(Particle);
-			AsyncSpatialData.SpatialIdx = Particle.SpatialIdx();
+			AsyncSpatialData = SpatialData;
 		}
 	}
 
