Chaos - missing building parts fix

diff --git a/Engine/Source/Runtime/Experimental/Chaos/Public/Chaos/PBDRigidsEvolution.h b/Engine/Source/Runtime/Experimental/Chaos/Public/Chaos/PBDRigidsEvolution.h
--- a/Engine/Source/Runtime/Experimental/Chaos/Public/Chaos/PBDRigidsEvolution.h
+++ b/Engine/Source/Runtime/Experimental/Chaos/Public/Chaos/PBDRigidsEvolution.h
@@ -453,7 +453,7 @@ public:
 		Particles.EnableParticle(Particle);
 		EnableConstraints(Particle);
 		IslandManager.AddParticle(Particle);
-		DirtyParticle(*Particle);
+		DirtyParticle(*Particle, EPendingSpatialDataOperation::Add);
 	}
 
 	/**
@@ -502,7 +502,7 @@ public:
 	}
 
 	template <bool bPersistent>
-	FORCEINLINE_DEBUGGABLE void DirtyParticle(TGeometryParticleHandleImp<FReal, 3, bPersistent>& Particle)
+	FORCEINLINE_DEBUGGABLE void DirtyParticle(TGeometryParticleHandleImp<FReal, 3, bPersistent>& Particle, const EPendingSpatialDataOperation Op = EPendingSpatialDataOperation::Update)
 	{
 		const TPBDRigidParticleHandleImp<FReal, 3, bPersistent>* AsRigid = Particle.CastToRigidParticle();
 		if(AsRigid && AsRigid->Disabled())
@@ -560,6 +560,10 @@ public:
 			const FUniqueIdx UniqueIdx = Particle.UniqueIdx();
 			FPendingSpatialData& SpatialData = InternalAccelerationQueue.FindOrAdd(UniqueIdx, EPendingSpatialDataOperation::Update);
 			ensure(SpatialData.Operation != EPendingSpatialDataOperation::Delete);
+			if (Op == EPendingSpatialDataOperation::Add)
+			{
+				SpatialData.Operation = Op;
+			}
 			
 			SpatialData.AccelerationHandle = FAccelerationStructureHandle(Particle);
 			SpatialData.SpatialIdx = Particle.SpatialIdx();
