Iris - Subobject destruction fix when root object is destroyed or pending destroy.

#rnx

diff --git a/Engine/Source/Runtime/Experimental/Iris/Core/Private/Iris/ReplicationSystem/ReplicationWriter.cpp b/Engine/Source/Runtime/Experimental/Iris/Core/Private/Iris/ReplicationSystem/ReplicationWriter.cpp
--- a/Engine/Source/Runtime/Experimental/Iris/Core/Private/Iris/ReplicationSystem/ReplicationWriter.cpp
+++ b/Engine/Source/Runtime/Experimental/Iris/Core/Private/Iris/ReplicationSystem/ReplicationWriter.cpp
@@ -1595,8 +1595,8 @@ uint32 FReplicationWriter::WriteObjectsPendingDestroy(FNetSerializationContext&
 		{
 			if (Info.SubObjectPendingDestroy)
 			{
-				// If the owner is pending destroy then we need to replicate the destruction for the subobject here.
-				if (ObjectsPendingDestroy.GetBit(ObjectData.SubObjectRootIndex))
+				// If the owner is pending destroy or already destroyed then we need to replicate the destruction for the subobject here.
+				if (ObjectsPendingDestroy.GetBit(ObjectData.SubObjectRootIndex) || (ReplicatedObjects[ObjectData.SubObjectRootIndex].GetState() == EReplicatedObjectState::Invalid))
 				{
 					//			UE_LOG(LogTemp, Warning, TEXT("Object %s was a subobject but owner was destroyed so we need to destroy object normally"), ToCStr(ObjectData.RefHandle.ToString()));
 					SetState(InternalIndex, EReplicatedObjectState::PendingDestroy);
