Chaos : fix GC replication forcing all children of cluster union to break off and be processed in the deferred removal
- fixed by stopping the force removal when finding a cluster union particle

diff --git a/Engine/Source/Runtime/Experimental/Chaos/Private/Chaos/PBDRigidClustering.cpp b/Engine/Source/Runtime/Experimental/Chaos/Private/Chaos/PBDRigidClustering.cpp
--- a/Engine/Source/Runtime/Experimental/Chaos/Private/Chaos/PBDRigidClustering.cpp
+++ b/Engine/Source/Runtime/Experimental/Chaos/Private/Chaos/PBDRigidClustering.cpp
@@ -1050,11 +1050,18 @@ namespace Chaos
 				// first release any parent if any
 				if (FPBDRigidClusteredParticleHandle* ParentCluster = ChildClusteredParticle->Parent())
 				{
-					// we need now to force parents to break
-					ForceReleaseChildParticleAndParents(ParentCluster, bTriggerBreakEvents);
+					if (const IPhysicsProxyBase* ParentProxy = ParentCluster->PhysicsProxy())
+					{
+						// we shoudl not break cluster union parent and stop recursion there
+						if (ParentProxy->GetType() != FClusterUnionPhysicsProxy::ConcreteType())
+						{
+							// we need now to force parents to break
+							ForceReleaseChildParticleAndParents(ParentCluster, bTriggerBreakEvents);
 
-					SetExternalStrain(ChildClusteredParticle, TNumericLimits<FRealSingle>::Max());
-					ReleaseClusterParticles(ParentCluster, bTriggerBreakEvents);
+							SetExternalStrain(ChildClusteredParticle, TNumericLimits<FRealSingle>::Max());
+							ReleaseClusterParticles(ParentCluster, bTriggerBreakEvents);
+						}
+					}
 				}
 			}
 		}
