Chaos - fix spatial acceleration desync crash

We were leaving deleted particles in the async spatial acceleration stuructures which if used in scene queries would result in a garbage pointer dereference

diff --git a/Engine/Source/Runtime/Experimental/Chaos/Private/Chaos/PBDRigidsEvolution.cpp b/Engine/Source/Runtime/Experimental/Chaos/Private/Chaos/PBDRigidsEvolution.cpp
--- a/Engine/Source/Runtime/Experimental/Chaos/Private/Chaos/PBDRigidsEvolution.cpp
+++ b/Engine/Source/Runtime/Experimental/Chaos/Private/Chaos/PBDRigidsEvolution.cpp
@@ -630,7 +630,7 @@ namespace Chaos
 		for (int32 Idx = ExternalQueue.PendingData.Num() - 1; Idx >=0; --Idx)
 		{
 			const FPendingSpatialData& SpatialData = ExternalQueue.PendingData[Idx];
-			if(SpatialData.SyncTimestamp > SyncTimestamp)
+			if(SpatialData.SyncTimestamp >= SyncTimestamp)
 			{
 				//operation still pending so update structure
 				//note: do we care about roll over? if game ticks at 60fps we'd get 385+ days
