replicated cluster unions need to handle being streamed out on clients properly

- The bRecomputeMassOrientation flag needs to be respected on the client irregardless of the number of children.

diff --git a/Engine/Source/Runtime/Experimental/Chaos/Private/Chaos/PBDRigidClusteringAlgo.cpp b/Engine/Source/Runtime/Experimental/Chaos/Private/Chaos/PBDRigidClusteringAlgo.cpp
--- a/Engine/Source/Runtime/Experimental/Chaos/Private/Chaos/PBDRigidClusteringAlgo.cpp
+++ b/Engine/Source/Runtime/Experimental/Chaos/Private/Chaos/PBDRigidClusteringAlgo.cpp
@@ -115,14 +115,22 @@ namespace Chaos
 		if (!bHasChild || !(Parent->M() > UE_SMALL_NUMBER))
 		{
 			Parent->M() = 1.0;
-			Parent->X() = FVec3(0);
+			if (ForceMassOrientation)
+			{
+				Parent->X() = ForceMassOrientation->GetLocation();
+				Parent->R() = ForceMassOrientation->GetRotation();
+			}
+			else
+			{
+				Parent->X() = FVec3(0);
+				Parent->R() = FRotation3(FMatrix::Identity);
+			}
 			Parent->V() = FVec3(0);
 			Parent->PreV() = Parent->V();
 			Parent->InvM() = 1;
 			Parent->P() = Parent->X();
 			Parent->W() = FVec3(0);
 			Parent->PreW() = Parent->W();
-			Parent->R() = FRotation3(FMatrix::Identity);
 			Parent->Q() = Parent->R();
 			Parent->I() = TVec3<FRealSingle>(1);
 			Parent->InvI() = TVec3<FRealSingle>(1);
