Proper fix input/output serialization for dataflow nodes to allow adding inputs and outputs

diff --git a/Engine/Plugins/Experimental/ChaosClothAssetEditor/Source/ChaosClothAssetDataflowNodes/Private/ChaosClothAsset/TerminalNode.cpp b/Engine/Plugins/Experimental/ChaosClothAssetEditor/Source/ChaosClothAssetDataflowNodes/Private/ChaosClothAsset/TerminalNode.cpp
--- a/Engine/Plugins/Experimental/ChaosClothAssetEditor/Source/ChaosClothAssetDataflowNodes/Private/ChaosClothAsset/TerminalNode.cpp
+++ b/Engine/Plugins/Experimental/ChaosClothAssetEditor/Source/ChaosClothAssetDataflowNodes/Private/ChaosClothAsset/TerminalNode.cpp
@@ -193,4 +193,21 @@ TArray<const FManagedArrayCollection*> FChaosClothAssetTerminalNode::GetCollecti
 	return CollectionLods;
 }
 
+void FChaosClothAssetTerminalNode::Serialize(FArchive& Ar)
+{
+	// because we add pins we need to make sure we restore them when loading
+	// to make sure they can get properly reconnected
+
+	if (Ar.IsLoading())
+	{
+		const int32 NumLodToAdd = (NumLods - 1);
+		NumLods = 1; // need to reset back to default because add pin will increment it again 
+		for (int32 LodIndex = 0; LodIndex < NumLodToAdd; ++LodIndex)
+		{
+			AddPin();
+		}
+		ensure(NumLodToAdd == (NumLods - 1));
+	}
+}
+
 #undef LOCTEXT_NAMESPACE
diff --git a/Engine/Plugins/Experimental/ChaosClothAssetEditor/Source/ChaosClothAssetDataflowNodes/Public/ChaosClothAsset/TerminalNode.h b/Engine/Plugins/Experimental/ChaosClothAssetEditor/Source/ChaosClothAssetDataflowNodes/Public/ChaosClothAsset/TerminalNode.h
--- a/Engine/Plugins/Experimental/ChaosClothAssetEditor/Source/ChaosClothAssetDataflowNodes/Public/ChaosClothAsset/TerminalNode.h
+++ b/Engine/Plugins/Experimental/ChaosClothAssetEditor/Source/ChaosClothAssetDataflowNodes/Public/ChaosClothAsset/TerminalNode.h
@@ -48,6 +48,7 @@ public:
 	virtual bool CanAddPin() const override { return NumLods < MaxLods; }
 	virtual bool CanRemovePin() const override { return NumLods > 1; }
 	virtual Dataflow::FPin RemovePin() override;
+	virtual void Serialize(FArchive& Ar);
 
 private:
 	TArray<const FManagedArrayCollection*> GetCollectionLods() const;
diff --git a/Engine/Source/Runtime/Core/Public/UObject/UE5MainStreamObjectVersion.h b/Engine/Source/Runtime/Core/Public/UObject/UE5MainStreamObjectVersion.h
--- a/Engine/Source/Runtime/Core/Public/UObject/UE5MainStreamObjectVersion.h
+++ b/Engine/Source/Runtime/Core/Public/UObject/UE5MainStreamObjectVersion.h
@@ -357,6 +357,9 @@ struct CORE_API FUE5MainStreamObjectVersion
 		
 		// Store the RigVM execute context struct the VM uses in the archive
 		RigVMExternalExecuteContextStruct,
+
+		// serialization inputs and outputs as two different sections
+		DataflowSeparateInputOutputSerialization,
 		
 		// -----<new versions can be added above this line>-------------------------------------------------
 		VersionPlusOne,
diff --git a/Engine/Source/Runtime/Experimental/Dataflow/Core/Private/Dataflow/DataflowGraph.cpp b/Engine/Source/Runtime/Experimental/Dataflow/Core/Private/Dataflow/DataflowGraph.cpp
--- a/Engine/Source/Runtime/Experimental/Dataflow/Core/Private/Dataflow/DataflowGraph.cpp
+++ b/Engine/Source/Runtime/Experimental/Dataflow/Core/Private/Dataflow/DataflowGraph.cpp
@@ -7,6 +7,7 @@
 #include "Dataflow/DataflowNodeFactory.h"
 #include "Logging/LogMacros.h"
 #include "Dataflow/DataflowArchive.h"
+#include "UObject/UE5MainStreamObjectVersion.h"
 
 DEFINE_LOG_CATEGORY_STATIC(DATAFLOW_LOG, Error, All);
 
@@ -120,6 +121,8 @@ namespace Dataflow
 
 	void FGraph::Serialize(FArchive& Ar)
 	{
+		Ar.UsingCustomVersion(FUE5MainStreamObjectVersion::GUID);
+
 		Ar << Guid;
 		if (Ar.IsSaving())
 		{
@@ -137,20 +140,30 @@ namespace Dataflow
 
 				DATAFLOW_OPTIONAL_BLOCK_WRITE_BEGIN()
 				{
-					TArray< FDataflowConnection* > IO;
-					IO.Append(Node->GetOutputs());
-					IO.Append(Node->GetInputs());
-					ArNum = IO.Num();
-					Ar << ArNum;
-					for (FDataflowConnection* Conn : IO)
+					// Node needs to be serialized first to make sure it registers all the dynamic input/output for when input and output will be deserialized
+					Node->SerializeInternal(Ar);
+
+					// keep outputs and inputs separated even though their serialization code looks almost identical
+					// this is to make sure we can handle when number of inputs or outputs have changed on the node
+					int32 ArNumOutputs = Node->GetOutputs().Num();
+					Ar << ArNumOutputs;
+					for (FDataflowConnection* Output : Node->GetOutputs())
 					{
-						ArGuid = Conn->GetGuid();
-						ArType = Conn->GetType();
-						ArName = Conn->GetName();
+						ArGuid = Output->GetGuid();
+						ArType = Output->GetType();
+						ArName = Output->GetName();
 						Ar << ArGuid << ArType << ArName;
 					}
 
-					Node->SerializeInternal(Ar);
+					int32 ArNumInputs = Node->GetInputs().Num();
+					Ar << ArNumInputs;
+					for (FDataflowConnection* Input: Node->GetInputs())
+					{
+						ArGuid = Input->GetGuid();
+						ArType = Input->GetType();
+						ArName = Input->GetName();
+						Ar << ArGuid << ArType << ArName;
+					}
 				}
 				DATAFLOW_OPTIONAL_BLOCK_WRITE_END();
 			}
@@ -170,64 +183,131 @@ namespace Dataflow
 			Ar << ArNum;
 			for (int32 Ndx = ArNum; Ndx > 0; Ndx--)
 			{
-				Ar << ArGuid << ArType << ArName;
+				FName ArNodeName;
+				Ar << ArGuid << ArType << ArNodeName;
 
-				TSharedPtr<FDataflowNode> Node = FNodeFactory::GetInstance()->NewNodeFromRegisteredType(*this, { ArGuid,ArType,ArName });
+				TSharedPtr<FDataflowNode> Node = FNodeFactory::GetInstance()->NewNodeFromRegisteredType(*this, { ArGuid,ArType, ArNodeName });
 				DATAFLOW_OPTIONAL_BLOCK_READ_BEGIN(Node != nullptr)
 				{
 					ensure(!NodeGuidMap.Contains(ArGuid));
 					NodeGuidMap.Add(ArGuid, Node);
 
-					int ArNumInputsOutputs;
-					Ar << ArNumInputsOutputs;
-					TArray< FDataflowConnection* > InputsOutputs;
-					InputsOutputs.Append(Node->GetOutputs());
-					InputsOutputs.Append(Node->GetInputs());
+					if ((Ar.CustomVer(FUE5MainStreamObjectVersion::GUID) < FUE5MainStreamObjectVersion::DataflowSeparateInputOutputSerialization))
+					{
+
+						// former input / output serialization method where we only store aggregate number of inputs and outputs
+						// this has limitation when adding more outputs or inputs
+						int ArNumInputsOutputs;
+						Ar << ArNumInputsOutputs;
+						TArray< FDataflowConnection* > InputsOutputs;
+						InputsOutputs.Append(Node->GetOutputs());
+						InputsOutputs.Append(Node->GetInputs());
+
+						// skip offset is use to correct the mismatch of outputs have been added
+						int SkipOffset = 0;
+						for (int ConnectionIndex = 0; ConnectionIndex < ArNumInputsOutputs; ConnectionIndex++)
+						{
+							Ar << ArGuid << ArType << ArName;
+
+							int AdjustedConnectionIndex = ConnectionIndex + SkipOffset;
+							if (InputsOutputs.IsValidIndex(AdjustedConnectionIndex))
+							{
+								FDataflowConnection* Connection = InputsOutputs[AdjustedConnectionIndex];
+
+								// if the name does not match this means the node has changed since the last serialization 
+								// ( added outputs for example that shift the index )
+								// in that case we try to recover by finding the next good node
+								// note we cannot just find by name as some nodes have inputs and outputs named the same 
+								// todo: implement a better way to serialize inputs and outputs seperately to avoid this case
+								while (Connection && Connection->GetName() != ArName)
+								{
+									SkipOffset++;
+									AdjustedConnectionIndex = ConnectionIndex + SkipOffset;
+									if (InputsOutputs.IsValidIndex(AdjustedConnectionIndex))
+									{
+										Connection = InputsOutputs[AdjustedConnectionIndex];
+									}
+									else
+									{
+										Connection = nullptr;
+									}
+								}
+								if (Connection)
+								{
+									check(Connection->GetType() == ArType);
+									Connection->SetGuid(ArGuid);
+									ensure(!ConnectionGuidMap.Contains(ArGuid));
+									ConnectionGuidMap.Add(ArGuid, Connection);
+								}
+							}
+						}
 
-					// skip offset is use to correct the mismatch of outputs have been added
-					int SkipOffset = 0;
-					for (int ConnectionIndex = 0; ConnectionIndex < ArNumInputsOutputs; ConnectionIndex++)
+						Node->SerializeInternal(Ar);
+					}
+					else
 					{
-						Ar << ArGuid << ArType << ArName;
+						// we need to deserilaize the node first because if it may add more inputs that may 
+						// be referenced when deserializing them below ( see Dataflow Node AddPin method )
+						Node->SerializeInternal(Ar);
 
-						int AdjustedConnectionIndex = ConnectionIndex + SkipOffset;
-						if (InputsOutputs.IsValidIndex(AdjustedConnectionIndex))
+						// Outputs deserialization
 						{
-							FDataflowConnection* Connection = InputsOutputs[AdjustedConnectionIndex];
-
-							// if the name does not match this means the node has changed since the last serialization 
-							// ( added outputs for example that shift the index )
-							// in that case we try to recover by finding the next good node
-							// note we cannot just find by name as some nodes have inputs and outputs named the same 
-							// todo: implement a better way to serialize inputs and outputs seperately to avoid this case
-							while (Connection && Connection->GetName() != ArName)
+							int32 ArNumOutputs;
+							Ar << ArNumOutputs;
+
+							for (int32 OutputIndex = 0; OutputIndex < ArNumOutputs; OutputIndex++)
 							{
-								SkipOffset++;
-								AdjustedConnectionIndex = ConnectionIndex + SkipOffset;
-								if (InputsOutputs.IsValidIndex(AdjustedConnectionIndex))
+								Ar << ArGuid << ArType << ArName;
+
+								if (FDataflowOutput* Output = Node->FindOutput(ArName))
 								{
-									Connection = InputsOutputs[AdjustedConnectionIndex];
+									check(Output->GetType() == ArType);
+									Output->SetGuid(ArGuid);
+									ensure(!ConnectionGuidMap.Contains(ArGuid));
+									ConnectionGuidMap.Add(ArGuid, Output);
 								}
 								else
 								{
-									Connection = nullptr;
+									// output has been serialized but cannot be found
+									// this means the definition of the node has changed and the output is no longer registered
+									ensureMsgf(false,
+										TEXT("Error: Cannot find registered output (%s) in node (%s) - this may result in missing connection ")
+										, *ArName.ToString(), *ArNodeName.ToString());
 								}
 							}
-							if (Connection)
+						}
+
+						// Inputs deserialization
+						{
+							int32 ArNumInputs;
+							Ar << ArNumInputs;
+
+							for (int32 InputIndex = 0; InputIndex < ArNumInputs; InputIndex++)
 							{
-								check(Connection->GetType() == ArType);
-								Connection->SetGuid(ArGuid);
-								ensure(!ConnectionGuidMap.Contains(ArGuid));
-								ConnectionGuidMap.Add(ArGuid, Connection);
+								Ar << ArGuid << ArType << ArName;
+
+								if (FDataflowInput* Input = Node->FindInput(ArName))
+								{
+									check(Input->GetType() == ArType);
+									Input->SetGuid(ArGuid);
+									ensure(!ConnectionGuidMap.Contains(ArGuid));
+									ConnectionGuidMap.Add(ArGuid, Input);
+								}
+								else
+								{
+									// input has been serialized but cannot be found
+									// this means the definition of the node has changed and the input is no longer registered
+									ensureMsgf(false,
+										TEXT("Error: Cannot find registered input (%s) in node (%s) - this may result in missing connection ")
+										, *ArName.ToString(), *ArNodeName.ToString());
+								}
 							}
 						}
 					}
-
-					Node->SerializeInternal(Ar);
 				}
 				DATAFLOW_OPTIONAL_BLOCK_READ_ELSE()
 				{
-					DisabledNodes.Add(ArName);
+					DisabledNodes.Add(ArNodeName);
 					ensureMsgf(false,
 						TEXT("Error: Missing registered node type (%s) will be removed from graph on load. Graph will fail to evaluate due to missing node (%s).")
 						, *ArType.ToString(), *ArName.ToString());
@@ -254,4 +334,3 @@ namespace Dataflow
 		}
 	}
 }
-
