ChaosFlesh : Add nodal coloring for Gauss Seidel method

diff --git a/Engine/Source/Runtime/Experimental/Chaos/Private/Chaos/GraphColoring.cpp b/Engine/Source/Runtime/Experimental/Chaos/Private/Chaos/GraphColoring.cpp
--- a/Engine/Source/Runtime/Experimental/Chaos/Private/Chaos/GraphColoring.cpp
+++ b/Engine/Source/Runtime/Experimental/Chaos/Private/Chaos/GraphColoring.cpp
@@ -259,6 +259,74 @@ static bool VerifyWeakConstraintsColoring(const TArray<TArray<int32>>& Indices,
 
 }
 
+template <typename T>
+static bool VerifyNodalColoring(const TArray<Chaos::TVec4<int32>>& Graph, const Chaos::TDynamicParticles<T, 3>& InParticles, const int32 GraphParticlesStart, const int32 GraphParticlesEnd, const TArray<TArray<int32>>& IncidentElements, const TArray<TArray<int32>>& IncidentElementsLocalIndex, const TArray<TArray<int32>> ParticlesPerColor)
+{
+	TArray<bool> ParticleIsIncluded;
+	ParticleIsIncluded.Init(false, InParticles.Size());
+	for (int32 i = 0; i < ParticlesPerColor.Num(); i++) 
+	{
+		for (int32 j = 0; j < ParticlesPerColor[i].Num(); j++)
+		{
+			ParticleIsIncluded[ParticlesPerColor[i][j]] = true;
+		}
+	}
+
+	for (int32 p = 0; p < GraphParticlesEnd - GraphParticlesStart; p++)
+	{
+		int32 ParticleIndex = p + GraphParticlesStart;
+		if (InParticles.InvM(ParticleIndex) != (T)0.)
+		{
+			if (!ParticleIsIncluded[ParticleIndex])
+			{
+				return false;
+			}
+		}
+	}
+
+	TArray<int32> Particle2Incident;
+	Particle2Incident.Init(-1, GraphParticlesEnd - GraphParticlesStart);
+	for (int32 i = 0; i < IncidentElements.Num(); i++)
+	{
+		if (IncidentElements[i].Num() > 0)
+		{
+			Particle2Incident[Graph[IncidentElements[i][0]][IncidentElementsLocalIndex[i][0]]] = i;
+		}
+	}
+
+	for (int32 i = 0; i < ParticlesPerColor.Num(); i++)
+	{
+		TSet<int32> IncidentParticles;
+		for (int32 j = 0; j < ParticlesPerColor[i].Num(); j++)
+		{
+			int32 ParticleIndex = ParticlesPerColor[i][j];
+			TSet<int32> LocalIncidentParticles;
+			for (int32 k = 0; k < IncidentElements[Particle2Incident[ParticleIndex]].Num(); k++)
+			{
+				int32 ElementIndex = IncidentElements[Particle2Incident[ParticleIndex]][k];
+				for (int32 ie = 0; ie < 4; ie++)
+				{
+					LocalIncidentParticles.Emplace(Graph[ElementIndex][ie]);
+				}
+			}
+			if (IncidentParticles.Contains(ParticleIndex))
+			{
+				return false;
+			}
+			else
+			{
+				for (int32 LocalParticle : LocalIncidentParticles)
+				{
+					IncidentParticles.Emplace(LocalParticle);
+				}
+			}
+		}
+	}
+
+	return true;
+
+}
+
 template<typename T>
 TArray<TArray<int32>> Chaos::FGraphColoring::ComputeGraphColoring(const TArray<Chaos::TVec2<int32>>& Graph, const Chaos::TDynamicParticles<T, 3>& InParticles, const int32 GraphParticlesStart, const int32 GraphParticlesEnd)
 {
@@ -1291,6 +1359,105 @@ void Chaos::ComputeWeakConstraintsColoring(const TArray<TArray<int32>>& Indices,
 }
 
 
+template<typename T>
+TArray<TArray<int32>> Chaos::ComputeNodalColoring(const TArray<TVec4<int32>>& Graph, const Chaos::TDynamicParticles<T, 3>& InParticles, const int32 GraphParticlesStart, const int32 GraphParticlesEnd, const TArray<TArray<int32>>& IncidentElements, const TArray<TArray<int32>>& IncidentElementsLocalIndex)
+{
+	using namespace Chaos;
+
+	checkSlow(GraphParticlesStart <= GraphParticlesEnd);
+	checkSlow(GraphParticlesEnd <= (int32)InParticles.Size());
+	TArray<TArray<int32>> ParticlesPerColor;
+
+	TArray<int32> Particle2Incident;
+	Particle2Incident.Init(-1, GraphParticlesEnd - GraphParticlesStart);
+	for (int32 i = 0; i < IncidentElements.Num(); i++) 
+	{
+		if (IncidentElements[i].Num() > 0)
+		{
+			Particle2Incident[Graph[IncidentElements[i][0]][IncidentElementsLocalIndex[i][0]]] = i;
+		}
+	}
+
+	TArray<TSet<int32>*> ElementColorsSet;
+	ElementColorsSet.Init(nullptr, Graph.Num());
+	TArray<int32> ParticleColors;
+	ParticleColors.Init(-1, GraphParticlesEnd - GraphParticlesStart);
+
+	for (int32 p = 0; p < GraphParticlesEnd - GraphParticlesStart; p++) 
+	{
+		int32 ParticleIndex = p + GraphParticlesStart;
+		if (InParticles.InvM(ParticleIndex) != (T)0. && Particle2Incident[ParticleIndex] != -1) 
+		{
+			for (int32 j = 0; j < IncidentElements[Particle2Incident[ParticleIndex]].Num(); j++) 
+			{
+				int32 LocalIndex = IncidentElementsLocalIndex[Particle2Incident[ParticleIndex]][j];
+				int32 e = IncidentElements[Particle2Incident[ParticleIndex]][j];
+				if (!ElementColorsSet[e]) 
+				{
+					ElementColorsSet[e] = new TSet<int32>();
+				}
+			}
+			int32 color_to_use = 0;
+			while (true) 
+			{
+				bool ColorFound = false;
+				for (int32 j = 0; j < IncidentElements[Particle2Incident[ParticleIndex]].Num(); j++) 
+				{
+					int32 e = IncidentElements[Particle2Incident[ParticleIndex]][j];
+					if (ElementColorsSet[e]->Contains(color_to_use)) 
+					{
+						ColorFound = true;
+						break;
+					}
+				}
+				if (!ColorFound) 
+				{
+					ParticleColors[p] = color_to_use;
+					for (int32 j = 0; j < IncidentElements[Particle2Incident[ParticleIndex]].Num(); j++) 
+					{
+						int32 e = IncidentElements[Particle2Incident[ParticleIndex]][j];
+						ElementColorsSet[e]->Emplace(color_to_use);
+						if (ElementColorsSet[e]->Num() == 4) 
+						{
+							delete ElementColorsSet[e];
+							ElementColorsSet[e] = nullptr;
+						}
+					}
+					break;
+				}
+				color_to_use++;
+			}
+		}
+	}
+
+	for (int32 ii = 0; ii < ElementColorsSet.Num(); ii++) 
+	{
+		if (ElementColorsSet[ii]) 
+		{
+			delete (ElementColorsSet[ii]);
+		}
+	}
+
+	int32 SizeColors = ParticleColors.Max();
+
+	ParticlesPerColor.SetNum(0);
+	ParticlesPerColor.Init(TArray<int32>(), SizeColors + 1);
+
+	for (int32 j = 0; j < ParticleColors.Num(); j++) 
+	{
+		if (ParticleColors[j] != -1) 
+		{
+			ParticlesPerColor[ParticleColors[j]].Emplace(j + GraphParticlesStart);
+		}
+	}
+
+	checkSlow(VerifyNodalColoring<T>(Graph, InParticles, GraphParticlesStart, GraphParticlesEnd, IncidentElements, IncidentElementsLocalIndex, ParticlesPerColor));
+
+	return ParticlesPerColor;
+}
+
+
+
 
 template CHAOS_API TArray<TArray<int32>> Chaos::FGraphColoring::ComputeGraphColoring<Chaos::FRealSingle>(const TArray<Chaos::TVector<int32, 2>>&, const Chaos::TDynamicParticles<Chaos::FRealSingle, 3>&, const int32 GraphParticlesStart, const int32 GraphParticlesEnd);
 template CHAOS_API TArray<TArray<int32>> Chaos::FGraphColoring::ComputeGraphColoring<Chaos::FRealDouble>(const TArray<Chaos::TVector<int32, 2>>&, const Chaos::TDynamicParticles<Chaos::FRealDouble, 3>&, const int32 GraphParticlesStart, const int32 GraphParticlesEnd);
@@ -1303,4 +1470,6 @@ template CHAOS_API TArray<TArray<int32>> Chaos::FGraphColoring::ComputeGraphColo
 template CHAOS_API void Chaos::ComputeGridBasedGraphSubColoringPointer(const TArray<TArray<int32>>& ElementsPerColor, const TMPMGrid<Chaos::FRealSingle>& Grid, const int32 GridSize, TArray<TArray<int32>>*& PreviousColoring, const TArray<TArray<int32>>& ConstraintsNodesSet, TArray<TArray<TArray<int32>>>& ElementsPerSubColors);
 template CHAOS_API void Chaos::ComputeGridBasedGraphSubColoringPointer(const TArray<TArray<int32>>& ElementsPerColor, const TMPMGrid<Chaos::FRealDouble>& Grid, const int32 GridSize, TArray<TArray<int32>>*& PreviousColoring, const TArray<TArray<int32>>& ConstraintsNodesSet, TArray<TArray<TArray<int32>>>& ElementsPerSubColors);
 template CHAOS_API void Chaos::ComputeWeakConstraintsColoring<Chaos::FRealSingle>(const TArray<TArray<int32>>& Indices, const TArray<TArray<int32>>& SecondIndices, const Chaos::TDynamicParticles<Chaos::FRealSingle, 3>& InParticles, TArray<TArray<int32>>& ConstraintsPerColor);
-template CHAOS_API void Chaos::ComputeWeakConstraintsColoring<Chaos::FRealDouble>(const TArray<TArray<int32>>& Indices, const TArray<TArray<int32>>& SecondIndices, const Chaos::TDynamicParticles<Chaos::FRealDouble, 3>& InParticles, TArray<TArray<int32>>& ConstraintsPerColor);
\ No newline at end of file
+template CHAOS_API void Chaos::ComputeWeakConstraintsColoring<Chaos::FRealDouble>(const TArray<TArray<int32>>& Indices, const TArray<TArray<int32>>& SecondIndices, const Chaos::TDynamicParticles<Chaos::FRealDouble, 3>& InParticles, TArray<TArray<int32>>& ConstraintsPerColor);
+template CHAOS_API TArray<TArray<int32>> Chaos::ComputeNodalColoring<Chaos::FRealSingle>(const TArray<TVec4<int32>>& Graph, const Chaos::TDynamicParticles<Chaos::FRealSingle, 3>& InParticles, const int32 GraphParticlesStart, const int32 GraphParticlesEnd, const TArray<TArray<int32>>& IncidentElements, const TArray<TArray<int32>>& IncidentElementsLocalIndex);
+template CHAOS_API TArray<TArray<int32>> Chaos::ComputeNodalColoring<Chaos::FRealDouble>(const TArray<TVec4<int32>>& Graph, const Chaos::TDynamicParticles<Chaos::FRealDouble, 3>& InParticles, const int32 GraphParticlesStart, const int32 GraphParticlesEnd, const TArray<TArray<int32>>& IncidentElements, const TArray<TArray<int32>>& IncidentElementsLocalIndex);
\ No newline at end of file
diff --git a/Engine/Source/Runtime/Experimental/Chaos/Public/Chaos/GraphColoring.h b/Engine/Source/Runtime/Experimental/Chaos/Public/Chaos/GraphColoring.h
--- a/Engine/Source/Runtime/Experimental/Chaos/Public/Chaos/GraphColoring.h
+++ b/Engine/Source/Runtime/Experimental/Chaos/Public/Chaos/GraphColoring.h
@@ -98,4 +98,7 @@ CHAOS_API void ComputeGridBasedGraphSubColoringPointer(const TArray<TArray<int32
 template<typename T>
 CHAOS_API void ComputeWeakConstraintsColoring(const TArray<TArray<int32>>& Indices, const TArray<TArray<int32>>& SecondIndices, const Chaos::TDynamicParticles<T, 3>& InParticles, TArray<TArray<int32>>& ConstraintsPerColor);
 
+template<typename T>
+CHAOS_API TArray<TArray<int32>> ComputeNodalColoring(const TArray<TVec4<int32>>& Graph, const Chaos::TDynamicParticles<T, 3>& InParticles, const int32 GraphParticlesStart, const int32 GraphParticlesEnd, const TArray<TArray<int32>>& IncidentElements, const TArray<TArray<int32>>& IncidentElementsLocalIndex);
+
 }
