Iris
* Filter are now done in two passes
* The first pass executes regular engine filter options then any dynamic filters not using fragment data.
* Then any filtered object not relevant to atleast one connection will get culled out.
* Culled out objects will not get polled or copy their dirty data into the fragments.
* After the poll and copy is done, we now filter again only for filters that work with fragment data.
* This final filter list gets passed to be prioritized and given to the ReplicationWriter.

#rnx

diff --git a/Engine/Plugins/Runtime/ReplicationSystemTestPlugin/Source/Private/Tests/ReplicationSystem/Filtering/MockNetObjectFilter.cpp b/Engine/Plugins/Runtime/ReplicationSystemTestPlugin/Source/Private/Tests/ReplicationSystem/Filtering/MockNetObjectFilter.cpp
--- a/Engine/Plugins/Runtime/ReplicationSystemTestPlugin/Source/Private/Tests/ReplicationSystem/Filtering/MockNetObjectFilter.cpp
+++ b/Engine/Plugins/Runtime/ReplicationSystemTestPlugin/Source/Private/Tests/ReplicationSystem/Filtering/MockNetObjectFilter.cpp
@@ -9,7 +9,7 @@ UMockNetObjectFilter::UMockNetObjectFilter()
 {
 }
 
-void UMockNetObjectFilter::Init(FNetObjectFilterInitParams& Params)
+void UMockNetObjectFilter::OnInit(FNetObjectFilterInitParams& Params)
 {
 	++CallStatus.CallCounts.Init;
 
@@ -138,11 +138,11 @@ void UMockNetObjectFilter::PostFilter(FNetObjectPostFilteringParams& Params)
 // UMockNetObjectFilterUsingFragmentData
 //**************************************************************************************************
 
-void UMockNetObjectFilterUsingFragmentData::Init(FNetObjectFilterInitParams& Params)
+void UMockNetObjectFilterUsingFragmentData::OnInit(FNetObjectFilterInitParams& Params)
 {
-	Super::Init(Params);
+	Super::OnInit(Params);
 
-	//TODO: Set to PostFilter pass when added
+	SetupFilterType(ENetFilterType::PostPoll_FragmentBased);
 }
 
 bool UMockNetObjectFilterUsingFragmentData::AddObject(uint32 ObjectIndex, FNetObjectFilterAddObjectParams& Params)
diff --git a/Engine/Plugins/Runtime/ReplicationSystemTestPlugin/Source/Private/Tests/ReplicationSystem/Filtering/MockNetObjectFilter.h b/Engine/Plugins/Runtime/ReplicationSystemTestPlugin/Source/Private/Tests/ReplicationSystem/Filtering/MockNetObjectFilter.h
--- a/Engine/Plugins/Runtime/ReplicationSystemTestPlugin/Source/Private/Tests/ReplicationSystem/Filtering/MockNetObjectFilter.h
+++ b/Engine/Plugins/Runtime/ReplicationSystemTestPlugin/Source/Private/Tests/ReplicationSystem/Filtering/MockNetObjectFilter.h
@@ -66,7 +66,7 @@ public:
 	inline void ResetFunctionCallStatus() { CallStatus = FFunctionCallStatus({}); }
 
 protected:
-	virtual void Init(FNetObjectFilterInitParams&) override;
+	virtual void OnInit(FNetObjectFilterInitParams&) override;
 	virtual void AddConnection(uint32 ConnectionId) override;
 	virtual void RemoveConnection(uint32 ConnectionId) override;
 	virtual bool AddObject(uint32 ObjectIndex, FNetObjectFilterAddObjectParams&) override;
@@ -98,7 +98,7 @@ class UMockNetObjectFilterUsingFragmentData : public UMockNetObjectFilter
 
 protected:
 
-	virtual void Init(FNetObjectFilterInitParams&) override;
+	virtual void OnInit(FNetObjectFilterInitParams&) override;
 	virtual bool AddObject(uint32 ObjectIndex, FNetObjectFilterAddObjectParams&) override;
 	virtual void RemoveObject(uint32 ObjectIndex, const FNetObjectFilteringInfo&) override;
 	virtual void UpdateObjects(FNetObjectFilterUpdateParams&) override;
diff --git a/Engine/Source/Runtime/Experimental/Iris/Core/Private/Iris/ReplicationSystem/Filtering/FilterOutNetObjectFilter.cpp b/Engine/Source/Runtime/Experimental/Iris/Core/Private/Iris/ReplicationSystem/Filtering/FilterOutNetObjectFilter.cpp
--- a/Engine/Source/Runtime/Experimental/Iris/Core/Private/Iris/ReplicationSystem/Filtering/FilterOutNetObjectFilter.cpp
+++ b/Engine/Source/Runtime/Experimental/Iris/Core/Private/Iris/ReplicationSystem/Filtering/FilterOutNetObjectFilter.cpp
@@ -2,7 +2,7 @@
 
 #include "Iris/ReplicationSystem/Filtering/FilterOutNetObjectFilter.h"
 
-void UFilterOutNetObjectFilter::Init(FNetObjectFilterInitParams& Params)
+void UFilterOutNetObjectFilter::OnInit(FNetObjectFilterInitParams& Params)
 {
 }
 
diff --git a/Engine/Source/Runtime/Experimental/Iris/Core/Private/Iris/ReplicationSystem/Filtering/NetObjectFilter.cpp b/Engine/Source/Runtime/Experimental/Iris/Core/Private/Iris/ReplicationSystem/Filtering/NetObjectFilter.cpp
--- a/Engine/Source/Runtime/Experimental/Iris/Core/Private/Iris/ReplicationSystem/Filtering/NetObjectFilter.cpp
+++ b/Engine/Source/Runtime/Experimental/Iris/Core/Private/Iris/ReplicationSystem/Filtering/NetObjectFilter.cpp
@@ -14,6 +14,16 @@ UNetObjectFilter::UNetObjectFilter()
 {
 }
 
+void UNetObjectFilter::Init(FNetObjectFilterInitParams& Params)
+{
+	if (Params.Config)
+	{
+		FilterType = Params.Config->FilterType;
+	}
+
+	OnInit(Params);
+}
+
 void UNetObjectFilter::AddConnection(uint32 ConnectionId)
 {
 }
diff --git a/Engine/Source/Runtime/Experimental/Iris/Core/Private/Iris/ReplicationSystem/Filtering/NetObjectGridFilter.cpp b/Engine/Source/Runtime/Experimental/Iris/Core/Private/Iris/ReplicationSystem/Filtering/NetObjectGridFilter.cpp
--- a/Engine/Source/Runtime/Experimental/Iris/Core/Private/Iris/ReplicationSystem/Filtering/NetObjectGridFilter.cpp
+++ b/Engine/Source/Runtime/Experimental/Iris/Core/Private/Iris/ReplicationSystem/Filtering/NetObjectGridFilter.cpp
@@ -9,7 +9,7 @@
 #include "Iris/Core/IrisLog.h"
 #include "Iris/Core/IrisProfiler.h"
 
-void UNetObjectGridFilter::Init(FNetObjectFilterInitParams& Params)
+void UNetObjectGridFilter::OnInit(FNetObjectFilterInitParams& Params)
 {
 	Config = TStrongObjectPtr<UNetObjectGridFilterConfig>(CastChecked<UNetObjectGridFilterConfig>(Params.Config));
 	checkf(Config.IsValid(), TEXT("Need config to operate."));
@@ -75,8 +75,6 @@ void UNetObjectGridFilter::RemoveObject(uint32 ObjectIndex, const FNetObjectFilt
 
 void UNetObjectGridFilter::UpdateObjects(FNetObjectFilterUpdateParams& Params)
 {
-	IRIS_PROFILER_SCOPE(UNetObjectGridFilter_UpdateObjects);
-
 	for (SIZE_T ObjectIt = 0, ObjectEndIt = Params.ObjectCount; ObjectIt != ObjectEndIt; ++ObjectIt)
 	{
 		const uint32 ObjectIndex = Params.ObjectIndices[ObjectIt];
@@ -402,9 +400,11 @@ bool UNetObjectGridFilter::DoesCellContainCoord(const UNetObjectGridFilter::FCel
 // UNetObjectGridWorldLocFilter
 //*************************************************************************************************
 
-void UNetObjectGridWorldLocFilter::Init(FNetObjectFilterInitParams& Params)
+void UNetObjectGridWorldLocFilter::OnInit(FNetObjectFilterInitParams& Params)
 {
-	Super::Init(Params);
+	Super::OnInit(Params);
+
+	SetupFilterType(ENetFilterType::PrePoll_Raw);
 
 	WorldLocations = &Params.ReplicationSystem->GetWorldLocations();
 }
@@ -435,9 +435,11 @@ bool UNetObjectGridWorldLocFilter::BuildObjectInfo(uint32 ObjectIndex, FNetObjec
 // UNetObjectGridFragmentLocFilter
 //*************************************************************************************************
 
-void UNetObjectGridFragmentLocFilter::Init(FNetObjectFilterInitParams& InitParams)
+void UNetObjectGridFragmentLocFilter::OnInit(FNetObjectFilterInitParams& InitParams)
 {
-	Super::Init(InitParams);
+	Super::OnInit(InitParams);
+
+	SetupFilterType(ENetFilterType::PostPoll_FragmentBased);
 }
 
 bool UNetObjectGridFragmentLocFilter::BuildObjectInfo(uint32 ObjectIndex, FNetObjectFilterAddObjectParams& Params)
diff --git a/Engine/Source/Runtime/Experimental/Iris/Core/Private/Iris/ReplicationSystem/Filtering/NetObjectGroups.cpp b/Engine/Source/Runtime/Experimental/Iris/Core/Private/Iris/ReplicationSystem/Filtering/NetObjectGroups.cpp
--- a/Engine/Source/Runtime/Experimental/Iris/Core/Private/Iris/ReplicationSystem/Filtering/NetObjectGroups.cpp
+++ b/Engine/Source/Runtime/Experimental/Iris/Core/Private/Iris/ReplicationSystem/Filtering/NetObjectGroups.cpp
@@ -76,6 +76,7 @@ void FNetObjectGroups::Init(const FNetObjectGroupInitParams& Params)
 	Groups.Add(FNetObjectGroup());
 
 	GroupMemberships.SetNumZeroed(Params.MaxObjectCount);
+	GroupFilteredObjects.Init(Params.MaxObjectCount);
 }
 
 FNetObjectGroupHandle FNetObjectGroups::CreateGroup()
@@ -201,6 +202,7 @@ void FNetObjectGroups::AddToGroup(FNetObjectGroupHandle GroupHandle, FInternalNe
 		if (AddGroupMembership(GroupMemberships[InternalIndex], GroupHandle))
 		{
 			Group->Members.AddUnique(InternalIndex);
+			GroupFilteredObjects.SetBit(InternalIndex);
 		}
 		else
 		{
@@ -214,10 +216,22 @@ void FNetObjectGroups::RemoveFromGroup(FNetObjectGroupHandle GroupHandle, FInter
 	FNetObjectGroup* Group = GetGroup(GroupHandle);
 	if (InternalIndex != FNetRefHandleManager::InvalidInternalIndex && Group)
 	{
-		checkSlow(IsMemberOf(GroupMemberships[InternalIndex], GroupHandle));
+		FNetObjectGroupMembership& GroupMembership = GroupMemberships[InternalIndex];
+		checkSlow(IsMemberOf(GroupMembership, GroupHandle));
 
-		RemoveGroupMembership(GroupMemberships[InternalIndex], GroupHandle);
+		RemoveGroupMembership(GroupMembership, GroupHandle);
 		Group->Members.RemoveSingle(InternalIndex);
+
+		bool bIsGroupFiltered = false;
+		for (FNetObjectGroupHandle AssignedGroup : GroupMembership.Groups)
+		{
+			bIsGroupFiltered |= (AssignedGroup != InvalidNetObjectGroupHandle);
+		}
+
+		if (!bIsGroupFiltered)
+		{
+			GroupFilteredObjects.ClearBit(InternalIndex);
+		}
 	}
 }
 
diff --git a/Engine/Source/Runtime/Experimental/Iris/Core/Private/Iris/ReplicationSystem/Filtering/NetObjectGroups.h b/Engine/Source/Runtime/Experimental/Iris/Core/Private/Iris/ReplicationSystem/Filtering/NetObjectGroups.h
--- a/Engine/Source/Runtime/Experimental/Iris/Core/Private/Iris/ReplicationSystem/Filtering/NetObjectGroups.h
+++ b/Engine/Source/Runtime/Experimental/Iris/Core/Private/Iris/ReplicationSystem/Filtering/NetObjectGroups.h
@@ -7,6 +7,7 @@
 #include "Containers/Map.h"
 #include "Containers/SparseArray.h"
 #include "UObject/NameTypes.h"
+#include "Net/Core/NetBitArray.h"
 
 namespace UE::Net
 {
@@ -78,6 +79,12 @@ public:
 	// Destroy Named group
 	void DestroyNamedGroup(FName GroupName);
 
+	/** Returns a list of all objects currently part of a group filter */
+	const FNetBitArrayView GetGroupFilteredObjects() const
+	{
+		return MakeNetBitArrayView(GroupFilteredObjects);
+	}
+
 private:
 	struct FNetObjectGroupMembership
 	{
@@ -97,6 +104,11 @@ private:
 	TArray<FNetObjectGroupMembership> GroupMemberships;
 	uint32 MaxGroupCount;
 
+	// List of objects that are filtered by group memberships
+	//$IRIS TODO: This currently assumes all groups are registered to filter. 
+	//Fix this so it only objects in a group that called RepFiltering::AddGroupFilter are set.
+	FNetBitArray GroupFilteredObjects;
+
 	TMap<FName, FNetObjectGroupHandle> NamedGroups;
 };
 
diff --git a/Engine/Source/Runtime/Experimental/Iris/Core/Private/Iris/ReplicationSystem/Filtering/NopNetObjectFilter.cpp b/Engine/Source/Runtime/Experimental/Iris/Core/Private/Iris/ReplicationSystem/Filtering/NopNetObjectFilter.cpp
--- a/Engine/Source/Runtime/Experimental/Iris/Core/Private/Iris/ReplicationSystem/Filtering/NopNetObjectFilter.cpp
+++ b/Engine/Source/Runtime/Experimental/Iris/Core/Private/Iris/ReplicationSystem/Filtering/NopNetObjectFilter.cpp
@@ -2,7 +2,7 @@
 
 #include "Iris/ReplicationSystem/Filtering/NopNetObjectFilter.h"
 
-void UNopNetObjectFilter::Init(FNetObjectFilterInitParams& Params)
+void UNopNetObjectFilter::OnInit(FNetObjectFilterInitParams& Params)
 {
 }
 
diff --git a/Engine/Source/Runtime/Experimental/Iris/Core/Private/Iris/ReplicationSystem/Filtering/ReplicationFiltering.cpp b/Engine/Source/Runtime/Experimental/Iris/Core/Private/Iris/ReplicationSystem/Filtering/ReplicationFiltering.cpp
--- a/Engine/Source/Runtime/Experimental/Iris/Core/Private/Iris/ReplicationSystem/Filtering/ReplicationFiltering.cpp
+++ b/Engine/Source/Runtime/Experimental/Iris/Core/Private/Iris/ReplicationSystem/Filtering/ReplicationFiltering.cpp
@@ -10,6 +10,7 @@
 #include "Iris/ReplicationSystem/ReplicationConnections.h"
 #include "Iris/ReplicationSystem/ReplicationProtocol.h"
 #include "Iris/ReplicationSystem/ReplicationSystem.h"
+#include "Iris/ReplicationSystem/ReplicationSystemInternal.h"
 #include "Iris/ReplicationSystem/DeltaCompression/DeltaCompressionBaselineInvalidationTracker.h"
 #include "Iris/ReplicationSystem/Filtering/NetObjectGroups.h"
 #include "Iris/ReplicationSystem/Filtering/NetObjectFilterDefinitions.h"
@@ -19,6 +20,8 @@
 namespace UE::Net::Private
 {
 
+DEFINE_LOG_CATEGORY_STATIC(LogIrisFiltering, Log, All);
+
 FName GetStaticFilterName(FNetObjectFilterHandle Filter)
 {
 	switch (Filter)
@@ -76,6 +79,10 @@ private:
 	static constexpr FNetObjectFilterHandle DynamicNetObjectFilterHandleFlag = 1U << (sizeof(FNetObjectFilterHandle)*8U - 1U);
 };
 
+//*************************************************************************************************
+// FUpdateDirtyObjectsBatchHelper
+//*************************************************************************************************
+
 class FReplicationFiltering::FUpdateDirtyObjectsBatchHelper
 {
 public:
@@ -86,25 +93,57 @@ public:
 
 	struct FPerFilterInfo
 	{
-		uint32* ObjectIndices;
-		uint32 ObjectCount;
+		uint32* ObjectIndices = nullptr;
+		uint32 ObjectCount = 0;
 
-		FReplicationInstanceProtocol const** InstanceProtocols;
+		FReplicationInstanceProtocol const** InstanceProtocols = nullptr;
 	};
 
-	FUpdateDirtyObjectsBatchHelper(const FNetRefHandleManager* InNetRefHandleManager, uint32 FilterCount)
+	FUpdateDirtyObjectsBatchHelper(const FNetRefHandleManager* InNetRefHandleManager, const TArray<FFilterInfo>& DynamicFilters, ENetFilterType ActiveFilterType)
 		: NetRefHandleManager(InNetRefHandleManager)
 	{
-		PerFilterInfos.SetNumUninitialized(FilterCount);
-		ObjectIndicesStorage.SetNumUninitialized(FilterCount * MaxObjectCountPerBatch);
-		InstanceProtocolsStorage.SetNumUninitialized(FilterCount * MaxObjectCountPerBatch);
+		bHasProtocolBuffers = (ActiveFilterType == ENetFilterType::PostPoll_FragmentBased);
 
-		uint32 FilterIndex = 0;
-		for (FPerFilterInfo& PerFilterInfo : PerFilterInfos)
+		uint32 NumActiveFilters = 0;
+		for (const FFilterInfo& Info : DynamicFilters)
+		{
+			if (Info.Type == ActiveFilterType)
+			{
+				NumActiveFilters++;
+			}
+		}
+		ensureMsgf(NumActiveFilters > 0, TEXT("Should never be called when no filters of type %u exist"), ActiveFilterType);
+
+		const int32 TotalFilters = DynamicFilters.Num();
+
+		// We create PerFilterInfo for every dynamic filters, but allocate the memory buffers only for active filters
+		constexpr bool bDontShrink = false;
+		PerFilterInfos.SetNum(TotalFilters, bDontShrink);
+
+		ObjectIndicesStorage.SetNumUninitialized(NumActiveFilters * MaxObjectCountPerBatch);
+		
+		if (bHasProtocolBuffers)
+		{
+			InstanceProtocolsStorage.SetNumUninitialized(NumActiveFilters * MaxObjectCountPerBatch);
+		}
+
+		uint32 BufferIndex = 0;
+		for (int32 i=0; i < TotalFilters; ++i)
 		{
-			PerFilterInfo.ObjectIndices = ObjectIndicesStorage.GetData() + FilterIndex * MaxObjectCountPerBatch;
-			PerFilterInfo.InstanceProtocols = InstanceProtocolsStorage.GetData() + FilterIndex * MaxObjectCountPerBatch;
-			++FilterIndex;
+			const FFilterInfo& FilterInfo = DynamicFilters[i];
+		
+			if (FilterInfo.Type == ActiveFilterType)
+			{
+				FPerFilterInfo& PerFilterInfo = PerFilterInfos[i];
+				PerFilterInfo.ObjectIndices = ObjectIndicesStorage.GetData() + BufferIndex * MaxObjectCountPerBatch;
+
+				if (bHasProtocolBuffers)
+				{
+					PerFilterInfo.InstanceProtocols = InstanceProtocolsStorage.GetData() + BufferIndex * MaxObjectCountPerBatch;
+				}
+
+				++BufferIndex;
+			}
 		}
 	}
 
@@ -112,7 +151,6 @@ public:
 	{
 		ResetBatch();
 
-		FPerFilterInfo* PerFilterInfosData = PerFilterInfos.GetData();
 		for (const uint32 ObjectIndex : MakeArrayView(ObjectIndices, ObjectCount))
 		{
 			const uint8 FilterIndex = FilterIndices[ObjectIndex];
@@ -121,16 +159,29 @@ public:
 				continue;
 			}
 
-			if (const FReplicationInstanceProtocol* InstanceProtocol = NetRefHandleManager->GetReplicatedObjectDataNoCheck(ObjectIndex).InstanceProtocol)
+			FPerFilterInfo& PerFilterInfo = PerFilterInfos[FilterIndex];
+
+			// If the info has a buffer assigned then it's an active filter
+			if (PerFilterInfo.ObjectIndices)
 			{
-				FPerFilterInfo& PerFilterInfo = PerFilterInfosData[FilterIndex];
 				PerFilterInfo.ObjectIndices[PerFilterInfo.ObjectCount] = ObjectIndex;
-				PerFilterInfo.InstanceProtocols[PerFilterInfo.ObjectCount] = InstanceProtocol;
+
+				if (bHasProtocolBuffers)
+				{
+					const FReplicationInstanceProtocol* InstanceProtocol = NetRefHandleManager->GetReplicatedObjectDataNoCheck(ObjectIndex).InstanceProtocol;
+					check(InstanceProtocol);
+							
+					check(PerFilterInfo.InstanceProtocols);
+					PerFilterInfo.InstanceProtocols[PerFilterInfo.ObjectCount] = InstanceProtocol;
+				}
+
 				++PerFilterInfo.ObjectCount;
 			}
 		}
 	}
 
+	bool HasProtocolBuffers() const { return bHasProtocolBuffers; }
+
 	TArray<FPerFilterInfo, TInlineAllocator<16>> PerFilterInfos;
 
 private:
@@ -144,16 +195,24 @@ private:
 
 	TArray<uint32> ObjectIndicesStorage;
 	TArray<const FReplicationInstanceProtocol*> InstanceProtocolsStorage;
-	const FNetRefHandleManager* NetRefHandleManager;
+	
+	const FNetRefHandleManager* NetRefHandleManager = nullptr;
+
+	bool bHasProtocolBuffers = false;
 };
 
-DEFINE_LOG_CATEGORY_STATIC(LogIrisFiltering, Log, All);
+//*************************************************************************************************
+// FReplicationFiltering
+//*************************************************************************************************
 
 FReplicationFiltering::FReplicationFiltering()
 : bHasNewConnection(0)
 , bHasRemovedConnection(0)
 , bHasDirtyConnectionFilter(0)
 , bHasDirtyOwner(0)
+, bHasDynamicFilters(0)
+, bHasDynamicRawFilters(0)
+, bHasDynamicFragmentFilters(0)
 {
 	StaticChecks();
 }
@@ -190,6 +249,9 @@ void FReplicationFiltering::Init(FReplicationFilteringInitParams& Params)
 	ObjectsWithOwnerFilter.Init(MaxObjectCount);
 	ObjectsWithPerObjectInfo.Init(MaxObjectCount);
 
+	AllConnectionFilteredObjects.Init(MaxObjectCount);
+	DynamicFilterEnabledObjects.Init(MaxObjectCount);
+
 	// Group filtering
 	MaxGroupCount = Params.MaxGroupCount;
 	GroupInfos.SetNumZeroed(MaxGroupCount);
@@ -220,10 +282,8 @@ void FReplicationFiltering::Init(FReplicationFilteringInitParams& Params)
 	InitFilters();
 }
 
-void FReplicationFiltering::Filter(const FNetBitArrayView& DirtyObjects)
+void FReplicationFiltering::FilterPrePoll()
 {
-	IRIS_PROFILER_SCOPE(FReplicationFiltering_Filter);
-
 	ResetRemovedConnections();
 
 	InitNewConnections();
@@ -237,36 +297,80 @@ void FReplicationFiltering::Filter(const FNetBitArrayView& DirtyObjects)
 
 	UpdateSubObjectFilters();
 
-	/**
-	 * This last filtering allows users to filter out objects based on arbitrary criteria.
-	 */
-	if (HasDynamicFilters())
+	const FNetBitArrayView DirtyObjects = ReplicationSystem->GetReplicationSystemInternal()->GetDirtyNetObjectTracker().GetDirtyNetObjects();
+
+	if (HasRawFilters())
 	{
-		NotifyFiltersOfDirtyObjects(DirtyObjects);
+		UpdateDynamicFilters(ENetFilterType::PrePoll_Raw, DirtyObjects);
+	}
+	else
+	{
+		// Dynamic filters are responsible for updating ObjectsInScope. 
+		// Do it here if no filters were executed.
+		auto UpdateConnectionScope = [this](uint32 ConnectionId)
+		{
+			FPerConnectionInfo& ConnectionInfo = ConnectionInfos[ConnectionId];
+			ConnectionInfo.ObjectsInScope.Copy(ConnectionInfo.ObjectsInScopeBeforeDynamicFiltering);
+		};
 
-		PreUpdateDynamicFiltering();
-		UpdateDynamicFiltering();
-		PostUpdateDynamicFiltering();
+		ValidConnections.ForAllSetBits(UpdateConnectionScope);
 	}
+
+	//$IRIS TODO: Filter culling is still WIP instead update the relevant list so it's equal to the global scope
+	//FilterNonRelevantObjects(DirtyObjects);
+	NetRefHandleManager->GetRelevantObjectsInternalIndices().Copy(MakeNetBitArrayView(NetRefHandleManager->GetScopableInternalIndices()));
 }
 
-FNetBitArrayView FReplicationFiltering::GetObjectsInScope(uint32 ConnectionId) const
+void FReplicationFiltering::FilterNonRelevantObjects(const FNetBitArrayView& DirtyObjects)
 {
-	const FPerConnectionInfo& ConnectionInfo = ConnectionInfos[ConnectionId];
-	if (HasDynamicFilters())
+	IRIS_PROFILER_SCOPE(FReplicationFiltering_FilterNonRelevantObjects);
+
+	const uint32* const GlobalScopeData = NetRefHandleManager->GetScopableInternalIndices().GetData();
+	const uint32* const WithOwnerData = ObjectsWithOwnerFilter.GetData();
+	const uint32* const ConnectionFiltersData = AllConnectionFilteredObjects.GetData();
+	const uint32* const DynamicFilteredData = DynamicFilterEnabledObjects.GetData();
+	const uint32* const GroupFilteredData = Groups->GetGroupFilteredObjects().GetData();
+
+	FNetBitArrayView GlobalRelevantObjects = NetRefHandleManager->GetRelevantObjectsInternalIndices();
+	
+	uint32* GlobalRelevantData = GlobalRelevantObjects.GetData();
+
+	//$IRIS TODO: Could be worth sampling if it's faster for the AlwaysRelevant list to be kept cached instead of recalculating it here frame.
+	const uint32 MaxWords = GlobalRelevantObjects.GetNumWords();
+	for (uint32 WordIndex = 0; WordIndex < MaxWords; ++WordIndex)
 	{
-		return MakeNetBitArrayView(ConnectionInfo.ObjectsInScope);
+		GlobalRelevantData[WordIndex] = GlobalScopeData[WordIndex] & ~(WithOwnerData[WordIndex] | ConnectionFiltersData[WordIndex] | DynamicFilteredData[WordIndex] | GroupFilteredData[WordIndex]);
 	}
-	else
+
+	// Build the list of currently relevant objects. e.g. always relevant objects + filterable objects relevant to atleast one connection.	
+	auto MergeConnectionScopes = [this, &GlobalRelevantObjects](uint32 ConnectionId)
 	{
-		return MakeNetBitArrayView(ConnectionInfo.ObjectsInScopeBeforeDynamicFiltering);
-	}
+		const FNetBitArrayView ConnectionFilteredScope = MakeNetBitArrayView(ConnectionInfos[ConnectionId].ObjectsInScope);
+		GlobalRelevantObjects.Combine(ConnectionFilteredScope, FNetBitArray::OrOp);
+	};
+	
+	ValidConnections.ForAllSetBits(MergeConnectionScopes);
+}
+
+void FReplicationFiltering::UpdateDynamicFilters(ENetFilterType FilterPass, const FNetBitArrayView& DirtyObjects)
+{
+	/**
+	* Dynamic filters allows users to filter out objects based on arbitrary criteria.
+	*/
+	NotifyFiltersOfDirtyObjects(FilterPass, DirtyObjects);
+
+	PreUpdateDynamicFiltering(FilterPass);
+	UpdateDynamicFiltering(FilterPass);
+	PostUpdateDynamicFiltering(FilterPass);
 }
 
-FNetBitArrayView FReplicationFiltering::GetGroupFilteredOutObjects(uint32 ConnectionId) const
+void FReplicationFiltering::FilterPostPoll()
 {
-	const FPerConnectionInfo& ConnectionInfo = ConnectionInfos[ConnectionId];
-	return MakeNetBitArrayView(ConnectionInfo.GroupFilteredOutObjects);
+	if (HasFragmentFilters())
+	{
+		const FNetBitArrayView DirtyObjects = ReplicationSystem->GetReplicationSystemInternal()->GetDirtyNetObjectTracker().GetDirtyNetObjects();
+		UpdateDynamicFilters(ENetFilterType::PostPoll_FragmentBased, DirtyObjects);
+	}
 }
 
 void FReplicationFiltering::SetOwningConnection(FInternalNetRefIndex ObjectIndex, uint32 ConnectionId)
@@ -485,7 +589,7 @@ FName FReplicationFiltering::GetFilterName(FNetObjectFilterHandle Filter) const
 		return DynamicFilterInfos[DynamicFilterIndex].Name;
 	}
 
-	return UE::Net::Private::GetStaticFilterName(Filter);
+	return GetStaticFilterName(Filter);
 }
 
 // Connection handling
@@ -571,6 +675,7 @@ void FReplicationFiltering::InitNewConnections()
 		if (HasDynamicFilters())
 		{
 			ConnectionInfo.DynamicFilteredOutObjects.Init(MaxObjectCount);
+			ConnectionInfo.InProgressDynamicFilteredOutObjects.Init(MaxObjectCount);
 		}
 
 		// Update group filter by iterating over all registered groups for the connection
@@ -889,6 +994,11 @@ void FReplicationFiltering::UpdateOwnerAndConnectionFiltering()
 	// Update filtering
 	if (bHasDirtyConnectionFilter)
 	{
+		ObjectsWithDirtyConnectionFilter.ForAllSetBits([this](uint32 DirtyObjectIndex)
+		{
+			AllConnectionFilteredObjects.SetBitValue(DirtyObjectIndex, HasConnectionFilter(DirtyObjectIndex));
+		});
+
 		auto UpdateConnectionScope = [this, &GlobalObjectsInScope](uint32 ConnectionId)
 		{
 			FPerConnectionInfo& ConnectionInfo = this->ConnectionInfos[ConnectionId];
@@ -955,6 +1065,8 @@ void FReplicationFiltering::UpdateGroupFiltering()
 	{
 		if (this->FilterGroups.GetBit(GroupIndex))
 		{
+			IRIS_PROFILER_SCOPE(FReplicationFiltering_UpdateGroupFiltering);
+
 			const FPerObjectInfo* ConnectionStateInfo = GetPerObjectInfo(this->GroupInfos[GroupIndex].ConnectionStateIndex);
 			const FNetObjectGroup* Group = this->Groups->GetGroup(GroupIndex);
 			FPerConnectionInfo* LocalConnectionInfos = this->ConnectionInfos.GetData();
@@ -992,7 +1104,7 @@ void FReplicationFiltering::UpdateGroupFiltering()
 	DirtyFilterGroups.Reset();
 }
 
-void FReplicationFiltering::PreUpdateDynamicFiltering()
+void FReplicationFiltering::PreUpdateDynamicFiltering(ENetFilterType FilterType)
 {
 	IRIS_PROFILER_SCOPE(FReplicationFiltering_PreUpdateDynamicFiltering);
 
@@ -1005,8 +1117,11 @@ void FReplicationFiltering::PreUpdateDynamicFiltering()
 			{
 				continue;
 			}
-			
-			Info.Filter->PreFilter(PreFilteringParams);
+
+			if (Info.Type == FilterType)
+			{
+				Info.Filter->PreFilter(PreFilteringParams);
+			}
 		}
 	}
 }
@@ -1019,7 +1134,7 @@ void FReplicationFiltering::PreUpdateDynamicFiltering()
  * of objects it should not be concerned about.
  */
 
-void FReplicationFiltering::UpdateDynamicFiltering()
+void FReplicationFiltering::UpdateDynamicFiltering(ENetFilterType FilterType)
 {
 	IRIS_PROFILER_SCOPE(FReplicationFiltering_UpdateDynamicFiltering);
 
@@ -1027,16 +1142,14 @@ void FReplicationFiltering::UpdateDynamicFiltering()
 	TArray<uint32, TInlineAllocator<256>> DisabledDependentObjects;
 
 	uint32* AllowedObjectsData = static_cast<uint32*>(FMemory_Alloca(WordCountForObjectBitArrays * sizeof(uint32)));
-	uint32* NewDynamicFilteredOutObjectsData = static_cast<uint32*>(FMemory_Alloca(WordCountForObjectBitArrays * sizeof(uint32)));
-
 	FNetBitArrayView AllowedObjects(AllowedObjectsData, MaxObjectCount, FNetBitArrayView::NoResetNoValidate);
-	FNetBitArrayView NewDynamicFilteredOutObjects(NewDynamicFilteredOutObjectsData, MaxObjectCount, FNetBitArrayView::NoResetNoValidate);
 
 	// Subobjects will never be added to a dynamic filter, but objects can become dependent at any time.
 	// We need to make sure they are not filtered out.
 	const uint32* SubObjectsData = NetRefHandleManager->GetSubObjectInternalIndices().GetData();
 	const uint32* DependentObjectsData = NetRefHandleManager->GetDependentObjectInternalIndices().GetData();
 	const uint32* ObjectsRequiringDynamicFilterUpdateData = ObjectsRequiringDynamicFilterUpdate.GetData();
+	uint8* const* ObjectsStateBuffers = FilterType == ENetFilterType::PostPoll_FragmentBased ? NetRefHandleManager->GetReplicatedObjectStateBuffers().GetData() : nullptr;
 
 	uint32* ConnectionIds = static_cast<uint32*>(FMemory_Alloca(ValidConnections.GetNumBits() * sizeof(uint32)));
 	uint32 ConnectionCount = 0;
@@ -1046,7 +1159,14 @@ void FReplicationFiltering::UpdateDynamicFiltering()
 	{
 		FPerConnectionInfo& ConnectionInfo = this->ConnectionInfos[ConnId];
 
-		NewDynamicFilteredOutObjects.Reset();
+		if (FilterType == ENetFilterType::PrePoll_Raw)
+		{
+			// Reset the array in the first filter pass
+			ConnectionInfo.InProgressDynamicFilteredOutObjects.Reset();
+		}
+
+		uint32* InProgressDynamicFilteredOutObjectsData = ConnectionInfo.InProgressDynamicFilteredOutObjects.GetData();
+		FNetBitArrayView InProgressDynamicFilteredOutObjects = MakeNetBitArrayView(ConnectionInfo.InProgressDynamicFilteredOutObjects);
 
 		/*
 		 * Apply dynamic filters.
@@ -1060,7 +1180,7 @@ void FReplicationFiltering::UpdateDynamicFiltering()
 		{
 			for (FFilterInfo& Info : DynamicFilterInfos)
 			{
-				if (Info.ObjectCount == 0U)
+				if (Info.ObjectCount == 0U || Info.Type != FilterType)
 				{
 					continue;
 				}
@@ -1068,9 +1188,11 @@ void FReplicationFiltering::UpdateDynamicFiltering()
 				FNetObjectFilteringParams FilteringParams(MakeNetBitArrayView(Info.FilteredObjects));
 				FilteringParams.OutAllowedObjects = AllowedObjects;
 				FilteringParams.FilteringInfos = NetObjectFilteringInfos.GetData();
-				FilteringParams.StateBuffers = NetRefHandleManager->GetReplicatedObjectStateBuffers().GetData();
+				FilteringParams.StateBuffers = ObjectsStateBuffers;
 				FilteringParams.ConnectionId = ConnId;
 				FilteringParams.View = Connections->GetReplicationView(ConnId);
+
+				// Execute the filter here
 				Info.Filter->Filter(FilteringParams);
 
 				const uint32* FilteredObjectsData = Info.FilteredObjects.GetData();
@@ -1079,7 +1201,7 @@ void FReplicationFiltering::UpdateDynamicFiltering()
 					const uint32 FilteredObjects = FilteredObjectsData[WordIt];
 					const uint32 FilterAllowedObjects = AllowedObjectsData[WordIt];
 					// Keep status of objects not affected by this filter. Inverse the effect of the filter and add which objects are filtered out.
-					NewDynamicFilteredOutObjectsData[WordIt] = (NewDynamicFilteredOutObjectsData[WordIt] & ~FilteredObjects) | (~FilterAllowedObjects & FilteredObjects);
+					InProgressDynamicFilteredOutObjectsData[WordIt] = (InProgressDynamicFilteredOutObjectsData[WordIt] & ~FilteredObjects) | (~FilterAllowedObjects & FilteredObjects);
 				}
 			}
 		}
@@ -1101,7 +1223,7 @@ void FReplicationFiltering::UpdateDynamicFiltering()
 				const uint32 DependentObjects = DependentObjectsData[WordIt];
 				const uint32 ObjectsRequiringUpdate = ObjectsRequiringDynamicFilterUpdateData[WordIt];
 				const uint32 PrevFilteredOutObjects = DynamicFilteredOutObjectsData[WordIt];
-				const uint32 NewFilteredOutObjects = NewDynamicFilteredOutObjectsData[WordIt];
+				const uint32 NewFilteredOutObjects = InProgressDynamicFilteredOutObjectsData[WordIt];
 				const uint32 FilterEnabledObjects = DynamicFilterEnabledObjectsData[WordIt];
 
 				const uint32 ModifiedScopeObjects = (PrevFilteredOutObjects ^ NewFilteredOutObjects);
@@ -1182,10 +1304,12 @@ void FReplicationFiltering::UpdateDynamicFiltering()
 		}
 	}
 
+
+	// It's ok to reset this in the first filter pass since objects in it would be fully treated.
 	ObjectsRequiringDynamicFilterUpdate.Reset();
 }
 
-void FReplicationFiltering::PostUpdateDynamicFiltering()
+void FReplicationFiltering::PostUpdateDynamicFiltering(ENetFilterType FilterType)
 {
 	IRIS_PROFILER_SCOPE(FReplicationFiltering_PostUpdateDynamicFiltering);
 
@@ -1199,14 +1323,17 @@ void FReplicationFiltering::PostUpdateDynamicFiltering()
 				continue;
 			}
 			
-			Info.Filter->PostFilter(PostFilteringParams);
+			if (Info.Type == FilterType)
+			{
+				Info.Filter->PostFilter(PostFilteringParams);
+			}
 		}
 	}
 }
 
-void FReplicationFiltering::NotifyFiltersOfDirtyObjects(const FNetBitArrayView& DirtyObjects)
+void FReplicationFiltering::NotifyFiltersOfDirtyObjects(ENetFilterType FilterType, const FNetBitArrayView& DirtyObjects)
 {
-	FUpdateDirtyObjectsBatchHelper BatchHelper(NetRefHandleManager, DynamicFilterInfos.Num());
+	FUpdateDirtyObjectsBatchHelper BatchHelper(NetRefHandleManager, DynamicFilterInfos, FilterType);
 
 	constexpr SIZE_T MaxBatchObjectCount = FUpdateDirtyObjectsBatchHelper::Constants::MaxObjectCountPerBatch;
 	uint32 ObjectIndices[MaxBatchObjectCount];
@@ -1224,14 +1351,18 @@ void FReplicationFiltering::NotifyFiltersOfDirtyObjects(const FNetBitArrayView&
 	}
 }
 
-void FReplicationFiltering::BatchNotifyFiltersOfDirtyObjects(FUpdateDirtyObjectsBatchHelper& BatchHelper, const uint32* ObjectIndices, uint32 ObjectCount)
+void FReplicationFiltering::BatchNotifyFiltersOfDirtyObjects(FUpdateDirtyObjectsBatchHelper& BatchHelper, const uint32* DirtyObjectIndices, uint32 ObjectCount)
 {
-	BatchHelper.PrepareBatch(ObjectIndices, ObjectCount, ObjectIndexToDynamicFilterIndex.GetData());
+	BatchHelper.PrepareBatch(DirtyObjectIndices, ObjectCount, ObjectIndexToDynamicFilterIndex.GetData());
 
 	FNetObjectFilterUpdateParams UpdateParameters;
-	UpdateParameters.StateBuffers = NetRefHandleManager->GetReplicatedObjectStateBuffers().GetData();
 	UpdateParameters.FilteringInfos = NetObjectFilteringInfos.GetData();
 
+	if (BatchHelper.HasProtocolBuffers())
+	{
+		UpdateParameters.StateBuffers = NetRefHandleManager->GetReplicatedObjectStateBuffers().GetData();
+	}
+
 	for (const FUpdateDirtyObjectsBatchHelper::FPerFilterInfo& PerFilterInfo : BatchHelper.PerFilterInfos)
 	{
 		if (PerFilterInfo.ObjectCount == 0)
@@ -1885,15 +2016,15 @@ void FReplicationFiltering::InitFilters()
 		Filter->Init(InitParams);
 
 		FFilterInfo& Info = DynamicFilterInfos.Emplace_GetRef();
+		Info.Type = Filter->GetFilterType();
 		Info.Filter = MoveTemp(Filter);
 		Info.Name = FilterDefinition.FilterName;
 		Info.ObjectCount = 0;
 		Info.FilteredObjects.Init(MaxObjectCount);
-	}
 
-	if (HasDynamicFilters())
-	{
-		DynamicFilterEnabledObjects.Init(MaxObjectCount);
+		bHasDynamicFilters = true;
+		bHasDynamicRawFilters |= (Info.Type == ENetFilterType::PrePoll_Raw);
+		bHasDynamicFragmentFilters |= (Info.Type == ENetFilterType::PostPoll_FragmentBased);
 	}
 }
 
@@ -1947,7 +2078,9 @@ void FReplicationFiltering::InvalidateBaselinesForObject(uint32 ObjectIndex, uin
 	}
 }
 
+//*************************************************************************************************
 // FNetObjectFilterHandleUtil
+//*************************************************************************************************
 inline bool FNetObjectFilterHandleUtil::IsInvalidHandle(FNetObjectFilterHandle Handle)
 {
 	return Handle == InvalidNetObjectFilterHandle;
diff --git a/Engine/Source/Runtime/Experimental/Iris/Core/Private/Iris/ReplicationSystem/Filtering/ReplicationFiltering.h b/Engine/Source/Runtime/Experimental/Iris/Core/Private/Iris/ReplicationSystem/Filtering/ReplicationFiltering.h
--- a/Engine/Source/Runtime/Experimental/Iris/Core/Private/Iris/ReplicationSystem/Filtering/ReplicationFiltering.h
+++ b/Engine/Source/Runtime/Experimental/Iris/Core/Private/Iris/ReplicationSystem/Filtering/ReplicationFiltering.h
@@ -45,14 +45,32 @@ public:
 	void Init(FReplicationFilteringInitParams& Params);
 
 	/**
-	 * Determine which objects are allowed to be replicated to which connections.
-	 * This is done for all connections every tick.
+	 * First pass to determine which objects are relevant to each connection.
+	 * Executes group, owner and connection filtering then any raw dynamic filters.
+	 * At the end any object that is not relevant to at least one connection will be removed from the scoped object list.
+	 * Exception to this rule are always relevant (e.g. non-filtered) objects or objects set to be filtered by a fragment-based dynamic filter.
 	 */
-	void Filter(const FNetBitArrayView& DirtyObjects);
+	void FilterPrePoll();
 
-	FNetBitArrayView GetObjectsInScope(uint32 ConnectionId) const;
+	/**
+	 * Second pass to determine which objects are relevant to each connection.
+	 * Executes only fragment-based dynamic filters.
+	 */
+	void FilterPostPoll();
 
-	FNetBitArrayView GetGroupFilteredOutObjects(uint32 ConnectionId) const;
+	/** 
+	* Returns the list of objects relevant to a given connection.
+	* This represents the global scope list minus the objects that were filtered out for the given connection.
+	*/
+	const FNetBitArrayView GetRelevantObjectsInScope(uint32 ConnectionId) const
+	{
+		return MakeNetBitArrayView(ConnectionInfos[ConnectionId].ObjectsInScope);
+	}
+
+	const FNetBitArrayView GetGroupFilteredOutObjects(uint32 ConnectionId) const
+	{
+		return MakeNetBitArrayView(ConnectionInfos[ConnectionId].GroupFilteredOutObjects);
+	}
 
 	// Who owns what?
 	void SetOwningConnection(FInternalNetRefIndex ObjectIndex, uint32 ConnectionId);
@@ -114,6 +132,9 @@ private:
 
 		// Which objects are filtered out by dynamic filters.
 		FNetBitArray DynamicFilteredOutObjects;
+
+		// List of objects currently filtered out by a either dynamic filter passes
+		FNetBitArray InProgressDynamicFilteredOutObjects;
 	};
 
 	struct FPerObjectInfo
@@ -132,6 +153,7 @@ private:
 
 	struct FFilterInfo
 	{
+		ENetFilterType Type = ENetFilterType::PrePoll_Raw;
 		TStrongObjectPtr<UNetObjectFilter> Filter;
 		FName Name;
 		uint32 ObjectCount = 0;
@@ -151,12 +173,19 @@ private:
 	void UpdateObjectsInScope();
 	void UpdateOwnerAndConnectionFiltering();
 	void UpdateGroupFiltering();
-	void PreUpdateDynamicFiltering();
-	void UpdateDynamicFiltering();
-	void PostUpdateDynamicFiltering();
 	void UpdateSubObjectFilters();
 
+	void UpdateDynamicFilters(ENetFilterType FilterPass, const FNetBitArrayView& DirtyObjects);
+	void PreUpdateDynamicFiltering(ENetFilterType FilterType);
+	void UpdateDynamicFiltering(ENetFilterType FilterType);
+	void PostUpdateDynamicFiltering(ENetFilterType FilterType);
+
+	/** Build the list of always relevant objects + objects that are currently relevant to at least one connection. */
+	void FilterNonRelevantObjects(const FNetBitArrayView& DirtyObjects);
+
 	bool HasDynamicFilters() const;
+	bool HasRawFilters() const;
+	bool HasFragmentFilters() const;
 
 	// Helper to update and reset group filter effects if objects are removed from a filter or after a filter status change, returns true if the group filter was changed
 	bool UpdateGroupFilterEffectsForObject(uint32 ObjectIndex, uint32 ConnectionId, const FNetBitArray& ScopableObjects);
@@ -183,7 +212,7 @@ private:
 
 	void RemoveFromDynamicFilter(uint32 ObjectIndex, uint32 FilterIndex);
 
-	void NotifyFiltersOfDirtyObjects(const FNetBitArrayView& DirtyObjects);
+	void NotifyFiltersOfDirtyObjects(ENetFilterType FilterType, const FNetBitArrayView& DirtyObjects);
 	void BatchNotifyFiltersOfDirtyObjects(FUpdateDirtyObjectsBatchHelper& BatchHelper, const uint32* ObjectIndices, uint32 ObjectCount);
 
 	void InvalidateBaselinesForObject(uint32 ObjectIndex, uint32 NewOwningConnectionId, uint32 PrevOwningConnectionId);
@@ -230,10 +259,13 @@ private:
 	// Groups
 	TArray<FPerGroupInfo> GroupInfos;
 	uint32 MaxGroupCount = 0;
+
+	//$IRIS TODO: These need better documentation
 	FNetBitArray FilterGroups;
 	FNetBitArray DirtyFilterGroups;
 	FNetBitArray SubObjectFilterGroups;
 	FNetBitArray DirtySubObjectFilterGroups;
+	FNetBitArray AllConnectionFilteredObjects;
 
 	TArray<PerObjectInfoIndexType> ObjectIndexToPerObjectInfoIndex;
 	uint32 PerObjectInfoStorageCountForConnections = 0;
@@ -244,6 +276,7 @@ private:
 	TArray<FNetObjectFilteringInfo> NetObjectFilteringInfos;
 	TArray<uint8> ObjectIndexToDynamicFilterIndex;
 	TArray<FFilterInfo> DynamicFilterInfos;
+
 	FNetBitArray DynamicFilterEnabledObjects;
 	FNetBitArray ObjectsRequiringDynamicFilterUpdate;
 
@@ -251,11 +284,23 @@ private:
 	uint32 bHasRemovedConnection : 1;
 	uint32 bHasDirtyConnectionFilter: 1;
 	uint32 bHasDirtyOwner : 1;
+	uint32 bHasDynamicFilters : 1;
+	uint32 bHasDynamicRawFilters : 1;
+	uint32 bHasDynamicFragmentFilters : 1;
 };
 
 inline bool FReplicationFiltering::HasDynamicFilters() const
 {
-	return DynamicFilterInfos.Num() > 0;
+	return bHasDynamicFilters;
+}
+
+inline bool FReplicationFiltering::HasRawFilters() const
+{
+	return bHasDynamicRawFilters;
+}
+inline bool FReplicationFiltering::HasFragmentFilters() const
+{
+	return bHasDynamicFragmentFilters;
 }
 
 }
diff --git a/Engine/Source/Runtime/Experimental/Iris/Core/Private/Iris/ReplicationSystem/NetRefHandleManager.cpp b/Engine/Source/Runtime/Experimental/Iris/Core/Private/Iris/ReplicationSystem/NetRefHandleManager.cpp
--- a/Engine/Source/Runtime/Experimental/Iris/Core/Private/Iris/ReplicationSystem/NetRefHandleManager.cpp
+++ b/Engine/Source/Runtime/Experimental/Iris/Core/Private/Iris/ReplicationSystem/NetRefHandleManager.cpp
@@ -23,6 +23,7 @@ FNetRefHandleManager::FNetRefHandleManager(FReplicationProtocolManager& InReplic
 , ReplicationSystemId(InReplicationSystemId)
 , ScopableInternalIndices(MaxActiveObjectCount)
 , PrevFrameScopableInternalIndices(MaxActiveObjectCount)
+, RelevantObjectsInternalIndices(MaxActiveObjectCount)
 , AssignedInternalIndices(MaxActiveObjectCount)
 , SubObjectInternalIndices(MaxActiveObjectCount)
 , DependentObjectInternalIndices(MaxActiveObjectCount)
diff --git a/Engine/Source/Runtime/Experimental/Iris/Core/Private/Iris/ReplicationSystem/NetRefHandleManager.h b/Engine/Source/Runtime/Experimental/Iris/Core/Private/Iris/ReplicationSystem/NetRefHandleManager.h
--- a/Engine/Source/Runtime/Experimental/Iris/Core/Private/Iris/ReplicationSystem/NetRefHandleManager.h
+++ b/Engine/Source/Runtime/Experimental/Iris/Core/Private/Iris/ReplicationSystem/NetRefHandleManager.h
@@ -189,6 +189,9 @@ public:
 	const FNetBitArray& GetPrevFrameScopableInternalIndices() const { return PrevFrameScopableInternalIndices; }
 	void SetPrevFrameScopableInternalIndicesToCurrent() { PrevFrameScopableInternalIndices = ScopableInternalIndices; }
 
+	/** List of objects that are always relevant or currently relevant to at least one connection. */
+	FNetBitArrayView GetRelevantObjectsInternalIndices() const { return MakeNetBitArrayView(RelevantObjectsInternalIndices); }
+
 	// Get bitarray for all internal indices that currently are assigned
 	const FNetBitArray& GetAssignedInternalIndices() const { return AssignedInternalIndices; }
 
@@ -301,6 +304,9 @@ private:
 	// Which internal indices were used last net frame. This can be used to find out which ones are new and deleted this frame. 
 	FNetBitArray PrevFrameScopableInternalIndices;
 
+	/** The contains the ScopableInternalIndices list minus filtered objects that are not relevant to any connection this frame. */
+	FNetBitArray RelevantObjectsInternalIndices;
+
 	// Bitset containing all internal indices that are assigned
 	FNetBitArray AssignedInternalIndices;
 
diff --git a/Engine/Source/Runtime/Experimental/Iris/Core/Private/Iris/ReplicationSystem/ObjectReplicationBridge.cpp b/Engine/Source/Runtime/Experimental/Iris/Core/Private/Iris/ReplicationSystem/ObjectReplicationBridge.cpp
--- a/Engine/Source/Runtime/Experimental/Iris/Core/Private/Iris/ReplicationSystem/ObjectReplicationBridge.cpp
+++ b/Engine/Source/Runtime/Experimental/Iris/Core/Private/Iris/ReplicationSystem/ObjectReplicationBridge.cpp
@@ -859,22 +859,22 @@ void UObjectReplicationBridge::PreUpdateAndPollImpl(FNetRefHandle Handle)
 		return;
 	}
 	
-	// If we poll all scopebable objects
+	// Update every relevant objects from here
 
-	const FNetBitArray& ScopableObjects = LocalNetRefHandleManager.GetScopableInternalIndices();
+	const FNetBitArrayView RelevantObjects = LocalNetRefHandleManager.GetRelevantObjectsInternalIndices();
 	const FNetBitArray& WantToBeDormantObjects = LocalNetRefHandleManager.GetWantToBeDormantInternalIndices();
 
 	// Filter the set of objects considered for pre-update and polling
 	// We always want to consider objects marked as dirty and their subobjects
 	FNetBitArray ObjectsConsideredForPolling(LocalNetRefHandleManager.GetMaxActiveObjectCount());
+	FNetBitArrayView FrequencyBasedObjectsToPollView = MakeNetBitArrayView(ObjectsConsideredForPolling);
 	if (bUseFrequencyBasedPolling)
 	{
-		FNetBitArrayView FrequencyBasedObjectsToPollView = MakeNetBitArrayView(ObjectsConsideredForPolling);
-		PollFrequencyLimiter->Update(MakeNetBitArrayView(ScopableObjects), DirtyObjects, FrequencyBasedObjectsToPollView);
+		PollFrequencyLimiter->Update(RelevantObjects, DirtyObjects, FrequencyBasedObjectsToPollView);
 	}
 	else
 	{
-		ObjectsConsideredForPolling.Copy(ScopableObjects);
+		FrequencyBasedObjectsToPollView.Copy(RelevantObjects);
 	}
 
 	// Mask off objects pending dormancy as we do not want to poll/pre-update them unless they are marked for flush or are dirty
@@ -882,8 +882,8 @@ void UObjectReplicationBridge::PreUpdateAndPollImpl(FNetRefHandle Handle)
 	{
 		IRIS_PROFILER_SCOPE(PreUpdateAndPollImpl_Dormancy);
 
-		// Mask off objects pending dormancy, masked off dirty objects are restored later
-		ObjectsConsideredForPolling.Combine(WantToBeDormantObjects, FNetBitArrayView::AndNotOp);
+		// Mask off objects pending dormancy that are not dirty
+		FrequencyBasedObjectsToPollView.CombineMultiple(FNetBitArrayView::AndNotOp, MakeNetBitArrayView(WantToBeDormantObjects), FNetBitArrayView::AndOp, DirtyObjects);
 
 		// Add objects that have requested to flush dormancy
 		for (FNetRefHandle HandlePendingFlush : MakeArrayView(DormantHandlesPendingFlush))
@@ -924,10 +924,6 @@ void UObjectReplicationBridge::PreUpdateAndPollImpl(FNetRefHandle Handle)
 		auto PropagateOwnerDirtinessToSubObjectsAndDependentObjects = [&LocalNetRefHandleManager, &ObjectsConsideredForPolling](uint32 InternalObjectIndex)
 		{
 			const FNetRefHandleManager::FReplicatedObjectData& ObjectData = LocalNetRefHandleManager.GetReplicatedObjectDataNoCheck(InternalObjectIndex);
-
-			// Mark owner as well as we might have masked it out in the dormancy pass
-			ObjectsConsideredForPolling.SetBit(InternalObjectIndex);
-
 			for (const FInternalNetRefIndex SubObjectInternalIndex : LocalNetRefHandleManager.GetSubObjects(InternalObjectIndex))
 			{
 				ObjectsConsideredForPolling.SetBit(SubObjectInternalIndex);
diff --git a/Engine/Source/Runtime/Experimental/Iris/Core/Private/Iris/ReplicationSystem/ReplicationSystem.cpp b/Engine/Source/Runtime/Experimental/Iris/Core/Private/Iris/ReplicationSystem/ReplicationSystem.cpp
--- a/Engine/Source/Runtime/Experimental/Iris/Core/Private/Iris/ReplicationSystem/ReplicationSystem.cpp
+++ b/Engine/Source/Runtime/Experimental/Iris/Core/Private/Iris/ReplicationSystem/ReplicationSystem.cpp
@@ -235,22 +235,33 @@ public:
 		ReplicationSystemInternal.GetReplicationBridge()->CallUpdateInstancesWorldLocation();
 	}
 
-	void UpdateFiltering(const FNetBitArrayView& DirtyObjects)
+	void UpdateFilterPrePoll()
 	{
-		IRIS_PROFILER_SCOPE(FReplicationSystem_UpdateFiltering);
+		IRIS_PROFILER_SCOPE(FReplicationSystem_UpdateFilterPrePoll);
 		LLM_SCOPE_BYTAG(Iris);
 
 		FReplicationFiltering& Filtering = ReplicationSystemInternal.GetFiltering();
-		Filtering.Filter(DirtyObjects);
+		Filtering.FilterPrePoll();
+	}
+
+	void UpdateFilterPostPoll()
+	{
+		IRIS_PROFILER_SCOPE(FReplicationSystem_UpdateFilterPostPoll);
+		LLM_SCOPE_BYTAG(Iris);
+
+		FReplicationFiltering& Filtering = ReplicationSystemInternal.GetFiltering();
+		Filtering.FilterPostPoll();
 
 		// Iterate over all valid connections and propagate updated scopes
 		FReplicationConnections& Connections = ReplicationSystemInternal.GetConnections();
+		
 		auto UpdateConnectionScope = [&Filtering, &Connections](uint32 ConnectionId)
 		{
 			FReplicationConnection* Conn = Connections.GetConnection(ConnectionId);
-			const FNetBitArrayView& ObjectsInScope = Filtering.GetObjectsInScope(ConnectionId);
+			const FNetBitArrayView ObjectsInScope = Filtering.GetRelevantObjectsInScope(ConnectionId);
 			Conn->ReplicationWriter->UpdateScope(ObjectsInScope);
 		};
+
 		const FNetBitArray& ValidConnections = Connections.GetValidConnections();
 		ValidConnections.ForAllSetBits(UpdateConnectionScope);
 	}
@@ -328,6 +339,10 @@ public:
 		// Copy all ReplicatedObjects with dirty state data
 		DirtyObjects.ForAllSetBits(CopyFunction);
 
+		//$IRIS TODO Only copy relevant objects but keep dirty flag for those not copied.
+		//FNetBitArrayView RelevantObjects = NetRefHandleManager.GetRelevantObjectsInternalIndices();
+		//FNetBitArrayView::ForAllSetBits(DirtyObjects, RelevantObjects, FNetBitArrayView::AndOp, CopyFunction);
+
 		const uint32 ReplicationSystemId = ReplicationSystem->GetId();
 		UE_NET_TRACE_FRAME_STATSCOUNTER(ReplicationSystemId, ReplicationSystem.CopiedObjectCount, CopiedObjectCount, ENetTraceVerbosity::Trace);
 	}
@@ -552,12 +567,15 @@ void UReplicationSystem::PreSendUpdate(float DeltaSeconds)
 		// Refresh dirty object list
 		Impl->UpdateDirtyObjectList();
 
-		// Invoke any operations we need to do before copying state data
-		InternalSys.GetReplicationBridge()->CallPreSendUpdate(DeltaSeconds);
-
 		// Update world locations. We need this to happen before both filtering and prioritization.
 		Impl->UpdateWorldLocations();
 
+		// Filters to reduce the top-level scoped object list
+		Impl->UpdateFilterPrePoll();
+
+		// Invoke any operations we need to do before copying state data
+		InternalSys.GetReplicationBridge()->CallPreSendUpdate(DeltaSeconds);		
+
 		// Update conditionals
 		Impl->UpdateConditionals();
 
@@ -568,10 +586,7 @@ void UReplicationSystem::PreSendUpdate(float DeltaSeconds)
 		Impl->ProcessNetObjectAttachmentSendQueue(FNetBlobManager::EProcessMode::ProcessObjectsGoingOutOfScope);
 
 		// Update filtering and scope for all connections
-		{
-			FNetBitArrayView DirtyObjects = InternalSys.GetDirtyNetObjectTracker().GetDirtyNetObjects();
-			Impl->UpdateFiltering(DirtyObjects);
-		}
+		Impl->UpdateFilterPostPoll();
 
 		// Propagate dirty changes to all connections
 		Impl->PropagateDirtyChanges();
diff --git a/Engine/Source/Runtime/Experimental/Iris/Core/Public/Iris/ReplicationSystem/Filtering/FilterOutNetObjectFilter.h b/Engine/Source/Runtime/Experimental/Iris/Core/Public/Iris/ReplicationSystem/Filtering/FilterOutNetObjectFilter.h
--- a/Engine/Source/Runtime/Experimental/Iris/Core/Public/Iris/ReplicationSystem/Filtering/FilterOutNetObjectFilter.h
+++ b/Engine/Source/Runtime/Experimental/Iris/Core/Public/Iris/ReplicationSystem/Filtering/FilterOutNetObjectFilter.h
@@ -18,7 +18,7 @@ class UFilterOutNetObjectFilter final : public UNetObjectFilter
 
 protected:
 	// UNetObjectFilter interface
-	IRISCORE_API virtual void Init(FNetObjectFilterInitParams&) override;
+	IRISCORE_API virtual void OnInit(FNetObjectFilterInitParams&) override;
 	IRISCORE_API virtual bool AddObject(uint32 ObjectIndex, FNetObjectFilterAddObjectParams&) override;
 	IRISCORE_API virtual void RemoveObject(uint32 ObjectIndex, const FNetObjectFilteringInfo&) override;
 	IRISCORE_API virtual void UpdateObjects(FNetObjectFilterUpdateParams&) override;
diff --git a/Engine/Source/Runtime/Experimental/Iris/Core/Public/Iris/ReplicationSystem/Filtering/NetObjectFilter.h b/Engine/Source/Runtime/Experimental/Iris/Core/Public/Iris/ReplicationSystem/Filtering/NetObjectFilter.h
--- a/Engine/Source/Runtime/Experimental/Iris/Core/Public/Iris/ReplicationSystem/Filtering/NetObjectFilter.h
+++ b/Engine/Source/Runtime/Experimental/Iris/Core/Public/Iris/ReplicationSystem/Filtering/NetObjectFilter.h
@@ -49,7 +49,7 @@ enum class ENetFilterStatus : uint32
 	Allow,
 };
 
-}
+} // end namespace UE::Net
 
 /**
  * Parameters passed to UNetObjectFilter::Filter.
@@ -70,10 +70,13 @@ struct FNetObjectFilteringParams
 	UE::Net::FNetBitArrayView OutAllowedObjects;
 
 	/** FilteringInfos for all objects. Index using the set bit indices in FilteredObjects. */
-	const FNetObjectFilteringInfo* FilteringInfos;
+	const FNetObjectFilteringInfo* FilteringInfos = nullptr;
 
-	/** State buffers for all objects. Index using the set bit indices in FilteredObjects. */
-	uint8 *const* StateBuffers;
+	/** 
+	* State buffers for all objects. Index using the set bit indices in FilteredObjects. 
+	*/
+	
+	uint8 *const* StateBuffers = nullptr;
 
 	/** ID of the connection that the filtering applies to. */
 	uint32 ConnectionId;
@@ -105,14 +108,39 @@ struct alignas(8) FNetObjectFilteringInfo
 	uint16 Data[4];
 };
 
+/** This configures when a filter gets executed inside PreSendUpdate and what data it has access to. */
+UENUM()
+enum class ENetFilterType : uint8
+{
+	/**
+	 * The default setting of filters.
+	 * This type of filter is applied before we poll objects and copy their data into network fragments.
+	 * Thus a filter cannot operate on fragment data during it's operation, only whatever it stored or is accessing directly from the object.
+	 * The benefit is that objects not relevant to any connection will get culled from being polled and forced to copy their dirty state for nothing.
+	 */
+	PrePoll_Raw,
+
+	/**
+	 * When set to FragmentBased, the filter gets access to the up to date fragment data of an object when it executes.
+	 * Those objects will always be polled and have their dirty state copied even if not relevant.
+	 */
+	PostPoll_FragmentBased,
+};
+
 /**
  * Base class for filter specific configuration.
  * @see FNetObjectFilterDefinition
  */
-UCLASS(Transient, MinimalAPI)
+UCLASS(Transient, MinimalAPI, config=Engine)
 class UNetObjectFilterConfig : public UObject
 {
 	GENERATED_BODY()
+
+public:
+
+	/** Can be used to modify when the filter is executed */
+	UPROPERTY(Config)
+	ENetFilterType FilterType = ENetFilterType::PrePoll_Raw;
 };
 
 /** Parameters passed to the filter's Init() call. */
@@ -150,18 +178,24 @@ struct FNetObjectFilterAddObjectParams
 struct FNetObjectFilterUpdateParams
 {
 	/** Indices of the updated objects. */
-	const uint32* ObjectIndices;
+	const uint32* ObjectIndices = nullptr;
 	/** The number of objects that have been updated. */
-	uint32 ObjectCount;
-
-	/** InstanceProtocols for updated objects. Index using 0..ObjectCount-1. */
-	UE::Net::FReplicationInstanceProtocol const*const* InstanceProtocols;
-
-	/** State buffers for all objects. Index using ObjectIndices[0..ObjectCount-1]. */
-	uint8 *const* StateBuffers;
+	uint32 ObjectCount = 0;
 
 	/** Infos for all objects. Index using ObjectIndices[0..ObjectCount-1]. */
-	FNetObjectFilteringInfo* FilteringInfos;
+	FNetObjectFilteringInfo* FilteringInfos = nullptr;
+
+	/** 
+	* InstanceProtocols for updated objects. Index using 0..ObjectCount-1.
+	* NOTE: Only for filters of type FragmentBased; null for Raw types
+	*/
+	UE::Net::FReplicationInstanceProtocol const* const* InstanceProtocols = nullptr;
+
+	/** 
+	* State buffers for all objects. Index using ObjectIndices[0..ObjectCount-1]. 
+	* NOTE: Only for filters of type FragmentBased; null for Raw types
+	*/
+	uint8* const* StateBuffers = nullptr;
 };
 
 UCLASS(Abstract)
@@ -170,8 +204,7 @@ class UNetObjectFilter : public UObject
 	GENERATED_BODY()
 
 public:
-	/** Called right after constructor for enabled filters. Must be overriden. */ 
-	IRISCORE_API virtual void Init(FNetObjectFilterInitParams&) PURE_VIRTUAL(Init,);
+	IRISCORE_API void Init(FNetObjectFilterInitParams& Params);
 
 	/** A new connection has been added. An opportunity for the filter to allocate per connection info. */
 	IRISCORE_API virtual void AddConnection(uint32 ConnectionId);
@@ -202,6 +235,19 @@ public:
 	 */
 	IRISCORE_API virtual void PostFilter(FNetObjectPostFilteringParams&);
 
+	/** Returns what type of filter it is. Default is to cull and be executed before dirty state copying. */
+	ENetFilterType GetFilterType() const { return FilterType; }
+
 protected:
 	IRISCORE_API UNetObjectFilter();
+
+	/** Called right after constructor for enabled filters. Must be overriden. */
+	IRISCORE_API virtual void OnInit(FNetObjectFilterInitParams&) PURE_VIRTUAL(OnInit, );
+
+	/** Directly set when you want your dynamic filter to be executed. */
+	void SetupFilterType(ENetFilterType NewFilterType) { FilterType = NewFilterType; }
+
+private:
+
+	ENetFilterType FilterType = ENetFilterType::PrePoll_Raw;
 };
diff --git a/Engine/Source/Runtime/Experimental/Iris/Core/Public/Iris/ReplicationSystem/Filtering/NetObjectGridFilter.h b/Engine/Source/Runtime/Experimental/Iris/Core/Public/Iris/ReplicationSystem/Filtering/NetObjectGridFilter.h
--- a/Engine/Source/Runtime/Experimental/Iris/Core/Public/Iris/ReplicationSystem/Filtering/NetObjectGridFilter.h
+++ b/Engine/Source/Runtime/Experimental/Iris/Core/Public/Iris/ReplicationSystem/Filtering/NetObjectGridFilter.h
@@ -56,7 +56,7 @@ class UNetObjectGridFilter : public UNetObjectFilter
 
 protected:
 	// UNetObjectFilter interface
-	IRISCORE_API virtual void Init(FNetObjectFilterInitParams&) override;
+	IRISCORE_API virtual void OnInit(FNetObjectFilterInitParams&) override;
 	IRISCORE_API virtual void AddConnection(uint32 ConnectionId) override;
 	IRISCORE_API virtual void RemoveConnection(uint32 ConnectionId) override;
 	IRISCORE_API virtual bool AddObject(uint32 ObjectIndex, FNetObjectFilterAddObjectParams&) override;
@@ -199,7 +199,7 @@ class UNetObjectGridWorldLocFilter : public UNetObjectGridFilter
 
 protected:
 
-	virtual void Init(FNetObjectFilterInitParams&) override;
+	virtual void OnInit(FNetObjectFilterInitParams&) override;
 	virtual void UpdateObjectInfo(FPerObjectInfo& PerObjectInfo, const UNetObjectGridFilter::FObjectLocationInfo& ObjectLocationInfo, const UE::Net::FReplicationInstanceProtocol* InstanceProtocol) override;
 	virtual bool BuildObjectInfo(uint32 ObjectIndex, FNetObjectFilterAddObjectParams& Params) override;
 
@@ -220,7 +220,7 @@ class UNetObjectGridFragmentLocFilter : public UNetObjectGridFilter
 
 protected:
 
-	virtual void Init(FNetObjectFilterInitParams&) override;
+	virtual void OnInit(FNetObjectFilterInitParams&) override;
 	virtual void UpdateObjectInfo(FPerObjectInfo& PerObjectInfo, const UNetObjectGridFilter::FObjectLocationInfo& ObjectLocationInfo, const UE::Net::FReplicationInstanceProtocol* InstanceProtocol) override;
 	virtual bool BuildObjectInfo(uint32 ObjectIndex, FNetObjectFilterAddObjectParams& Params) override;
 	virtual void OnObjectRemoved(uint32 ObjectIndex) override;
diff --git a/Engine/Source/Runtime/Experimental/Iris/Core/Public/Iris/ReplicationSystem/Filtering/NopNetObjectFilter.h b/Engine/Source/Runtime/Experimental/Iris/Core/Public/Iris/ReplicationSystem/Filtering/NopNetObjectFilter.h
--- a/Engine/Source/Runtime/Experimental/Iris/Core/Public/Iris/ReplicationSystem/Filtering/NopNetObjectFilter.h
+++ b/Engine/Source/Runtime/Experimental/Iris/Core/Public/Iris/ReplicationSystem/Filtering/NopNetObjectFilter.h
@@ -18,7 +18,7 @@ class UNopNetObjectFilter final : public UNetObjectFilter
 
 protected:
 	// UNetObjectFilter interface
-	IRISCORE_API virtual void Init(FNetObjectFilterInitParams&) override;
+	IRISCORE_API virtual void OnInit(FNetObjectFilterInitParams&) override;
 	IRISCORE_API virtual bool AddObject(uint32 ObjectIndex, FNetObjectFilterAddObjectParams&) override;
 	IRISCORE_API virtual void RemoveObject(uint32 ObjectIndex, const FNetObjectFilteringInfo&) override;
 	IRISCORE_API virtual void UpdateObjects(FNetObjectFilterUpdateParams&) override;
