Chaos : add ways to display internal cluster particles in Chaos debug draw

diff --git a/Engine/Source/Runtime/Experimental/Chaos/Private/Chaos/ChaosDebugDraw.cpp b/Engine/Source/Runtime/Experimental/Chaos/Private/Chaos/ChaosDebugDraw.cpp
--- a/Engine/Source/Runtime/Experimental/Chaos/Private/Chaos/ChaosDebugDraw.cpp
+++ b/Engine/Source/Runtime/Experimental/Chaos/Private/Chaos/ChaosDebugDraw.cpp
@@ -55,6 +55,9 @@ namespace Chaos
 		bool bChaosDebugDebugDrawColorShapesByIsland = false;
 		FAutoConsoleVariableRef CVarChaosDebugDebugDrawColorShapesByIsland(TEXT("p.Chaos.DebugDraw.ColorShapesByIsland"), bChaosDebugDebugDrawColorShapesByIsland, TEXT("Whether to use particle island to define the color of the shapes instead of using the particle state "));
 
+		bool bChaosDebugDebugDrawColorShapesByInternalCluster = false;
+		FAutoConsoleVariableRef CVarChaosDebugDebugDrawColorShapesByInternalCluster(TEXT("p.Chaos.DebugDraw.ColorShapesByInternalCluster"), bChaosDebugDebugDrawColorShapesByInternalCluster, TEXT("Whether to check if the particle is an internal cluster to define its color (black : regular particle: red :internal cluster )"));
+
 		bool bChaosDebugDebugDrawColorBoundsByShapeType = false;
 		FAutoConsoleVariableRef CVarChaosDebugDebugDrawColorBoundsByShapeType(TEXT("p.Chaos.DebugDraw.ColorBoundsByShapeType"), bChaosDebugDebugDrawColorBoundsByShapeType, TEXT("Whether to use shape type to define the color of the bounds instead of using the particle state (if multiple shapes , will use the first one)"));
 
@@ -610,6 +613,17 @@ namespace Chaos
 				{
 					ShapeColor = GetIslandColor(FConstGenericParticleHandle(Particle)->IslandIndex(), true);
 				}
+				if (bChaosDebugDebugDrawColorShapesByInternalCluster)
+				{
+					ShapeColor = FColor::Black;
+					if (const FPBDRigidClusteredParticleHandle* ClusteredParticle = Particle->CastToClustered())
+					{
+						if (ClusteredParticle->InternalCluster())
+						{
+							ShapeColor = FColor::Red;
+						}
+					}
+				}
 			}
 
 			// If we get here, we have an actual shape to render
