Iris
- Fixed issue with a bitstream error for batches only containng a subobject.

diff --git a/Engine/Source/Runtime/Experimental/Iris/Core/Private/Iris/ReplicationSystem/ReplicationWriter.cpp b/Engine/Source/Runtime/Experimental/Iris/Core/Private/Iris/ReplicationSystem/ReplicationWriter.cpp
--- a/Engine/Source/Runtime/Experimental/Iris/Core/Private/Iris/ReplicationSystem/ReplicationWriter.cpp
+++ b/Engine/Source/Runtime/Experimental/Iris/Core/Private/Iris/ReplicationSystem/ReplicationWriter.cpp
@@ -521,7 +521,7 @@ const FReplicationWriter::FReplicationInfo& FReplicationWriter::GetReplicationIn
 void FReplicationWriter::WriteNetRefHandleId(FNetSerializationContext& Context, FNetRefHandle Handle)
 {
 	FNetBitStreamWriter* Writer = Context.GetBitStreamWriter();
-	UE_NET_TRACE_SCOPE(Handle, *Writer, Context.GetTraceCollector(), ENetTraceVerbosity::Trace);
+	UE_NET_TRACE_OBJECT_SCOPE(Handle, *Writer, Context.GetTraceCollector(), ENetTraceVerbosity::Verbose);
 	WritePackedUint64(Writer, Handle.GetId());
 }
 
@@ -1617,7 +1617,6 @@ uint32 FReplicationWriter::WriteObjectsPendingDestroy(FNetSerializationContext&
 					if (OwnerInfo.GetState() != EReplicatedObjectState::Invalid && OwnerInfo.GetState() < EReplicatedObjectState::PendingDestroy)
 					{
 						//			UE_LOG(LogTemp, Warning, TEXT("SubObject %s Should be destroyed by replicating owner"), ToCStr(ObjectData.RefHandle.ToString()));
-
 						SetState(InternalIndex, EReplicatedObjectState::SubObjectPendingDestroy);
 						Info.SubObjectPendingDestroy = 1U;
 
@@ -1894,6 +1893,7 @@ FReplicationWriter::EWriteObjectStatus FReplicationWriter::WriteObjectAndSubObje
 
 	if (bHasState | bWriteAttachments | bSentTearOff | Info.SubObjectPendingDestroy)
 	{
+		// Only need to write the handle if this is a subobject
 		if (Info.IsSubObject)
 		{
 			// We send the Index of the handle to the remote end
@@ -1909,7 +1909,7 @@ FReplicationWriter::EWriteObjectStatus FReplicationWriter::WriteObjectAndSubObje
 		const bool bWriteReplicatedDestroyHeader = !bIsObjectIndexForAttachment;
 		if (bWriteReplicatedDestroyHeader)
 		{
-			// Write destroy header bits, we always want to write the same number of bits to be able to update the header afterwards we know what data made it into the packet
+			// Write destroy header bits, we always want to write the same number of bits to be able to update the header afterwards when we know what data made it into the packet
 			Writer.WriteBits(0U, ReplicatedDestroyHeaderFlags_BitCount);
 		}
 
@@ -2115,6 +2115,11 @@ FReplicationWriter::EWriteObjectStatus FReplicationWriter::WriteObjectAndSubObje
 				FNetBitStreamWriteScope WriteScope(Writer, ReplicatedDestroyHeaderBitPos);
 				Writer.WriteBits(ReplicatedDestroyHeaderFlags, ReplicatedDestroyHeaderFlags_BitCount);
 			}
+			else if (bWriteBatchInfo && !(BatchEntry.bSentState || BatchEntry.bSentAttachments))
+			{
+				// No need for the destroy header as we did not write any data at all for the batch.
+				Writer.Seek(ReplicatedDestroyHeaderBitPos);
+			}
 		}
 	}
 
@@ -2208,7 +2213,7 @@ FReplicationWriter::EWriteObjectStatus FReplicationWriter::WriteObjectAndSubObje
 				Writer.WriteBool(bWroteExports);
 			}
 		}
-		// If we did not write any data we rollback any written headers and report a succcess
+		// If we did not write any data we rollback any written headers and report a success
 		else
 		{
 			// if we or our subobjects did not write any data, rollback and forget about everything
@@ -2252,9 +2257,6 @@ FReplicationWriter::EWriteObjectStatus FReplicationWriter::WriteObjectInBatch(FN
 				EWriteObjectStatus DependentObjectWriteStatus = WriteObjectInBatch(Context, DependentInternalIndex, WriteObjectFlags, OutBatchInfo);
 				if (!IsWriteObjectSuccess(DependentObjectWriteStatus))
 				{
-					// Need to remove the batch entry from BatchInfo.
-					OutBatchInfo.ObjectInfos.RemoveAt(OutBatchInfo.ObjectInfos.Num() - 1);
-
 					// Restore ParentInternalIndex
 					OutBatchInfo.ParentInternalIndex = OldBatchInfoParentInternalIndex;
 					return DependentObjectWriteStatus;
@@ -2531,7 +2533,10 @@ int FReplicationWriter::WriteDestructionInfo(FNetSerializationContext& Context,
 	}
 
 #if UE_NET_USE_READER_WRITER_SENTINEL
-	WriteSentinelBits(&Writer, 8);
+	{
+		UE_NET_TRACE_SCOPE(Sentinel, Writer, Context.GetTraceCollector(), ENetTraceVerbosity::Trace);
+		WriteSentinelBits(&Writer, 8);
+	}
 #endif
 
 	// Push record
