Geometry collection : Make sure GC custom renderer system is properly handling bHiddenInGame flag for GC component

diff --git a/Engine/Source/Runtime/Experimental/GeometryCollectionEngine/Private/GeometryCollection/GeometryCollectionComponent.cpp b/Engine/Source/Runtime/Experimental/GeometryCollectionEngine/Private/GeometryCollection/GeometryCollectionComponent.cpp
--- a/Engine/Source/Runtime/Experimental/GeometryCollectionEngine/Private/GeometryCollection/GeometryCollectionComponent.cpp
+++ b/Engine/Source/Runtime/Experimental/GeometryCollectionEngine/Private/GeometryCollection/GeometryCollectionComponent.cpp
@@ -124,7 +124,6 @@ FAutoConsoleVariableRef CVarGeometryCollectionRemovalTimerMultiplier(TEXT("p.Cha
 // temporary cvar , shoul dbe removed when the root event is no longer no ecessary
 bool GeometryCollectionEmitRootBreakingEvent = true;
 FAutoConsoleVariableRef CVarGeometryCollectionEmitRootBreakingEvent(TEXT("p.Chaos.GC.EmitRootBreakingEvent"), GeometryCollectionEmitRootBreakingEvent, TEXT("When true send a breaking event when root is breaking"));
-
 DEFINE_LOG_CATEGORY_STATIC(UGCC_LOG, Error, All);
 
 extern FGeometryCollectionDynamicDataPool GDynamicDataPool;
@@ -2423,6 +2422,12 @@ void UGeometryCollectionComponent::AsyncPhysicsTickComponent(float DeltaTime, fl
 	}
 }
 
+void UGeometryCollectionComponent::OnHiddenInGameChanged()
+{
+	Super::OnHiddenInGameChanged();
+	RefreshCustomRenderer();
+}
+
 void UGeometryCollectionComponent::OnRegister()
 {
 	//UE_LOG(UGCC_LOG, Log, TEXT("GeometryCollectionComponent[%p]::OnRegister()[%p]"), this,RestCollection );
@@ -4393,7 +4398,7 @@ void UGeometryCollectionComponent::RefreshCustomRenderer()
 				const int32 RootIndex = GetRootIndex();
 				const bool bIsBroken = DynamicCollection ? !DynamicCollection->Active[RootIndex] : false;
 
-				RendererInterface->UpdateState(*RestCollection, ComponentTransform, bIsBroken);
+				RendererInterface->UpdateState(*RestCollection, ComponentTransform, bIsBroken, !bHiddenInGame);
 
 				if (bIsBroken)
 				{
diff --git a/Engine/Source/Runtime/Experimental/GeometryCollectionEngine/Private/GeometryCollection/GeometryCollectionISMPoolRenderer.cpp b/Engine/Source/Runtime/Experimental/GeometryCollectionEngine/Private/GeometryCollection/GeometryCollectionISMPoolRenderer.cpp
--- a/Engine/Source/Runtime/Experimental/GeometryCollectionEngine/Private/GeometryCollection/GeometryCollectionISMPoolRenderer.cpp
+++ b/Engine/Source/Runtime/Experimental/GeometryCollectionEngine/Private/GeometryCollection/GeometryCollectionISMPoolRenderer.cpp
@@ -25,26 +25,34 @@ void UGeometryCollectionISMPoolRenderer::OnUnregisterGeometryCollection()
 	ISMPoolActor = nullptr;
 }
 
-void UGeometryCollectionISMPoolRenderer::UpdateState(UGeometryCollection const& InGeometryCollection, FTransform const& InComponentTransform, bool bIsBroken)
+void UGeometryCollectionISMPoolRenderer::UpdateState(UGeometryCollection const& InGeometryCollection, FTransform const& InComponentTransform, bool bIsBroken, bool bIsVisible)
 {
 	ComponentTransform = InComponentTransform;
 
-	if (!bIsBroken && MergedMeshGroup.GroupIndex == INDEX_NONE)
+	if (bIsVisible == false)
 	{
-		// Remove broken primitives.
 		ReleaseGroup(InstancesGroup);
-
-		// Add merged mesh.
-		InitMergedMeshFromGeometryCollection(InGeometryCollection);
+		ReleaseGroup(MergedMeshGroup);
 	}
-
-	if (bIsBroken && InstancesGroup.GroupIndex == INDEX_NONE)
+	else
 	{
-		// Remove merged mesh.
-		ReleaseGroup(MergedMeshGroup);
+		if (!bIsBroken && MergedMeshGroup.GroupIndex == INDEX_NONE)
+		{
+			// Remove broken primitives.
+			ReleaseGroup(InstancesGroup);
 
-		// Add broken primitives.
-		InitInstancesFromGeometryCollection(InGeometryCollection);
+			// Add merged mesh.
+			InitMergedMeshFromGeometryCollection(InGeometryCollection);
+		}
+
+		if (bIsBroken && InstancesGroup.GroupIndex == INDEX_NONE)
+		{
+			// Remove merged mesh.
+			ReleaseGroup(MergedMeshGroup);
+
+			// Add broken primitives.
+			InitInstancesFromGeometryCollection(InGeometryCollection);
+		}
 	}
 }
 
diff --git a/Engine/Source/Runtime/Experimental/GeometryCollectionEngine/Private/GeometryCollection/GeometryCollectionISMPoolRenderer.h b/Engine/Source/Runtime/Experimental/GeometryCollectionEngine/Private/GeometryCollection/GeometryCollectionISMPoolRenderer.h
--- a/Engine/Source/Runtime/Experimental/GeometryCollectionEngine/Private/GeometryCollection/GeometryCollectionISMPoolRenderer.h
+++ b/Engine/Source/Runtime/Experimental/GeometryCollectionEngine/Private/GeometryCollection/GeometryCollectionISMPoolRenderer.h
@@ -19,7 +19,7 @@ public:
 	//~ Begin IGeometryCollectionExternalRenderInterface Interface.
 	virtual void OnRegisterGeometryCollection(UGeometryCollectionComponent const& InComponent) override;
 	virtual void OnUnregisterGeometryCollection() override;
-	virtual void UpdateState(UGeometryCollection const& InGeometryCollection, FTransform const& InComponentTransform, bool bInIsBroken) override;
+	virtual void UpdateState(UGeometryCollection const& InGeometryCollection, FTransform const& InComponentTransform, bool bInIsBroken, bool bInIsVisible) override;
 	virtual void UpdateRootTransform(UGeometryCollection const& InGeometryCollection, FTransform const& InRootTransform) override;
 	virtual void UpdateTransforms(UGeometryCollection const& InGeometryCollection, TArrayView<const FMatrix> InMatrices) override;
 	//~ End IGeometryCollectionExternalRenderInterface Interface.
diff --git a/Engine/Source/Runtime/Experimental/GeometryCollectionEngine/Public/GeometryCollection/GeometryCollectionComponent.h b/Engine/Source/Runtime/Experimental/GeometryCollectionEngine/Public/GeometryCollection/GeometryCollectionComponent.h
--- a/Engine/Source/Runtime/Experimental/GeometryCollectionEngine/Public/GeometryCollection/GeometryCollectionComponent.h
+++ b/Engine/Source/Runtime/Experimental/GeometryCollectionEngine/Public/GeometryCollection/GeometryCollectionComponent.h
@@ -481,6 +481,8 @@ public:
 	
 	virtual void TickComponent(float DeltaTime, enum ELevelTick TickType, FActorComponentTickFunction *ThisTickFunction) override;
 	virtual void AsyncPhysicsTickComponent(float DeltaTime, float SimTime) override;
+
+	virtual void OnHiddenInGameChanged() override;
 	//~ Begin USceneComponent Interface.
 
 
diff --git a/Engine/Source/Runtime/Experimental/GeometryCollectionEngine/Public/GeometryCollection/GeometryCollectionExternalRenderInterface.h b/Engine/Source/Runtime/Experimental/GeometryCollectionEngine/Public/GeometryCollection/GeometryCollectionExternalRenderInterface.h
--- a/Engine/Source/Runtime/Experimental/GeometryCollectionEngine/Public/GeometryCollection/GeometryCollectionExternalRenderInterface.h
+++ b/Engine/Source/Runtime/Experimental/GeometryCollectionEngine/Public/GeometryCollection/GeometryCollectionExternalRenderInterface.h
@@ -22,7 +22,7 @@ class GEOMETRYCOLLECTIONENGINE_API IGeometryCollectionExternalRenderInterface
 public:
 	virtual void OnRegisterGeometryCollection(UGeometryCollectionComponent const& InComponent) = 0;
 	virtual void OnUnregisterGeometryCollection() = 0;
-	virtual void UpdateState(UGeometryCollection const& InGeometryCollection, FTransform const& InComponentTransform, bool bInIsBroken) = 0;
+	virtual void UpdateState(UGeometryCollection const& InGeometryCollection, FTransform const& InComponentTransform, bool bInIsBroken, bool bInIsVisible) = 0;
 	virtual void UpdateRootTransform(UGeometryCollection const& InGeometryCollection, FTransform const& InRootTransform) = 0;
 	virtual void UpdateTransforms(UGeometryCollection const& InGeometryCollection, TArrayView<const FMatrix> InMatrices) = 0;
 };
