Add null checks to SortedCCDConstraints loop in FCCDManager::UpdateParticleSweptConstraints

diff --git a/Engine/Source/Runtime/Experimental/Chaos/Private/Chaos/CCDUtilities.cpp b/Engine/Source/Runtime/Experimental/Chaos/Private/Chaos/CCDUtilities.cpp
--- a/Engine/Source/Runtime/Experimental/Chaos/Private/Chaos/CCDUtilities.cpp
+++ b/Engine/Source/Runtime/Experimental/Chaos/Private/Chaos/CCDUtilities.cpp
@@ -758,7 +758,10 @@ namespace Chaos
 		
 		for (FCCDConstraint* CCDConstraint : SortedCCDConstraints)
 		{ 
-			CCDConstraint->SweptConstraint->SetCCDResults(CCDConstraint->NetImpulse);
+			if (CCDConstraint && CCDConstraint->SweptConstraint)
+			{
+				CCDConstraint->SweptConstraint->SetCCDResults(CCDConstraint->NetImpulse);
+			}
 		}
 		
 		return HasResweptConstraint;
