Constraints:  Allow for deleting constraints again.
#trivial

diff --git a/Engine/Source/Runtime/Experimental/Animation/Constraints/Private/ConstraintsManager.cpp b/Engine/Source/Runtime/Experimental/Animation/Constraints/Private/ConstraintsManager.cpp
--- a/Engine/Source/Runtime/Experimental/Animation/Constraints/Private/ConstraintsManager.cpp
+++ b/Engine/Source/Runtime/Experimental/Animation/Constraints/Private/ConstraintsManager.cpp
@@ -502,7 +502,7 @@ int32 FConstraintsManagerController::GetConstraintIndex(const FName& InConstrain
 	
 bool FConstraintsManagerController::RemoveConstraint(UTickableConstraint* InConstraint, bool bDoNotCompensate) const
 {
-	if (FConstraintsManagerController::bDoNotRemoveConstraint == false)
+	if (FConstraintsManagerController::bDoNotRemoveConstraint == true)
 	{
 		return false;
 	}
