add 'UpdateVolumeAttributes' node to GeometryCollection Dataflow nodes
+ convert various utility functions to work on Managed Array Collection (so the Dataflow node does not need to convert to GeometryCollection)

diff --git a/Engine/Plugins/Experimental/GeometryCollectionPlugin/Source/GeometryCollectionNodes/Private/Dataflow/GeometryCollectionUtilityNodes.cpp b/Engine/Plugins/Experimental/GeometryCollectionPlugin/Source/GeometryCollectionNodes/Private/Dataflow/GeometryCollectionUtilityNodes.cpp
--- a/Engine/Plugins/Experimental/GeometryCollectionPlugin/Source/GeometryCollectionNodes/Private/Dataflow/GeometryCollectionUtilityNodes.cpp
+++ b/Engine/Plugins/Experimental/GeometryCollectionPlugin/Source/GeometryCollectionNodes/Private/Dataflow/GeometryCollectionUtilityNodes.cpp
@@ -19,6 +19,7 @@ namespace Dataflow
 		DATAFLOW_NODE_REGISTER_CREATION_FACTORY(FCreateNonOverlappingConvexHullsDataflowNode);
 		DATAFLOW_NODE_REGISTER_CREATION_FACTORY(FGenerateClusterConvexHullsFromLeafHullsDataflowNode);
 		DATAFLOW_NODE_REGISTER_CREATION_FACTORY(FGenerateClusterConvexHullsFromChildrenHullsDataflowNode);
+		DATAFLOW_NODE_REGISTER_CREATION_FACTORY(FUpdateVolumeAttributesDataflowNode);
 	}
 }
 
@@ -147,4 +148,25 @@ void FGenerateClusterConvexHullsFromChildrenHullsDataflowNode::Evaluate(Dataflow
 			SetValue<FManagedArrayCollection>(Context, static_cast<const FManagedArrayCollection>(*GeomCollection), &Collection);
 		}
 	}
-}
\ No newline at end of file
+}
+
+
+FUpdateVolumeAttributesDataflowNode::FUpdateVolumeAttributesDataflowNode(const Dataflow::FNodeParameters& InParam, FGuid InGuid)
+	: FDataflowNode(InParam, InGuid)
+{
+	RegisterInputConnection(&Collection);
+	RegisterOutputConnection(&Collection);
+}
+
+void FUpdateVolumeAttributesDataflowNode::Evaluate(Dataflow::FContext& Context, const FDataflowOutput* Out) const
+{
+	if (Out->IsA(&Collection))
+	{
+		FManagedArrayCollection InCollection = GetValue(Context, &Collection);
+		if (InCollection.NumElements(FGeometryCollection::TransformGroup) > 0)
+		{
+			FGeometryCollectionConvexUtility::SetVolumeAttributes(&InCollection);
+		}
+		SetValue(Context, MoveTemp(InCollection), &Collection);
+	}
+}
diff --git a/Engine/Plugins/Experimental/GeometryCollectionPlugin/Source/GeometryCollectionNodes/Public/Dataflow/GeometryCollectionUtilityNodes.h b/Engine/Plugins/Experimental/GeometryCollectionPlugin/Source/GeometryCollectionNodes/Public/Dataflow/GeometryCollectionUtilityNodes.h
--- a/Engine/Plugins/Experimental/GeometryCollectionPlugin/Source/GeometryCollectionNodes/Public/Dataflow/GeometryCollectionUtilityNodes.h
+++ b/Engine/Plugins/Experimental/GeometryCollectionPlugin/Source/GeometryCollectionNodes/Public/Dataflow/GeometryCollectionUtilityNodes.h
@@ -132,6 +132,26 @@ public:
 
 };
 
+/**
+ *
+ * Update the Volume and Size attributes on the target Collection (and add them if they were not present)
+ *
+ */
+USTRUCT(meta = (DataflowGeometryCollection))
+struct FUpdateVolumeAttributesDataflowNode : public FDataflowNode
+{
+	GENERATED_USTRUCT_BODY()
+	DATAFLOW_NODE_DEFINE_INTERNAL(FUpdateVolumeAttributesDataflowNode, "UpdateVolumeAttributes", "GeometryCollection|Utilities", "")
+
+public:
+	UPROPERTY(meta = (DataflowInput, DataflowOutput))
+	FManagedArrayCollection Collection;
+
+	FUpdateVolumeAttributesDataflowNode(const Dataflow::FNodeParameters& InParam, FGuid InGuid = FGuid::NewGuid());
+
+	virtual void Evaluate(Dataflow::FContext& Context, const FDataflowOutput* Out) const override;
+};
+
 
 namespace Dataflow
 {
diff --git a/Engine/Source/Runtime/Experimental/Chaos/Private/GeometryCollection/Facades/CollectionMeshFacade.cpp b/Engine/Source/Runtime/Experimental/Chaos/Private/GeometryCollection/Facades/CollectionMeshFacade.cpp
--- a/Engine/Source/Runtime/Experimental/Chaos/Private/GeometryCollection/Facades/CollectionMeshFacade.cpp
+++ b/Engine/Source/Runtime/Experimental/Chaos/Private/GeometryCollection/Facades/CollectionMeshFacade.cpp
@@ -13,6 +13,7 @@ namespace GeometryCollection::Facades
 	FCollectionMeshFacade::FCollectionMeshFacade(FManagedArrayCollection& InCollection)
 		: FCollectionUVFacade(InCollection)
 		, TransformToGeometryIndexAttribute(InCollection, "TransformToGeometryIndex", FGeometryCollection::TransformGroup)
+		, TransformIndexAttribute(InCollection, "TransformIndex", FGeometryCollection::GeometryGroup)
 		, VertexAttribute(InCollection, "Vertex", FGeometryCollection::VerticesGroup)
 		, TangentUAttribute(InCollection, "TangentU", FGeometryCollection::VerticesGroup)
 		, TangentVAttribute(InCollection, "TangentV", FGeometryCollection::VerticesGroup)
@@ -34,6 +35,7 @@ namespace GeometryCollection::Facades
 	FCollectionMeshFacade::FCollectionMeshFacade(const FManagedArrayCollection& InCollection)
 		: FCollectionUVFacade(InCollection)
 		, TransformToGeometryIndexAttribute(InCollection, "TransformToGeometryIndex", FGeometryCollection::TransformGroup)
+		, TransformIndexAttribute(InCollection, "TransformIndex", FGeometryCollection::GeometryGroup)
 		, VertexAttribute(InCollection, "Vertex", FGeometryCollection::VerticesGroup)
 		, TangentUAttribute(InCollection, "TangentU", FGeometryCollection::VerticesGroup)
 		, TangentVAttribute(InCollection, "TangentV", FGeometryCollection::VerticesGroup)
@@ -56,6 +58,7 @@ namespace GeometryCollection::Facades
 	{
 		return FCollectionUVFacade::IsValid()
 			&& TransformToGeometryIndexAttribute.IsValid()
+			&& TransformIndexAttribute.IsValid()
 			&& VertexAttribute.IsValid()
 			&& TangentUAttribute.IsValid()
 			&& TangentVAttribute.IsValid()
@@ -78,6 +81,7 @@ namespace GeometryCollection::Facades
 	{
 		FCollectionUVFacade::DefineSchema();
 		TransformToGeometryIndexAttribute.Add();
+		TransformIndexAttribute.Add();
 		VertexAttribute.Add();
 		TangentUAttribute.Add();
 		TangentVAttribute.Add();
diff --git a/Engine/Source/Runtime/Experimental/Chaos/Private/GeometryCollection/GeometryCollectionAlgo.cpp b/Engine/Source/Runtime/Experimental/Chaos/Private/GeometryCollection/GeometryCollectionAlgo.cpp
--- a/Engine/Source/Runtime/Experimental/Chaos/Private/GeometryCollection/GeometryCollectionAlgo.cpp
+++ b/Engine/Source/Runtime/Experimental/Chaos/Private/GeometryCollection/GeometryCollectionAlgo.cpp
@@ -6,6 +6,7 @@
 #include "GeometryCollection/GeometryCollectionAlgo.h"
 #include "GeometryCollection/GeometryCollection.h"
 #include "GeometryCollection/RecordedTransformTrack.h"
+#include "GeometryCollection/Facades/CollectionHierarchyFacade.h"
 #include "Async/ParallelFor.h"
 
 DEFINE_LOG_CATEGORY_STATIC(GeometryCollectionAlgoLog, Log, All);
@@ -1423,11 +1424,17 @@ namespace GeometryCollectionAlgo
 	}
 
 
-	TArray<int32> ComputeRecursiveOrder(const FGeometryCollection& Collection)
+	TArray<int32> ComputeRecursiveOrder(const FManagedArrayCollection& Collection)
 	{
 		const int32 NumTransforms = Collection.NumElements(FGeometryCollection::TransformGroup);
-		const TManagedArray<int32>& Parent = Collection.Parent;
-		const TManagedArray<TSet<int32>>& Children = Collection.Children;
+
+		Chaos::Facades::FCollectionHierarchyFacade HierarchyFacade(Collection);
+		if (!ensure(HierarchyFacade.IsValid()))
+		{
+			// We cannot compute a recursive ordering without hierarchy attributes
+			TArray<int32> Transforms;
+			return Transforms;
+		}
 
 		//traverse cluster hierarchy in depth first and record order
 		struct FClusterProcessing
@@ -1446,7 +1453,7 @@ namespace GeometryCollectionAlgo
 		//enqueue all roots
 		for (int32 TransformGroupIndex = 0; TransformGroupIndex < NumTransforms; TransformGroupIndex++)
 		{
-			if (Parent[TransformGroupIndex] == FGeometryCollection::Invalid)
+			if (HierarchyFacade.GetParent(TransformGroupIndex) == FGeometryCollection::Invalid)
 			{
 				ClustersToProcess.Emplace(TransformGroupIndex);
 			}
@@ -1466,13 +1473,14 @@ namespace GeometryCollectionAlgo
 			}
 			else
 			{
-				if (Children[ClusterTransformIdx].Num())
+				const TSet<int32>* ClusterChildren = HierarchyFacade.FindChildren(ClusterTransformIdx);
+				if (ClusterChildren && ClusterChildren->Num())
 				{
 					CurCluster.State = FClusterProcessing::VisitingChildren;
 					ClustersToProcess.Add(CurCluster);
 
 					//order of children doesn't matter as long as all children appear before parent
-					for (int32 ChildIdx : Children[ClusterTransformIdx])
+					for (int32 ChildIdx : *ClusterChildren)
 					{
 						ClustersToProcess.Emplace(ChildIdx);
 					}
diff --git a/Engine/Source/Runtime/Experimental/Chaos/Private/GeometryCollection/GeometryCollectionConvexUtility.cpp b/Engine/Source/Runtime/Experimental/Chaos/Private/GeometryCollection/GeometryCollectionConvexUtility.cpp
--- a/Engine/Source/Runtime/Experimental/Chaos/Private/GeometryCollection/GeometryCollectionConvexUtility.cpp
+++ b/Engine/Source/Runtime/Experimental/Chaos/Private/GeometryCollection/GeometryCollectionConvexUtility.cpp
@@ -8,6 +8,7 @@
 #include "GeometryCollection/GeometryCollectionProximityUtility.h"
 #include "GeometryCollection/Facades/CollectionTransformFacade.h"
 #include "GeometryCollection/Facades/CollectionHierarchyFacade.h"
+#include "GeometryCollection/Facades/CollectionMeshFacade.h"
 #include "CompGeom/ConvexHull3.h"
 #include "Templates/Sorting.h"
 #include "Spatial/PointHashGrid3.h"
@@ -1314,14 +1315,21 @@ bool CopyHulls(
 
 // helper to compute the volume of an individual piece of geometry
 double ComputeGeometryVolume(
-	const FGeometryCollection* Collection,
+	const FManagedArrayCollection* Collection,
 	int32 GeometryIdx,
 	const FTransform& GlobalTransform,
 	double ScalePerDimension
 )
 {
-	int32 VStart = Collection->VertexStart[GeometryIdx];
-	int32 VEnd = VStart + Collection->VertexCount[GeometryIdx];
+	GeometryCollection::Facades::FCollectionMeshFacade MeshFacade(*Collection);
+	const TManagedArray<int32>& VertexStart = MeshFacade.VertexStartAttribute.Get();
+	const TManagedArray<int32>& VertexCount = MeshFacade.VertexCountAttribute.Get();
+	const TManagedArray<FVector3f>& Vertex = MeshFacade.VertexAttribute.Get();
+	const TManagedArray<int32>& FaceStart = MeshFacade.FaceStartAttribute.Get();
+	const TManagedArray<int32>& FaceCount = MeshFacade.FaceCountAttribute.Get();
+	const TManagedArray<FIntVector>& Indices = MeshFacade.IndicesAttribute.Get();
+	int32 VStart = VertexStart[GeometryIdx];
+	int32 VEnd = VStart + VertexCount[GeometryIdx];
 	if (VStart == VEnd)
 	{
 		return 0.0;
@@ -1329,19 +1337,19 @@ double ComputeGeometryVolume(
 	FVector3d Center = FVector::ZeroVector;
 	for (int32 VIdx = VStart; VIdx < VEnd; VIdx++)
 	{
-		FVector Pos = GlobalTransform.TransformPosition((FVector)Collection->Vertex[VIdx]);
+		FVector Pos = GlobalTransform.TransformPosition((FVector)Vertex[VIdx]);
 		Center += (FVector3d)Pos;
 	}
 	Center /= double(VEnd - VStart);
-	int32 FStart = Collection->FaceStart[GeometryIdx];
-	int32 FEnd = FStart + Collection->FaceCount[GeometryIdx];
+	int32 FStart = FaceStart[GeometryIdx];
+	int32 FEnd = FStart + FaceCount[GeometryIdx];
 	double VolOut = 0;
 	for (int32 FIdx = FStart; FIdx < FEnd; FIdx++)
 	{
-		FIntVector Tri = Collection->Indices[FIdx];
-		FVector3d V0 = (FVector3d)GlobalTransform.TransformPosition((FVector)Collection->Vertex[Tri.X]);
-		FVector3d V1 = (FVector3d)GlobalTransform.TransformPosition((FVector)Collection->Vertex[Tri.Y]);
-		FVector3d V2 = (FVector3d)GlobalTransform.TransformPosition((FVector)Collection->Vertex[Tri.Z]);
+		FIntVector Tri = Indices[FIdx];
+		FVector3d V0 = (FVector3d)GlobalTransform.TransformPosition((FVector)Vertex[Tri.X]);
+		FVector3d V1 = (FVector3d)GlobalTransform.TransformPosition((FVector)Vertex[Tri.Y]);
+		FVector3d V2 = (FVector3d)GlobalTransform.TransformPosition((FVector)Vertex[Tri.Z]);
 
 		// add volume of the tetrahedron formed by the triangles and the reference point
 		FVector3d V1mRef = (V1 - Center) * ScalePerDimension;
@@ -1945,23 +1953,32 @@ static float GetRelativeSizeDimensionFromVolume(const float Volume)
 	return FGenericPlatformMath::Pow(Volume, 1.0f / 3.0f);
 }
 
-void FGeometryCollectionConvexUtility::SetVolumeAttributes(FGeometryCollection* Collection)
+void FGeometryCollectionConvexUtility::SetVolumeAttributes(FManagedArrayCollection* Collection)
 {
 	TManagedArray<float>& Volumes = Collection->AddAttribute<float>("Volume", FTransformCollection::TransformGroup);
 	TManagedArray<float>& Sizes = Collection->AddAttribute<float>("Size", FTransformCollection::TransformGroup);
 
-	const TManagedArray<int32>& SimulationType = Collection->SimulationType;
-	const TManagedArray<int32>& TransformToGeometryIndex = Collection->TransformToGeometryIndex;
-	const TManagedArray<int32>& Parent = Collection->Parent;
+	GeometryCollection::Facades::FCollectionMeshFacade MeshFacade(*Collection);
+	GeometryCollection::Facades::FCollectionTransformFacade TransformFacade(*Collection);
+	const TManagedArray<int32>* FoundSimulationType = Collection->FindAttributeTyped<int32>("SimulationType", FTransformCollection::TransformGroup);
+	if (!MeshFacade.IsValid() || !TransformFacade.IsValid() || !FoundSimulationType)
+	{
+		ensureMsgf(false, TEXT("Cannot compute Volume and Size attributes for Collection: Missing required attributes"));
+		return;
+	}
+
+	const TManagedArray<int32>& SimulationType = *FoundSimulationType;
+	const TManagedArray<int32>& TransformToGeometryIndex = MeshFacade.TransformToGeometryIndexAttribute.Get();
+	const TManagedArray<int32>& Parent = *TransformFacade.GetParents();
 
 	TArray<FTransform> Transforms;
-	GeometryCollectionAlgo::GlobalMatrices(Collection->Transform, Collection->Parent, Transforms);
+	GeometryCollectionAlgo::GlobalMatrices(*TransformFacade.FindTransforms(), Parent, Transforms);
 
 	TArray<float> GeoVolumes;
 	GeoVolumes.SetNumZeroed(Collection->NumElements(FGeometryCollection::GeometryGroup));
 	ParallelFor(GeoVolumes.Num(), [&](int32 GeoIdx)
 	{
-		int32 Bone = Collection->TransformIndex[GeoIdx];
+		int32 Bone = MeshFacade.TransformIndexAttribute[GeoIdx];
 		if (SimulationType[Bone] == FGeometryCollection::ESimulationTypes::FST_Rigid)
 		{
 			GeoVolumes[GeoIdx] = (float)ComputeGeometryVolume(Collection, GeoIdx, Transforms[Bone], 1.0);
diff --git a/Engine/Source/Runtime/Experimental/Chaos/Public/GeometryCollection/Facades/CollectionMeshFacade.h b/Engine/Source/Runtime/Experimental/Chaos/Public/GeometryCollection/Facades/CollectionMeshFacade.h
--- a/Engine/Source/Runtime/Experimental/Chaos/Public/GeometryCollection/Facades/CollectionMeshFacade.h
+++ b/Engine/Source/Runtime/Experimental/Chaos/Public/GeometryCollection/Facades/CollectionMeshFacade.h
@@ -53,6 +53,7 @@ namespace GeometryCollection::Facades
 		 void BakeTransform(int32 TransformIdx, const FTransform& InTransform);
 
 		TManagedArrayAccessor<int32> TransformToGeometryIndexAttribute;
+		TManagedArrayAccessor<int32> TransformIndexAttribute;
 		TManagedArrayAccessor<FVector3f> VertexAttribute;
 		TManagedArrayAccessor<FVector3f> TangentUAttribute;
 		TManagedArrayAccessor<FVector3f> TangentVAttribute;
diff --git a/Engine/Source/Runtime/Experimental/Chaos/Public/GeometryCollection/GeometryCollectionAlgo.h b/Engine/Source/Runtime/Experimental/Chaos/Public/GeometryCollection/GeometryCollectionAlgo.h
--- a/Engine/Source/Runtime/Experimental/Chaos/Public/GeometryCollection/GeometryCollectionAlgo.h
+++ b/Engine/Source/Runtime/Experimental/Chaos/Public/GeometryCollection/GeometryCollectionAlgo.h
@@ -281,6 +281,6 @@ namespace GeometryCollectionAlgo
 	*/
 	TArray<int32>
 	CHAOS_API
-	ComputeRecursiveOrder(const FGeometryCollection& Collection);
+	ComputeRecursiveOrder(const FManagedArrayCollection& Collection);
 
 }
diff --git a/Engine/Source/Runtime/Experimental/Chaos/Public/GeometryCollection/GeometryCollectionConvexUtility.h b/Engine/Source/Runtime/Experimental/Chaos/Public/GeometryCollection/GeometryCollectionConvexUtility.h
--- a/Engine/Source/Runtime/Experimental/Chaos/Public/GeometryCollection/GeometryCollectionConvexUtility.h
+++ b/Engine/Source/Runtime/Experimental/Chaos/Public/GeometryCollection/GeometryCollectionConvexUtility.h
@@ -102,7 +102,7 @@ public:
 	static bool ValidateConvexData(const FGeometryCollection* GeometryCollection);
 
 	/** Set Volume and Size attributes on the Collection (will be called by CreateNonOverlappingConvexHullData -- Volumes must be up to date for convex calc) */
-	static void SetVolumeAttributes(FGeometryCollection* Collection);
+	static void SetVolumeAttributes(FManagedArrayCollection* Collection);
 
 	/**
 	 * Copy convex hulls from *below* FromTransformIdx over to all live at ToTransformIdx.
