Iris
* Added thread-safe checks to FDirtyNetObjectTracker

#rnx

diff --git a/Engine/Plugins/Runtime/ReplicationSystemTestPlugin/Source/Private/Tests/ReplicationSystem/TestDirtyNetObjectTracker.cpp b/Engine/Plugins/Runtime/ReplicationSystemTestPlugin/Source/Private/Tests/ReplicationSystem/TestDirtyNetObjectTracker.cpp
--- a/Engine/Plugins/Runtime/ReplicationSystemTestPlugin/Source/Private/Tests/ReplicationSystem/TestDirtyNetObjectTracker.cpp
+++ b/Engine/Plugins/Runtime/ReplicationSystemTestPlugin/Source/Private/Tests/ReplicationSystem/TestDirtyNetObjectTracker.cpp
@@ -19,7 +19,9 @@ public:
 		Super::SetUp();
 		DirtyNetObjectTracker = &ReplicationSystem->GetReplicationSystemInternal()->GetDirtyNetObjectTracker();
 		ReplicationSystemId = ReplicationSystem->GetId();
-		NetObjectIndexRangeEnd = DirtyNetObjectTracker->GetDirtyNetObjects().GetNumBits() - 1U;
+		
+		FDirtyObjectsAccessor DirtyObjectsAccessor(*DirtyNetObjectTracker);
+		NetObjectIndexRangeEnd = DirtyObjectsAccessor.GetDirtyNetObjects().GetNumBits() - 1U;
 	}
 
 	virtual void TearDown() override
@@ -39,7 +41,9 @@ protected:
 
 UE_NET_TEST_FIXTURE(FDirtyNetObjectTrackerTestFixture, TestNoObjectIsDirtyFromStart)
 {
-	const FNetBitArrayView& DirtyObjects = DirtyNetObjectTracker->GetDirtyNetObjects();
+	FDirtyObjectsAccessor DirtyObjectsAccessor(*DirtyNetObjectTracker);
+
+	const FNetBitArrayView DirtyObjects = DirtyObjectsAccessor.GetDirtyNetObjects();
 	UE_NET_ASSERT_EQ(DirtyObjects.GetNumBits(), NetObjectIndexRangeEnd + 1U);
 
 	UE_NET_ASSERT_FALSE_MSG(DirtyObjects.IsAnyBitSet(), "Objects are marked as dirty before marking any object state as dirty");
@@ -50,7 +54,8 @@ UE_NET_TEST_FIXTURE(FDirtyNetObjectTrackerTestFixture, CannotMarkInvalidObjectAs
 	MarkNetObjectStateDirty(ReplicationSystemId, FNetRefHandle().GetId());
 	MarkNetObjectStateDirty(ReplicationSystemId, NetObjectIndexRangeEnd + 1);
 
-	const FNetBitArrayView& DirtyObjects = DirtyNetObjectTracker->GetDirtyNetObjects();
+	FDirtyObjectsAccessor DirtyObjectsAccessor(*DirtyNetObjectTracker);
+	const FNetBitArrayView DirtyObjects = DirtyObjectsAccessor.GetDirtyNetObjects();
 	UE_NET_ASSERT_FALSE(DirtyObjects.IsAnyBitSet());
 }
 
@@ -61,7 +66,8 @@ UE_NET_TEST_FIXTURE(FDirtyNetObjectTrackerTestFixture, CanMarkValidObjectAsDirty
 	MarkNetObjectStateDirty(ReplicationSystemId, NetObjectIndexRangeStart);
 	MarkNetObjectStateDirty(ReplicationSystemId, NetObjectIndexRangeEnd);
 
-	const FNetBitArrayView& DirtyObjects = DirtyNetObjectTracker->GetDirtyNetObjects();
+	FDirtyObjectsAccessor DirtyObjectsAccessor(*DirtyNetObjectTracker);
+	const FNetBitArrayView DirtyObjects = DirtyObjectsAccessor.GetDirtyNetObjects();
 	UE_NET_ASSERT_TRUE(DirtyObjects.GetBit(IndexInRange));
 	UE_NET_ASSERT_TRUE(DirtyObjects.GetBit(NetObjectIndexRangeStart));
 	UE_NET_ASSERT_TRUE(DirtyObjects.GetBit(NetObjectIndexRangeEnd));
@@ -73,7 +79,9 @@ UE_NET_TEST_FIXTURE(FDirtyNetObjectTrackerTestFixture, CanClearDirtyObjects)
 	MarkNetObjectStateDirty(ReplicationSystemId, NetObjectIndexRangeEnd);
 
 	DirtyNetObjectTracker->ClearDirtyNetObjects();
-	const FNetBitArrayView& DirtyObjects = DirtyNetObjectTracker->GetDirtyNetObjects();
+
+	FDirtyObjectsAccessor DirtyObjectsAccessor(*DirtyNetObjectTracker);
+	const FNetBitArrayView DirtyObjects = DirtyObjectsAccessor.GetDirtyNetObjects();
 	UE_NET_ASSERT_FALSE(DirtyObjects.IsAnyBitSet());
 }
 
diff --git a/Engine/Source/Runtime/Experimental/Iris/Core/Private/Iris/ReplicationSystem/DirtyNetObjectTracker.cpp b/Engine/Source/Runtime/Experimental/Iris/Core/Private/Iris/ReplicationSystem/DirtyNetObjectTracker.cpp
--- a/Engine/Source/Runtime/Experimental/Iris/Core/Private/Iris/ReplicationSystem/DirtyNetObjectTracker.cpp
+++ b/Engine/Source/Runtime/Experimental/Iris/Core/Private/Iris/ReplicationSystem/DirtyNetObjectTracker.cpp
@@ -35,7 +35,6 @@ FDirtyNetObjectTracker::FDirtyNetObjectTracker()
 , NetObjectIdRangeStart(0U)
 , NetObjectIdRangeEnd(0U)
 , NetObjectIdCount(0)
-, bHasPolledGlobalDirtyTracker(false)
 {
 }
 
@@ -58,13 +57,15 @@ void FDirtyNetObjectTracker::Init(const FDirtyNetObjectTrackerInitParams& Params
 	 * For now we support all IDs up to RangeEnd. This could be expensive if we partition things more in some way or other.
 	 * In the latter case we would have to add functionality to FNetBitArrayView to handle an offset or add a "FNetSparseBitArray".
 	 */
-	NetObjectIdCount = NetObjectIdRangeEnd + 1;
+	NetObjectIdCount = Params.MaxObjectCount;
 
 	GlobalDirtyTrackerPollHandle = FGlobalDirtyNetObjectTracker::CreatePoller();
 
 	DirtyNetObjectWordCount = (NetObjectIdCount + StorageTypeBitCount - 1)/StorageTypeBitCount;
 	DirtyNetObjectContainer = new StorageType[DirtyNetObjectWordCount];
-	ClearDirtyNetObjects();
+	FMemory::Memzero(DirtyNetObjectContainer, DirtyNetObjectWordCount * sizeof(StorageType));
+
+	AllowExternalAccess();
 
 	UE_LOG_DIRTYOBJECTTRACKER(TEXT("FDirtyNetObjectTracker::Init %u Id, Start:%u, End: %u"), ReplicationSystemId, NetObjectIdRangeStart, NetObjectIdRangeEnd);
 }
@@ -87,6 +88,8 @@ void FDirtyNetObjectTracker::UpdateDirtyNetObjects()
 
 	IRIS_PROFILER_SCOPE(FDirtyNetObjectTracker_UpdateDirtyNetObjects)
 
+	LockExternalAccess();
+
 	bHasPolledGlobalDirtyTracker = true;
 	const TSet<FNetHandle>& GlobalDirtyNetObjects = FGlobalDirtyNetObjectTracker::GetDirtyNetObjects(GlobalDirtyTrackerPollHandle);
 	for (FNetHandle NetHandle : GlobalDirtyNetObjects)
@@ -99,10 +102,16 @@ void FDirtyNetObjectTracker::UpdateDirtyNetObjects()
 			DirtyNetObjectContainer[BitOffset/StorageTypeBitCount] |= BitMask;
 		}
 	}
+
+	AllowExternalAccess();
 }
 
 void FDirtyNetObjectTracker::MarkNetObjectDirty(uint32 NetObjectIndex)
 {
+#if UE_NET_THREAD_SAFETY_CHECK
+	checkf(bIsExternalAccessAllowed, TEXT("Cannot mark objects dirty while the bitarray is locked for modifications."));
+#endif
+
 	if ((NetObjectIndex >= NetObjectIdRangeStart) & (NetObjectIndex <= NetObjectIdRangeEnd))
 	{
 		const uint32 BitOffset = NetObjectIndex;
@@ -115,20 +124,41 @@ void FDirtyNetObjectTracker::MarkNetObjectDirty(uint32 NetObjectIndex)
 	}
 }
 
+void FDirtyNetObjectTracker::LockExternalAccess()
+{
+#if UE_NET_THREAD_SAFETY_CHECK
+	bIsExternalAccessAllowed = false;
+#endif
+}
+
+void FDirtyNetObjectTracker::AllowExternalAccess()
+{
+#if UE_NET_THREAD_SAFETY_CHECK
+	bIsExternalAccessAllowed = true;
+#endif
+}
+
 FNetBitArrayView FDirtyNetObjectTracker::GetDirtyNetObjects() const
 {
+#if UE_NET_THREAD_SAFETY_CHECK
+	checkf(!bIsExternalAccessAllowed, TEXT("Cannot access the DirtyNetObjects bitarray unless its locked for multithread access."));
+#endif
 	return FNetBitArrayView(DirtyNetObjectContainer, NetObjectIdCount);
 }
 
 void FDirtyNetObjectTracker::ClearDirtyNetObjects()
 {
+	LockExternalAccess();
+
 	if (bHasPolledGlobalDirtyTracker)
 	{
 		bHasPolledGlobalDirtyTracker = false;
 		FGlobalDirtyNetObjectTracker::ResetDirtyNetObjects(GlobalDirtyTrackerPollHandle);
 	}
 
-	FMemory::Memset(DirtyNetObjectContainer, 0, DirtyNetObjectWordCount*sizeof(StorageType));
+	FMemory::Memzero(DirtyNetObjectContainer, DirtyNetObjectWordCount*sizeof(StorageType));
+	
+	AllowExternalAccess();
 
 	std::atomic_thread_fence(std::memory_order_seq_cst);
 }
diff --git a/Engine/Source/Runtime/Experimental/Iris/Core/Private/Iris/ReplicationSystem/DirtyNetObjectTracker.h b/Engine/Source/Runtime/Experimental/Iris/Core/Private/Iris/ReplicationSystem/DirtyNetObjectTracker.h
--- a/Engine/Source/Runtime/Experimental/Iris/Core/Private/Iris/ReplicationSystem/DirtyNetObjectTracker.h
+++ b/Engine/Source/Runtime/Experimental/Iris/Core/Private/Iris/ReplicationSystem/DirtyNetObjectTracker.h
@@ -5,9 +5,12 @@
 #include "Net/Core/NetBitArray.h"
 #include "Net/Core/DirtyNetObjectTracker/GlobalDirtyNetObjectTracker.h"
 
+#include "Iris/IrisConfig.h"
+
 namespace UE::Net::Private
 {
 	class FNetRefHandleManager;
+	class FDirtyObjectsAccessor;
 }
 
 namespace UE::Net::Private
@@ -19,6 +22,7 @@ struct FDirtyNetObjectTrackerInitParams
 {
 	const FNetRefHandleManager* NetRefHandleManager = nullptr;
 	uint32 ReplicationSystemId = 0;
+	uint32 MaxObjectCount = 0;
 	uint32 NetObjectIndexRangeStart = 0;
 	uint32 NetObjectIndexRangeEnd = 0;
 };
@@ -34,20 +38,29 @@ public:
 	/* Update dirty objects with the set of globally marked dirty objects. **/
 	void UpdateDirtyNetObjects();
 
-	// N.B. These methods do not use atomics.
-	// It's up to the user to determine whether it's safe to call the functions or not.
-	FNetBitArrayView GetDirtyNetObjects() const;
+	/** Set safety permissions so no one can write in the bit array via the public methods */
+	void LockExternalAccess();
+
+	/** Release safety permissions and allow to write in the bit array via the public methods */
+	void AllowExternalAccess();
+
+	/** Reset the global and local dirty objects lists */
 	void ClearDirtyNetObjects();
 
 private:
 	friend IRISCORE_API void MarkNetObjectStateDirty(uint32 ReplicationSystemId, uint32 NetObjectIndex);
 
+	friend FDirtyObjectsAccessor;
+
 	using StorageType = FNetBitArrayView::StorageWordType;
 	static constexpr uint32 StorageTypeBitCount = FNetBitArrayView::WordBitCount;
 
 	void Deinit();
 	void MarkNetObjectDirty(uint32 NetObjectIndex);
 
+	/** Can only be accessed via FDirtyObjectsAccessor */
+	FNetBitArrayView GetDirtyNetObjects() const;
+
 	const FNetRefHandleManager* NetRefHandleManager;
 	StorageType* DirtyNetObjectContainer;
 	FGlobalDirtyNetObjectTracker::FPollHandle GlobalDirtyTrackerPollHandle;
@@ -56,7 +69,33 @@ private:
 	uint32 NetObjectIdRangeStart;
 	uint32 NetObjectIdRangeEnd;
 	uint32 NetObjectIdCount;
-	bool bHasPolledGlobalDirtyTracker;
+	
+	bool bHasPolledGlobalDirtyTracker = false;
+
+#if UE_NET_THREAD_SAFETY_CHECK
+	std::atomic_bool bIsExternalAccessAllowed = false;
+#endif
+};
+
+class FDirtyObjectsAccessor
+{
+public:
+	FDirtyObjectsAccessor(FDirtyNetObjectTracker& InDirtyNetObjectTracker)
+		: DirtyNetObjectTracker(InDirtyNetObjectTracker)
+	{
+		DirtyNetObjectTracker.LockExternalAccess();
+	}
+
+	~FDirtyObjectsAccessor()
+	{
+		DirtyNetObjectTracker.AllowExternalAccess();
+	}
+
+	FNetBitArrayView GetDirtyNetObjects()				{ return DirtyNetObjectTracker.GetDirtyNetObjects(); }
+	const FNetBitArrayView GetDirtyNetObjects() const	{ return DirtyNetObjectTracker.GetDirtyNetObjects(); }
+
+private:
+	FDirtyNetObjectTracker& DirtyNetObjectTracker;
 };
 
 }
diff --git a/Engine/Source/Runtime/Experimental/Iris/Core/Private/Iris/ReplicationSystem/Filtering/ReplicationFiltering.cpp b/Engine/Source/Runtime/Experimental/Iris/Core/Private/Iris/ReplicationSystem/Filtering/ReplicationFiltering.cpp
--- a/Engine/Source/Runtime/Experimental/Iris/Core/Private/Iris/ReplicationSystem/Filtering/ReplicationFiltering.cpp
+++ b/Engine/Source/Runtime/Experimental/Iris/Core/Private/Iris/ReplicationSystem/Filtering/ReplicationFiltering.cpp
@@ -297,7 +297,8 @@ void FReplicationFiltering::FilterPrePoll()
 
 	UpdateSubObjectFilters();
 
-	const FNetBitArrayView DirtyObjects = ReplicationSystem->GetReplicationSystemInternal()->GetDirtyNetObjectTracker().GetDirtyNetObjects();
+	FDirtyObjectsAccessor DirtyObjectsAccessor(ReplicationSystem->GetReplicationSystemInternal()->GetDirtyNetObjectTracker());
+	const FNetBitArrayView DirtyObjects = DirtyObjectsAccessor.GetDirtyNetObjects();
 
 	if (HasRawFilters())
 	{
@@ -368,7 +369,9 @@ void FReplicationFiltering::FilterPostPoll()
 {
 	if (HasFragmentFilters())
 	{
-		const FNetBitArrayView DirtyObjects = ReplicationSystem->GetReplicationSystemInternal()->GetDirtyNetObjectTracker().GetDirtyNetObjects();
+		FDirtyObjectsAccessor DirtyObjectsAccessor(ReplicationSystem->GetReplicationSystemInternal()->GetDirtyNetObjectTracker());
+		const FNetBitArrayView DirtyObjects = DirtyObjectsAccessor.GetDirtyNetObjects();
+
 		UpdateDynamicFilters(ENetFilterType::PostPoll_FragmentBased, DirtyObjects);
 	}
 }
@@ -1077,8 +1080,9 @@ void FReplicationFiltering::UpdateGroupFiltering()
 				{
 					FNetBitArray& GroupFilteredOutObjects = LocalConnectionInfos[ConnectionId].GroupFilteredOutObjects;
 					FNetBitArray& ObjectsInScopeBeforeDynamicFiltering = LocalConnectionInfos[ConnectionId].ObjectsInScopeBeforeDynamicFiltering;
-					const FNetBitArray& ScopableObjects = NetRefHandleManager->GetScopableInternalIndices();
-					for (uint32 ObjectIndex : MakeArrayView(Group->Members.GetData(), Group->Members.Num()))
+					const FNetBitArray& GlobalScopableObjects = NetRefHandleManager->GetScopableInternalIndices();
+
+					for (uint32 ObjectIndex : Group->Members)
 					{
 						GroupFilteredOutObjects.SetBit(ObjectIndex);
 						ObjectsInScopeBeforeDynamicFiltering.ClearBit(ObjectIndex);
@@ -1086,7 +1090,7 @@ void FReplicationFiltering::UpdateGroupFiltering()
 						// Filter subobjects
 						for (const FInternalNetRefIndex SubObjectIndex : NetRefHandleManager->GetSubObjects(ObjectIndex))
 						{
-							const bool bIsScopable = ScopableObjects.GetBit(SubObjectIndex);
+							const bool bIsScopable = GlobalScopableObjects.GetBit(SubObjectIndex);
 							GroupFilteredOutObjects.SetBitValue(SubObjectIndex, bIsScopable);
 							ObjectsInScopeBeforeDynamicFiltering.ClearBit(SubObjectIndex);
 						}
@@ -1627,13 +1631,19 @@ bool FReplicationFiltering::GetSubObjectFilterStatus(FNetObjectGroupHandle Group
 
 void FReplicationFiltering::AddGroupFilter(FNetObjectGroupHandle GroupHandle)
 {
-	check(!SubObjectFilterGroups.GetBit(GroupHandle));
+	const bool bIsValidGroup = ensureMsgf(GroupHandle != InvalidNetObjectGroupHandle, TEXT("AddGroupFilter received an invalid group handle"));
 
-	if (ensure(GroupHandle && !FilterGroups.GetBit(GroupHandle)))
+	if (bIsValidGroup)
 	{
-		const bool bGroupExists = FilterGroups.GetBit(GroupHandle);
-		if (ensure(bGroupExists == false))
+		const bool bIsFiltering = FilterGroups.GetBit(GroupHandle);
+		ensureMsgf(!bIsFiltering, TEXT("NetObjectGroup %u was asked to start filtering but it was already doing so."), GroupHandle);
+
+		check(!SubObjectFilterGroups.GetBit(GroupHandle));
+		
+		if (!bIsFiltering)
 		{
+			Groups->AddFilterTrait(GroupHandle);
+
 			FilterGroups.SetBit(GroupHandle);
 			DirtyFilterGroups.SetBit(GroupHandle);
 			GroupInfos[GroupHandle].ConnectionStateIndex = AllocPerObjectInfo();
@@ -1646,12 +1656,15 @@ void FReplicationFiltering::AddGroupFilter(FNetObjectGroupHandle GroupHandle)
 	}
 }
 
-// Remove group from filtered groups
 void FReplicationFiltering::RemoveGroupFilter(FNetObjectGroupHandle GroupHandle)
 {
+	// Remove the filter trait from the group
 	if (GroupHandle && FilterGroups.GetBit(GroupHandle))
 	{
-		ValidConnections.ForAllSetBits( [GroupHandle, this](uint32 ConnectionIndex) { this->SetGroupFilterStatus(GroupHandle, ConnectionIndex, ENetFilterStatus::Allow); } );
+		ValidConnections.ForAllSetBits( [GroupHandle, this](uint32 ConnectionIndex) 
+		{ 
+			this->SetGroupFilterStatus(GroupHandle, ConnectionIndex, ENetFilterStatus::Allow); 
+		});
 
 		// Mark group as no longer a group filter group
 		FilterGroups.ClearBit(GroupHandle);
@@ -1659,6 +1672,8 @@ void FReplicationFiltering::RemoveGroupFilter(FNetObjectGroupHandle GroupHandle)
 		GroupInfos[GroupHandle].ConnectionStateIndex = 0U;
 		FreePerObjectInfo(ConnectionStateIndex);
 
+		Groups->RemoveFilterTrait(GroupHandle);
+
 		UE_LOG(LogIrisFiltering, Verbose, TEXT("ReplicationFiltering::RemoveGroupFilter GroupIndex: %u"), GroupHandle);
 	}
 }
diff --git a/Engine/Source/Runtime/Experimental/Iris/Core/Private/Iris/ReplicationSystem/ObjectReplicationBridge.cpp b/Engine/Source/Runtime/Experimental/Iris/Core/Private/Iris/ReplicationSystem/ObjectReplicationBridge.cpp
--- a/Engine/Source/Runtime/Experimental/Iris/Core/Private/Iris/ReplicationSystem/ObjectReplicationBridge.cpp
+++ b/Engine/Source/Runtime/Experimental/Iris/Core/Private/Iris/ReplicationSystem/ObjectReplicationBridge.cpp
@@ -705,7 +705,10 @@ void UObjectReplicationBridge::PreUpdateAndPollImpl(FNetRefHandle Handle)
 	const FNetRefHandleManager& LocalNetRefHandleManager = ReplicationSystemInternal->GetNetRefHandleManager();
 	const TArray<UObject*>& ReplicatedInstances = LocalNetRefHandleManager.GetReplicatedInstances();
 	const bool bIsUsingPushModel = IsIrisPushModelEnabled();
-	FNetBitArrayView DirtyObjects = ReplicationSystemInternal->GetDirtyNetObjectTracker().GetDirtyNetObjects();
+
+	FDirtyObjectsAccessor DirtyObjectsAccessor(ReplicationSystemInternal->GetDirtyNetObjectTracker());
+
+	FNetBitArrayView DirtyObjects = DirtyObjectsAccessor.GetDirtyNetObjects();
 
 	struct FPreUpdateAndPollStats
 	{
@@ -951,6 +954,10 @@ void UObjectReplicationBridge::PreUpdateAndPollImpl(FNetRefHandle Handle)
 
 	{
 		IRIS_PROFILER_SCOPE(PreUpdateAndPollImpl_Poll);
+
+		// From here we call into user code via PreUpdateInstanceFunction, so allow external code to set dirty flags since DirtyObjects is not read anymore.
+		ReplicationSystemInternal->GetDirtyNetObjectTracker().AllowExternalAccess();
+
 		if (IsIrisPushModelEnabled())
 		{
 			ObjectsConsideredForPolling.ForAllSetBits(PushModelUpdateAndPollFunction);
@@ -996,7 +1003,8 @@ void UObjectReplicationBridge::UpdateInstancesWorldLocation()
 		}
 	};
 
-	const FNetBitArrayView DirtyObjects = ReplicationSystemInternal->GetDirtyNetObjectTracker().GetDirtyNetObjects();
+	FDirtyObjectsAccessor DirtyObjectsAccessor(ReplicationSystemInternal->GetDirtyNetObjectTracker());
+	const FNetBitArrayView DirtyObjects = DirtyObjectsAccessor.GetDirtyNetObjects();
 	DirtyObjects.ForAllSetBits(UpdateInstanceWorldLocation);
 
 }
diff --git a/Engine/Source/Runtime/Experimental/Iris/Core/Private/Iris/ReplicationSystem/ReplicationSystem.cpp b/Engine/Source/Runtime/Experimental/Iris/Core/Private/Iris/ReplicationSystem/ReplicationSystem.cpp
--- a/Engine/Source/Runtime/Experimental/Iris/Core/Private/Iris/ReplicationSystem/ReplicationSystem.cpp
+++ b/Engine/Source/Runtime/Experimental/Iris/Core/Private/Iris/ReplicationSystem/ReplicationSystem.cpp
@@ -87,6 +87,7 @@ public:
 		{
 			DirtyNetObjectTrackerInitParams.NetRefHandleManager = &NetRefHandleManager;
 			DirtyNetObjectTrackerInitParams.ReplicationSystemId = ReplicationSystemId;
+			DirtyNetObjectTrackerInitParams.MaxObjectCount = MaxObjectCount;
 			DirtyNetObjectTrackerInitParams.NetObjectIndexRangeStart = 1;
 			DirtyNetObjectTrackerInitParams.NetObjectIndexRangeEnd = MaxObjectCount - 1U;
 
@@ -221,7 +222,9 @@ public:
 	{
 		FDirtyNetObjectTracker& DirtyNetObjectTracker = ReplicationSystemInternal.GetDirtyNetObjectTracker();
 		DirtyNetObjectTracker.UpdateDirtyNetObjects();
-		FNetBitArrayView DirtyObjects = DirtyNetObjectTracker.GetDirtyNetObjects();
+
+		FDirtyObjectsAccessor DirtyObjectsAccessor(ReplicationSystemInternal.GetDirtyNetObjectTracker());
+		FNetBitArrayView DirtyObjects = DirtyObjectsAccessor.GetDirtyNetObjects();
 
 		// Due to objects having been marked as dirty and later removed we must make sure that all dirty objects are still in scope.
 		const FNetBitArrayView ScopableInternalIndices = MakeNetBitArrayView(ReplicationSystemInternal.GetNetRefHandleManager().GetScopableInternalIndices());
@@ -331,10 +334,10 @@ public:
 			CopiedObjectCount += FReplicationInstanceOperationsInternal::CopyObjectStateData(ChangeMaskWriter, Cache, NetRefHandleManager, SerializationContext, DirtyIndex);
 		};
 
-		FDirtyNetObjectTracker& DirtyNetObjectTracker = ReplicationSystemInternal.GetDirtyNetObjectTracker();
 
-		// Only iterate over dirty objects
-		FNetBitArrayView DirtyObjects = DirtyNetObjectTracker.GetDirtyNetObjects();
+		// Only iterate over dirty and relevant objects
+		FDirtyObjectsAccessor DirtyObjectsAccessor(ReplicationSystemInternal.GetDirtyNetObjectTracker());
+		const FNetBitArrayView DirtyObjects = DirtyObjectsAccessor.GetDirtyNetObjects();
 
 		// Copy all ReplicatedObjects with dirty state data
 		DirtyObjects.ForAllSetBits(CopyFunction);
@@ -359,7 +362,9 @@ public:
 			FReplicationInstanceOperationsInternal::ResetObjectStateDirtiness(NetRefHandleManager, DirtyIndex);
 		};
 
-		FNetBitArrayView DirtyObjects = DirtyNetObjectTracker.GetDirtyNetObjects();
+		FDirtyObjectsAccessor DirtyObjectsAccessor(DirtyNetObjectTracker);
+		const FNetBitArrayView DirtyObjects = DirtyObjectsAccessor.GetDirtyNetObjects();
+
 		DirtyObjects.ForAllSetBits(ResetDirtinessFunction);
 
 		DirtyNetObjectTracker.ClearDirtyNetObjects();
@@ -599,7 +604,9 @@ void UReplicationSystem::PreSendUpdate(float DeltaSeconds)
 	if (bAllowObjectReplication)
 	{
 		// Update object priorities
-		FNetBitArrayView DirtyObjects = InternalSys.GetDirtyNetObjectTracker().GetDirtyNetObjects();
+		FDirtyObjectsAccessor DirtyObjectsAccessor(InternalSys.GetDirtyNetObjectTracker());
+		const FNetBitArrayView DirtyObjects = DirtyObjectsAccessor.GetDirtyNetObjects();
+
 		Impl->UpdatePrioritization(ReplicatingConnections, DirtyObjects);
 
 		// Delta compression preparations before send
diff --git a/Engine/Source/Runtime/Experimental/Iris/Core/Public/Iris/IrisConfig.h b/Engine/Source/Runtime/Experimental/Iris/Core/Public/Iris/IrisConfig.h
--- a/Engine/Source/Runtime/Experimental/Iris/Core/Public/Iris/IrisConfig.h
+++ b/Engine/Source/Runtime/Experimental/Iris/Core/Public/Iris/IrisConfig.h
@@ -48,3 +48,7 @@ IRISCORE_API void SetUseIrisReplication(bool EnableIrisReplication);
 #	define UE_NET_IRIS_VERBOSE_CSV_STATS 0
 #endif
 
+/** Enables code that detects non-thread safe access to network data */
+#ifndef UE_NET_THREAD_SAFETY_CHECK
+#	define UE_NET_THREAD_SAFETY_CHECK DO_CHECK
+#endif
\ No newline at end of file
