Iris ENetObjectGroupTraits enum entry use hardcoded bits now.

#rnx

diff --git a/Engine/Source/Runtime/Experimental/Iris/Core/Private/Iris/ReplicationSystem/Filtering/NetObjectGroups.h b/Engine/Source/Runtime/Experimental/Iris/Core/Private/Iris/ReplicationSystem/Filtering/NetObjectGroups.h
--- a/Engine/Source/Runtime/Experimental/Iris/Core/Private/Iris/ReplicationSystem/Filtering/NetObjectGroups.h
+++ b/Engine/Source/Runtime/Experimental/Iris/Core/Private/Iris/ReplicationSystem/Filtering/NetObjectGroups.h
@@ -23,9 +23,9 @@ namespace UE::Net::Private
 
 enum class ENetObjectGroupTraits : uint32
 {
-	None = 0U,
-	IsFindableByName = 1U << None,
-	IsFiltering = 1U << IsFindableByName,
+	None				= 0x0000,
+	IsFindableByName	= 0x0001,
+	IsFiltering			= 0x0002,
 };
 ENUM_CLASS_FLAGS(ENetObjectGroupTraits);
 
