immediately add geometry collection and cluster union particles into the external SQ queue to ensure there's never a moment where they are not in the external SQ even after it gets swapped

diff --git a/Engine/Source/Runtime/Experimental/Chaos/Private/PBDRigidsSolver.cpp b/Engine/Source/Runtime/Experimental/Chaos/Private/PBDRigidsSolver.cpp
--- a/Engine/Source/Runtime/Experimental/Chaos/Private/PBDRigidsSolver.cpp
+++ b/Engine/Source/Runtime/Experimental/Chaos/Private/PBDRigidsSolver.cpp
@@ -721,6 +721,15 @@ namespace Chaos
 		InProxy->Initialize(GetEvolution());
 		InProxy->NewData(); // Buffers data on the proxy.
 
+		// Need to immediately add to the SQ to make sure that the GC is in the external SQ even after the SQ's flip.
+		for (const TUniquePtr<FPBDRigidParticle>& Particle : InProxy->GetExternalParticles())
+		{
+			if (Particle && !Particle->Disabled())
+			{
+				UpdateParticleInAccelerationStructure_External(Particle.Get(), /*bDelete=*/false);
+			}
+		}
+
 		// There used to be an EnqueueCommandImmediate here to push the initialization of the GC
 		// onto the physics thread. We should use AddDirtyProxy here instead to better match up with
 		// the initialization order done in the ProcessSinglePushedData_Internal. Using EnqueueCommandImmediate
@@ -784,6 +793,7 @@ namespace Chaos
 			}
 
 			Particle->SetUniqueIdx(GetEvolution()->GenerateUniqueIdx());
+			UpdateParticleInAccelerationStructure_External(Particle, /*bDelete=*/false);
 		}
 		Proxy->SetSolver(this);
 		AddDirtyProxy(Proxy);
