Added ability to check if something was succesfully removed or updated an Acceleration structure to enhance future error checking and debugging.

diff --git a/Engine/Source/Runtime/Experimental/Chaos/Public/Chaos/ISpatialAcceleration.h b/Engine/Source/Runtime/Experimental/Chaos/Public/Chaos/ISpatialAcceleration.h
--- a/Engine/Source/Runtime/Experimental/Chaos/Public/Chaos/ISpatialAcceleration.h
+++ b/Engine/Source/Runtime/Experimental/Chaos/Public/Chaos/ISpatialAcceleration.h
@@ -302,14 +302,14 @@ public:
 		return true;
 	}
 
-	virtual void RemoveElementFrom(const TPayloadType& Payload, FSpatialAccelerationIdx Idx)
+	virtual bool RemoveElementFrom(const TPayloadType& Payload, FSpatialAccelerationIdx Idx)
 	{
-		RemoveElement(Payload);
+		return RemoveElement(Payload);
 	}
 
-	virtual void UpdateElementIn(const TPayloadType& Payload, const TAABB<T, d>& NewBounds, bool bHasBounds, FSpatialAccelerationIdx Idx)
+	virtual bool UpdateElementIn(const TPayloadType& Payload, const TAABB<T, d>& NewBounds, bool bHasBounds, FSpatialAccelerationIdx Idx)
 	{
-		UpdateElement(Payload, NewBounds, bHasBounds);
+		return UpdateElement(Payload, NewBounds, bHasBounds);
 	}
 
 	virtual TUniquePtr<ISpatialAcceleration<TPayloadType, T, d>> Copy() const
diff --git a/Engine/Source/Runtime/Experimental/Chaos/Public/Chaos/SpatialAccelerationCollection.h b/Engine/Source/Runtime/Experimental/Chaos/Public/Chaos/SpatialAccelerationCollection.h
--- a/Engine/Source/Runtime/Experimental/Chaos/Public/Chaos/SpatialAccelerationCollection.h
+++ b/Engine/Source/Runtime/Experimental/Chaos/Public/Chaos/SpatialAccelerationCollection.h
@@ -570,10 +570,11 @@ public:
 		return Indices;
 	}
 
-	virtual void RemoveElementFrom(const TPayloadType& Payload, FSpatialAccelerationIdx SpatialIdx) override
+	// Returns true if the element was in fact removed from some SpatialIdx, otherwise returns false if the element was not found anywhere
+	virtual bool RemoveElementFrom(const TPayloadType& Payload, FSpatialAccelerationIdx SpatialIdx) override
 	{
 		const uint16 UseBucket = ((1 << SpatialIdx.Bucket) & this->ActiveBucketsMask) ? SpatialIdx.Bucket : 0;
-		const bool Success = Buckets[UseBucket].Objects[SpatialIdx.InnerIdx].Acceleration->RemoveElement(Payload);
+		bool Success = Buckets[UseBucket].Objects[SpatialIdx.InnerIdx].Acceleration->RemoveElement(Payload);
 		//ensure(Success); // Debug check to see if SpatialIdx was what the caller expected it to be
 		// Make sure that we remove this Payload even if the SpatialIdx is wrong
 		for (FSpatialAccelerationIdx Idx : GetAllSpatialIndices())
@@ -581,15 +582,18 @@ public:
 			if (!(Idx == SpatialIdx))
 			{
 				const uint16 Buckt = ((1 << Idx.Bucket) & this->ActiveBucketsMask) ? Idx.Bucket : 0;
-				const bool Removed = Buckets[Buckt].Objects[Idx.InnerIdx].Acceleration->RemoveElement(Payload);					
+				const bool Removed = Buckets[Buckt].Objects[Idx.InnerIdx].Acceleration->RemoveElement(Payload);
+				Success = Success || Removed;
 			}
-		}		
+		}
+		return Success;
 	}
 
-	virtual void UpdateElementIn(const TPayloadType& Payload, const TAABB<T, d>& NewBounds, bool bHasBounds, FSpatialAccelerationIdx SpatialIdx)
+	// Returns true if the element was in fact updated or if it was moved from one SpatialIdx to another during the update or false if not found at any SpatialIdx
+	virtual bool UpdateElementIn(const TPayloadType& Payload, const TAABB<T, d>& NewBounds, bool bHasBounds, FSpatialAccelerationIdx SpatialIdx)
 	{
 		const uint16 UseBucket = ((1 << SpatialIdx.Bucket) & this->ActiveBucketsMask) ? SpatialIdx.Bucket : 0;
-		const bool bElementExisted = Buckets[UseBucket].Objects[SpatialIdx.InnerIdx].Acceleration->UpdateElement(Payload, NewBounds, bHasBounds);
+		bool bElementExisted = Buckets[UseBucket].Objects[SpatialIdx.InnerIdx].Acceleration->UpdateElement(Payload, NewBounds, bHasBounds);
 		// In case spatial index changed, remove this element in all other substructures
 		if (!bElementExisted)
 		{
@@ -598,10 +602,12 @@ public:
 				if (!(Idx == SpatialIdx))
 				{
 					const uint16 Buckt = ((1 << Idx.Bucket) & this->ActiveBucketsMask) ? Idx.Bucket : 0;
-					Buckets[Buckt].Objects[Idx.InnerIdx].Acceleration->RemoveElement(Payload);
+					const bool Removed = Buckets[Buckt].Objects[Idx.InnerIdx].Acceleration->RemoveElement(Payload);
+					bElementExisted = bElementExisted || Removed;
 				}
 			}
 		}
+		return bElementExisted;
 	}
 
 	virtual TUniquePtr<ISpatialAcceleration<TPayloadType, T, d>> Copy() const
