Chaos Cloth: Fix multiplier calculation in new FXPBDAnisotropicBendingConstraints.

Multipliers weren't getting calculated correctly due to a typo.

diff --git a/Engine/Source/Runtime/Experimental/Chaos/Private/Chaos/XPBDAnisotropicBendingConstraints.cpp b/Engine/Source/Runtime/Experimental/Chaos/Private/Chaos/XPBDAnisotropicBendingConstraints.cpp
--- a/Engine/Source/Runtime/Experimental/Chaos/Private/Chaos/XPBDAnisotropicBendingConstraints.cpp
+++ b/Engine/Source/Runtime/Experimental/Chaos/Private/Chaos/XPBDAnisotropicBendingConstraints.cpp
@@ -191,7 +191,7 @@ TArray<FSolverVec3> FXPBDAnisotropicBendingConstraints::GenerateWarpWeftBiasBase
 		else
 		{
 			// Stiffness is blend between weft and bias
-			UVDiffTransformed = FSolverVec3((FSolverReal)0.f, UVDiffAbs.Y - UVDiff.X, UVDiffAbs.X);
+			UVDiffTransformed = FSolverVec3((FSolverReal)0.f, UVDiffAbs.Y - UVDiffAbs.X, UVDiffAbs.X);
 		}
 		const FSolverReal Denom = UVDiffTransformed.X + UVDiffTransformed.Y + UVDiffTransformed.Z;
 		return Denom > UE_SMALL_NUMBER ? UVDiffTransformed / Denom : FSolverVec3(1.f, 0.f, 0.f); // Default to Warp if zero length
