Add bUpdateKinematicProperties parameter to UpdateAllClusterUnionProperties, allowing it to be called without updating kinematics. Use case is in removal of  a child.

diff --git a/Engine/Source/Runtime/Experimental/Chaos/Private/Chaos/ClusterUnionManager.cpp b/Engine/Source/Runtime/Experimental/Chaos/Private/Chaos/ClusterUnionManager.cpp
--- a/Engine/Source/Runtime/Experimental/Chaos/Private/Chaos/ClusterUnionManager.cpp
+++ b/Engine/Source/Runtime/Experimental/Chaos/Private/Chaos/ClusterUnionManager.cpp
@@ -332,7 +332,7 @@ namespace Chaos
 			// The anchored flag is taken care of in UpdateKinematicProperties so it must be set before that is called.
 			Cluster->InternalCluster->SetIsAnchored(true);
 		}
-		UpdateAllClusterUnionProperties(*Cluster, bIsNewCluster);
+		UpdateAllClusterUnionProperties(*Cluster, bIsNewCluster, true);
 
 		if (OldProxy)
 		{
@@ -363,7 +363,7 @@ namespace Chaos
 	void FClusterUnionManager::PostRemovalClusterUnionUpdate(FClusterUnion& Union)
 	{
 		SCOPE_CYCLE_COUNTER(STAT_PostRemovalClusterUnionUpdate);
-		UpdateAllClusterUnionProperties(Union, false);
+		UpdateAllClusterUnionProperties(Union, false, false);
 
 		if (!Union.ChildParticles.IsEmpty())
 		{
@@ -485,7 +485,7 @@ namespace Chaos
 	}
 
 	DECLARE_CYCLE_STAT(TEXT("FClusterUnionManager::UpdateClusterUnionProperties"), STAT_UpdateClusterUnionProperties, STATGROUP_Chaos);
-	void FClusterUnionManager::UpdateAllClusterUnionProperties(FClusterUnion& ClusterUnion, bool bRecomputeMassOrientation)
+	void FClusterUnionManager::UpdateAllClusterUnionProperties(FClusterUnion& ClusterUnion, bool bRecomputeMassOrientation, bool bUpdateKinematicProperties)
 	{
 		SCOPE_CYCLE_COUNTER(STAT_UpdateClusterUnionProperties);
 		// Update cluster properties.
@@ -497,7 +497,11 @@ namespace Chaos
 
 		const FRigidTransform3 ForceMassOrientation{ ClusterUnion.InternalCluster->X(), ClusterUnion.InternalCluster->R() };
 		UpdateClusterMassProperties(ClusterUnion.InternalCluster, FullChildrenSet, ClusterInertia, (bRecomputeMassOrientation && ClusterUnion.bNeedsXRInitialization) ? nullptr : &ForceMassOrientation);
-		UpdateKinematicProperties(ClusterUnion.InternalCluster, MClustering.GetChildrenMap(), MEvolution);
+
+		if (bUpdateKinematicProperties)
+		{
+			UpdateKinematicProperties(ClusterUnion.InternalCluster, MClustering.GetChildrenMap(), MEvolution);
+		}
 
 		// We must reset collisions etc on this particle since geometry will be changed
 		MEvolution.InvalidateParticle(ClusterUnion.InternalCluster);
@@ -665,7 +669,7 @@ namespace Chaos
 			}
 		}
 
-		UpdateAllClusterUnionProperties(*ClusterUnion, false);
+		UpdateAllClusterUnionProperties(*ClusterUnion, false, true);
 
 		for (TPair<FPBDRigidClusteredParticleHandle*, bool>& Pair : DirtyParticleLockStates)
 		{
diff --git a/Engine/Source/Runtime/Experimental/Chaos/Public/Chaos/ClusterUnionManager.h b/Engine/Source/Runtime/Experimental/Chaos/Public/Chaos/ClusterUnionManager.h
--- a/Engine/Source/Runtime/Experimental/Chaos/Public/Chaos/ClusterUnionManager.h
+++ b/Engine/Source/Runtime/Experimental/Chaos/Public/Chaos/ClusterUnionManager.h
@@ -137,7 +137,7 @@ namespace Chaos
 		void UpdateClusterUnionParticlesChildToParent(FClusterUnionIndex Index, const TArray<FPBDRigidParticleHandle*>& Particles, const TArray<FTransform>& ChildToParent, bool bLock);
 
 		// Update the cluster union's properties after its set of particle changes.
-		void UpdateAllClusterUnionProperties(FClusterUnion& ClusterUnion, bool bRecomputeMassOrientation);
+		void UpdateAllClusterUnionProperties(FClusterUnion& ClusterUnion, bool bRecomputeMassOrientation, bool bUpdateKinematicProperties=true);
 
 		// Returns all cluster unions. Really meant only to be used for debugging.
 		const TMap<FClusterUnionIndex, FClusterUnion>& GetAllClusterUnions() const { return ClusterUnions; }
