Add physics material option on geometry collection asset
- Add the ability to use density from the assigned material  ( or the overriden ones on the component )
- Add ability to apply default from asset on the component  ( as a button )
- Hide mass and inertia scale implementation details as functions  in physics proxy
- Move density converstion functions to MassProperties chaos headers

diff --git a/Engine/Source/Runtime/Experimental/Chaos/Private/PhysicsProxy/GeometryCollectionPhysicsProxy.cpp b/Engine/Source/Runtime/Experimental/Chaos/Private/PhysicsProxy/GeometryCollectionPhysicsProxy.cpp
--- a/Engine/Source/Runtime/Experimental/Chaos/Private/PhysicsProxy/GeometryCollectionPhysicsProxy.cpp
+++ b/Engine/Source/Runtime/Experimental/Chaos/Private/PhysicsProxy/GeometryCollectionPhysicsProxy.cpp
@@ -497,9 +497,7 @@ void FGeometryCollectionPhysicsProxy::Initialize(Chaos::FPBDRigidsEvolutionBase
 		
 	GameThreadCollection.AddExternalAttribute<TUniquePtr<FParticle>>(FGeometryCollection::ParticlesAttribute, FTransformCollection::TransformGroup, GTParticles);
 	
-
 	const FVector Scale = Parameters.WorldTransform.GetScale3D();
-	const FVector::FReal MassScale = Scale.X * Scale.Y * Scale.Z;
 	const TManagedArray<float>& Mass = Parameters.RestCollection->GetAttribute<float>("Mass", FTransformCollection::TransformGroup);
 
 	TArray<int32> ChildrenToCheckForParentFix;
@@ -515,6 +513,8 @@ void FGeometryCollectionPhysicsProxy::Initialize(Chaos::FPBDRigidsEvolutionBase
 
 				GTParticles[Index]->SetUniqueIdx(Evolution->GenerateUniqueIdx());
 
+				const Chaos::FReal ScaledMass = AdjustMassForScale(Mass[Index]);
+
 				// Note that this transform must match the physics thread transform computation for initialization.
 				// Take for example, a geometry collection in the editor that is linked with a joint constraint. The joint
 				// constraint will query the game thread particle position/rotation for the geometry collection to compute its
@@ -523,7 +523,7 @@ void FGeometryCollectionPhysicsProxy::Initialize(Chaos::FPBDRigidsEvolutionBase
 				const FTransform& T = GameThreadCollection.MassToLocal[Index] * GameThreadCollection.Transform[Index] * Parameters.WorldTransform;
 				P->SetX(T.GetTranslation(), false);
 				P->SetR(T.GetRotation(), false);
-				P->SetM(Mass[Index] * MassScale);
+				P->SetM(ScaledMass);
 				P->SetUserData(Parameters.UserData);
 				P->SetProxy(this);
 
@@ -953,20 +953,17 @@ void FGeometryCollectionPhysicsProxy::InitializeBodiesPT(Chaos::FPBDRigidsSolver
 		const float StrainDefault = Parameters.DamageThreshold.Num() ? Parameters.DamageThreshold[0] : 0;
 		// Add the rigid bodies
 
-		const FVector WorldScale = Parameters.WorldTransform.GetScale3D();
-		const FVector::FReal MassScale = WorldScale.X * WorldScale.Y * WorldScale.Z;
-
 		const Chaos::Facades::FCollectionAnchoringFacade AnchoringFacade(PhysicsThreadCollection);
 
-
 		TArray<float> Masses;
 		TArray<Chaos::FVec3f> Inertias;
 		TArray < FTransform> Transforms;
 
 		for (int32 Index = 0; Index < NumTransforms; Index++)
 		{
-			const Chaos::FVec3f ScaledInertia = Chaos::Utilities::ScaleInertia<float>((Chaos::FVec3f)InertiaTensor[Index], (Chaos::FVec3f)(WorldScale), true);
-			Masses.Add(Mass[Index] * float(MassScale));
+			const float ScaledMass = (float)AdjustMassForScale(Mass[Index]);
+			const Chaos::FVec3f ScaledInertia = AdjustInertiaForScale((Chaos::FVec3f)InertiaTensor[Index]);
+			Masses.Add(ScaledMass);
 			Inertias.Add(ScaledInertia);
 			const FTransform WorldTransform = MassToLocal[Index] * Transform[Index] * Parameters.WorldTransform;
 			Transforms.Add(WorldTransform);
@@ -982,7 +979,8 @@ void FGeometryCollectionPhysicsProxy::InitializeBodiesPT(Chaos::FPBDRigidsSolver
 			if (FClusterHandle* Handle = SolverParticleHandles[TransformGroupIndex])
 			{
 				// Mass space -> Composed parent space -> world
-				const Chaos::FVec3f ScaledInertia = Chaos::Utilities::ScaleInertia<float>((Chaos::FVec3f)InertiaTensor[TransformGroupIndex], (Chaos::FVec3f)(WorldScale), true);
+				const Chaos::FReal ScaledMass = AdjustMassForScale(Mass[TransformGroupIndex]);
+				const Chaos::FVec3f ScaledInertia = AdjustInertiaForScale((Chaos::FVec3f)InertiaTensor[TransformGroupIndex]);
 
 				PopulateSimulatedParticle(
 					Handle,
@@ -991,7 +989,7 @@ void FGeometryCollectionPhysicsProxy::InitializeBodiesPT(Chaos::FPBDRigidsSolver
 					Implicits[TransformGroupIndex],
 					SimFilter,
 					QueryFilter,
-					Mass[TransformGroupIndex] * MassScale,
+					ScaledMass,
 					ScaledInertia,
 					Transforms[TransformGroupIndex],
 					static_cast<uint8>(DynamicState[TransformGroupIndex]),
@@ -1376,6 +1374,21 @@ float FGeometryCollectionPhysicsProxy::ComputeDamageThreshold(const FGeometryDyn
 	return DamageThreshold;
 }
 
+Chaos::FReal FGeometryCollectionPhysicsProxy::AdjustMassForScale(Chaos::FReal Mass) const
+{
+	const FVector WorldScale = Parameters.WorldTransform.GetScale3D();
+	Chaos::FReal MassScale = WorldScale.X * WorldScale.Y * WorldScale.Z;
+	MassScale *= Parameters.MaterialOverrideMassScaleMultiplier;
+	return (Mass * MassScale);
+}
+
+Chaos::FVec3f FGeometryCollectionPhysicsProxy::AdjustInertiaForScale(const Chaos::FVec3f& Inertia) const
+{
+	const FVector3f WorldScale(Parameters.WorldTransform.GetScale3D());
+	const FVector3f InertiaScale = WorldScale * Parameters.MaterialOverrideMassScaleMultiplier;
+	return Chaos::Utilities::ScaleInertia<float>(Inertia, InertiaScale, true);
+}
+
 Chaos::TPBDGeometryCollectionParticleHandle<Chaos::FReal, 3>* FGeometryCollectionPhysicsProxy::BuildNonClusters_Internal(const uint32 CollectionClusterIndex, Chaos::FPBDRigidsSolver* RigidsSolver, float Mass, Chaos::FVec3f Inertia)
 {
 	FGeometryDynamicCollection& DynamicCollection = PhysicsThreadCollection;
@@ -1581,9 +1594,8 @@ FGeometryCollectionPhysicsProxy::BuildClusters_Internal(
 	const TManagedArray<FVector3f>& InertiaTensor = 
 		Parameters.RestCollection->GetAttribute<FVector3f>("InertiaTensor", FTransformCollection::TransformGroup);
 
-	const FVector WorldScale = Parameters.WorldTransform.GetScale3D();
-	const FVector::FReal MassScale = WorldScale.X * WorldScale.Y * WorldScale.Z;
-	const Chaos::FVec3f ScaledInertia = Chaos::Utilities::ScaleInertia<float>((Chaos::FVec3f)InertiaTensor[CollectionClusterIndex], FVector3f(WorldScale), true);  
+	const Chaos::FReal ScaledMass = AdjustMassForScale(Mass[CollectionClusterIndex]);
+	const Chaos::FVec3f ScaledInertia = AdjustInertiaForScale((Chaos::FVec3f)InertiaTensor[CollectionClusterIndex]);
 	
 	// NOTE: The particle creation call above (CreateClusterParticle) activates the particle, which does various things including adding
 	// it to the constraint graph. It seems a bit dodgy to be completely changing what the particle is after it has been enabled. Maybe we should fix this.
@@ -1594,7 +1606,7 @@ FGeometryCollectionPhysicsProxy::BuildClusters_Internal(
 		nullptr, // Parent->Geometry() ? Parent->Geometry() : Implicits[CollectionClusterIndex], 
 		SimFilter,
 		QueryFilter,
-		Parent->M() > 0.0 ? Parent->M() : Mass[CollectionClusterIndex] * MassScale, 
+		Parent->M() > 0.0 ? Parent->M() : ScaledMass,
 		Parent->I() != Chaos::FVec3f(0.0) ? Parent->I() : ScaledInertia,
 		ParticleTM, 
 		(uint8)DynamicState[CollectionClusterIndex], 
diff --git a/Engine/Source/Runtime/Experimental/Chaos/Public/Chaos/MassProperties.h b/Engine/Source/Runtime/Experimental/Chaos/Public/Chaos/MassProperties.h
--- a/Engine/Source/Runtime/Experimental/Chaos/Public/Chaos/MassProperties.h
+++ b/Engine/Source/Runtime/Experimental/Chaos/Public/Chaos/MassProperties.h
@@ -95,6 +95,24 @@ namespace Chaos
 		const TArray<bool>& bContributesToMass,
 		TFunction<FPerShapeData* (int32 ShapeIndex)> GetShapeDelegate);
 
+	template<typename T>
+	T KgCm3ToKgM3(T Density)
+	{
+		return Density * (T)1000000;
+	}
+
+	template<typename T>
+	T KgM3ToKgCm3(T Density)
+	{
+		return Density / (T)1000000;
+	}
+
+	template<typename T>
+	T GCm3ToKgCm3(T Density)
+	{
+		return Density / (T)1000;
+	}
+
 	template <typename T, int d>
 	using TMassProperties UE_DEPRECATED(4.27, "Deprecated. this class is to be deleted, use FMassProperties instead") = FMassProperties;
 }
diff --git a/Engine/Source/Runtime/Experimental/Chaos/Public/GeometryCollection/GeometryCollectionSimulationCoreTypes.h b/Engine/Source/Runtime/Experimental/Chaos/Public/GeometryCollection/GeometryCollectionSimulationCoreTypes.h
--- a/Engine/Source/Runtime/Experimental/Chaos/Public/GeometryCollection/GeometryCollectionSimulationCoreTypes.h
+++ b/Engine/Source/Runtime/Experimental/Chaos/Public/GeometryCollection/GeometryCollectionSimulationCoreTypes.h
@@ -195,6 +195,7 @@ struct FSimulationParameters
 		, bClearCache(false)
 		, ObjectType(EObjectStateTypeEnum::Chaos_NONE)
 		, StartAwake(true)
+		, MaterialOverrideMassScaleMultiplier(1.0f)
 		, bGenerateBreakingData(false)
 		, bGenerateCollisionData(false)
 		, bGenerateTrailingData(false)
@@ -246,6 +247,7 @@ struct FSimulationParameters
 		, ObjectType(Other.ObjectType)
 		, StartAwake(Other.StartAwake)
 		, PhysicalMaterialHandle(Other.PhysicalMaterialHandle)
+		, MaterialOverrideMassScaleMultiplier(Other.MaterialOverrideMassScaleMultiplier)
 		, bGenerateBreakingData(Other.bGenerateBreakingData)
 		, bGenerateCollisionData(Other.bGenerateCollisionData)
 		, bGenerateTrailingData(Other.bGenerateTrailingData)
@@ -319,6 +321,8 @@ struct FSimulationParameters
 
 	Chaos::FMaterialHandle PhysicalMaterialHandle;
 
+	float MaterialOverrideMassScaleMultiplier;
+
 	bool bGenerateBreakingData;
 	bool bGenerateCollisionData;
 	bool bGenerateTrailingData;
diff --git a/Engine/Source/Runtime/Experimental/Chaos/Public/PhysicsProxy/GeometryCollectionPhysicsProxy.h b/Engine/Source/Runtime/Experimental/Chaos/Public/PhysicsProxy/GeometryCollectionPhysicsProxy.h
--- a/Engine/Source/Runtime/Experimental/Chaos/Public/PhysicsProxy/GeometryCollectionPhysicsProxy.h
+++ b/Engine/Source/Runtime/Experimental/Chaos/Public/PhysicsProxy/GeometryCollectionPhysicsProxy.h
@@ -465,6 +465,11 @@ protected:
 	*/
 	float ComputeDamageThreshold(const FGeometryDynamicCollection& DynamicCollection, int32 TransformIndex) const;
 
+	/** adjust scalar mass to account for per component scale properties ( from material override and world transform scale ) */ 
+	Chaos::FReal AdjustMassForScale(Chaos::FReal Mass) const;
+
+	/** adjust inertia to account for per component scale properties ( from material override and world transform scale ) */
+	Chaos::FVec3f AdjustInertiaForScale(const Chaos::FVec3f& Inertia) const;
 
 	Chaos::TPBDGeometryCollectionParticleHandle<Chaos::FReal, 3>* BuildNonClusters_Internal(const uint32 CollectionClusterIndex, Chaos::FPBDRigidsSolver* RigidsSolver, float Mass, Chaos::FVec3f Inertia);
 
diff --git a/Engine/Source/Runtime/Experimental/GeometryCollectionEngine/Private/GeometryCollection/GeometryCollectionComponent.cpp b/Engine/Source/Runtime/Experimental/GeometryCollectionEngine/Private/GeometryCollection/GeometryCollectionComponent.cpp
--- a/Engine/Source/Runtime/Experimental/GeometryCollectionEngine/Private/GeometryCollection/GeometryCollectionComponent.cpp
+++ b/Engine/Source/Runtime/Experimental/GeometryCollectionEngine/Private/GeometryCollection/GeometryCollectionComponent.cpp
@@ -2766,6 +2766,24 @@ void UGeometryCollectionComponent::RegisterAndInitializePhysicsProxy()
 		if (ensure(EnginePhysicalMaterial))
 		{
 			SimulationParameters.PhysicalMaterialHandle = EnginePhysicalMaterial->GetPhysicsMaterial();
+
+			// mass properties are cooked in the GC asset, but component can override physics material with a different density
+			// if the user wants the density to be set from the physics material we need adjust by scaling mass by the ratio override/cooked defined properties
+			SimulationParameters.MaterialOverrideMassScaleMultiplier = 1.0f;
+			if (RestCollection && RestCollection->bDensityFromPhysicsMaterial)
+			{
+				bool bMassAsDensity = false;
+				const float AssetMassOrDensity = RestCollection->GetMassOrDensity(bMassAsDensity);
+				if (ensureMsgf(bMassAsDensity, TEXT("Need a density to be able to compute the mass scale multiplier, this may result in incorrect mass properties")))
+				{
+					if (ensureMsgf(AssetMassOrDensity > SMALL_NUMBER, TEXT("Asset density is set to a too small number, ignoring it for mass adjustment")))
+					{
+						const float OverrideMaterialDensity = Chaos::GCm3ToKgCm3(EnginePhysicalMaterial->Density);
+						SimulationParameters.MaterialOverrideMassScaleMultiplier = OverrideMaterialDensity / AssetMassOrDensity;
+					}
+				}
+				
+			}
 		}
 		GetInitializationCommands(SimulationParameters.InitializationCommands);
 	}
@@ -3251,6 +3269,24 @@ void UGeometryCollectionComponent::BuildInitialFilterData()
 	}
 }
 
+void UGeometryCollectionComponent::ApplyAssetDefaults()
+{
+	if (RestCollection)
+	{
+		// initialize the component per level damage threshold from the asset defaults 
+		DamageThreshold = RestCollection->DamageThreshold;
+		bUseSizeSpecificDamageThreshold = RestCollection->bUseSizeSpecificDamageThreshold;
+
+		// initialize the component damage progataion data from the asset defaults 
+		DamagePropagationData = RestCollection->DamagePropagationData;
+
+		if (RestCollection->PhysicsMaterial)
+		{
+			BodyInstance.SetPhysMaterialOverride(RestCollection->PhysicsMaterial);
+		}
+	}
+}
+
 void UGeometryCollectionComponent::SetRestCollection(const UGeometryCollection* RestCollectionIn, bool bApplyAssetDefaults)
 {
 	//UE_LOG(UGCC_LOG, Log, TEXT("GeometryCollectionComponent[%p]::SetRestCollection()"), this);
@@ -3267,12 +3303,7 @@ void UGeometryCollectionComponent::SetRestCollection(const UGeometryCollection*
 
 		if (bApplyAssetDefaults)
 		{
-			// initialize the component per level damage threshold from the asset defaults 
-			DamageThreshold = RestCollection->DamageThreshold;
-			bUseSizeSpecificDamageThreshold = RestCollection->bUseSizeSpecificDamageThreshold;
-
-			// initialize the component damage progataion data from the asset defaults 
-			DamagePropagationData = RestCollection->DamagePropagationData;
+			ApplyAssetDefaults();
 		}
 	}
 }
diff --git a/Engine/Source/Runtime/Experimental/GeometryCollectionEngine/Private/GeometryCollection/GeometryCollectionObject.cpp b/Engine/Source/Runtime/Experimental/GeometryCollectionEngine/Private/GeometryCollection/GeometryCollectionObject.cpp
--- a/Engine/Source/Runtime/Experimental/GeometryCollectionEngine/Private/GeometryCollection/GeometryCollectionObject.cpp
+++ b/Engine/Source/Runtime/Experimental/GeometryCollectionEngine/Private/GeometryCollection/GeometryCollectionObject.cpp
@@ -38,6 +38,7 @@
 
 #include "GeometryCollection/GeometryCollectionSimulationCoreTypes.h"
 #include "Chaos/ChaosArchive.h"
+#include "Chaos/MassProperties.h"
 #include "GeometryCollectionProxyData.h"
 #include "GeometryCollection/Facades/CollectionHierarchyFacade.h"
 #include "GeometryCollection/Facades/CollectionInstancedMeshFacade.h"
@@ -104,6 +105,7 @@ UGeometryCollection::UGeometryCollection(const FObjectInitializer& ObjectInitial
 	, MaxClusterLevelSetResolution_DEPRECATED(50)
 	, CollisionObjectReductionPercentage_DEPRECATED(0.0f)
 #endif
+	, bDensityFromPhysicsMaterial(false)
 	, bMassAsDensity(true)
 	, Mass(2500.0f)
 	, MinimumMassClamp(0.1f)
@@ -124,6 +126,7 @@ UGeometryCollection::UGeometryCollection(const FObjectInitializer& ObjectInitial
 	RenderDataGuid = StateGuid;
 	bStripOnCook = GeometryCollectionAssetForceStripOnCook;
 #endif
+	PhysicsMaterial = GEngine? GEngine->DefaultPhysMaterial: nullptr;
 }
 
 FGeometryCollectionLevelSetData::FGeometryCollectionLevelSetData()
@@ -352,22 +355,33 @@ void UGeometryCollection::PostInitProperties()
 	Super::PostInitProperties();
 }
 
-float KgCm3ToKgM3(float Density)
+float UGeometryCollection::GetMassOrDensity(bool& bOutIsDensity) const
 {
-	return Density * 1000000;
-}
-
-float KgM3ToKgCm3(float Density)
-{
-	return Density / 1000000;
+	bOutIsDensity = bMassAsDensity;
+	float MassOrDensity = bMassAsDensity ? Chaos::KgM3ToKgCm3(Mass) : Mass;
+	
+	if (bDensityFromPhysicsMaterial)
+	{
+		UPhysicalMaterial* PhysicsMaterialForDensity = PhysicsMaterial;
+		if (!PhysicsMaterialForDensity)
+		{
+			PhysicsMaterialForDensity = GEngine ? GEngine->DefaultPhysMaterial : nullptr;
+		}
+		if (ensureMsgf(PhysicsMaterialForDensity, TEXT("bDensityFromPhysicsMaterial is true but no physics material has been set (and engine default cannot be found )")))
+		{
+			// materials only provide density
+			bOutIsDensity = true;
+			MassOrDensity = Chaos::GCm3ToKgCm3(PhysicsMaterial->Density);
+		}
+	}
+	return MassOrDensity;
 }
 
 void UGeometryCollection::GetSharedSimulationParams(FSharedSimulationParameters& OutParams) const
 {
 	const FGeometryCollectionSizeSpecificData& SizeSpecificDefault = GetDefaultSizeSpecificData();
 
-	OutParams.bMassAsDensity = bMassAsDensity;
-	OutParams.Mass = bMassAsDensity ? KgM3ToKgCm3(Mass) : Mass;	//todo(ocohen): we still have the solver working in old units. This is mainly to fix ui issues. Long term need to normalize units for best precision
+	OutParams.Mass = GetMassOrDensity(OutParams.bMassAsDensity);
 	OutParams.MinimumMassClamp = MinimumMassClamp;
 
 	FGeometryCollectionSizeSpecificData InfSize;
@@ -866,7 +880,7 @@ void UGeometryCollection::Serialize(FArchive& Ar)
 	{
 		if (bMassAsDensity)
 		{
-			Mass = KgCm3ToKgM3(Mass);
+			Mass = Chaos::KgCm3ToKgM3(Mass);
 		}
 	}
 
diff --git a/Engine/Source/Runtime/Experimental/GeometryCollectionEngine/Public/GeometryCollection/GeometryCollectionComponent.h b/Engine/Source/Runtime/Experimental/GeometryCollectionEngine/Public/GeometryCollection/GeometryCollectionComponent.h
--- a/Engine/Source/Runtime/Experimental/GeometryCollectionEngine/Public/GeometryCollection/GeometryCollectionComponent.h
+++ b/Engine/Source/Runtime/Experimental/GeometryCollectionEngine/Public/GeometryCollection/GeometryCollectionComponent.h
@@ -727,6 +727,10 @@ public:
 	UPROPERTY(EditAnywhere, NoClear, BlueprintReadOnly, Category = "ChaosPhysics")
 	TObjectPtr<const UGeometryCollection> RestCollection;
 
+	/** Apply default values from asset ( damage related data and physics material ) */
+	UFUNCTION(BlueprintCallable, CallInEditor, Category = "ChaosPhysics", meta = (EditCondition = "RestCollection != nullptr"))
+	void ApplyAssetDefaults();
+
 	UPROPERTY(EditAnywhere, NoClear, BlueprintReadOnly, Category = "ChaosPhysics")
 	TArray<TObjectPtr<const AFieldSystemActor>> InitializationFields;
 
diff --git a/Engine/Source/Runtime/Experimental/GeometryCollectionEngine/Public/GeometryCollection/GeometryCollectionObject.h b/Engine/Source/Runtime/Experimental/GeometryCollectionEngine/Public/GeometryCollection/GeometryCollectionObject.h
--- a/Engine/Source/Runtime/Experimental/GeometryCollectionEngine/Public/GeometryCollection/GeometryCollectionObject.h
+++ b/Engine/Source/Runtime/Experimental/GeometryCollectionEngine/Public/GeometryCollection/GeometryCollectionObject.h
@@ -22,6 +22,7 @@ class UDataflow;
 class UGeometryCollectionCache;
 class UMaterial;
 class UMaterialInterface;
+class UPhysicalMaterial;
 
 USTRUCT(BlueprintType)
 struct GEOMETRYCOLLECTIONENGINE_API FGeometryCollectionSource
@@ -480,6 +481,13 @@ public:
 	/** Fills params struct with parameters used for precomputing content. */
 	void GetSharedSimulationParams(FSharedSimulationParameters& OutParams) const;
 
+	/**
+	* Get Mass or density as set by the asset 
+	* Mass is return in Kg and Density is returned in Kg/Cm3
+	* @param bOutIsDensity  is set to true by the function if the returned value is to be treated as a density
+	*/
+	float GetMassOrDensity(bool& bOutIsDensity) const;
+
 	/** Accessors for the two guids used to identify this collection */
 	FGuid GetIdGuid() const;
 	FGuid GetStateGuid() const;
@@ -624,16 +632,26 @@ public:
 	FSoftObjectPath RootProxy_DEPRECATED;
 #endif
 	
+	/**	Physics material to use for the geometry collection */
+	UPROPERTY(EditAnywhere, BlueprintReadOnly, Category = "Collisions")
+	TObjectPtr<UPhysicalMaterial> PhysicsMaterial;
+
 	/**
-	* Mass As Density, units are in kg/m^3
+	* Whether to use density ( for mass computation ) from physics material ( if physics material is not set, use the component one or defaults )
 	*/
-	UPROPERTY(EditAnywhere, BlueprintReadWrite, Category = "Collisions")
+	UPROPERTY(EditAnywhere, BlueprintReadWrite, Category = "Collisions", meta = (EditCondition = "PhysicsMaterial != nullptr"))
+	bool bDensityFromPhysicsMaterial;
+
+	/**
+	* Mass As Density, units are in kg/m^3 ( only enabled if physics material is not set )
+	*/
+	UPROPERTY(EditAnywhere, BlueprintReadWrite, Category = "Collisions", meta = (EditCondition = "!bDensityFromPhysicsMaterial"))
 	bool bMassAsDensity;
 
 	/**
-	* Total Mass of Collection. If density, units are in kg/m^3
+	* Total Mass of Collection. If density, units are in kg/m^3 ( only enabled if physics material is not set )
 	*/
-	UPROPERTY(EditAnywhere, BlueprintReadWrite, Category = "Collisions")
+	UPROPERTY(EditAnywhere, BlueprintReadWrite, Category = "Collisions", meta = (EditCondition = "!bDensityFromPhysicsMaterial"))
 	float Mass;
 
 	/**
