Fix non-unity compilation issue

diff --git a/Engine/Source/Runtime/Experimental/GeometryCollectionEngine/Private/GeometryCollection/GeometryCollectionObject.cpp b/Engine/Source/Runtime/Experimental/GeometryCollectionEngine/Private/GeometryCollection/GeometryCollectionObject.cpp
--- a/Engine/Source/Runtime/Experimental/GeometryCollectionEngine/Private/GeometryCollection/GeometryCollectionObject.cpp
+++ b/Engine/Source/Runtime/Experimental/GeometryCollectionEngine/Private/GeometryCollection/GeometryCollectionObject.cpp
@@ -17,11 +17,13 @@
 #include "Materials/MaterialInstance.h"
 #include "ProfilingDebugging/CookStats.h"
 #include "EngineUtils.h"
+#include "Engine/Engine.h"
 #include "Engine/StaticMesh.h"
 #include "PhysicsEngine/PhysicsSettings.h"
 #include "EditorFramework/AssetImportData.h"
 #include "Rendering/NaniteResources.h"
 #include "Engine/AssetUserData.h"
+#include "PhysicalMaterials/PhysicalMaterial.h"
 
 #if WITH_EDITOR
 #include "GeometryCollection/DerivedDataGeometryCollectionCooker.h"
