Chaos VD - Fix CIS issue in ChaosVisualDebuggerTrace.h

#rnx

diff --git a/Engine/Source/Runtime/Experimental/Chaos/Private/Chaos/ChaosArchive.cpp b/Engine/Source/Runtime/Experimental/Chaos/Private/Chaos/ChaosArchive.cpp
--- a/Engine/Source/Runtime/Experimental/Chaos/Private/Chaos/ChaosArchive.cpp
+++ b/Engine/Source/Runtime/Experimental/Chaos/Private/Chaos/ChaosArchive.cpp
@@ -6,6 +6,18 @@
 namespace Chaos
 {
 
+FChaosArchiveContext::FChaosArchiveContext(): TagCount(0)
+{
+}
+
+FChaosArchiveContext::~FChaosArchiveContext()
+{
+	for (auto Itr : ObjToSharedPtrHolder)
+	{
+		delete Itr.Value;
+	}
+}
+
 void FChaosArchive::SerializeLegacy(TUniquePtr<FImplicitObject>& Obj)
 {
 	FImplicitObject::SerializeLegacyHelper(InnerArchive, Obj);
diff --git a/Engine/Source/Runtime/Experimental/Chaos/Public/Chaos/ChaosArchive.h b/Engine/Source/Runtime/Experimental/Chaos/Public/Chaos/ChaosArchive.h
--- a/Engine/Source/Runtime/Experimental/Chaos/Public/Chaos/ChaosArchive.h
+++ b/Engine/Source/Runtime/Experimental/Chaos/Public/Chaos/ChaosArchive.h
@@ -80,18 +80,9 @@ public:
 	int32 TagCount;
 
 
-	FChaosArchiveContext()
-		: TagCount(0)
-	{
-	}
+	FChaosArchiveContext();
 
-	~FChaosArchiveContext()
-	{
-		for (auto Itr : ObjToSharedPtrHolder)
-		{
-			delete Itr.Value;
-		}
-	}
+	~FChaosArchiveContext();
 
 	template <typename T, ESPMode Mode>
 	TSharedPtr<T, Mode>& ToSharedPointerHelper(TSerializablePtr<T>& Obj)
