Fix for Non-Unity CIS ChaosVisualDebuggerTrace

#rnx

diff --git a/Engine/Source/Runtime/Experimental/Chaos/Private/ChaosVisualDebugger/ChaosVisualDebuggerTrace.cpp b/Engine/Source/Runtime/Experimental/Chaos/Private/ChaosVisualDebugger/ChaosVisualDebuggerTrace.cpp
--- a/Engine/Source/Runtime/Experimental/Chaos/Private/ChaosVisualDebugger/ChaosVisualDebuggerTrace.cpp
+++ b/Engine/Source/Runtime/Experimental/Chaos/Private/ChaosVisualDebugger/ChaosVisualDebuggerTrace.cpp
@@ -11,6 +11,7 @@
 #include "HAL/CriticalSection.h"
 #include "Misc/ScopeRWLock.h"
 #include "Serialization/MemoryWriter.h"
+#include "Templates/UniquePtr.h"
 
 UE_TRACE_EVENT_DEFINE(ChaosVDLogger, ChaosVDSolverFrameStart)
 
diff --git a/Engine/Source/Runtime/Experimental/Chaos/Public/ChaosVisualDebugger/ChaosVisualDebuggerTrace.h b/Engine/Source/Runtime/Experimental/Chaos/Public/ChaosVisualDebugger/ChaosVisualDebuggerTrace.h
--- a/Engine/Source/Runtime/Experimental/Chaos/Public/ChaosVisualDebugger/ChaosVisualDebuggerTrace.h
+++ b/Engine/Source/Runtime/Experimental/Chaos/Public/ChaosVisualDebugger/ChaosVisualDebuggerTrace.h
@@ -7,6 +7,8 @@
 #include "Chaos/Vector.h"
 #include "Containers/Set.h"
 
+#include "Chaos/ChaosArchive.h"
+
 #ifndef CHAOS_VISUAL_DEBUGGER_ENABLED
 	#define CHAOS_VISUAL_DEBUGGER_ENABLED (WITH_CHAOS_VISUAL_DEBUGGER && UE_TRACE_ENABLED)
 #endif
@@ -211,7 +213,6 @@ struct FChaosVDContext;
 
 namespace Chaos
 {
-	class FChaosArchiveContext;
 	class FImplicitObject;
 	class FPhysicsSolverBase;
 	template <typename T, int d>
