Chaos Physics Update/Add preparation for future work.
No functional changes as of yet.
Acceleration structure command buffers will now include Add/Update/Delete instead of just Update and Add to prevent accidental updates of deleted objects causing dangling references.
Update and Add commands will be handled identically for the time being until everything that uses this API is audited to use the correct commands.

#[fyi] Benn.Gallagher, Chris.Caulfield

diff --git a/Engine/Source/Runtime/Engine/Private/PhysicsEngine/Experimental/PhysScene_Chaos.cpp b/Engine/Source/Runtime/Engine/Private/PhysicsEngine/Experimental/PhysScene_Chaos.cpp
--- a/Engine/Source/Runtime/Engine/Private/PhysicsEngine/Experimental/PhysScene_Chaos.cpp
+++ b/Engine/Source/Runtime/Engine/Private/PhysicsEngine/Experimental/PhysScene_Chaos.cpp
@@ -21,6 +21,7 @@
 #include "PhysicsProxy/ClusterUnionPhysicsProxy.h"
 #include "PhysicsProxy/SkeletalMeshPhysicsProxy.h"
 #include "PhysicsProxy/StaticMeshPhysicsProxy.h"
+#include "Chaos/PendingSpatialData.h"
 #include "Chaos/PhysicsSolverBaseImpl.h"
 
 #if WITH_EDITOR
@@ -699,7 +700,7 @@ void FPhysScene_Chaos::RemoveObject(FGeometryCollectionPhysicsProxy* InObject)
 			RemoveActorFromAccelerationStructureImp(GTParticle);
 			if (Solver)
 			{
-				Solver->UpdateParticleInAccelerationStructure_External(GTParticle, /*bDelete=*/true); // ensures deletion will be applied to structures passed from physics thread.
+				Solver->UpdateParticleInAccelerationStructure_External(GTParticle,Chaos::EPendingSpatialDataOperation::Delete); // ensures deletion will be applied to structures passed from physics thread.
 			}
 		}
 	}
@@ -720,7 +721,7 @@ void FPhysScene_Chaos::RemoveObject(Chaos::FClusterUnionPhysicsProxy* InObject)
 
 	if (Solver)
 	{
-		Solver->UpdateParticleInAccelerationStructure_External(InObject->GetParticle_External(), true);
+		Solver->UpdateParticleInAccelerationStructure_External(InObject->GetParticle_External(), Chaos::EPendingSpatialDataOperation::Delete);
 		Solver->UnregisterObject(InObject);
 	}
 
diff --git a/Engine/Source/Runtime/Experimental/Chaos/Private/Chaos/Framework/PhysicsSolverBase.cpp b/Engine/Source/Runtime/Experimental/Chaos/Private/Chaos/Framework/PhysicsSolverBase.cpp
--- a/Engine/Source/Runtime/Experimental/Chaos/Private/Chaos/Framework/PhysicsSolverBase.cpp
+++ b/Engine/Source/Runtime/Experimental/Chaos/Private/Chaos/Framework/PhysicsSolverBase.cpp
@@ -284,16 +284,16 @@ namespace Chaos
 		delete &InSolver;
 	}
 
-	void FPhysicsSolverBase::UpdateParticleInAccelerationStructure_External(FGeometryParticle* Particle,bool bDelete)
+	void FPhysicsSolverBase::UpdateParticleInAccelerationStructure_External(FGeometryParticle* Particle, EPendingSpatialDataOperation InOperation)
 	{
 		//mark it as pending for async structure being built
 		FAccelerationStructureHandle AccelerationHandle(Particle);
 		FPendingSpatialData& SpatialData = PendingSpatialOperations_External->FindOrAdd(Particle->UniqueIdx());
 
 		//make sure any new operations (i.e not currently being consumed by sim) are not acting on a deleted object
-		ensure(SpatialData.SyncTimestamp < MarshallingManager.GetExternalTimestamp_External() || !SpatialData.bDelete);
+		ensure(SpatialData.SyncTimestamp < MarshallingManager.GetExternalTimestamp_External() || SpatialData.Operation != EPendingSpatialDataOperation::Delete);
 
-		SpatialData.bDelete = bDelete;
+		SpatialData.Operation = InOperation;
 		SpatialData.SpatialIdx = Particle->SpatialIdx();
 		SpatialData.AccelerationHandle = AccelerationHandle;
 		SpatialData.SyncTimestamp = MarshallingManager.GetExternalTimestamp_External();
diff --git a/Engine/Source/Runtime/Experimental/Chaos/Private/Chaos/PBDRigidsEvolution.cpp b/Engine/Source/Runtime/Experimental/Chaos/Private/Chaos/PBDRigidsEvolution.cpp
--- a/Engine/Source/Runtime/Experimental/Chaos/Private/Chaos/PBDRigidsEvolution.cpp
+++ b/Engine/Source/Runtime/Experimental/Chaos/Private/Chaos/PBDRigidsEvolution.cpp
@@ -524,11 +524,12 @@ namespace Chaos
 
 	void FPBDRigidsEvolutionBase::ApplyParticlePendingData(const FPendingSpatialData& SpatialData, FAccelerationStructure& AccelerationStructure, bool bUpdateCache, bool bUpdateDynamicTrees)
 	{
-		if (SpatialData.bDelete)
+		if (SpatialData.Operation ==  EPendingSpatialDataOperation::Delete)
 		{
 			if (bUpdateDynamicTrees || !FDefaultCollectionFactory::IsDynamicTree(SpatialData.SpatialIdx))
 			{
-				AccelerationStructure.RemoveElementFrom(SpatialData.AccelerationHandle, SpatialData.SpatialIdx);
+				const bool bExisted = AccelerationStructure.RemoveElementFrom(SpatialData.AccelerationHandle, SpatialData.SpatialIdx);
+				//ensure(bExisted);
 			}
 
 			if (bUpdateCache)
@@ -555,7 +556,15 @@ namespace Chaos
 			if (bUpdateDynamicTrees || !FDefaultCollectionFactory::IsDynamicTree(SpatialData.SpatialIdx))
 			{
 				//CSV_SCOPED_TIMING_STAT(ChaosPhysicsTimers, AABBTreeUpdate)
-				AccelerationStructure.UpdateElementIn(UpdateParticle, UpdateParticle->WorldSpaceInflatedBounds(), UpdateParticle->HasBounds(), SpatialData.SpatialIdx);
+				const bool bExisted = AccelerationStructure.UpdateElementIn(UpdateParticle, UpdateParticle->WorldSpaceInflatedBounds(), UpdateParticle->HasBounds(), SpatialData.SpatialIdx);
+				if (SpatialData.Operation == EPendingSpatialDataOperation::Add)
+				{
+					//ensure(bExisted == false); // TODO use this for ADD/UPDATE Audit
+				}
+				else
+				{
+					//ensure(bExisted == true);	// TODO use this for ADD/UPDATE Audit				
+				}
 			}
 
 			if (bUpdateCache)
@@ -634,9 +643,10 @@ namespace Chaos
 			{
 				//operation still pending so update structure
 				//note: do we care about roll over? if game ticks at 60fps we'd get 385+ days
-				if(SpatialData.bDelete)
+				if(SpatialData.Operation == EPendingSpatialDataOperation::Delete)
 				{
-					Acceleration.RemoveElementFrom(SpatialData.AccelerationHandle,SpatialData.SpatialIdx);
+					const bool bExisted = Acceleration.RemoveElementFrom(SpatialData.AccelerationHandle,SpatialData.SpatialIdx);
+					//ensure(bExisted);
 				}
 				else
 				{
@@ -648,7 +658,15 @@ namespace Chaos
 						TRigidTransform<FReal,3> WorldTM(UpdateParticle->X(),UpdateParticle->R());
 						WorldBounds = UpdateParticle->Geometry()->BoundingBox().TransformedAABB(WorldTM);
 					}
-					Acceleration.UpdateElementIn(UpdateParticle,WorldBounds,bHasBounds,SpatialData.SpatialIdx);
+					const bool bExisted = Acceleration.UpdateElementIn(UpdateParticle,WorldBounds,bHasBounds,SpatialData.SpatialIdx);
+					if (SpatialData.Operation == EPendingSpatialDataOperation::Add)
+					{
+						//ensure(bExisted == false); // TODO use this for ADD/UPDATE Audit
+					}
+					else
+					{
+						//ensure(bExisted == true); // TODO use this for ADD/UPDATE Audit						
+					}
 				}
 			}
 			else
diff --git a/Engine/Source/Runtime/Experimental/Chaos/Private/PBDRigidsSolver.cpp b/Engine/Source/Runtime/Experimental/Chaos/Private/PBDRigidsSolver.cpp
--- a/Engine/Source/Runtime/Experimental/Chaos/Private/PBDRigidsSolver.cpp
+++ b/Engine/Source/Runtime/Experimental/Chaos/Private/PBDRigidsSolver.cpp
@@ -630,7 +630,7 @@ namespace Chaos
 		Proxy->GetParticle_LowLevel()->SetProxy(Proxy);
 		AddDirtyProxy(Proxy);
 
-		UpdateParticleInAccelerationStructure_External(Proxy->GetParticle_LowLevel(), /*bDelete=*/false);
+		UpdateParticleInAccelerationStructure_External(Proxy->GetParticle_LowLevel(), EPendingSpatialDataOperation::Add);
 	}
 
 	int32 LogCorruptMap = 0;
@@ -645,7 +645,7 @@ namespace Chaos
 
 		ClearGTParticle_External(*Proxy->GetParticle_LowLevel());	//todo: remove this
 
-		UpdateParticleInAccelerationStructure_External(Proxy->GetParticle_LowLevel(), /*bDelete=*/true);
+		UpdateParticleInAccelerationStructure_External(Proxy->GetParticle_LowLevel(), EPendingSpatialDataOperation::Delete);
 
 		// remove the proxy from the invalidation list
 		RemoveDirtyProxy(Proxy);
@@ -726,7 +726,7 @@ namespace Chaos
 		{
 			if (Particle && !Particle->Disabled())
 			{
-				UpdateParticleInAccelerationStructure_External(Particle.Get(), /*bDelete=*/false);
+				UpdateParticleInAccelerationStructure_External(Particle.Get(), EPendingSpatialDataOperation::Add);
 			}
 		}
 
@@ -793,7 +793,7 @@ namespace Chaos
 			}
 
 			Particle->SetUniqueIdx(GetEvolution()->GenerateUniqueIdx());
-			UpdateParticleInAccelerationStructure_External(Particle, /*bDelete=*/false);
+			UpdateParticleInAccelerationStructure_External(Particle, EPendingSpatialDataOperation::Add);
 		}
 		Proxy->SetSolver(this);
 		AddDirtyProxy(Proxy);
diff --git a/Engine/Source/Runtime/Experimental/Chaos/Public/Chaos/Framework/PhysicsSolverBase.h b/Engine/Source/Runtime/Experimental/Chaos/Public/Chaos/Framework/PhysicsSolverBase.h
--- a/Engine/Source/Runtime/Experimental/Chaos/Public/Chaos/Framework/PhysicsSolverBase.h
+++ b/Engine/Source/Runtime/Experimental/Chaos/Public/Chaos/Framework/PhysicsSolverBase.h
@@ -31,6 +31,7 @@ namespace Chaos
 	class FChaosResultsManager;
 	class FRewindData;
 	class IRewindCallback;
+	enum EPendingSpatialDataOperation : uint8;
 
 	extern CHAOS_API int32 GSingleThreadedPhysics;
 	extern CHAOS_API int32 UseAsyncInterpolation;
@@ -495,7 +496,7 @@ namespace Chaos
 			}
 		}
 
-		void UpdateParticleInAccelerationStructure_External(FGeometryParticle* Particle,bool bDelete);
+		void UpdateParticleInAccelerationStructure_External(FGeometryParticle* Particle, EPendingSpatialDataOperation InOperation);
 
 		bool IsPaused_External() const
 		{
diff --git a/Engine/Source/Runtime/Experimental/Chaos/Public/Chaos/PBDRigidsEvolution.h b/Engine/Source/Runtime/Experimental/Chaos/Public/Chaos/PBDRigidsEvolution.h
--- a/Engine/Source/Runtime/Experimental/Chaos/Public/Chaos/PBDRigidsEvolution.h
+++ b/Engine/Source/Runtime/Experimental/Chaos/Public/Chaos/PBDRigidsEvolution.h
@@ -453,7 +453,7 @@ public:
 		Particles.EnableParticle(Particle);
 		EnableConstraints(Particle);
 		IslandManager.AddParticle(Particle);
-		DirtyParticle(*Particle);
+		DirtyParticle(*Particle, EPendingSpatialDataOperation::Add);
 	}
 
 	/**
@@ -502,8 +502,9 @@ public:
 	}
 
 	template <bool bPersistent>
-	FORCEINLINE_DEBUGGABLE void DirtyParticle(TGeometryParticleHandleImp<FReal, 3, bPersistent>& Particle)
+	FORCEINLINE_DEBUGGABLE void DirtyParticle(TGeometryParticleHandleImp<FReal, 3, bPersistent>& Particle, const EPendingSpatialDataOperation Op = EPendingSpatialDataOperation::Update)
 	{
+		ensure(Op != EPendingSpatialDataOperation::Delete); // Don't use the function to delete particles
 		const TPBDRigidParticleHandleImp<FReal, 3, bPersistent>* AsRigid = Particle.CastToRigidParticle();
 		if(AsRigid && AsRigid->Disabled())
 		{
@@ -556,15 +557,17 @@ public:
 		//only add to acceleration structure if it has collision
 		if (Particle.HasCollision() || ForceNoCollisionIntoSQ)
 		{
-			//TODO: distinguish between new particles and dirty particles
+			//TODO: distinguish between new particles and dirty particles - Adds and updates are treated the same right now
 			const FUniqueIdx UniqueIdx = Particle.UniqueIdx();
 			FPendingSpatialData& SpatialData = InternalAccelerationQueue.FindOrAdd(UniqueIdx);
-			ensure(SpatialData.bDelete == false);
+			ensure(SpatialData.Operation != EPendingSpatialDataOperation::Delete);
+
+			SpatialData.Operation = Op;
 			SpatialData.AccelerationHandle = FAccelerationStructureHandle(Particle);
 			SpatialData.SpatialIdx = Particle.SpatialIdx();
 
 			auto& AsyncSpatialData = AsyncAccelerationQueue.FindOrAdd(UniqueIdx);
-			ensure(SpatialData.bDelete == false);
+			// ensure(AsyncSpatialData.Operation != EPendingSpatialDataOperation::Delete); // TODO: This may be hit: Potentially due to UniqueIdx reuse?
 			AsyncSpatialData = SpatialData;
 		}
 	}
@@ -952,9 +955,9 @@ public:
 		//TODO: at the moment we don't distinguish between the first time a particle is created and when it's just moved
 		// If we had this distinction we could simply remove the entry for the async queue
 		const FUniqueIdx UniqueIdx = ParticleHandle.UniqueIdx();
-		FPendingSpatialData& SpatialData = AsyncAccelerationQueue.FindOrAdd(UniqueIdx);
+		FPendingSpatialData& SpatialData = AsyncAccelerationQueue.FindOrAdd(UniqueIdx, EPendingSpatialDataOperation::Delete);
 
-		SpatialData.bDelete = true;
+		SpatialData.Operation = EPendingSpatialDataOperation::Delete;
 		SpatialData.SpatialIdx = ParticleHandle.SpatialIdx();
 		SpatialData.AccelerationHandle = FAccelerationStructureHandle(ParticleHandle);
 
diff --git a/Engine/Source/Runtime/Experimental/Chaos/Public/Chaos/PendingSpatialData.h b/Engine/Source/Runtime/Experimental/Chaos/Public/Chaos/PendingSpatialData.h
--- a/Engine/Source/Runtime/Experimental/Chaos/Public/Chaos/PendingSpatialData.h
+++ b/Engine/Source/Runtime/Experimental/Chaos/Public/Chaos/PendingSpatialData.h
@@ -8,6 +8,14 @@
 
 namespace Chaos
 { 
+enum EPendingSpatialDataOperation : uint8 
+{	
+	Delete,
+	// Note: Updates and Adds are treated the same right now. TODO: Distinguish between them
+	Add, // Use this if the element does not exist in the acceleration structure
+	Update // Use this when it is known that the element already exists in the acceleration structure
+	
+};
 	
 /** Used for updating intermediate spatial structures when they are finished */
 struct FPendingSpatialData
@@ -15,11 +23,11 @@ struct FPendingSpatialData
 	FAccelerationStructureHandle AccelerationHandle;
 	FSpatialAccelerationIdx SpatialIdx;
 	int32 SyncTimestamp;	//indicates the inputs timestamp associated with latest change. Only relevant for external queue
-	bool bDelete;
+	EPendingSpatialDataOperation Operation;
 
 	FPendingSpatialData()
 	: SyncTimestamp(0)
-	, bDelete(false)
+	, Operation(Add)
 	{}
 
 	void Serialize(FChaosArchive& Ar)
@@ -66,7 +74,7 @@ struct FPendingSpatialDataQueue
 		return PendingData.Num();
 	}
 
-	FPendingSpatialData& FindOrAdd(const FUniqueIdx UniqueIdx)
+	FPendingSpatialData& FindOrAdd(const FUniqueIdx UniqueIdx, EPendingSpatialDataOperation Operation = EPendingSpatialDataOperation::Add)
 	{
 		if(int32* Existing = ParticleToPendingData.Find(UniqueIdx))
 		{
@@ -75,6 +83,7 @@ struct FPendingSpatialDataQueue
 		{
 			const int32 NewIdx = PendingData.AddDefaulted(1);
 			ParticleToPendingData.Add(UniqueIdx,NewIdx);
+			PendingData[NewIdx].Operation = Operation;
 			return PendingData[NewIdx];
 		}
 	}
diff --git a/Engine/Source/Runtime/PhysicsCore/Private/ChaosScene.cpp b/Engine/Source/Runtime/PhysicsCore/Private/ChaosScene.cpp
--- a/Engine/Source/Runtime/PhysicsCore/Private/ChaosScene.cpp
+++ b/Engine/Source/Runtime/PhysicsCore/Private/ChaosScene.cpp
@@ -233,7 +233,7 @@ void FChaosScene::UpdateActorInAccelerationStructure(const FPhysicsActorHandle&
 			SpatialAcceleration->UpdateElementIn(AccelerationHandle,WorldBounds,bHasBounds, Body_External.SpatialIdx());
 		}
 
-		GetSolver()->UpdateParticleInAccelerationStructure_External(Actor->GetParticle_LowLevel(),/*bDelete=*/false);
+		GetSolver()->UpdateParticleInAccelerationStructure_External(Actor->GetParticle_LowLevel(), EPendingSpatialDataOperation::Update);
 	}
 }
 
@@ -275,7 +275,7 @@ void FChaosScene::UpdateActorsInAccelerationStructure(const TArrayView<FPhysicsA
 			const FPhysicsActorHandle& Actor = Actors[ActorIndex];
 			if(Actor)
 			{
-				GetSolver()->UpdateParticleInAccelerationStructure_External(Actor->GetParticle_LowLevel(),/*bDelete=*/false);
+				GetSolver()->UpdateParticleInAccelerationStructure_External(Actor->GetParticle_LowLevel(), EPendingSpatialDataOperation::Update);
 			}
 		}
 	}
