Iris
* Added IrisDebugHelper::SetIrisDebugInternalIndexViaNetHandle

#rnx

diff --git a/Engine/Source/Runtime/Experimental/Iris/Core/Private/Iris/Core/IrisDebugging.cpp b/Engine/Source/Runtime/Experimental/Iris/Core/Private/Iris/Core/IrisDebugging.cpp
--- a/Engine/Source/Runtime/Experimental/Iris/Core/Private/Iris/Core/IrisDebugging.cpp
+++ b/Engine/Source/Runtime/Experimental/Iris/Core/Private/Iris/Core/IrisDebugging.cpp
@@ -175,6 +175,36 @@ void SetIrisDebugInternalIndex(UE::Net::Private::FInternalNetRefIndex InternalIn
 	IrisDebugHelperInternal::GIrisDebugInternalIndex = InternalIndex;
 }
 
+void SetIrisDebugInternalIndexViaNetHandle(FNetRefHandle RefHandle)
+{
+	using namespace UE::Net::Private;
+
+	if (!RefHandle.IsValid())
+	{
+		IrisDebugHelperInternal::GIrisDebugNetRefHandle = RefHandle;
+		IrisDebugHelperInternal::GIrisDebugInternalIndex = FNetRefHandleManager::InvalidInternalIndex;
+		return;
+	}
+
+	UReplicationSystem* ReplicationSystem = GetReplicationSystem(RefHandle.GetReplicationSystemId());
+	if (!ReplicationSystem)
+	{
+		return;
+	}
+
+	FReplicationSystemInternal* ReplicationSystemInternal = ReplicationSystem->GetReplicationSystemInternal();
+	const FNetRefHandleManager& NetRefHandleManager = ReplicationSystemInternal->GetNetRefHandleManager();
+
+	const FInternalNetRefIndex InternalNetRefIndex = NetRefHandleManager.GetInternalIndex(RefHandle);
+	if (!InternalNetRefIndex)
+	{
+		return;
+	}
+
+	IrisDebugHelperInternal::GIrisDebugNetRefHandle = RefHandle;
+	IrisDebugHelperInternal::GIrisDebugInternalIndex = InternalNetRefIndex;
+}
+
 
 UReplicationSystem* GetReplicationSystemForDebug(uint32 Id)
 {
diff --git a/Engine/Source/Runtime/Experimental/Iris/Core/Private/Iris/ReplicationSystem/Filtering/ReplicationFiltering.cpp b/Engine/Source/Runtime/Experimental/Iris/Core/Private/Iris/ReplicationSystem/Filtering/ReplicationFiltering.cpp
--- a/Engine/Source/Runtime/Experimental/Iris/Core/Private/Iris/ReplicationSystem/Filtering/ReplicationFiltering.cpp
+++ b/Engine/Source/Runtime/Experimental/Iris/Core/Private/Iris/ReplicationSystem/Filtering/ReplicationFiltering.cpp
@@ -435,7 +435,6 @@ bool FReplicationFiltering::SetFilter(FInternalNetRefIndex ObjectIndex, FNetObje
 		return false;
 	}
 
-	const FNetBitArray& ObjectsInScope = NetRefHandleManager->GetScopableInternalIndices();
 	auto ClearAndDirtyStaticFilter = [this](uint32 ObjIndex)
 	{
 		this->bHasDirtyConnectionFilter = 1;
diff --git a/Engine/Source/Runtime/Experimental/Iris/Core/Public/Iris/Core/IrisDebugging.h b/Engine/Source/Runtime/Experimental/Iris/Core/Public/Iris/Core/IrisDebugging.h
--- a/Engine/Source/Runtime/Experimental/Iris/Core/Public/Iris/Core/IrisDebugging.h
+++ b/Engine/Source/Runtime/Experimental/Iris/Core/Public/Iris/Core/IrisDebugging.h
@@ -110,6 +110,9 @@ extern "C" IRISCORE_API void SetIrisDebugNetRefHandle(uint32 NetHandleId);
 /** Set the InternalIndex to break on */
 extern "C" IRISCORE_API void SetIrisDebugInternalIndex(UE::Net::Private::FInternalNetRefIndex InternalIndex);
 
+/** Set the InternalIndex to break on via it's NetHandle */
+extern "C" IRISCORE_API void SetIrisDebugInternalIndexViaNetHandle(FNetRefHandle RefHandle);
+
 /** Set the RPC Name to break on */
 extern "C" IRISCORE_API void SetIrisDebugRPCName(const ANSICHAR* NameBuffer);
 
