Iris: Added forced ref to debug function

#rnx

diff --git a/Engine/Source/Runtime/Experimental/Iris/Core/Private/Iris/Core/IrisDebugging.cpp b/Engine/Source/Runtime/Experimental/Iris/Core/Private/Iris/Core/IrisDebugging.cpp
--- a/Engine/Source/Runtime/Experimental/Iris/Core/Private/Iris/Core/IrisDebugging.cpp
+++ b/Engine/Source/Runtime/Experimental/Iris/Core/Private/Iris/Core/IrisDebugging.cpp
@@ -230,6 +230,7 @@ uint64 Init()
 	UE_NET_FORCE_REFERENCE_DEBUGFUNCTION(SetIrisDebugObjectName);
 	UE_NET_FORCE_REFERENCE_DEBUGFUNCTION(SetIrisDebugNetRefHandle);
 	UE_NET_FORCE_REFERENCE_DEBUGFUNCTION(SetIrisDebugInternalIndex);
+	UE_NET_FORCE_REFERENCE_DEBUGFUNCTION(SetIrisDebugInternalIndexViaNetHandle);
 	UE_NET_FORCE_REFERENCE_DEBUGFUNCTION(SetIrisDebugRPCName);
 	
 	return FunctionReferenceAccumulator;
