make cluster union sub-component tracing better match expected behavior as if tracing through the SQ

- Zero-length sweeps should be converted to overlaps with normals computed using MTD information.
- Cluster union line traces and sweeps should properly have the ability to return multiple results.
- UPrimitiveComponent::SweepComponent should pass in the swept geometry to ConvertQueryImpactHit to properly find opposing geometry normals.
- FPhysicsObjectCollisionInterface::ShapeSweep needs to do distance checks to find the best hit.

diff --git a/Engine/Source/Runtime/Engine/Classes/PhysicsEngine/ClusterUnionComponent.h b/Engine/Source/Runtime/Engine/Classes/PhysicsEngine/ClusterUnionComponent.h
--- a/Engine/Source/Runtime/Engine/Classes/PhysicsEngine/ClusterUnionComponent.h
+++ b/Engine/Source/Runtime/Engine/Classes/PhysicsEngine/ClusterUnionComponent.h
@@ -159,6 +159,10 @@ public:
 
 	bool HasReceivedTransform() const { return bHasReceivedTransform; }
 
+	// Multi-trace/sweep functions that only make sense in the context of a cluster union.
+	bool LineTraceComponent(TArray<FHitResult>& OutHit, const FVector Start, const FVector End, ECollisionChannel TraceChannel, const struct FCollisionQueryParams& Params, const struct FCollisionResponseParams& ResponseParams, const struct FCollisionObjectQueryParams& ObjectParams);
+	bool SweepComponent(TArray<FHitResult>& OutHit, const FVector Start, const FVector End, const FQuat& ShapeWorldRotation, const FPhysicsGeometry& Geometry, ECollisionChannel TraceChannel, const struct FCollisionQueryParams& Params, const struct FCollisionResponseParams& ResponseParams, const struct FCollisionObjectQueryParams& ObjectParams);
+
 	friend class UClusterUnionReplicatedProxyComponent;
 protected:
 
diff --git a/Engine/Source/Runtime/Engine/Private/Collision/SceneQuery.cpp b/Engine/Source/Runtime/Engine/Private/Collision/SceneQuery.cpp
--- a/Engine/Source/Runtime/Engine/Private/Collision/SceneQuery.cpp
+++ b/Engine/Source/Runtime/Engine/Private/Collision/SceneQuery.cpp
@@ -424,7 +424,7 @@ bool TSceneCastCommonImp(const UWorld* World, typename Traits::TOutHits& OutHits
 			// TODO: Is there a way to get this to generalize better to the PT as well? Right now it depends on GT functions on the cluster union component.
 			if constexpr (Traits::IsExternalData())
 			{
-				auto DoClusterUnionTrace = [&GeomInputs, &Start, &End, TraceChannel, &Params, &ResponseParams, &ObjectParams](const FHitResult& OriginalHit, FHitResult& NewHit)
+				auto DoClusterUnionTrace = [&GeomInputs, &Start, &End, TraceChannel, &Params, &ResponseParams, &ObjectParams](const FHitResult& OriginalHit, auto& NewHit)
 				{
 					FClusterUnionHit Result;
 					if (UClusterUnionComponent* ClusterUnion = Cast<UClusterUnionComponent>(OriginalHit.GetComponent()))
@@ -452,7 +452,7 @@ bool TSceneCastCommonImp(const UWorld* World, typename Traits::TOutHits& OutHits
 
 						for (int32 Index = 0; Index < OutHits.Num(); ++Index)
 						{
-							FHitResult NewHit;
+							TArray<FHitResult> NewHit;
 							FClusterUnionHit ClusterUnionHit = DoClusterUnionTrace(OutHits[Index], NewHit);
 							if (ClusterUnionHit.bIsClusterUnion)
 							{
@@ -460,7 +460,7 @@ bool TSceneCastCommonImp(const UWorld* World, typename Traits::TOutHits& OutHits
 
 								if (ClusterUnionHit.bHit)
 								{
-									AllNewHits.Add(NewHit);
+									AllNewHits.Append(NewHit);
 								}
 							}
 						}
diff --git a/Engine/Source/Runtime/Engine/Private/Components/PrimitiveComponent.cpp b/Engine/Source/Runtime/Engine/Private/Components/PrimitiveComponent.cpp
--- a/Engine/Source/Runtime/Engine/Private/Components/PrimitiveComponent.cpp
+++ b/Engine/Source/Runtime/Engine/Private/Components/PrimitiveComponent.cpp
@@ -3017,7 +3017,7 @@ bool UPrimitiveComponent::SweepComponent(FHitResult& OutHit, const FVector Start
 		QueryFilter.Word1 = 0xFFFFF;
 		ChaosInterface::SetFlags(BestHit, EHitFlags::Distance | EHitFlags::Normal | EHitFlags::Position | EHitFlags::FaceIndex);
 
-		ConvertQueryImpactHit(GetWorld(), BestHit, OutHit, (End - Start).Size(), QueryFilter, Start, End, nullptr, FTransform{ Start }, false, false);
+		ConvertQueryImpactHit(GetWorld(), BestHit, OutHit, (End - Start).Size(), QueryFilter, Start, End, &Geometry, FTransform{ Start }, false, false);
 		return true;
 	}
 
diff --git a/Engine/Source/Runtime/Engine/Private/PhysicsEngine/ClusterUnionComponent.cpp b/Engine/Source/Runtime/Engine/Private/PhysicsEngine/ClusterUnionComponent.cpp
--- a/Engine/Source/Runtime/Engine/Private/PhysicsEngine/ClusterUnionComponent.cpp
+++ b/Engine/Source/Runtime/Engine/Private/PhysicsEngine/ClusterUnionComponent.cpp
@@ -709,6 +709,46 @@ void UClusterUnionComponent::SetSimulatePhysics(bool bSimulate)
 	PhysicsProxy->SetObjectState_External(bSimulate ? Chaos::EObjectStateType::Dynamic : Chaos::EObjectStateType::Kinematic);
 }
 
+bool UClusterUnionComponent::LineTraceComponent(TArray<FHitResult>& OutHit, const FVector Start, const FVector End, ECollisionChannel TraceChannel, const struct FCollisionQueryParams& Params, const struct FCollisionResponseParams& ResponseParams, const struct FCollisionObjectQueryParams& ObjectParams)
+{
+	OutHit.Reset();
+
+	VisitAllCurrentChildComponentsForCollision(TraceChannel, Params, ResponseParams, ObjectParams,
+		[&Start, &End, TraceChannel, &Params, &ResponseParams, &ObjectParams, &OutHit](UPrimitiveComponent* Component)
+		{
+			FHitResult ComponentHit;
+			if (Component->LineTraceComponent(ComponentHit, Start, End, TraceChannel, Params, ResponseParams, ObjectParams))
+			{
+				OutHit.Add(ComponentHit);
+			}
+
+			return true;
+		}
+	);
+
+	return !OutHit.IsEmpty();
+}
+
+bool UClusterUnionComponent::SweepComponent(TArray<FHitResult>& OutHit, const FVector Start, const FVector End, const FQuat& ShapeWorldRotation, const FPhysicsGeometry& Geometry, ECollisionChannel TraceChannel, const struct FCollisionQueryParams& Params, const struct FCollisionResponseParams& ResponseParams, const struct FCollisionObjectQueryParams& ObjectParams)
+{
+	OutHit.Reset();
+
+	VisitAllCurrentChildComponentsForCollision(TraceChannel, Params, ResponseParams, ObjectParams,
+		[&Geometry, &Start, &End, &ShapeWorldRotation, TraceChannel, &Params, &ResponseParams, &ObjectParams, &OutHit](UPrimitiveComponent* Component)
+		{
+			FHitResult ComponentHit;
+			if (Component->SweepComponent(ComponentHit, Start, End, ShapeWorldRotation, Geometry, TraceChannel, Params, ResponseParams, ObjectParams))
+			{
+				OutHit.Add(ComponentHit);
+			}
+
+			return true;
+		}
+	);
+
+	return !OutHit.IsEmpty();
+}
+
 bool UClusterUnionComponent::LineTraceComponent(FHitResult& OutHit, const FVector Start, const FVector End, ECollisionChannel TraceChannel, const struct FCollisionQueryParams& Params, const struct FCollisionResponseParams& ResponseParams, const struct FCollisionObjectQueryParams& ObjectParams)
 {
 	bool bHasHit = false;
diff --git a/Engine/Source/Runtime/Experimental/Chaos/Public/Chaos/Interface/SQTypes.h b/Engine/Source/Runtime/Experimental/Chaos/Public/Chaos/Interface/SQTypes.h
--- a/Engine/Source/Runtime/Experimental/Chaos/Public/Chaos/Interface/SQTypes.h
+++ b/Engine/Source/Runtime/Experimental/Chaos/Public/Chaos/Interface/SQTypes.h
@@ -154,6 +154,9 @@ namespace ChaosInterface
 	{
 	};
 
+	template<Chaos::EThreadContext Id>
+	using TThreadQueryHit = std::conditional_t<Id == Chaos::EThreadContext::External, FQueryHit, FPTQueryHit>;
+
 	template<Chaos::EThreadContext Id>
 	using TThreadRaycastHit = std::conditional_t<Id == Chaos::EThreadContext::External, FRaycastHit, FPTRaycastHit>;
 
diff --git a/Engine/Source/Runtime/Experimental/Chaos/Public/Chaos/PhysicsObjectCollisionInterface.h b/Engine/Source/Runtime/Experimental/Chaos/Public/Chaos/PhysicsObjectCollisionInterface.h
--- a/Engine/Source/Runtime/Experimental/Chaos/Public/Chaos/PhysicsObjectCollisionInterface.h
+++ b/Engine/Source/Runtime/Experimental/Chaos/Public/Chaos/PhysicsObjectCollisionInterface.h
@@ -138,6 +138,12 @@ namespace Chaos
 
 		template<typename TOverlapHit>
 		bool ShapeOverlap(TArrayView<const FConstPhysicsObjectHandle> InObjects, const Chaos::FImplicitObject& InGeom, const FTransform& GeomTransform, TArray<TOverlapHit>& OutOverlaps)
+		{
+			return ShapeOverlapWithMTD(InObjects, InGeom, GeomTransform, OutOverlaps, nullptr);
+		}
+
+		template<typename TOverlapHit>
+		bool ShapeOverlapWithMTD(TArrayView<const FConstPhysicsObjectHandle> InObjects, const Chaos::FImplicitObject& InGeom, const FTransform& GeomTransform, TArray<TOverlapHit>& OutOverlaps, FMTDInfo* MTD)
 		{
 			static_assert(std::is_same_v<TOverlapHit, ChaosInterface::TThreadOverlapHit<Id>>);
 			bool bHasOverlap = false;
@@ -147,15 +153,15 @@ namespace Chaos
 
 				Interface.VisitEveryShape(
 					{ &Object, 1 },
-					[this, &bHasOverlap, &WorldTM, &InGeom, &GeomTransform, &OutOverlaps](const FConstPhysicsObjectHandle IterObject, TThreadShapeInstance<Id>* Shape)
+					[this, &bHasOverlap, &WorldTM, &InGeom, &GeomTransform, &OutOverlaps, MTD](const FConstPhysicsObjectHandle IterObject, TThreadShapeInstance<Id>* Shape)
 					{
 						check(Shape);
 						const bool bOverlap = Chaos::Utilities::CastHelper(
 							InGeom,
 							GeomTransform,
-							[Shape, &WorldTM](const auto& Downcast, const auto& FullTransformB)
+							[Shape, &WorldTM, MTD](const auto& Downcast, const auto& FullTransformB)
 							{
-								return Chaos::OverlapQuery(*Shape->GetGeometry(), WorldTM, Downcast, FullTransformB, 0, nullptr);
+								return Chaos::OverlapQuery(*Shape->GetGeometry(), WorldTM, Downcast, FullTransformB, 0, MTD);
 							}
 						);
 
@@ -183,12 +189,26 @@ namespace Chaos
 			const FVector StartPos = StartTM.GetTranslation();
 			const FVector Delta = EndPos - StartPos;
 			const FReal DeltaMag = Delta.Size();
-			if (DeltaMag < UE_KINDA_SMALL_NUMBER)
+
+			// This maintains compatability with SQ querying code (see TSQVisitor::Visit where it checks Sweeps with 0 length).
+			if (DeltaMag < UE_SMALL_NUMBER)
 			{
-				return false;
+				TArray<ChaosInterface::TThreadOverlapHit<Id>> OverlapHits;
+				FMTDInfo MTDInfo;
+				const bool bOverlap = ShapeOverlapWithMTD(InObjects, InGeom, StartTM, OverlapHits, &MTDInfo);
+				if (!OverlapHits.IsEmpty())
+				{
+					const ChaosInterface::TThreadQueryHit<Id>& OverlapQueryHit = OverlapHits[0];
+					ChaosInterface::TThreadQueryHit<Id>& OutQueryHit = OutBestHit;
+					OutQueryHit = OverlapQueryHit;
+					OutBestHit.WorldNormal = MTDInfo.Normal * MTDInfo.Penetration;
+				}
+				return bOverlap;
 			}
-			const FVec3 Dir = Delta / DeltaMag;
 
+			const FVector Dir = Delta / DeltaMag;
+
+			OutBestHit.Distance = TNumericLimits<float>::Max();
 			for (const FConstPhysicsObjectHandle Object : InObjects)
 			{
 				const FTransform WorldTM = Interface.GetTransform(Object);
@@ -220,14 +240,15 @@ namespace Chaos
 								}
 							);
 
-							if (bShapeHit)
+							const float FloatDistance = static_cast<float>(Distance);
+							if (bShapeHit && FloatDistance < OutBestHit.Distance)
 							{
 								bHit = true;
 
 								OutBestHit.Shape = Shape;
 								OutBestHit.WorldPosition = WorldPosition;
 								OutBestHit.WorldNormal = WorldNormal;
-								OutBestHit.Distance = static_cast<float>(Distance);
+								OutBestHit.Distance = FloatDistance;
 								OutBestHit.FaceIndex = FaceIdx;
 								if (OutBestHit.Distance > 0.f)
 								{
