ChaosFlesh : Add Gauss Seildel corotated method

diff --git a/Engine/Plugins/Experimental/ChaosFlesh/Source/ChaosFleshEngine/Private/ChaosFlesh/ChaosDeformableSolverComponent.cpp b/Engine/Plugins/Experimental/ChaosFlesh/Source/ChaosFleshEngine/Private/ChaosFlesh/ChaosDeformableSolverComponent.cpp
--- a/Engine/Plugins/Experimental/ChaosFlesh/Source/ChaosFleshEngine/Private/ChaosFlesh/ChaosDeformableSolverComponent.cpp
+++ b/Engine/Plugins/Experimental/ChaosFlesh/Source/ChaosFleshEngine/Private/ChaosFlesh/ChaosDeformableSolverComponent.cpp
@@ -213,6 +213,9 @@ void UDeformableSolverComponent::Reset()
 			, SolverForces.bEnableGravity
 			, SolverConstraints.CorotatedConstraints.bEnableCorotatedConstraint
 			, SolverConstraints.bEnablePositionTargets
+			, GaussSeidelConstraints.bUseGaussSeidelConstraints
+			, GaussSeidelConstraints.bUseSOR
+			, GaussSeidelConstraints.OmegaSOR
 		}));
 
 		for (TObjectPtr<UDeformablePhysicsComponent>& DeformableComponent : ConnectedObjects.DeformableComponents)
diff --git a/Engine/Plugins/Experimental/ChaosFlesh/Source/ChaosFleshEngine/Public/ChaosFlesh/ChaosDeformableSolverComponent.h b/Engine/Plugins/Experimental/ChaosFlesh/Source/ChaosFleshEngine/Public/ChaosFlesh/ChaosDeformableSolverComponent.h
--- a/Engine/Plugins/Experimental/ChaosFlesh/Source/ChaosFleshEngine/Public/ChaosFlesh/ChaosDeformableSolverComponent.h
+++ b/Engine/Plugins/Experimental/ChaosFlesh/Source/ChaosFleshEngine/Public/ChaosFlesh/ChaosDeformableSolverComponent.h
@@ -94,6 +94,21 @@ struct FSolverGridBasedCollisionsGroup
 		float GridDx = 25.;
 };
 
+USTRUCT(BlueprintType)
+struct FSolverGaussSeidelConstraintsGroup
+{
+	GENERATED_USTRUCT_BODY()
+
+	UPROPERTY(EditAnywhere, Category = "GaussSeidelConstraints")
+		bool bUseGaussSeidelConstraints = false;
+
+	UPROPERTY(EditAnywhere, Category = "GaussSeidelConstraints")
+		bool bUseSOR = true;
+
+	UPROPERTY(EditAnywhere, Category = "GaussSeidelConstraints")
+		float OmegaSOR = 1.6f;
+};
+
 
 USTRUCT(BlueprintType)
 struct FSolverCollisionsGroup
@@ -208,6 +223,9 @@ public:
 	UPROPERTY(EditAnywhere, BlueprintReadOnly, Category = "Physics")
 	FSolverConstraintsGroup SolverConstraints;
 
+	UPROPERTY(EditAnywhere, BlueprintReadOnly, Category = "Physics")
+	FSolverGaussSeidelConstraintsGroup GaussSeidelConstraints;
+
 	UPROPERTY(EditAnywhere, BlueprintReadOnly, Category = "Physics")
 	FSolverForcesGroup SolverForces;
 
diff --git a/Engine/Source/Runtime/Experimental/Chaos/Private/Chaos/Deformable/ChaosDeformableSolver.cpp b/Engine/Source/Runtime/Experimental/Chaos/Private/Chaos/Deformable/ChaosDeformableSolver.cpp
--- a/Engine/Source/Runtime/Experimental/Chaos/Private/Chaos/Deformable/ChaosDeformableSolver.cpp
+++ b/Engine/Source/Runtime/Experimental/Chaos/Private/Chaos/Deformable/ChaosDeformableSolver.cpp
@@ -59,6 +59,8 @@ DECLARE_CYCLE_STAT(TEXT("Chaos.Deformable.Solver.InitializeKinematicConstraint")
 DECLARE_CYCLE_STAT(TEXT("Chaos.Deformable.Solver.InitializeCollisionBodies"), STAT_ChaosDeformableSolver_InitializeCollisionBodies, STATGROUP_Chaos);
 DECLARE_CYCLE_STAT(TEXT("Chaos.Deformable.Solver.InitializeSelfCollisionVariables"), STAT_ChaosDeformableSolver_InitializeSelfCollisionVariables, STATGROUP_Chaos);
 DECLARE_CYCLE_STAT(TEXT("Chaos.Deformable.Solver.InitializeGridBasedConstraintVariables"), STAT_ChaosDeformableSolver_InitializeGridBasedConstraintVariables, STATGROUP_Chaos);
+DECLARE_CYCLE_STAT(TEXT("Chaos.Deformable.Solver.InitializeGaussSeidelConstraintVariables"), STAT_ChaosDeformableSolver_InitializeGaussSeidelConstraintVariables, STATGROUP_Chaos);
+DECLARE_CYCLE_STAT(TEXT("Chaos.Deformable.Solver.InitializeGaussSeidelConstraint"), STAT_ChaosDeformableSolver_InitializeGaussSeidelConstraint, STATGROUP_Chaos);
 DECLARE_CYCLE_STAT(TEXT("Chaos.Deformable.Solver.RemoveSimulationObjects"), STAT_ChaosDeformableSolver_RemoveSimulationObjects, STATGROUP_Chaos);
 DECLARE_CYCLE_STAT(TEXT("Chaos.Deformable.Solver.RemoveProxy"), STAT_ChaosDeformableSolver_RemoveProxy, STATGROUP_Chaos);
 DECLARE_CYCLE_STAT(TEXT("Chaos.Deformable.Solver.AddProxy"), STAT_ChaosDeformableSolver_AddProxy, STATGROUP_Chaos);
@@ -142,6 +144,15 @@ namespace Chaos::Softs
 		{
 			AllElements.Reset(new TArray<Chaos::TVec4<int32>>());
 		}
+		if (Property.bUseGaussSeidelConstraints)
+		{
+			AllElements.Reset(new TArray<Chaos::TVec4<int32>>());
+			AllIncidentElements.Reset(new TArray<TArray<int32>>());
+			AllIncidentElementsLocal.Reset(new TArray<TArray<int32>>());
+			AllTetEMeshArray.Reset(new TArray<FSolverReal>());
+			AllTetNuMeshArray.Reset(new TArray<FSolverReal>());
+			AllTetAlphaJArray.Reset(new TArray<FSolverReal>());
+		}
 
 		InitializeKinematicConstraint();
 		Frame = 0;
@@ -315,6 +326,11 @@ namespace Chaos::Softs
 					{
 						InitializeGridBasedConstraintVariables();
 					}
+
+					if (Property.bUseGaussSeidelConstraints)
+					{
+						InitializeGaussSeidelConstraintVariables();
+					}
 				}
 				UninitializedProxys_Internal.SetNum(0, true);
 			}
@@ -356,9 +372,11 @@ namespace Chaos::Softs
 					InitializeWeakConstraint(*Proxy);
 					InitializeTetrahedralConstraint(*Proxy);
 					InitializeGidBasedConstraints(*Proxy);
+					InitializeGaussSeidelConstraints(*Proxy);
 				}
 			}
 		}
+
 		if (FCollisionManagerProxy* CollisionManagerProxy = InProxy.As< FCollisionManagerProxy>())
 		{
 			InitializeCollisionBodies(*CollisionManagerProxy);
@@ -419,6 +437,7 @@ namespace Chaos::Softs
 
 		Evolution->SetDamping(DampingMultiplier, GroupOffset - 1);
 
+
 		// Tet mesh points are in component space.  That means that if our sim space is:
 		//    World:     We need to multiply by the ComponentToWorldXf.
 		//	  Component: We do nothing.
@@ -803,6 +822,22 @@ namespace Chaos::Softs
 					}
 				}
 
+				if (Property.bUseGaussSeidelConstraints)
+				{
+					int32 ElementsOffset = AllTetEMeshArray->Num();
+					AllTetEMeshArray->SetNum(ElementsOffset + NumElements);
+					AllTetNuMeshArray->SetNum(ElementsOffset + NumElements);
+					AllTetAlphaJArray->SetNum(ElementsOffset + NumElements);
+						
+					for (uint32 edx = 0; edx < NumElements; ++edx)
+					{
+						(*AllTetEMeshArray)[edx + ElementsOffset] = TetStiffness[edx];
+						(*AllTetNuMeshArray)[edx + ElementsOffset] = TetNu[edx];
+						(*AllTetAlphaJArray)[edx + ElementsOffset] = AlphaJMesh[edx];
+					}
+				}
+
+
 				if (Property.bEnableCorotatedConstraints)
 				{
 
@@ -851,6 +886,8 @@ namespace Chaos::Softs
 						CorotatedConstraints.Add(TUniquePtr<FXPBDCorotatedConstraints<FSolverReal, FSolverParticles>>(CorotatedConstraint));
 					}
 				}
+			
+				
 			}
 
 		}
@@ -882,6 +919,50 @@ namespace Chaos::Softs
 		}
 	}
 
+	void FDeformableSolver::InitializeGaussSeidelConstraints(FFleshThreadingProxy& Proxy)
+	{
+		PERF_SCOPE(STAT_ChaosDeformableSolver_InitializeGaussSeidelConstraint);
+
+		if (Property.bUseGaussSeidelConstraints)
+		{
+			auto ChaosTet = [](FIntVector4 V, int32 dp) { return Chaos::TVec4<int32>(dp + V.X, dp + V.Y, dp + V.Z, dp + V.W); };
+
+			const FManagedArrayCollection& Rest = Proxy.GetRestCollection();
+			const TManagedArray<FIntVector4>& Tetrahedron = Rest.GetAttribute<FIntVector4>("Tetrahedron", "Tetrahedral");
+			const TManagedArray<TArray<int32>>& Tetrahedron1 = Rest.GetAttribute<TArray<int32>>("IncidentElements", "Vertices");
+			const TManagedArray<TArray<int32>>& Tetrahedron2 = Rest.GetAttribute<TArray<int32>>("IncidentElementsLocalIndex", "Vertices");
+			const TManagedArray<TArray<int32>>* IncidentElementsPointer = Rest.FindAttribute<TArray<int32>>("IncidentElements", "Vertices");
+			const TManagedArray<TArray<int32>>* IncidentElementsLocalPointer = Rest.FindAttribute<TArray<int32>>("IncidentElementsLocalIndex", "Vertices");
+
+			if (uint32 NumElements = Tetrahedron.Num())
+			{
+				const FIntVector2& Range = Proxy.GetSolverParticleRange();
+
+				int32 ElementsOffset = AllElements->Num();
+				AllElements->SetNum(ElementsOffset + NumElements);
+				for (uint32 edx = 0; edx < NumElements; ++edx)
+				{
+					(*AllElements)[edx + ElementsOffset] = ChaosTet(Tetrahedron[edx], Range[0]);
+				}
+				if (uint32 NumIncident = IncidentElementsPointer->Num())
+				{
+					int32 IncidentOffSet = AllIncidentElements->Num();
+					AllIncidentElements->SetNum(IncidentOffSet + NumIncident); 
+					AllIncidentElementsLocal->SetNum(IncidentOffSet + NumIncident);
+					for (uint32 i = 0; i < NumIncident; ++i)
+					{
+						(*AllIncidentElements)[i + IncidentOffSet] = (*IncidentElementsPointer)[i];
+						for (int32 j = 0; j < (*AllIncidentElements)[i + IncidentOffSet].Num(); j++)
+						{
+							(*AllIncidentElements)[i + IncidentOffSet][j] += ElementsOffset;
+						}
+						(*AllIncidentElementsLocal)[i + IncidentOffSet] = (*IncidentElementsLocalPointer)[i];
+					}
+				}
+			}
+		}
+	}
+
 
 	void FDeformableSolver::InitializeKinematicConstraint()
 	{
@@ -1064,10 +1145,30 @@ namespace Chaos::Softs
 		{
 			this->GridBasedCorotatedConstraint->TimeStepPostprocessing(InParticles, Dt);
 		};
-	
-	
 	}
 
+	void FDeformableSolver::InitializeGaussSeidelConstraintVariables()
+	{
+		PERF_SCOPE(STAT_ChaosDeformableSolver_InitializeGaussSeidelConstraintVariables);
+
+		GSCorotatedConstraints.Reset(new Chaos::Softs::FGaussSeidelCorotatedConstraints<FSolverReal, FSolverParticles>(
+			Evolution->Particles(), *AllElements, *AllTetEMeshArray, *AllTetNuMeshArray, MoveTemp(*AllTetAlphaJArray), MoveTemp(*AllIncidentElements), MoveTemp(*AllIncidentElementsLocal), 0, Evolution->Particles().Size(), Property.bDoQuasistatics, Property.bUseSOR, Property.OmegaSOR, GDeformableXPBDCorotatedParams));
+		Evolution->ResetConstraintRules();
+		int32 InitIndex1 = Evolution->AddConstraintInitRange(1, true);
+		Evolution->ConstraintInits()[InitIndex1] =
+			[this](FSolverParticles& InParticles, const FSolverReal Dt)
+		{
+			this->GSCorotatedConstraints->Init(Dt, InParticles);
+		};
+		int32 ConstraintIndex1 = Evolution->AddConstraintRuleRange(1, true);
+		Evolution->ConstraintRules()[ConstraintIndex1] =
+			[this](FSolverParticles& InParticles, const FSolverReal Dt)
+		{
+			this->GSCorotatedConstraints->ApplyInParallel(InParticles, Dt);
+		};
+	}
+
+
 	void FDeformableSolver::RemoveSimulationObjects()
 	{
 		PERF_SCOPE(STAT_ChaosDeformableSolver_RemoveSimulationObjects);
@@ -1131,6 +1232,7 @@ namespace Chaos::Softs
 				{
 					InitializeTetrahedralConstraint(*Proxy);
 					InitializeGidBasedConstraints(*Proxy);
+					InitializeGaussSeidelConstraints(*Proxy);
 				}
 			}
 		}
diff --git a/Engine/Source/Runtime/Experimental/Chaos/Private/Chaos/GraphColoring.cpp b/Engine/Source/Runtime/Experimental/Chaos/Private/Chaos/GraphColoring.cpp
--- a/Engine/Source/Runtime/Experimental/Chaos/Private/Chaos/GraphColoring.cpp
+++ b/Engine/Source/Runtime/Experimental/Chaos/Private/Chaos/GraphColoring.cpp
@@ -1370,11 +1370,12 @@ TArray<TArray<int32>> Chaos::ComputeNodalColoring(const TArray<TVec4<int32>>& Gr
 
 	TArray<int32> Particle2Incident;
 	Particle2Incident.Init(-1, GraphParticlesEnd - GraphParticlesStart);
+	//Assuming that offset of Graph is GraphParticlesStart
 	for (int32 i = 0; i < IncidentElements.Num(); i++) 
 	{
 		if (IncidentElements[i].Num() > 0)
 		{
-			Particle2Incident[Graph[IncidentElements[i][0]][IncidentElementsLocalIndex[i][0]]] = i;
+			Particle2Incident[Graph[IncidentElements[i][0]][IncidentElementsLocalIndex[i][0]] - GraphParticlesStart] = i;
 		}
 	}
 
diff --git a/Engine/Source/Runtime/Experimental/Chaos/Private/Chaos/NewtonCorotatedCache.cpp b/Engine/Source/Runtime/Experimental/Chaos/Private/Chaos/NewtonCorotatedCache.cpp
--- a/Engine/Source/Runtime/Experimental/Chaos/Private/Chaos/NewtonCorotatedCache.cpp
+++ b/Engine/Source/Runtime/Experimental/Chaos/Private/Chaos/NewtonCorotatedCache.cpp
@@ -95,6 +95,44 @@ void CorotatedCache<T>::deltaP(const Chaos::PMatrix<T, 3, 3>& F, const Chaos::PM
 
 }
 
+template <typename T>
+T PsiCorotated(const Chaos::PMatrix<T, 3, 3>& F, const T mu, const T lambda) 
+{
+	Chaos::PMatrix<T, 3, 3> R((T)0.), S((T)0.), b((T)0.);
+	Chaos::PolarDecomposition(F, R, S);
+	T J = F.Determinant();
+	b = F * F.GetTransposed();
+	T STrace = T(0), bTrace = T(0);
+	for (int32 alpha = 0; alpha < 3; alpha++)
+	{
+		STrace += S.GetDiagonal()[alpha];
+		bTrace += b.GetDiagonal()[alpha];
+	}
+	
+	return mu * (bTrace - T(2) * STrace + T(3)) + lambda * (J - T(1)) * (J - T(1)) / T(2);
+}
+
+template <typename T>
+void PCorotated(const Chaos::PMatrix<T, 3, 3>& Fe, const T mu, const T lambda, Chaos::PMatrix<T, 3, 3>& P)
+{
+	Chaos::PMatrix<T, 3, 3> R((T)0.), S((T)0.), JFinvT((T)0.);
+	Chaos::PolarDecomposition(Fe, R, S);
+
+	JFinvT.SetAt(0, 0, Fe.GetAt(1, 1) * Fe.GetAt(2, 2) - Fe.GetAt(2, 1) * Fe.GetAt(1, 2));
+	JFinvT.SetAt(0, 1, Fe.GetAt(2, 0) * Fe.GetAt(1, 2) - Fe.GetAt(1, 0) * Fe.GetAt(2, 2));
+	JFinvT.SetAt(0, 2, Fe.GetAt(1, 0) * Fe.GetAt(2, 1) - Fe.GetAt(2, 0) * Fe.GetAt(1, 1));
+	JFinvT.SetAt(1, 0, Fe.GetAt(2, 1) * Fe.GetAt(0, 2) - Fe.GetAt(0, 1) * Fe.GetAt(2, 2));
+	JFinvT.SetAt(1, 1, Fe.GetAt(0, 0) * Fe.GetAt(2, 2) - Fe.GetAt(2, 0) * Fe.GetAt(0, 2));
+	JFinvT.SetAt(1, 2, Fe.GetAt(2, 0) * Fe.GetAt(0, 1) - Fe.GetAt(0, 0) * Fe.GetAt(2, 1));
+	JFinvT.SetAt(2, 0, Fe.GetAt(0, 1) * Fe.GetAt(1, 2) - Fe.GetAt(1, 1) * Fe.GetAt(0, 2));
+	JFinvT.SetAt(2, 1, Fe.GetAt(1, 0) * Fe.GetAt(0, 2) - Fe.GetAt(0, 0) * Fe.GetAt(1, 2));
+	JFinvT.SetAt(2, 2, Fe.GetAt(0, 0) * Fe.GetAt(1, 1) - Fe.GetAt(1, 0) * Fe.GetAt(0, 1));
+
+	T J = Fe.Determinant();
+
+	P = (T)2. * mu * (Fe - R) + lambda * (J - (T)1.) * JFinvT;
+}
+
 } 
 
 template CHAOS_API void Chaos::Softs::CorotatedCache<Chaos::FRealSingle>::UpdateCache(const Chaos::PMatrix<Chaos::FRealSingle, 3, 3>&);
@@ -103,3 +141,7 @@ template CHAOS_API void Chaos::Softs::CorotatedCache<Chaos::FRealSingle>::deltaP
 template CHAOS_API void Chaos::Softs::CorotatedCache<Chaos::FReal>::UpdateCache(const Chaos::PMatrix<Chaos::FReal, 3, 3>&);
 template CHAOS_API void Chaos::Softs::CorotatedCache<Chaos::FReal>::P(const Chaos::PMatrix<Chaos::FReal, 3, 3>&, const Chaos::FReal, const Chaos::FReal, Chaos::PMatrix<Chaos::FReal, 3, 3>&);
 template CHAOS_API void Chaos::Softs::CorotatedCache<Chaos::FReal>::deltaP(const Chaos::PMatrix<Chaos::FReal, 3, 3>&, const Chaos::PMatrix<Chaos::FReal, 3, 3>&, const Chaos::FReal, const Chaos::FReal, Chaos::PMatrix<Chaos::FReal, 3, 3>&);
+template CHAOS_API void Chaos::Softs::PCorotated<Chaos::FReal>(const Chaos::PMatrix<Chaos::FReal, 3, 3>& Fe, const Chaos::FReal mu, const Chaos::FReal lambda, Chaos::PMatrix<Chaos::FReal, 3, 3>& P);
+template CHAOS_API void Chaos::Softs::PCorotated<Chaos::FRealSingle>(const Chaos::PMatrix<Chaos::FRealSingle, 3, 3>& Fe, const Chaos::FRealSingle mu, const Chaos::FRealSingle lambda, Chaos::PMatrix<Chaos::FRealSingle, 3, 3>& P);
+template CHAOS_API Chaos::FReal Chaos::Softs::PsiCorotated<Chaos::FReal>(const Chaos::PMatrix<Chaos::FReal, 3, 3>& F, const Chaos::FReal mu, const Chaos::FReal lambda);
+template CHAOS_API Chaos::FRealSingle Chaos::Softs::PsiCorotated<Chaos::FRealSingle>(const Chaos::PMatrix<Chaos::FRealSingle, 3, 3>& F, const Chaos::FRealSingle mu, const Chaos::FRealSingle lambda);
diff --git a/Engine/Source/Runtime/Experimental/Chaos/Public/Chaos/Deformable/ChaosDeformableSolver.h b/Engine/Source/Runtime/Experimental/Chaos/Public/Chaos/Deformable/ChaosDeformableSolver.h
--- a/Engine/Source/Runtime/Experimental/Chaos/Public/Chaos/Deformable/ChaosDeformableSolver.h
+++ b/Engine/Source/Runtime/Experimental/Chaos/Public/Chaos/Deformable/ChaosDeformableSolver.h
@@ -11,6 +11,7 @@
 #include "Chaos/PBDTriangleMeshCollisions.h"
 #include "Chaos/TriangleMesh.h"
 #include "Chaos/XPBDCorotatedConstraints.h"
+#include "Chaos/Deformable/GaussSeidelCorotatedConstraints.h"
 #include "Chaos/XPBDWeakConstraints.h"
 #include "Chaos/BlendedXPBDCorotatedConstraints.h"
 #include "Chaos/XPBDGridBasedCorotatedConstraints.h"
@@ -121,11 +122,13 @@ namespace Chaos::Softs
 		void InitializeKinematicParticles(FFleshThreadingProxy&);
 		void InitializeTetrahedralConstraint(FFleshThreadingProxy&);
 		void InitializeGidBasedConstraints(FFleshThreadingProxy&);
+		void InitializeGaussSeidelConstraints(FFleshThreadingProxy& Proxy);
 		void InitializeWeakConstraint(FFleshThreadingProxy&);
 		void InitializeKinematicConstraint();
 		void InitializeCollisionBodies(FCollisionManagerProxy&);
 		void InitializeSelfCollisionVariables();
 		void InitializeGridBasedConstraintVariables();
+		void InitializeGaussSeidelConstraintVariables();
 		void UpdateCollisionBodies(FCollisionManagerProxy&, FThreadingProxy::FKey, FSolverReal DeltaTime);
 		void RemoveSimulationObjects();
 		TArray<Chaos::TVec3<FSolverReal>> ComputeParticleTargets(const TArray<TArray<int32>>& ParticleIndices);
@@ -172,6 +175,7 @@ namespace Chaos::Softs
 		// Simulation Variables
 		TUniquePtr<Softs::FPBDEvolution> Evolution;
 		TArray<TUniquePtr<Softs::FXPBDCorotatedConstraints<Softs::FSolverReal, Softs::FSolverParticles>>> CorotatedConstraints;
+		TUniquePtr<Softs::FGaussSeidelCorotatedConstraints<Softs::FSolverReal, Softs::FSolverParticles>> GSCorotatedConstraints;
 		TArray<TUniquePtr<Softs::FXPBDWeakConstraints<Softs::FSolverReal, Softs::FSolverParticles>>> WeakConstraints;
 		TArray<TUniquePtr<Softs::FBlendedXPBDCorotatedConstraints<Softs::FSolverReal, Softs::FSolverParticles>>> BlendedCorotatedConstraints;
 		TUniquePtr<Softs::FXPBDGridBasedCorotatedConstraints<Softs::FSolverReal, Softs::FSolverParticles>> GridBasedCorotatedConstraint;
@@ -181,6 +185,11 @@ namespace Chaos::Softs
 		TUniquePtr <TArray<TVec3<int32>>> SurfaceElements;
 		TUniquePtr <TArray<Chaos::TVec4<int32>>> AllElements;
 		TUniquePtr <FTriangleMesh> SurfaceTriangleMesh;
+		TUniquePtr <TArray<TArray<int32>>> AllIncidentElements;
+		TUniquePtr <TArray<TArray<int32>>> AllIncidentElementsLocal;
+		TUniquePtr <TArray<FSolverReal>> AllTetEMeshArray;
+		TUniquePtr <TArray<FSolverReal>> AllTetNuMeshArray;
+		TUniquePtr <TArray<FSolverReal>> AllTetAlphaJArray;
 
 		//typedef TMap<int32, TTuple<float, Chaos::Softs::FPAndInvM, FVector3f>> TransientConstraintBufferMap;
 		typedef TMap<int32, TTuple<float, float, FVector3f>> TransientConstraintBufferMap;
diff --git a/Engine/Source/Runtime/Experimental/Chaos/Public/Chaos/Deformable/ChaosDeformableSolverTypes.h b/Engine/Source/Runtime/Experimental/Chaos/Public/Chaos/Deformable/ChaosDeformableSolverTypes.h
--- a/Engine/Source/Runtime/Experimental/Chaos/Public/Chaos/Deformable/ChaosDeformableSolverTypes.h
+++ b/Engine/Source/Runtime/Experimental/Chaos/Public/Chaos/Deformable/ChaosDeformableSolverTypes.h
@@ -33,7 +33,10 @@ namespace Chaos::Softs
 			FSolverReal InDamping = (FSolverReal)0,
 			bool InbEnableGravity = true, 
 			bool InbEnableCorotatedConstraints = true, 
-			bool InbEnablePositionTargets = true)
+			bool InbEnablePositionTargets = true, 
+			bool InbUseGaussSeidelConstraints = false, 
+			bool InbUseSOR = true,
+			FSolverReal InOmegaSOR = (FSolverReal)1.6)
 			: NumSolverSubSteps(InNumSolverSubSteps)
 			, NumSolverIterations(InNumSolverIterations)
 			, FixTimeStep(InFixTimeStep)
@@ -52,6 +55,9 @@ namespace Chaos::Softs
 			, bEnableGravity(InbEnableGravity)
 			, bEnableCorotatedConstraints(InbEnableCorotatedConstraints)
 			, bEnablePositionTargets(InbEnablePositionTargets)
+			, bUseGaussSeidelConstraints(InbUseGaussSeidelConstraints)
+			, bUseSOR(InbUseSOR)
+			, OmegaSOR(InOmegaSOR)
 		{}
 
 		int32 NumSolverSubSteps = 5;
@@ -72,6 +78,9 @@ namespace Chaos::Softs
 		bool bEnableGravity = true;
 		bool bEnableCorotatedConstraints = true;
 		bool bEnablePositionTargets = true;
+		bool bUseGaussSeidelConstraints = false;
+		bool bUseSOR = true;
+		FSolverReal OmegaSOR = (FSolverReal)1.6;
 	};
 
 
diff --git a/Engine/Source/Runtime/Experimental/Chaos/Public/Chaos/Deformable/GaussSeidelCorotatedConstraints.h b/Engine/Source/Runtime/Experimental/Chaos/Public/Chaos/Deformable/GaussSeidelCorotatedConstraints.h
new file mode 100644
--- /dev/null
+++ b/Engine/Source/Runtime/Experimental/Chaos/Public/Chaos/Deformable/GaussSeidelCorotatedConstraints.h
@@ -0,0 +1,317 @@
+// Copyright Epic Games, Inc. All Rights Reserved.
+#pragma once
+
+#include "Chaos/PBDSpringConstraintsBase.h"
+#include "Chaos/XPBDCorotatedConstraints.h"
+#include "ChaosStats.h"
+#include "Chaos/ImplicitQRSVD.h"
+#include "Chaos/GraphColoring.h"
+#include "Chaos/NewtonCorotatedCache.h"
+#include "Chaos/Framework/Parallel.h"
+
+
+namespace Chaos::Softs
+{
+	template <typename T, typename ParticleType>
+	class FGaussSeidelCorotatedConstraints : public FXPBDCorotatedConstraints <T, ParticleType>
+	{
+
+		typedef FXPBDCorotatedConstraints<T, ParticleType> Base;
+		using Base::MeshConstraints;
+		using Base::LambdaArray;
+		using Base::Measure;
+		using Base::DmInverse;
+		using Base::MuElementArray;
+		using Base::LambdaElementArray;
+
+	public:
+		//this one only accepts tetmesh input and mesh
+		FGaussSeidelCorotatedConstraints(
+			const ParticleType& InParticles,
+			const TArray<TVector<int32, 4>>& InMesh,
+			const TArray<T>& EMeshArray,
+			const TArray<T>& NuMeshArray,
+			TArray<T>&& AlphaJMeshArray,
+			TArray<TArray<int32>>&& IncidentElementsIn,
+			TArray<TArray<int32>>&& IncidentElementsLocalIn,
+			const int32 ParticleStartIndexIn,
+			const int32 ParticleEndIndexIn,
+			const bool bDoQuasistaticsIn = false,
+			const bool bDoSORIn = true,
+			const T InOmegaSOR = (T)1.6,
+			const FDeformableXPBDCorotatedParams& InParams = FDeformableXPBDCorotatedParams(),
+			const T& NuMesh = (T).3,
+			const bool bRecordMetricIn = false
+		)
+			: Base(InParticles, InMesh, EMeshArray, NuMeshArray, MoveTemp(AlphaJMeshArray), InParams, NuMesh, bRecordMetricIn, false), IncidentElements(MoveTemp(IncidentElementsIn)), IncidentElementsLocal(IncidentElementsLocalIn),
+			ParticleStartIndex(ParticleStartIndexIn), ParticleEndIndex(ParticleEndIndexIn), bDoQuasistatics(bDoQuasistaticsIn), bDoSOR(bDoSORIn), OmegaSOR(InOmegaSOR)
+		{
+			Base::LambdaArray.SetNum(0);
+			Particle2Incident.Init(-1, ParticleEndIndexIn - ParticleStartIndexIn);
+			for (int32 i = 0; i < IncidentElements.Num(); i++)
+			{
+				if (IncidentElements[i].Num() > 0)
+				{
+					Particle2Incident[InMesh[IncidentElements[i][0]][IncidentElementsLocal[i][0]] - ParticleStartIndex] = i;
+				}
+			}
+			for (int32 e = 0; e < InMesh.Num(); e++)
+			{
+				for (int32 r = 0; r < 3; r++) {
+					for (int32 c = 0; c < 3; c++) {
+						DmInverse[(3 * 3) * e + 3 * r + c] *= Base::AlphaJArray[e];
+					}
+				}
+			}
+			if (!bDoQuasistatics)
+			{
+				xtilde.Init(TVector<T, 3>(T(0.)), ParticleEndIndex - ParticleStartIndex);
+			}
+			if (bDoSOR)
+			{
+				X_k_1.Init(TVector<T, 3>(T(0.)), ParticleEndIndex - ParticleStartIndex);
+				X_k.Init(TVector<T, 3>(T(0.)), ParticleEndIndex - ParticleStartIndex);
+			}
+			InitColor(InParticles);
+		}
+
+		virtual ~FGaussSeidelCorotatedConstraints() {}
+
+
+		void Apply(ParticleType& Particles, const T Dt) const 
+		{
+			//SCOPE_CYCLE_COUNTER(STAT_ChaosGaussSeidelCorotated);
+			TRACE_CPUPROFILER_EVENT_SCOPE(TEXT("STAT_ChaosGaussSeidelCorotatedApply"));
+			for (int32 i = 0; i < ParticlesPerColor.Num(); i++)
+			{
+				PhysicsParallelFor(ParticlesPerColor[i].Num(), [&](const int32 j)
+					{
+						int32 ParticleIndex = ParticlesPerColor[i][j];
+						Chaos::TVector<T, 3> Dx = ComputeDeltax(ParticleIndex, Particle2Incident[ParticleIndex - ParticleStartIndex], Particles,
+							Dt);
+						Particles.P(ParticleIndex) += Dx;
+					}, ParticlesPerColor[i].Num() < Base::CorotatedParams.XPBDCorotatedBatchThreshold);
+			}
+			ApplySOR(Particles, Dt);
+			CurrentIt += 1;
+		}
+
+		void Init(const T Dt, const ParticleType& Particles) const 
+		{
+			if (!bDoQuasistatics)
+			{
+				for (int32 i = 0; i < ParticleEndIndex - ParticleStartIndex; i++)
+				{
+					xtilde[i] = Particles.X(ParticleStartIndex + i) + Dt * Particles.V(ParticleStartIndex + i);
+				}
+			}
+			CurrentIt = 0;
+		}
+
+		void Init() const override {}
+
+
+		void ApplySOR(ParticleType& Particles, const T Dt) const
+		{
+			if (bDoSOR)
+			{
+				PhysicsParallelFor(ParticleEndIndex - ParticleStartIndex, [&](const int32 i)
+					{
+						int32 ParticleIndex = i + ParticleStartIndex;
+						if (Particles.InvM(ParticleIndex) != T(0) && CurrentIt > 3)
+						{
+							Particles.P(ParticleIndex) = OmegaSOR * (Particles.P(ParticleIndex) - X_k_1[i]) + X_k_1[i];
+						}
+						//X_k_1[i] = Particles.X(ParticleIndex);
+						X_k_1[i] = X_k[i];
+						X_k[i] = Particles.P(ParticleIndex);
+					}, ParticleEndIndex - ParticleStartIndex < 1000);
+			}
+		}
+
+		Chaos::TVector<T, 3> ComputePerParticleResidual(const int32 p, const int32 IncidentIndex, const ParticleType& Particles,
+			const T Dt, const bool AddMass = true) const
+		{
+
+			Chaos::TVector<T, 3> res(T(0));
+			if (AddMass) {
+				for (int32 alpha = 0; alpha < 3; alpha++) {
+					res[alpha] = Particles.M(p) * (Particles.P(p)[alpha] - xtilde[p - ParticleStartIndex][alpha]);
+				}
+			}
+			for (int32 j = 0; j < IncidentElements[IncidentIndex].Num(); j++) {
+				int32 local_index = IncidentElementsLocal[IncidentIndex][j];
+				int32 e = IncidentElements[IncidentIndex][j];
+				Chaos::PMatrix<T, 3, 3> DmInvT = Base::ElementDmInv(e).GetTransposed();
+				Chaos::PMatrix<T, 3, 3> Fe = Base::F(e, Particles);
+				Chaos::PMatrix<T, 3, 3> P((T)0.);
+
+				//corotated model may switch to others:
+				Chaos::Softs::PCorotated(Fe, MuElementArray[e], LambdaElementArray[e], P);
+
+				Chaos::PMatrix<T, 3, 3> force_term = -Measure[e] * DmInvT * P;
+				// last term: contribution from element e:
+				Chaos::TVector<T, 3> dx((T)0.);
+				if (local_index > 0)
+				{
+					for (int32 c = 0; c < 3; c++)
+					{
+						dx[c] += force_term.GetAt(c, local_index - 1);
+					}
+				}
+				else
+				{
+					for (int32 c = 0; c < 3; c++)
+					{
+						for (int32 h = 0; h < 3; h++)
+						{
+							dx[c] -= force_term.GetAt(c, h);
+						}
+					}
+				}
+
+				for (int32 ll = 0; ll < 3; ll++) {
+					dx[ll] *= Dt * Dt;
+				}
+
+				for (int32 alpha = 0; alpha < 3; alpha++) {
+					res[alpha] -= dx[alpha];
+				}
+			}
+			return res;
+		}
+
+
+		Chaos::PMatrix<T, 3, 3> ComputePerParticleCorotatedHessianSimple(const int32 p, const int32 IncidentIndex, const ParticleType& Particles,
+			const T Dt, const bool AddMass = true) const
+		{
+			Chaos::PMatrix<T, 3, 3> final_hessian = Chaos::PMatrix<T, 3, 3>::Zero;
+			if (AddMass) {
+				for (int32 alpha = 0; alpha < 3; alpha++) {
+					final_hessian.SetAt(alpha, alpha, Particles.M(p));
+				}
+			}
+
+			for (int32 j = 0; j < IncidentElements[IncidentIndex].Num(); j++) {
+				int32 local_index = IncidentElementsLocal[IncidentIndex][j];
+				int32 ElementIndex = IncidentElements[IncidentIndex][j];
+				Chaos::PMatrix<T, 3, 3> DmInv = Base::ElementDmInv(ElementIndex);
+				Chaos::PMatrix<T, 3, 3> Fe = Base::F(ElementIndex, Particles);
+
+				T Coeff = Dt * Dt * Measure[ElementIndex];
+
+				Chaos::PMatrix<T, 3, 3> JFinvT((T)0.);
+				JFinvT.SetAt(0, 0, Fe.GetAt(1, 1) * Fe.GetAt(2, 2) - Fe.GetAt(2, 1) * Fe.GetAt(1, 2));
+				JFinvT.SetAt(0, 1, Fe.GetAt(2, 0) * Fe.GetAt(1, 2) - Fe.GetAt(1, 0) * Fe.GetAt(2, 2));
+				JFinvT.SetAt(0, 2, Fe.GetAt(1, 0) * Fe.GetAt(2, 1) - Fe.GetAt(2, 0) * Fe.GetAt(1, 1));
+				JFinvT.SetAt(1, 0, Fe.GetAt(2, 1) * Fe.GetAt(0, 2) - Fe.GetAt(0, 1) * Fe.GetAt(2, 2));
+				JFinvT.SetAt(1, 1, Fe.GetAt(0, 0) * Fe.GetAt(2, 2) - Fe.GetAt(2, 0) * Fe.GetAt(0, 2));
+				JFinvT.SetAt(1, 2, Fe.GetAt(2, 0) * Fe.GetAt(0, 1) - Fe.GetAt(0, 0) * Fe.GetAt(2, 1));
+				JFinvT.SetAt(2, 0, Fe.GetAt(0, 1) * Fe.GetAt(1, 2) - Fe.GetAt(1, 1) * Fe.GetAt(0, 2));
+				JFinvT.SetAt(2, 1, Fe.GetAt(1, 0) * Fe.GetAt(0, 2) - Fe.GetAt(0, 0) * Fe.GetAt(1, 2));
+				JFinvT.SetAt(2, 2, Fe.GetAt(0, 0) * Fe.GetAt(1, 1) - Fe.GetAt(1, 0) * Fe.GetAt(0, 1));
+
+				Chaos::PMatrix<T, 3, 3> JFinv = JFinvT.GetTransposed();
+
+				if (local_index == 0) {
+					// first mu term:
+					T DmInvsum = T(0);
+					for (int32 nu = 0; nu < 3; nu++) {
+						T localDmsum = T(0);
+						for (int32 k = 0; k < 3; k++) {
+							localDmsum += DmInv.GetAt(k, nu);
+						}
+						DmInvsum += localDmsum * localDmsum;
+					}
+					for (int32 alpha = 0; alpha < 3; alpha++) {
+						final_hessian.SetAt(alpha, alpha, final_hessian.GetAt(alpha, alpha) + Coeff * T(2) * MuElementArray[ElementIndex] * DmInvsum);
+					}
+
+					// second lambda term:
+					Chaos::PMatrix<T, 3, 3> DmInvJFinv = JFinv * DmInv;
+					Chaos::TVector<T, 3> l((T)0.);
+					for (int32 alpha = 0; alpha < 3; alpha++) {
+						for (int32 k = 0; k < 3; k++) {
+							l[alpha] += DmInvJFinv.GetAt(k, alpha);
+						}
+					}
+					for (int32 alpha = 0; alpha < 3; alpha++)
+					{
+						final_hessian.SetRow(alpha, final_hessian.GetRow(alpha) + Coeff * LambdaElementArray[ElementIndex] * l[alpha] * l);
+					}
+
+				}
+				else {
+					// first mu term:
+					T DmInvsum = T(0);
+					for (int32 nu = 0; nu < 3; nu++)
+					{
+						DmInvsum += DmInv.GetAt(local_index - 1, nu) * DmInv.GetAt(local_index - 1, nu);
+					}
+					for (int32 alpha = 0; alpha < 3; alpha++)
+					{
+						final_hessian.SetAt(alpha, alpha, final_hessian.GetAt(alpha, alpha) + Coeff * T(2) * MuElementArray[ElementIndex] * DmInvsum);
+					}
+
+					// lambda term:
+					Chaos::PMatrix<T, 3, 3> DmInvJFinv = JFinv * DmInv;
+					Chaos::TVector<T, 3> l((T)0.);
+					for (int32 alpha = 0; alpha < 3; alpha++) {
+						l[alpha] = DmInvJFinv.GetAt(local_index - 1, alpha);
+					}
+					for (int32 alpha = 0; alpha < 3; alpha++)
+					{
+						final_hessian.SetRow(alpha, final_hessian.GetRow(alpha) + Coeff * LambdaElementArray[ElementIndex] * l[alpha] * l);
+					}
+				}
+			}
+			return final_hessian;
+		}
+
+	protected:
+		 void InitColor(const ParticleType& Particles)
+		{
+			ParticlesPerColor = ComputeNodalColoring(MeshConstraints, Particles, ParticleStartIndex, ParticleEndIndex, IncidentElements, IncidentElementsLocal);
+		}
+
+	private:
+
+		TVector<T, 3> ComputeDeltax(const int32 p, const int32 IncidentIndex, const ParticleType& Particles,
+			const T Dt) const
+		{
+			Chaos::TVector<T, 3> ParticleResidual = ComputePerParticleResidual(p, Particle2Incident[p], Particles, Dt, !bDoQuasistatics);
+
+			if (ParticleResidual.Size() > LocalNewtonTol)
+			{
+				Chaos::PMatrix<T, 3, 3> SimplifiedHessian = ComputePerParticleCorotatedHessianSimple(p, Particle2Incident[p], Particles, Dt, !bDoQuasistatics);
+				T HessianDet = SimplifiedHessian.Determinant();
+				if (HessianDet > UE_SMALL_NUMBER)
+				{
+					Chaos::PMatrix<T, 3, 3> HessianInv = SimplifiedHessian.SymmetricCofactorMatrix();
+					HessianInv *= T(1) / HessianDet;
+					return HessianInv.GetTransposed() * (-ParticleResidual);
+				}
+			}
+			return TVector<T, 3>((T)0.);
+		}
+
+		
+
+		TArray<int32> Particle2Incident;
+		TArray<TArray<int32>> IncidentElements;
+		TArray<TArray<int32>> IncidentElementsLocal;
+		T LocalNewtonTol = T(1e-5);
+		TArray<TArray<int32>> ParticlesPerColor;
+		int32 ParticleStartIndex;
+		int32 ParticleEndIndex;
+		mutable TArray<Chaos::TVector<T, 3>> xtilde;
+		mutable TArray<Chaos::TVector<T, 3>> X_k_1;
+		mutable TArray<Chaos::TVector<T, 3>> X_k;
+		bool bDoQuasistatics = false;
+		bool bDoSOR = true;
+		T OmegaSOR = T(1.6);
+		mutable int32 CurrentIt = 0;
+	};
+
+
+}// End namespace Chaos::Softs
diff --git a/Engine/Source/Runtime/Experimental/Chaos/Public/Chaos/NewtonCorotatedCache.h b/Engine/Source/Runtime/Experimental/Chaos/Public/Chaos/NewtonCorotatedCache.h
--- a/Engine/Source/Runtime/Experimental/Chaos/Public/Chaos/NewtonCorotatedCache.h
+++ b/Engine/Source/Runtime/Experimental/Chaos/Public/Chaos/NewtonCorotatedCache.h
@@ -27,4 +27,10 @@ struct CorotatedCache
 
 };
 
+template <typename T>
+T PsiCorotated(const Chaos::PMatrix<T, 3, 3>& F, const T mu, const T lambda);
+
+template <typename T>
+void PCorotated(const Chaos::PMatrix<T, 3, 3>& F, const T mu, const T lambda, Chaos::PMatrix<T, 3, 3>& P);
+
 }
diff --git a/Engine/Source/Runtime/Experimental/Chaos/Public/Chaos/XPBDCorotatedConstraints.h b/Engine/Source/Runtime/Experimental/Chaos/Public/Chaos/XPBDCorotatedConstraints.h
--- a/Engine/Source/Runtime/Experimental/Chaos/Public/Chaos/XPBDCorotatedConstraints.h
+++ b/Engine/Source/Runtime/Experimental/Chaos/Public/Chaos/XPBDCorotatedConstraints.h
@@ -117,7 +117,8 @@ namespace Chaos::Softs
 			TArray<T>&& AlphaJMeshArray,
 			const FDeformableXPBDCorotatedParams& InParams,
 			const T& NuMesh = (T).3,
-			const bool bRecordMetricIn = false
+			const bool bRecordMetricIn = false, 
+			const bool bDoColoring = true
 		)
 			: CorotatedParams(InParams), AlphaJArray(MoveTemp(AlphaJMeshArray)), bRecordMetric(bRecordMetricIn), MeshConstraints(InMesh)
 		{
@@ -162,8 +163,9 @@ namespace Chaos::Softs
 					Measure[e] = -Measure[e];
 				}
 			}
-
-			InitColor(InParticles);
+			if (bDoColoring) {
+				InitColor(InParticles);
+			}
 		}
 
 		FXPBDCorotatedConstraints(
@@ -392,7 +394,8 @@ namespace Chaos::Softs
 			return dC2;
 		}
 
-	private:
+
+	protected:
 
 		void InitColor(const ParticleType& Particles)
 		{
@@ -438,10 +441,6 @@ namespace Chaos::Softs
 			}
 		}
 
-	protected:
-
-		
-
 
 		virtual TVec4<TVector<T, 3>> GetDeterminantDelta(const ParticleType& Particles, const T Dt, const int32 ElementIndex, const T Tol = (T)1e-3) const
 		{
