modify primitive component wake/sleep events so that they don't depend on the existence of a BodyInstance which would allow it to work with geometry collections and cluster unions

diff --git a/Engine/Source/Runtime/Engine/Classes/Components/PrimitiveComponent.h b/Engine/Source/Runtime/Engine/Classes/Components/PrimitiveComponent.h
--- a/Engine/Source/Runtime/Engine/Classes/Components/PrimitiveComponent.h
+++ b/Engine/Source/Runtime/Engine/Classes/Components/PrimitiveComponent.h
@@ -2455,6 +2455,11 @@ public:
 
 	void DispatchWakeEvents(ESleepEvent WakeEvent, FName BoneName);
 
+	/**
+	 * Whether or not the primitive component should dispatch sleep/wake events.
+	 */
+	virtual bool ShouldDispatchWakeEvents(FName BoneName) const;
+
 	/**
 	 * Set collision params on OutParams (such as CollisionResponse) to match the settings on this PrimitiveComponent.
 	 */
diff --git a/Engine/Source/Runtime/Engine/Classes/PhysicsEngine/ClusterUnionComponent.h b/Engine/Source/Runtime/Engine/Classes/PhysicsEngine/ClusterUnionComponent.h
--- a/Engine/Source/Runtime/Engine/Classes/PhysicsEngine/ClusterUnionComponent.h
+++ b/Engine/Source/Runtime/Engine/Classes/PhysicsEngine/ClusterUnionComponent.h
@@ -272,6 +272,7 @@ public:
 	virtual bool SweepComponent(FHitResult& OutHit, const FVector Start, const FVector End, const FQuat& ShapeWorldRotation, const FPhysicsGeometry& Geometry, ECollisionChannel TraceChannel, const struct FCollisionQueryParams& Params, const struct FCollisionResponseParams& ResponseParams, const struct FCollisionObjectQueryParams& ObjectParams) override;
 	virtual bool OverlapComponentWithResult(const FVector& Pos, const FQuat& Rot, const FPhysicsGeometry& Geometry, ECollisionChannel TraceChannel, const struct FCollisionQueryParams& Params, const struct FCollisionResponseParams& ResponseParams, const struct FCollisionObjectQueryParams& ObjectParams, TArray<FOverlapResult>& OutOverlap) const override;
 	virtual bool ComponentOverlapComponentWithResultImpl(const class UPrimitiveComponent* const PrimComp, const FVector& Pos, const FQuat& Rot, const FCollisionQueryParams& Params, TArray<FOverlapResult>& OutOverlap) const override;
+	virtual bool ShouldDispatchWakeEvents(FName BoneName) const override { return true; }
 	//~ End UPrimitiveComponent Interface
 
 	//~ Begin USceneComponent Interface
diff --git a/Engine/Source/Runtime/Engine/Private/Components/PrimitiveComponent.cpp b/Engine/Source/Runtime/Engine/Private/Components/PrimitiveComponent.cpp
--- a/Engine/Source/Runtime/Engine/Private/Components/PrimitiveComponent.cpp
+++ b/Engine/Source/Runtime/Engine/Private/Components/PrimitiveComponent.cpp
@@ -2848,22 +2848,20 @@ void UPrimitiveComponent::DispatchBlockingHit(AActor& Owner, FHitResult const& B
 
 void UPrimitiveComponent::DispatchWakeEvents(ESleepEvent WakeEvent, FName BoneName)
 {
-	FBodyInstance* RootBI = GetBodyInstance(BoneName, false);
-	if(RootBI)
+	if (ShouldDispatchWakeEvents(BoneName))
 	{
-		if(RootBI->bGenerateWakeEvents)
+		if (WakeEvent == ESleepEvent::SET_Wakeup)
 		{
-			if (WakeEvent == ESleepEvent::SET_Wakeup)
-			{
-				OnComponentWake.Broadcast(this, BoneName);
-			}else
-			{
-				OnComponentSleep.Broadcast(this, BoneName);
-			}
+			OnComponentWake.Broadcast(this, BoneName);
+		}
+		else
+		{
+			OnComponentSleep.Broadcast(this, BoneName);
 		}
 	}
 	
 	//now update children that are welded
+	FBodyInstance* RootBI = GetBodyInstance(BoneName, false);
 	for(USceneComponent* SceneComp : GetAttachChildren())
 	{
 		if(UPrimitiveComponent* PrimComp = Cast<UPrimitiveComponent>(SceneComp))
@@ -2879,6 +2877,16 @@ void UPrimitiveComponent::DispatchWakeEvents(ESleepEvent WakeEvent, FName BoneNa
 	}
 }
 
+bool UPrimitiveComponent::ShouldDispatchWakeEvents(FName BoneName) const
+{
+	FBodyInstance* RootBI = GetBodyInstance(BoneName, false);
+	if (RootBI)
+	{
+		return RootBI->bGenerateWakeEvents;
+	}
+	return false;
+}
+
 void UPrimitiveComponent::GetNavigationData(FNavigationRelevantData& OutData) const
 {
 	if (bFillCollisionUnderneathForNavmesh)
@@ -3017,8 +3025,9 @@ bool UPrimitiveComponent::SweepComponent(FHitResult& OutHit, const FVector Start
 		QueryFilter.Word1 = 0xFFFFF;
 		ChaosInterface::SetFlags(BestHit, EHitFlags::Distance | EHitFlags::Normal | EHitFlags::Position | EHitFlags::FaceIndex);
 
-		ConvertQueryImpactHit(GetWorld(), BestHit, OutHit, (End - Start).Size(), QueryFilter, Start, End, &Geometry, FTransform{ Start }, false, false);
-		return true;
+		bool bHasHit = false;
+		ConvertTraceResults<ChaosInterface::FSweepHit>(bHasHit, GetWorld(), 1, &BestHit, (End - Start).Size(), QueryFilter, OutHit, Start, End, Geometry, FTransform{ ShapeWorldRotation, Start }, 0.f, Params.bReturnFaceIndex, Params.bReturnPhysicalMaterial);
+		return bHasHit;
 	}
 
 	return false;
diff --git a/Engine/Source/Runtime/Experimental/Chaos/Private/EventDefaults.cpp b/Engine/Source/Runtime/Experimental/Chaos/Private/EventDefaults.cpp
--- a/Engine/Source/Runtime/Experimental/Chaos/Private/EventDefaults.cpp
+++ b/Engine/Source/Runtime/Experimental/Chaos/Private/EventDefaults.cpp
@@ -463,33 +463,38 @@ namespace Chaos
 
 			Chaos::FPBDRigidsSolver* NonConstSolver = const_cast<Chaos::FPBDRigidsSolver*>(Solver);
 
-			NonConstSolver->Particles.GetDynamicParticles().GetSleepDataLock().ReadLock();
-			auto& SolverSleepingData = NonConstSolver->Particles.GetDynamicParticles().GetSleepData();
-			for(const TSleepData<FReal, 3>& SleepData : SolverSleepingData)
+			const TArray<FPBDRigidParticles*> RelevantParticleArrays = {
+				&NonConstSolver->Particles.GetDynamicParticles(),
+				&NonConstSolver->Particles.GetClusteredParticles(),
+				&NonConstSolver->Particles.GetGeometryCollectionParticles()
+			};
+
+			for (FPBDRigidParticles* ParticleArray : RelevantParticleArrays)
 			{
-				if(SleepData.Particle)
+				check(ParticleArray != nullptr);
+				ParticleArray->GetSleepDataLock().ReadLock();
+				const TArray<TSleepData<FReal, 3>>& SolverSleepingData = ParticleArray->GetSleepData();
+				for (const TSleepData<FReal, 3>& SleepData : SolverSleepingData)
 				{
-					FGeometryParticle* Particle = SleepData.Particle->GTGeometryParticle();
-					if (Particle && SleepData.Particle->PhysicsProxy())
+					if (SleepData.Particle)
 					{
-						int32 NewIdx = EventSleepDataArray.Add(FSleepingData());
-						FSleepingData& SleepingDataArrayItem = EventSleepDataArray[NewIdx];
-						SleepingDataArrayItem.Proxy = SleepData.Particle->PhysicsProxy();
-						SleepingDataArrayItem.Sleeping = SleepData.Sleeping;
+						FGeometryParticle* Particle = SleepData.Particle->GTGeometryParticle();
+						if (Particle && SleepData.Particle->PhysicsProxy())
+						{
+							int32 NewIdx = EventSleepDataArray.Add(FSleepingData());
+							FSleepingData& SleepingDataArrayItem = EventSleepDataArray[NewIdx];
+							SleepingDataArrayItem.Proxy = SleepData.Particle->PhysicsProxy();
+							SleepingDataArrayItem.Sleeping = SleepData.Sleeping;
+						}
 					}
 				}
+				ParticleArray->GetSleepDataLock().ReadUnlock();
+				ParticleArray->ClearSleepData();
 			}
-			NonConstSolver->Particles.GetDynamicParticles().GetSleepDataLock().ReadUnlock();
-
-			NonConstSolver->Particles.GetDynamicParticles().ClearSleepData();
 
 			// We don't care about sleep data added to these
 			NonConstSolver->Particles.GetDynamicKinematicParticles().ClearSleepData();
 			NonConstSolver->Particles.GetDynamicDisabledParticles().ClearSleepData();
-
-			// Sleep data is not supported for these particles yet, clear the buffers here so that we don't leak memory
-			NonConstSolver->Particles.GetClusteredParticles().ClearSleepData();
-			NonConstSolver->Particles.GetGeometryCollectionParticles().ClearSleepData(); 
 		});
 	}
 
diff --git a/Engine/Source/Runtime/Experimental/ChaosSolverEngine/Private/Chaos/ChaosGameplayEventDispatcher.cpp b/Engine/Source/Runtime/Experimental/ChaosSolverEngine/Private/Chaos/ChaosGameplayEventDispatcher.cpp
--- a/Engine/Source/Runtime/Experimental/ChaosSolverEngine/Private/Chaos/ChaosGameplayEventDispatcher.cpp
+++ b/Engine/Source/Runtime/Experimental/ChaosSolverEngine/Private/Chaos/ChaosGameplayEventDispatcher.cpp
@@ -299,20 +299,6 @@ void UChaosGameplayEventDispatcher::UnRegisterForCrumblingEvents(UPrimitiveCompo
 	}
 }
 
-void UChaosGameplayEventDispatcher::DispatchPendingWakeNotifies()
-{
-	for (auto MapItr = PendingSleepNotifies.CreateIterator(); MapItr; ++MapItr)
-	{
-		FBodyInstance* BodyInstance = MapItr.Key();
-		if (UPrimitiveComponent* PrimitiveComponent = BodyInstance->OwnerComponent.Get())
-		{
-			PrimitiveComponent->DispatchWakeEvents(MapItr.Value(), BodyInstance->BodySetup->BoneName);
-		}
-	}
-
-	PendingSleepNotifies.Empty();
-}
-
 void UChaosGameplayEventDispatcher::RegisterChaosEvents()
 {
 	if (FPhysScene* Scene = GetWorld()->GetPhysicsScene())
@@ -630,25 +616,21 @@ void UChaosGameplayEventDispatcher::HandleSleepingEvents(const Chaos::FSleepingE
 
 	for (const Chaos::FSleepingData& SleepData : SleepingArray)
 	{
-		if (SleepData.Proxy!= nullptr && Scene.GetOwningComponent<UPrimitiveComponent>(SleepData.Proxy) != nullptr)
+		ESleepEvent WakeSleepEvent = SleepData.Sleeping ? ESleepEvent::SET_Sleep : ESleepEvent::SET_Wakeup;
+		if (UPrimitiveComponent* PrimitiveComponent = Scene.GetOwningComponent<UPrimitiveComponent>(SleepData.Proxy))
 		{
+			FName BoneName = NAME_None;
 			if (FBodyInstance* BodyInstance = Scene.GetBodyInstanceFromProxy(SleepData.Proxy))
 			{
-				if (BodyInstance->bGenerateWakeEvents)
-				{
-					ESleepEvent WakeSleepEvent = SleepData.Sleeping ? ESleepEvent::SET_Sleep : ESleepEvent::SET_Wakeup;
-					AddPendingSleepingNotify(BodyInstance, WakeSleepEvent);
-				}
+				BoneName = BodyInstance->BodySetup->BoneName;
+			}
+
+			if (PrimitiveComponent->ShouldDispatchWakeEvents(BoneName))
+			{
+				PrimitiveComponent->DispatchWakeEvents(WakeSleepEvent, BoneName);
 			}
 		}
 	}
-
-	DispatchPendingWakeNotifies();
-}
-
-void UChaosGameplayEventDispatcher::AddPendingSleepingNotify(FBodyInstance* BodyInstance, ESleepEvent SleepEventType)
-{
-	PendingSleepNotifies.FindOrAdd(BodyInstance) = SleepEventType;
 }
 
 void UChaosGameplayEventDispatcher::HandleRemovalEvents(const Chaos::FRemovalEventData& Event)
diff --git a/Engine/Source/Runtime/Experimental/ChaosSolverEngine/Public/Chaos/ChaosGameplayEventDispatcher.h b/Engine/Source/Runtime/Experimental/ChaosSolverEngine/Public/Chaos/ChaosGameplayEventDispatcher.h
--- a/Engine/Source/Runtime/Experimental/ChaosSolverEngine/Public/Chaos/ChaosGameplayEventDispatcher.h
+++ b/Engine/Source/Runtime/Experimental/ChaosSolverEngine/Public/Chaos/ChaosGameplayEventDispatcher.h
@@ -236,9 +236,6 @@ private:
 	/** Holds the list of pending legacy notifies that are to be processed */
 	TArray<FCollisionNotifyInfo> PendingCollisionNotifies;
 
-	/** Holds the list of pending legacy sleep/wake notifies */
-	TMap<FBodyInstance*, ESleepEvent> PendingSleepNotifies;
-
 public:
 	/** 
 	 * Use to subscribe to collision events. 
@@ -278,7 +275,6 @@ private:
 	float LastCrumblingDataTime = -1.f;
 
 	void DispatchPendingCollisionNotifies();
-	void DispatchPendingWakeNotifies();
 
 	template <typename EventIterator>
 	void FillPhysicsProxy(FPhysScene_Chaos& Scene, TArray<UObject*>& Result, EventIterator& It);
@@ -292,7 +288,6 @@ private:
 	void HandleCollisionEvents(const Chaos::FCollisionEventData& CollisionData);
 	void HandleBreakingEvents(const Chaos::FBreakingEventData& BreakingData);
 	void HandleSleepingEvents(const Chaos::FSleepingEventData& SleepingData);
-	void AddPendingSleepingNotify(FBodyInstance* BodyInstance, ESleepEvent SleepEventType);
 	void HandleRemovalEvents(const Chaos::FRemovalEventData& RemovalData);
 	void HandleCrumblingEvents(const Chaos::FCrumblingEventData& CrumblingData);
 
