add DataFlow node for creating leaf convex hulls, with controls for whether to compute from geometry, convert from External/imported collision shapes, or take an intersection of both.
 + expose convex hull intersection algorithm in the convex hull utils

diff --git a/Engine/Plugins/Experimental/GeometryCollectionPlugin/Source/GeometryCollectionNodes/Private/Dataflow/GeometryCollectionUtilityNodes.cpp b/Engine/Plugins/Experimental/GeometryCollectionPlugin/Source/GeometryCollectionNodes/Private/Dataflow/GeometryCollectionUtilityNodes.cpp
--- a/Engine/Plugins/Experimental/GeometryCollectionPlugin/Source/GeometryCollectionNodes/Private/Dataflow/GeometryCollectionUtilityNodes.cpp
+++ b/Engine/Plugins/Experimental/GeometryCollectionPlugin/Source/GeometryCollectionNodes/Private/Dataflow/GeometryCollectionUtilityNodes.cpp
@@ -16,6 +16,7 @@ namespace Dataflow
 	{
 		static const FLinearColor CDefaultNodeBodyTintColor = FLinearColor(0.f, 0.f, 0.f, 0.5f);
 
+		DATAFLOW_NODE_REGISTER_CREATION_FACTORY(FCreateLeafConvexHullsDataflowNode);
 		DATAFLOW_NODE_REGISTER_CREATION_FACTORY(FCreateNonOverlappingConvexHullsDataflowNode);
 		DATAFLOW_NODE_REGISTER_CREATION_FACTORY(FGenerateClusterConvexHullsFromLeafHullsDataflowNode);
 		DATAFLOW_NODE_REGISTER_CREATION_FACTORY(FGenerateClusterConvexHullsFromChildrenHullsDataflowNode);
@@ -23,6 +24,47 @@ namespace Dataflow
 	}
 }
 
+FCreateLeafConvexHullsDataflowNode::FCreateLeafConvexHullsDataflowNode(const Dataflow::FNodeParameters& InParam, FGuid InGuid)
+	: FDataflowNode(InParam, InGuid)
+{
+	RegisterInputConnection(&Collection);
+	RegisterInputConnection(&OptionalSelectionFilter);
+	RegisterInputConnection(&SimplificationDistanceThreshold);
+	RegisterOutputConnection(&Collection);
+}
+
+void FCreateLeafConvexHullsDataflowNode::Evaluate(Dataflow::FContext& Context, const FDataflowOutput* Out) const
+{
+	if (Out->IsA(&Collection))
+	{
+		const FManagedArrayCollection& InCollection = GetValue(Context, &Collection);
+		if (InCollection.NumElements(FGeometryCollection::TransformGroup) == 0)
+		{
+			SetValue<FManagedArrayCollection>(Context, InCollection, &Collection);
+			return;
+		}
+
+		if (TUniquePtr<FGeometryCollection> GeomCollection = TUniquePtr<FGeometryCollection>(InCollection.NewCopy<FGeometryCollection>()))
+		{
+			TArray<int32> SelectedBones;
+			bool bRestrictToSelection = false;
+			if (IsConnected(&OptionalSelectionFilter))
+			{
+				const FDataflowTransformSelection& InOptionalSelectionFilter = GetValue<FDataflowTransformSelection>(Context, &OptionalSelectionFilter);
+				bRestrictToSelection = true;
+				SelectedBones = InOptionalSelectionFilter.AsArray();
+			}
+
+			float InSimplificationDistanceThreshold = GetValue(Context, &SimplificationDistanceThreshold);
+			FGeometryCollectionConvexUtility::FIntersectionFilters IntersectionFilters;
+			IntersectionFilters.OnlyIntersectIfComputedIsSmallerFactor = IntersectIfComputedIsSmallerByFactor;
+			IntersectionFilters.MinExternalVolumeToIntersect = MinExternalVolumeToIntersect;
+			FGeometryCollectionConvexUtility::GenerateLeafConvexHulls(*GeomCollection, bRestrictToSelection, SelectedBones, InSimplificationDistanceThreshold, GenerateMethod, IntersectionFilters);
+			SetValue(Context, (const FManagedArrayCollection&)(*GeomCollection), &Collection);
+		}
+	}
+}
+
 FCreateNonOverlappingConvexHullsDataflowNode::FCreateNonOverlappingConvexHullsDataflowNode(const Dataflow::FNodeParameters& InParam, FGuid InGuid)
 	: FDataflowNode(InParam, InGuid)
 {
@@ -87,7 +129,8 @@ void FGenerateClusterConvexHullsFromLeafHullsDataflowNode::Evaluate(Dataflow::FC
 					*GeomCollection,
 					InConvexCount,
 					InErrorToleranceInCm,
-					InOptionalSelectionFilter.AsArray()
+					InOptionalSelectionFilter.AsArray(),
+					bPreferExternalCollisionShapes
 				);
 			}
 			else
@@ -95,7 +138,8 @@ void FGenerateClusterConvexHullsFromLeafHullsDataflowNode::Evaluate(Dataflow::FC
 				FGeometryCollectionConvexUtility::GenerateClusterConvexHullsFromLeafHulls(
 					*GeomCollection,
 					InConvexCount,
-					InErrorToleranceInCm
+					InErrorToleranceInCm,
+					bPreferExternalCollisionShapes
 				);
 			}
 
@@ -132,7 +176,8 @@ void FGenerateClusterConvexHullsFromChildrenHullsDataflowNode::Evaluate(Dataflow
 					*GeomCollection,
 					InConvexCount,
 					InErrorToleranceInCm,
-					InOptionalSelectionFilter.AsArray()
+					InOptionalSelectionFilter.AsArray(),
+					bPreferExternalCollisionShapes
 				);
 			}
 			else
@@ -141,7 +186,7 @@ void FGenerateClusterConvexHullsFromChildrenHullsDataflowNode::Evaluate(Dataflow
 					*GeomCollection,
 					InConvexCount,
 					InErrorToleranceInCm,
-					TArrayView<const int32>()
+					bPreferExternalCollisionShapes
 				);
 			}
 
diff --git a/Engine/Plugins/Experimental/GeometryCollectionPlugin/Source/GeometryCollectionNodes/Public/Dataflow/GeometryCollectionUtilityNodes.h b/Engine/Plugins/Experimental/GeometryCollectionPlugin/Source/GeometryCollectionNodes/Public/Dataflow/GeometryCollectionUtilityNodes.h
--- a/Engine/Plugins/Experimental/GeometryCollectionPlugin/Source/GeometryCollectionNodes/Public/Dataflow/GeometryCollectionUtilityNodes.h
+++ b/Engine/Plugins/Experimental/GeometryCollectionPlugin/Source/GeometryCollectionNodes/Public/Dataflow/GeometryCollectionUtilityNodes.h
@@ -6,6 +6,7 @@
 #include "Dataflow/DataflowEngine.h"
 #include "Dataflow/DataflowSelection.h"
 #include "GeometryCollection/ManagedArrayCollection.h"
+#include "GeometryCollection/GeometryCollectionConvexUtility.h"
 
 #include "GeometryCollectionUtilityNodes.generated.h"
 
@@ -22,6 +23,42 @@ enum class EConvexOverlapRemovalMethodEnum : uint8
 	Dataflow_Max                UMETA(Hidden)
 };
 
+
+USTRUCT(meta = (DataflowGeometryCollection))
+struct FCreateLeafConvexHullsDataflowNode : public FDataflowNode
+{
+	GENERATED_USTRUCT_BODY()
+	DATAFLOW_NODE_DEFINE_INTERNAL(FCreateLeafConvexHullsDataflowNode, "CreateLeafConvexHulls", "GeometryCollection|Utilities", "")
+
+public:
+	UPROPERTY(meta = (DataflowInput, DataflowOutput))
+	FManagedArrayCollection Collection;
+
+	/** Optional transform selection to compute leaf hulls on -- if not provided, all leaf hulls will be computed. */
+	UPROPERTY(meta = (DataflowInput, DataflowIntrinsic))
+	FDataflowTransformSelection OptionalSelectionFilter;
+
+	/** How convex hulls are generated -- computed from geometry, imported from external collision shapes, or an intersection of both options. */
+	UPROPERTY(EditAnywhere, Category = Options)
+	EGenerateConvexMethod GenerateMethod = EGenerateConvexMethod::ExternalCollision;
+
+	/** If GenerateMethod is Intersect, only actually intersect when the volume of the Computed Hull is less than this fraction of the volume of the External Hull(s) */
+	UPROPERTY(EditAnywhere, Category = IntersectionFilters, meta = (ClampMin = 0.0, ClampMax = 1.0, EditCondition = "GenerateMethod == EGenerateConvexMethod::IntersectExternalWithComputed"))
+	float IntersectIfComputedIsSmallerByFactor = 1.0f;
+
+	/** If GenerateMethod is Intersect, only actually intersect if the volume of the External Hull(s) exceed this threshold */
+	UPROPERTY(EditAnywhere, Category = IntersectionFilters, meta = (ClampMin = 0.0, EditCondition = "GenerateMethod == EGenerateConvexMethod::IntersectExternalWithComputed"))
+	float MinExternalVolumeToIntersect = 0.0f;
+
+	/** Computed convex hulls are simplified to keep points spaced at least this far apart (except where needed to keep the hull from collapsing to zero volume) */
+	UPROPERTY(EditAnywhere, Category = "Convex", meta = (DataflowInput, UIMin = 0.f, EditCondition = "GenerateMethod != EGenerateConvexMethod::ExternalCollision"))
+	float SimplificationDistanceThreshold = 10.f;
+
+	FCreateLeafConvexHullsDataflowNode(const Dataflow::FNodeParameters& InParam, FGuid InGuid = FGuid::NewGuid());
+
+	virtual void Evaluate(Dataflow::FContext& Context, const FDataflowOutput* Out) const override;
+};
+
 /**
  *
  * Generates convex hull representation for the bones for simulation
@@ -37,18 +74,23 @@ public:
 	UPROPERTY(meta = (DataflowInput, DataflowOutput))
 	FManagedArrayCollection Collection;
 
+	/** Fraction of the convex hulls for a cluster that we can remove before using the hulls of the children */
 	UPROPERTY(EditAnywhere, Category = "Convex", meta = (DataflowInput, UIMin = 0.01f, UIMax = 1.f))
 	float CanRemoveFraction = 0.3f;
 
+	/** Computed convex hulls are simplified to keep points spaced at least this far apart (except where needed to keep the hull from collapsing to zero volume) */
 	UPROPERTY(EditAnywhere, Category = "Convex", meta = (DataflowInput, UIMin = 0.f))
-	float SimplificationDistanceThreshold = 0.f;
+	float SimplificationDistanceThreshold = 10.f;
 
+	/** Fraction (of geometry volume) by which a cluster's convex hull volume can exceed the actual geometry volume before instead using the hulls of the children.  0 means the convex volume cannot exceed the geometry volume; 1 means the convex volume is allowed to be 100% larger (2x) the geometry volume. */
 	UPROPERTY(EditAnywhere, Category = "Convex", meta = (DataflowInput, UIMin = 0.f))
 	float CanExceedFraction = 0.5f;
 
+	/** Whether and in what cases to automatically cut away overlapping parts of the convex hulls, to avoid the simulation 'popping' to fix the overlaps */
 	UPROPERTY(EditAnywhere, Category = "Convex")
 	EConvexOverlapRemovalMethodEnum OverlapRemovalMethod = EConvexOverlapRemovalMethodEnum::Dataflow_EConvexOverlapRemovalMethod_All;
 
+	/** Overlap removal will be computed as if convex hulls were this percentage smaller (in range 0-100) */
 	UPROPERTY(EditAnywhere, Category = "Convex", meta = (DataflowInput, UIMin = 0.f, UIMax = 100.f))
 	float OverlapRemovalShrinkPercent = 0.f;
 
@@ -84,6 +126,10 @@ public:
 	*/
 	UPROPERTY(EditAnywhere, Category = "Convex", meta = (DataflowInput, UIMin = "0", UIMax = "100.", Units = cm))
 	double ErrorTolerance = 0.0;
+	
+	/** Whether to prefer available External (imported) collision shapes instead of the computed convex hulls on the Collection */
+	UPROPERTY(EditAnywhere, Category = "Convex")
+	bool bPreferExternalCollisionShapes = true;
 
 	/** Optional transform selection to compute cluster hulls on -- if not provided, all cluster hulls will be computed. */
 	UPROPERTY(meta = (DataflowInput, DataflowIntrinsic))
@@ -104,7 +150,7 @@ USTRUCT(meta = (DataflowGeometryCollection))
 struct FGenerateClusterConvexHullsFromChildrenHullsDataflowNode : public FDataflowNode
 {
 	GENERATED_USTRUCT_BODY()
-	DATAFLOW_NODE_DEFINE_INTERNAL(FGenerateClusterConvexHullsFromChildrenHullsDataflowNode, "GenerateClusterConvexHullsFromChildrenHullsDataflowNode", "GeometryCollection|Utilities", "")
+	DATAFLOW_NODE_DEFINE_INTERNAL(FGenerateClusterConvexHullsFromChildrenHullsDataflowNode, "GenerateClusterConvexHullsFromChildrenHulls", "GeometryCollection|Utilities", "")
 
 public:
 	UPROPERTY(meta = (DataflowInput, DataflowOutput))
@@ -121,6 +167,10 @@ public:
 	*/
 	UPROPERTY(EditAnywhere, Category = "Convex", meta = (DataflowInput, UIMin = "0", UIMax = "100.", Units = cm))
 	double ErrorTolerance = 0.0;
+	
+	/** Whether to prefer available External (imported) collision shapes instead of the computed convex hulls on the Collection */
+	UPROPERTY(EditAnywhere, Category = "Convex")
+	bool bPreferExternalCollisionShapes = true;
 
 	/** Optional transform selection to compute cluster hulls on -- if not provided, all cluster hulls will be computed. */
 	UPROPERTY(meta = (DataflowInput, DataflowIntrinsic))
diff --git a/Engine/Source/Runtime/Experimental/Chaos/Private/GeometryCollection/GeometryCollectionConvexUtility.cpp b/Engine/Source/Runtime/Experimental/Chaos/Private/GeometryCollection/GeometryCollectionConvexUtility.cpp
--- a/Engine/Source/Runtime/Experimental/Chaos/Private/GeometryCollection/GeometryCollectionConvexUtility.cpp
+++ b/Engine/Source/Runtime/Experimental/Chaos/Private/GeometryCollection/GeometryCollectionConvexUtility.cpp
@@ -1172,7 +1172,8 @@ void HullsFromGeometry(
 	const TManagedArray<int32>& SimulationType,
 	int32 RigidType,
 	double SimplificationDistanceThreshold,
-	double OverlapRemovalShrinkPercent
+	double OverlapRemovalShrinkPercent,
+	TFunction<bool(int32)> SkipBoneFn = nullptr
 )
 {
 	TArray<FVector> GlobalVertices;
@@ -1192,6 +1193,11 @@ void HullsFromGeometry(
 	FCriticalSection HullCS;
 	ParallelFor(NumBones, [&](int32 Idx)
 	{
+		if (SkipBoneFn && !SkipBoneFn(Idx))
+		{
+			return;
+		}
+
 		int32 GeomIdx = Geometry.TransformToGeometryIndex[Idx];
 		if (OrigConvexData.IsSet() && HasCustomConvexFn(Idx))
 		{
@@ -1250,7 +1256,7 @@ void TransformHullsToLocal(
 	double OverlapRemovalShrinkPercent
 )
 {
-	checkSlow(Convexes.Num() == ConvexPivots.Num());
+	checkSlow((OverlapRemovalShrinkPercent == 0.0  && ConvexPivots.Num() == 0) || Convexes.Num() == ConvexPivots.Num());
 	
 	double ScaleFactor = 1.0 - OverlapRemovalShrinkPercent / 100.0;
 	double InvScaleFactor = 1.0;
@@ -1384,7 +1390,7 @@ void FGeometryCollectionConvexUtility::CreateConvexHullAttributesIfNeeded(FManag
 }
 
 UE::GeometryCollectionConvexUtility::FConvexHulls
-FGeometryCollectionConvexUtility::ComputeLeafHulls(FGeometryCollection* GeometryCollection, const TArray<FTransform>& GlobalTransformArray, double SimplificationDistanceThreshold, double OverlapRemovalShrinkPercent)
+FGeometryCollectionConvexUtility::ComputeLeafHulls(FGeometryCollection* GeometryCollection, const TArray<FTransform>& GlobalTransformArray, double SimplificationDistanceThreshold, double OverlapRemovalShrinkPercent, TFunction<bool(int32)> SkipBoneFn)
 {
 	check(GeometryCollection);
 
@@ -1401,7 +1407,7 @@ FGeometryCollectionConvexUtility::ComputeLeafHulls(FGeometryCollection* Geometry
 	TFunctionRef<bool(int32)> HasCustomConvexFn = (CustomConvexFlags != nullptr) ? (TFunctionRef<bool(int32)>)ConvexFlagsFromArray : (TFunctionRef<bool(int32)>)ConvexFlagsAlwaysFalse;
 
 	HullsFromGeometry(*GeometryCollection, GlobalTransformArray, HasCustomConvexFn, Hulls.Hulls, Hulls.Pivots, Hulls.TransformToHullsIndices, GeometryCollection->SimulationType,
-		FGeometryCollection::ESimulationTypes::FST_Rigid, SimplificationDistanceThreshold, Hulls.OverlapRemovalShrinkPercent);
+		FGeometryCollection::ESimulationTypes::FST_Rigid, SimplificationDistanceThreshold, Hulls.OverlapRemovalShrinkPercent, SkipBoneFn);
 
 	return Hulls;
 }
@@ -1457,21 +1463,206 @@ FGeometryCollectionConvexUtility::FGeometryCollectionConvexData FGeometryCollect
 	return { TransformToConvexIndices, ConvexHull };
 }
 
+void FGeometryCollectionConvexUtility::GenerateLeafConvexHulls(FGeometryCollection& Collection, bool bRestrictToSelection, const TArrayView<const int32> TransformSubset, double SimplificationDistanceThreshold, EGenerateConvexMethod GenerateMethod, FIntersectionFilters IntersectFilters)
+{
+	using FSharedImplicit = TSharedPtr<Chaos::FImplicitObject, ESPMode::ThreadSafe>;
+
+	const TManagedArrayAccessor<FSharedImplicit> ExternalCollisionAttribute(Collection, "ExternalCollisions", FGeometryCollection::TransformGroup);
+	bool bUseExternal = ExternalCollisionAttribute.IsValid() && (GenerateMethod == EGenerateConvexMethod::ExternalCollision || GenerateMethod == EGenerateConvexMethod::IntersectExternalWithComputed);
+	bool bUseGenerated = !bUseExternal || GenerateMethod == EGenerateConvexMethod::IntersectExternalWithComputed;
+	bool bUseIntersect = bUseExternal && GenerateMethod == EGenerateConvexMethod::IntersectExternalWithComputed;
+
+	int32 NumTransforms = Collection.NumElements(FGeometryCollection::TransformGroup);
+	
+	UE::GeometryCollectionConvexUtility::FConvexHulls ComputedHulls;
+	TArray<FTransformedConvex> ExternalHulls;
+	TArray<TSet<int32>> TransformToExternalHullsIndices;
+	TransformToExternalHullsIndices.SetNum(NumTransforms);
+
+	CreateConvexHullAttributesIfNeeded(Collection);
+
+	TManagedArray<TUniquePtr<Chaos::FConvex>>& ConvexHullAttrib = Collection.ModifyAttribute<TUniquePtr<Chaos::FConvex>>("ConvexHull", "Convex");
+	TManagedArray<TSet<int32>>& TransformToConvexIndicesAttrib = Collection.ModifyAttribute<TSet<int32>>("TransformToConvexIndices", FTransformCollection::TransformGroup);
+	TransformToConvexIndicesAttrib.Fill(TSet<int32>());
+
+	auto ComputeTransformedHull = [](const Chaos::FConvex& HullIn, const FTransform& TransformIn) -> TUniquePtr<Chaos::FConvex>
+	{
+		FTransform3f Transform = (FTransform3f)TransformIn;
+		TArray<Chaos::FConvex::FVec3Type> HullPts;
+		for (const Chaos::FConvex::FVec3Type& P : HullIn.GetVertices())
+		{
+			HullPts.Add(Transform.TransformPosition(P));
+		}
+		return MakeUnique<Chaos::FConvex>(HullPts, UE_KINDA_SMALL_NUMBER);
+	};
+
+	TArray<int32> LocalTransformIndices;
+	if (!bRestrictToSelection)
+	{
+		LocalTransformIndices.Reserve(NumTransforms);
+		for (int32 Idx = 0; Idx < NumTransforms; ++Idx)
+		{
+			if (Collection.SimulationType[Idx] == FGeometryCollection::ESimulationTypes::FST_Rigid)
+			{
+				LocalTransformIndices.Add(Idx);
+			}
+		}
+	}
+	const TArrayView<const int32> UseTransforms = bRestrictToSelection ? TransformSubset : TArrayView<const int32>(LocalTransformIndices);
+	if (UseTransforms.IsEmpty())
+	{
+		return;
+	}
+
+	if (bUseExternal)
+	{
+		for (int32 SourceTransformIdx : UseTransforms)
+		{
+			if (Collection.SimulationType[SourceTransformIdx] != FGeometryCollection::ESimulationTypes::FST_Rigid)
+			{
+				continue;
+			}
+
+			// convert the external collisions to convex hulls
+			const FSharedImplicit ExternalCollision = ExternalCollisionAttribute[SourceTransformIdx];
+
+			const int32 ExternalHullsStart = ExternalHulls.Num();
+			ConvertImplicitToConvexArray(*ExternalCollision, FTransform::Identity, ExternalHulls);
+			for (int32 HullIndex = ExternalHullsStart; HullIndex < ExternalHulls.Num(); HullIndex++)
+			{
+				TransformToExternalHullsIndices[SourceTransformIdx].Add(HullIndex);
+			}
+		}
+
+		if (!bUseIntersect)
+		{
+			TArray<TUniquePtr<Chaos::FConvex>> FinalHulls;
+			FinalHulls.SetNum(ExternalHulls.Num());
+
+			// Transform the hulls to the local space of the transform
+			ParallelFor(ExternalHulls.Num(), [&](int32 HullIdx)
+			{
+				FinalHulls[HullIdx] = ComputeTransformedHull(*ExternalHulls[HullIdx].Convex, ExternalHulls[HullIdx].Transform);
+			});
+
+			Collection.EmptyGroup("Convex");
+			Collection.Resize(FinalHulls.Num(), "Convex");
+			ConvexHullAttrib = MoveTemp(FinalHulls);
+			TransformToConvexIndicesAttrib = MoveTemp(TransformToExternalHullsIndices);
+		}
+	}
+	if (bUseGenerated)
+	{
+		// Identity Transform Array allows us to leave leaf hulls in their original local coordinate spaces
+		TArray<FTransform> IdentityTransformArray;
+		IdentityTransformArray.Init(FTransform::Identity, NumTransforms);
+		TFunction<bool(int32)> SkipBoneFn = nullptr;
+		TSet<int32> SelectionSet;
+		if (bRestrictToSelection)
+		{
+			SelectionSet.Append(UseTransforms);
+			SkipBoneFn = [SelectionSet](int32 BoneIdx) -> bool
+			{
+				return SelectionSet.Contains(BoneIdx);
+			};
+		}
+		ComputedHulls = ComputeLeafHulls(&Collection, IdentityTransformArray, SimplificationDistanceThreshold, 0, SkipBoneFn);
+
+		if (!bUseIntersect)
+		{
+			Collection.EmptyGroup("Convex");
+			Collection.Resize(ComputedHulls.Hulls.Num(), "Convex");
+			ConvexHullAttrib = MoveTemp(ComputedHulls.Hulls);
+			TransformToConvexIndicesAttrib = MoveTemp(ComputedHulls.TransformToHullsIndices);
+		}
+	}
+	if (bUseIntersect)
+	{
+		TArray<TUniquePtr<Chaos::FConvex>> FinalHulls;
+		for (int32 SourceTransformIdx : UseTransforms)
+		{
+			if (Collection.SimulationType[SourceTransformIdx] != FGeometryCollection::ESimulationTypes::FST_Rigid)
+			{
+				continue;
+			}
+
+			TSet<int32>& GeoHulls = ComputedHulls.TransformToHullsIndices[SourceTransformIdx];
+
+			if (IntersectFilters.OnlyIntersectIfComputedIsSmallerFactor < 1.0 || IntersectFilters.MinExternalVolumeToIntersect > 0.0)
+			{
+				double ComputedVolSum = 0.0;
+				for (int32 GeoHullIdx : GeoHulls)
+				{
+					ComputedVolSum += (double)ComputedHulls.Hulls[GeoHullIdx]->GetVolume();
+				}
+				double ExternalVolSum = 0.0;
+				for (int32 ExtHull : TransformToExternalHullsIndices[SourceTransformIdx])
+				{
+					ExternalVolSum += (double)ExternalHulls[ExtHull].Convex->GetVolume();
+				}
+				if (
+					(IntersectFilters.OnlyIntersectIfComputedIsSmallerFactor < 1.0 && ComputedVolSum >= ExternalVolSum * IntersectFilters.OnlyIntersectIfComputedIsSmallerFactor) ||
+					(ExternalVolSum < IntersectFilters.MinExternalVolumeToIntersect)
+					)
+				{
+					// Filters allow us to skip computing the intersection and just use the external hulls
+					for (int32 ExtHull : TransformToExternalHullsIndices[SourceTransformIdx])
+					{
+						int32 HullIdx = FinalHulls.Add(ComputeTransformedHull(*ExternalHulls[ExtHull].Convex, ExternalHulls[ExtHull].Transform));
+						TransformToConvexIndicesAttrib[SourceTransformIdx].Add(HullIdx);
+					}
+					continue;
+				}
+			}
+
+			// Note: Currently computed hulls path only generates one hull per transform. If this changes, and we have multiple
+			// compute hulls AND multiple external hulls, consider running an additional merge step after intersection
+			// to avoid multiplying the number of hulls used.
+			ensure(GeoHulls.Num() < 2);
+			for (int32 GeoHullIdx : GeoHulls)
+			{
+				if (TransformToExternalHullsIndices.IsEmpty())
+				{
+					TransformToConvexIndicesAttrib[SourceTransformIdx].Add(FinalHulls.Add(MoveTemp(ComputedHulls.Hulls[GeoHullIdx])));
 
-void FGeometryCollectionConvexUtility::GenerateClusterConvexHullsFromChildrenHulls(FGeometryCollection& Collection, int32 ConvexCount, double ErrorToleranceInCm, const TArrayView<const int32> TransformSubset)
+					continue;
+				}
+				for (int32 ExtHull : TransformToExternalHullsIndices[SourceTransformIdx])
+				{
+					int32 HullIdx = FinalHulls.Add(MakeUnique<Chaos::FConvex>());
+					UE::GeometryCollectionConvexUtility::IntersectConvexHulls(FinalHulls[HullIdx].Get(),
+						ComputedHulls.Hulls[GeoHullIdx].Get(), 0.0f, ExternalHulls[ExtHull].Convex.Get(),
+						nullptr, &ExternalHulls[ExtHull].Transform, &ExternalHulls[ExtHull].Transform);
+					TransformToConvexIndicesAttrib[SourceTransformIdx].Add(HullIdx);
+				}
+			}
+		}
+		Collection.EmptyGroup("Convex");
+		Collection.Resize(FinalHulls.Num(), "Convex");
+		ConvexHullAttrib = MoveTemp(FinalHulls);
+	}
+
+	RemoveEmptyConvexHulls(Collection);
+}
+
+void FGeometryCollectionConvexUtility::GenerateClusterConvexHullsFromChildrenHulls(FGeometryCollection& Collection, int32 ConvexCount, double ErrorToleranceInCm, const TArrayView<const int32> TransformSubset, bool bUseExternalCollisionIfAvailable)
+{
+	GenerateClusterConvexHullsFromLeafOrChildrenHullsInternal(Collection, ConvexCount, ErrorToleranceInCm, true, true/*bUseDirectChildren*/, TransformSubset, bUseExternalCollisionIfAvailable);
+}
+void FGeometryCollectionConvexUtility::GenerateClusterConvexHullsFromChildrenHulls(FGeometryCollection& Collection, int32 ConvexCount, double ErrorToleranceInCm, bool bUseExternalCollisionIfAvailable)
 {
-	GenerateClusterConvexHullsFromLeafOrChildrenHullsInternal(Collection, ConvexCount, ErrorToleranceInCm, true, true/*bUseDirectChildren*/, TransformSubset);
+	GenerateClusterConvexHullsFromLeafOrChildrenHullsInternal(Collection, ConvexCount, ErrorToleranceInCm, false, true/*bUseDirectChildren*/, TArrayView<const int32>(), bUseExternalCollisionIfAvailable);
 }
-void FGeometryCollectionConvexUtility::GenerateClusterConvexHullsFromLeafHulls(FGeometryCollection& Collection, int32 ConvexCount, double ErrorToleranceInCm, const TArrayView<const int32> TransformSubset)
+void FGeometryCollectionConvexUtility::GenerateClusterConvexHullsFromLeafHulls(FGeometryCollection& Collection, int32 ConvexCount, double ErrorToleranceInCm, const TArrayView<const int32> TransformSubset, bool bUseExternalCollisionIfAvailable)
 {
-	GenerateClusterConvexHullsFromLeafOrChildrenHullsInternal(Collection, ConvexCount, ErrorToleranceInCm, true, false/*bUseDirectChildren*/, TransformSubset);
+	GenerateClusterConvexHullsFromLeafOrChildrenHullsInternal(Collection, ConvexCount, ErrorToleranceInCm, true, false/*bUseDirectChildren*/, TransformSubset, bUseExternalCollisionIfAvailable);
 }
-void FGeometryCollectionConvexUtility::GenerateClusterConvexHullsFromLeafHulls(FGeometryCollection& Collection, int32 ConvexCount, double ErrorToleranceInCm)
+void FGeometryCollectionConvexUtility::GenerateClusterConvexHullsFromLeafHulls(FGeometryCollection& Collection, int32 ConvexCount, double ErrorToleranceInCm, bool bUseExternalCollisionIfAvailable)
 {
-	GenerateClusterConvexHullsFromLeafOrChildrenHullsInternal(Collection, ConvexCount, ErrorToleranceInCm, false, false/*bUseDirectChildren*/, TArrayView<const int32>());
+	GenerateClusterConvexHullsFromLeafOrChildrenHullsInternal(Collection, ConvexCount, ErrorToleranceInCm, false, false/*bUseDirectChildren*/, TArrayView<const int32>(), bUseExternalCollisionIfAvailable);
 }
 
-void FGeometryCollectionConvexUtility::GenerateClusterConvexHullsFromLeafOrChildrenHullsInternal(FGeometryCollection& Collection,	int32 ConvexCount, double ErrorToleranceInCm, bool bOnlySubset, bool bUseDirectChildren, const TArrayView<const int32> TransformSubset)
+void FGeometryCollectionConvexUtility::GenerateClusterConvexHullsFromLeafOrChildrenHullsInternal(FGeometryCollection& Collection, int32 ConvexCount, double ErrorToleranceInCm, bool bOnlySubset, bool bUseDirectChildren, const TArrayView<const int32> TransformSubset, bool bUseExternalCollisionIfAvailable)
 {
 	static FName ConvexGroupName("Convex");
 	static FName ConvexHullAttributeName("ConvexHull");
@@ -1563,7 +1754,7 @@ void FGeometryCollectionConvexUtility::GenerateClusterConvexHullsFromLeafOrChild
 					FTransform ChildToParentTransform = InnerTransform.GetRelativeTransform(ParentTransform);
 
 					const FSharedImplicit ExternalCollision = (ExternalCollisionAttribute.IsValid()) ? ExternalCollisionAttribute.Get()[SourceTransformIndex] : FSharedImplicit();
-					if (ExternalCollision)
+					if (ExternalCollision && bUseExternalCollisionIfAvailable)
 					{
 						const int32 ExternalHullsStart = ExternalHulls.Num();
 						ConvertImplicitToConvexArray(*ExternalCollision, FTransform::Identity, ExternalHulls);
@@ -2262,3 +2453,512 @@ void FGeometryCollectionConvexUtility::ConvertInstancedImplicitToConvexArray(
 	}
 }
 
+
+// local (static) helpers for the below Convex Utility code
+namespace
+{
+	// Helpful struct to represent a convex hull with a local editable representation, to perform intersections / clipping on it
+	struct FHullPolygons
+	{
+		// simple packed representation for convex hull faces, where each polygon's indices are listed sequentially,
+		// and negative values indicate the number of vertices in the next polygon. If no negative value is listed, polygon is a triangle.
+		TArray<int32> PackedPolygons;
+
+		// Copy of hull vertices, to be refined through plane cuts
+		TArray<Chaos::FVec3f> Vertices;
+		Chaos::FAABB3f Bounds;
+
+		void Reset()
+		{
+			Vertices.Reset();
+			PackedPolygons.Reset();
+			Bounds = Chaos::FAABB3f::EmptyAABB();
+		}
+
+		FHullPolygons() = default;
+
+		FHullPolygons(const Chaos::FConvex& HullIn)
+		{
+			Init(HullIn);
+		}
+
+		void Init(const Chaos::FConvex& HullIn)
+		{
+			Reset();
+			Vertices = HullIn.GetVertices();
+			const Chaos::FConvexStructureData& HullData = HullIn.GetStructureData();
+			int32 NumPlanes = HullIn.NumPlanes();
+			PackedPolygons.Reserve(NumPlanes * 3);
+			for (int PlaneIdx = 0; PlaneIdx < NumPlanes; PlaneIdx++)
+			{
+				int32 NumPlaneVerts = HullData.NumPlaneVertices(PlaneIdx);
+				if (NumPlaneVerts > 3)
+				{
+					PackedPolygons.Add(-NumPlaneVerts);
+				}
+				for (int32 PlaneVertexIdx = 0; PlaneVertexIdx < NumPlaneVerts; PlaneVertexIdx++)
+				{
+					PackedPolygons.Add(HullData.GetPlaneVertex(PlaneIdx, PlaneVertexIdx));
+				}
+			}
+			Bounds = HullIn.GetLocalBoundingBox();
+		}
+
+		void Intersect(const Chaos::FConvex& OtherHull, float ExpandAmount)
+		{
+			// Arrays to store intermediate plane cut data
+			TArray<int32> NewPolygons;
+			TMap<FIntVector2, int> NewVertices; // mapping from edges to new vertices
+			TArray<float> SignedDist; // signed distance from vertices to a cutting plane
+			TArray<int32> VertexRemap;
+			TMap<int32, int32> OpenEdgeVertMap;
+
+			// Cut by the plane through PlanePt with PlaneNormal. Note: Must pre-apply ExpandAmount; we do not assume the plane is shifted here.
+			// Uses the intermediate data above as temp storage
+			auto PlaneCut = [&NewPolygons, &NewVertices, &SignedDist, &VertexRemap, &OpenEdgeVertMap, this](Chaos::FVec3f PlanePt, Chaos::FVec3f PlaneNormal)
+			{
+				NewPolygons.Reset(PackedPolygons.Num());
+				NewVertices.Reset();
+				OpenEdgeVertMap.Reset();
+				SignedDist.SetNum(Vertices.Num(), false);
+				VertexRemap.SetNum(Vertices.Num(), false);
+				int32 OpenEdgeStart = -1;
+
+				// Possible optimization: if many vertices, check plane vs AABB corners first?
+
+				int32 OutTotal = 0;
+				for (int32 VertIdx = 0; VertIdx < Vertices.Num(); ++VertIdx)
+				{
+					float SD = (float)(Vertices[VertIdx] - PlanePt).Dot(PlaneNormal);
+					SignedDist[VertIdx] = SD;
+					OutTotal += int32(SD > 0);
+				}
+				// hull is fully outside plane's half-space
+				if (OutTotal == Vertices.Num())
+				{
+					Vertices.Empty();
+					PackedPolygons.Empty();
+					return;
+				}
+				// hull is fully inside plane's half-space
+				if (OutTotal == 0)
+				{
+					return;
+				}
+
+				for (int32 Idx = 0, PolyLen = 3; Idx < PackedPolygons.Num(); Idx += PolyLen)
+				{
+					// extract length of current polygon
+					PolyLen = 3;
+					int32 OrigStart = Idx;
+					if (PackedPolygons[Idx] < 0)
+					{
+						PolyLen = -PackedPolygons[Idx];
+						Idx++;
+					}
+					int32 Start = Idx;
+
+					// helper to convert index within polygon to index within vertices array
+					auto ToV = [this, Start, PolyLen](int32 SubIdx) -> int32
+					{
+						checkSlow(SubIdx >= 0 && SubIdx < PolyLen);
+						int32 VertIdx = PackedPolygons[Start + SubIdx];
+						checkSlow(VertIdx >= 0);
+						return VertIdx;
+					};
+					// track where the polygon crosses the plane to decide what to do with it
+					int32 OutCount = 0;
+					int32 FirstIn = -1, FirstOut = -1, LastIn = -1;
+					for (int32 SubIdx = 0; SubIdx < PolyLen; ++SubIdx)
+					{
+						float SD = SignedDist[ToV(SubIdx)];
+						bool IsOut = SD > 0;
+						if (FirstIn == -1)
+						{
+							if (!IsOut)
+							{
+								FirstIn = SubIdx;
+							}
+						}
+						else if (FirstOut == -1 && IsOut)
+						{
+							LastIn = SubIdx - 1;
+							FirstOut = SubIdx;
+						}
+						OutCount += int32(IsOut);
+					}
+					if (FirstOut == -1)
+					{
+						FirstOut = 0;
+						LastIn = PolyLen - 1;
+					}
+					if (OutCount == PolyLen)
+					{
+						continue;
+					}
+					if (OutCount == 0)
+					{
+						// copy original polygon data
+						for (int32 CopyIdx = OrigStart; CopyIdx < Start + PolyLen; ++CopyIdx)
+						{
+							NewPolygons.Add(PackedPolygons[CopyIdx]);
+						}
+						continue;
+					}
+					int32 NewPolyLen = LastIn + 1 - FirstIn;
+					if (FirstIn == 0)
+					{
+						int32 WalkBack = PolyLen - 1;
+						while (WalkBack > 0 && SignedDist[ToV(WalkBack)] <= 0)
+						{
+							FirstIn = WalkBack--;
+							NewPolyLen++;
+						}
+					}
+					auto GetCrossVertIdx = [&NewVertices, &SignedDist, this](int32 InsideVertIdx, int32 OutsideVertIdx) -> int32
+					{
+						checkSlow(InsideVertIdx != OutsideVertIdx);
+						// If within zero tolerance of plane, snap to plane
+						float InsideSD = SignedDist[InsideVertIdx];
+						checkSlow(InsideSD <= 0);
+						if (InsideSD > -FMathf::ZeroTolerance)
+						{
+							return -1;
+						}
+						FIntVector2 Key(InsideVertIdx, OutsideVertIdx);
+						int32* FoundVert = NewVertices.Find(Key);
+						if (!FoundVert)
+						{
+							float OutsideSD = SignedDist[OutsideVertIdx];
+							checkSlow(OutsideSD >= 0);
+							Chaos::FVec3f NewVert = FMath::Lerp(Vertices[InsideVertIdx], Vertices[OutsideVertIdx], InsideSD / (InsideSD - OutsideSD));
+							int32 NewVertIdx = Vertices.Add(NewVert);
+							NewVertices.Add(Key, NewVertIdx);
+							return NewVertIdx;
+						}
+						return *FoundVert;
+					};
+					int32 FirstCross = GetCrossVertIdx(ToV(FirstIn), ToV((FirstIn + PolyLen - 1) % PolyLen));
+					int32 LastCross = GetCrossVertIdx(ToV(LastIn), ToV(FirstOut));
+					int32 OpenPlaneEdgeVA = FirstCross;
+					int32 OpenPlaneEdgeVB = LastCross;
+					if (FirstCross != -1)
+					{
+						NewPolyLen++;
+					}
+					else
+					{
+						OpenPlaneEdgeVA = ToV(FirstIn);
+					}
+					if (LastCross != -1)
+					{
+						NewPolyLen++;
+					}
+					else
+					{
+						OpenPlaneEdgeVB = ToV(LastIn);
+					}
+
+					if (NewPolyLen < 2) // single co-incident vertex; no open edge here
+					{
+						continue;
+					}
+					OpenEdgeStart = OpenPlaneEdgeVB;
+					OpenEdgeVertMap.Add(OpenPlaneEdgeVA, OpenPlaneEdgeVB);
+					if (NewPolyLen == 2)
+					{
+						continue;
+					}
+					if (NewPolyLen > 3)
+					{
+						NewPolygons.Add(-NewPolyLen);
+					}
+					int32 NewPolygonStart = NewPolygons.Num();
+					if (FirstCross != -1)
+					{
+						NewPolygons.Add(FirstCross);
+					}
+					int32 AddStart = FirstIn;
+					if (FirstIn > LastIn)
+					{
+						for (int32 SubIdx = FirstIn; SubIdx < PolyLen; ++SubIdx)
+						{
+							NewPolygons.Add(ToV(SubIdx));
+						}
+						AddStart = 0;
+					}
+					for (int32 SubIdx = AddStart; SubIdx <= LastIn; ++SubIdx)
+					{
+						NewPolygons.Add(ToV(SubIdx));
+					}
+					if (LastCross != -1)
+					{
+						NewPolygons.Add(LastCross);
+					}
+
+					check(NewPolygons.Num() - NewPolygonStart == NewPolyLen);
+				}
+
+				// add the closing polygon
+				if (OpenEdgeStart != -1 && OpenEdgeVertMap.Num() > 2)
+				{
+					int32 OrigEnd = NewPolygons.Num();
+					int32 PolyEdges = OpenEdgeVertMap.Num();
+					if (PolyEdges > 3)
+					{
+						NewPolygons.Add(-PolyEdges);
+					}
+					int32 TraverseIdx = OpenEdgeStart;
+					int32 Added = 0;
+					do
+					{
+						NewPolygons.Add(TraverseIdx);
+						int32* FoundNext = OpenEdgeVertMap.Find(TraverseIdx);
+						if (!FoundNext)
+						{
+							break;
+						}
+						TraverseIdx = *FoundNext;
+						Added++;
+					} while (TraverseIdx != OpenEdgeStart && Added < PolyEdges);
+					if (Added != PolyEdges || TraverseIdx != OpenEdgeStart)
+					{
+						// failsafe if we didn't find a closed loop covering all edges:
+						// add a triangle fan closing off the edges that we did find
+						NewPolygons.SetNum(OrigEnd, false);
+						Chaos::FVec3f Center(0, 0, 0);
+						float CenterWt = 0;
+						int32 CenterIdx = Vertices.Num();
+						for (TPair<int32, int32> KV : OpenEdgeVertMap)
+						{
+							NewPolygons.Add(KV.Key);
+							NewPolygons.Add(KV.Value);
+							NewPolygons.Add(CenterIdx);
+							Center += Vertices[KV.Key];
+							CenterWt += 1.f;
+						}
+						Center /= CenterWt;
+						Vertices.Add(Center);
+					}
+				}
+
+				// NewPolygons now contains the updated polygon data
+				Swap(PackedPolygons, NewPolygons);
+				// Compress the vertex array to only include the vertices that weren't outside
+				// and track how the indices were remapped
+				int32 NumKept = 0;
+				const int32 OldVertCount = SignedDist.Num();
+				for (int32 OldV = 0; OldV < OldVertCount; ++OldV)
+				{
+					if (SignedDist[OldV] <= 0)
+					{
+						int32 UseNewV = NumKept++;
+						VertexRemap[OldV] = UseNewV;
+						checkSlow(OldV >= UseNewV);
+						Vertices[UseNewV] = Vertices[OldV];
+					}
+				}
+				// Translate back the new vertices
+				if (NumKept < OldVertCount)
+				{
+					for (int32 OldIdx = OldVertCount, AddedIdx = 0; OldIdx < Vertices.Num(); ++OldIdx, ++AddedIdx)
+					{
+						Vertices[NumKept + AddedIdx] = Vertices[OldIdx];
+					}
+					int32 NumNew = Vertices.Num() - OldVertCount;
+					Vertices.SetNum(NumKept + NumNew, false);
+				}
+				// Update the polygons w/ the compressed vertex indices
+				for (int32& VIdx : PackedPolygons)
+				{
+					if (VIdx >= 0) // Only remap vertices, not polygon sizes
+					{
+						if (VIdx < OldVertCount)
+						{
+							VIdx = VertexRemap[VIdx];
+						}
+						else // newly-created vertices are kept in the same order at the end of the array
+						{
+							VIdx = NumKept + (VIdx - OldVertCount);
+						}
+					}
+				}
+			};
+
+			// TODO: For performance, consider also pre-cutting with (some of) OtherHull's (expanded) bounding box planes
+			//Chaos::FConvex::FAABB3Type OtherBounds = OtherHull.GetLocalBoundingBox();
+			//OtherBounds.Thicken(ExpandAmount);
+
+			// Cut with each convex plane
+			const int32 NumPlanes = OtherHull.NumPlanes();
+			for (int32 PlaneIdx = 0; PlaneIdx < NumPlanes; ++PlaneIdx)
+			{
+				Chaos::TPlaneConcrete<float, 3> Plane = OtherHull.GetPlaneRaw(PlaneIdx);
+				Chaos::FVec3f N = Plane.Normal();
+				Chaos::FVec3f X = Plane.X();
+				X += N * ExpandAmount;
+				PlaneCut(X, N);
+			}
+
+			// When ExpandAmount is positive, also clip the hull at offsets of average edge planes for 'sharp' edges
+			if (ExpandAmount > 0)
+			{
+				const int32 NumEdges = OtherHull.NumEdges();
+				for (int32 EdgeIdx = 0; EdgeIdx < NumEdges; ++EdgeIdx)
+				{
+					Chaos::TPlaneConcrete<float, 3> EPlane0 = OtherHull.GetPlaneRaw(OtherHull.GetEdgePlane(EdgeIdx, 0));
+					Chaos::TPlaneConcrete<float, 3> EPlane1 = OtherHull.GetPlaneRaw(OtherHull.GetEdgePlane(EdgeIdx, 1));
+					float NormalDot = EPlane0.Normal().Dot(EPlane1.Normal());
+					if (NormalDot < -.1) // add an extra plane when not doing so would leave ~1.5x more space than the expected offset across from the edge, due to the miter
+					{
+						Chaos::FVec3f AvgNormal = EPlane0.Normal() + EPlane1.Normal();
+						if (AvgNormal.Normalize())
+						{
+							Chaos::FVec3f EdgeVert = OtherHull.GetVertex(OtherHull.GetEdgeVertex(EdgeIdx, 0));
+							PlaneCut(EdgeVert + AvgNormal * ExpandAmount, AvgNormal);
+						}
+					}
+
+				}
+			}
+		}
+
+		float ComputeArea()
+		{
+			float Area = 0;
+			for (int32 Idx = 0, PolyLen = 3; Idx < PackedPolygons.Num(); Idx += PolyLen)
+			{
+				// extract length of current polygon
+				PolyLen = 3;
+				if (PackedPolygons[Idx] < 0)
+				{
+					PolyLen = -PackedPolygons[Idx];
+					Idx++;
+				}
+				int32 Start = Idx;
+
+				// Add area of triangle fan covering the polygon
+				Chaos::FVec3f V0 = Vertices[PackedPolygons[Start]];
+				for (int32 SubIdx = 1; SubIdx + 1 < PolyLen; ++SubIdx)
+				{
+					Chaos::FVec3f V1 = Vertices[PackedPolygons[Start + SubIdx]];
+					Chaos::FVec3f V2 = Vertices[PackedPolygons[Start + SubIdx + 1]];
+					Area += UE::Geometry::VectorUtil::Area<float>(V0, V1, V2);
+				}
+			}
+			return Area;
+		}
+
+		void EstimateSharpContact(const Chaos::FConvex* HullA, const Chaos::FConvex* HullB, float& OutSharpContact, float& OutMaxSharpContact)
+		{
+			UE::Geometry::FExtremePoints3f ExtremePts(Vertices.Num(), [this](int32 Idx) {return Vertices[Idx];});
+			if (ExtremePts.Dimension < 1)
+			{
+				OutSharpContact = 0;
+				OutMaxSharpContact = 1;
+				return; // degenerate/empty contact
+			}
+			UE::Geometry::FInterval1f IntersectionIntervals[2];
+			UE::Geometry::FInterval1f HullAIntervals[2], HullBIntervals[2];
+			if (ExtremePts.Dimension > 1)
+			{
+				auto SetIntervals = [&ExtremePts](const TArray<Chaos::FVec3f>& UseVertices, UE::Geometry::FInterval1f* Intervals) -> void
+				{
+					for (FVector3f Vertex : UseVertices)
+					{
+						Intervals[0].Contain(Vertex.Dot(ExtremePts.Basis[1]));
+						if (ExtremePts.Dimension > 2)
+						{
+							Intervals[1].Contain(Vertex.Dot(ExtremePts.Basis[2]));
+						}
+					}
+				};
+				SetIntervals(Vertices, IntersectionIntervals);
+				SetIntervals(HullA->GetVertices(), HullAIntervals);
+				SetIntervals(HullB->GetVertices(), HullBIntervals);
+			}
+			auto IntervalsMaxLen = [](UE::Geometry::FInterval1f* Intervals)
+			{
+				return FMath::Max(Intervals[0].Length(), Intervals[1].Length());
+			};
+			OutSharpContact = IntervalsMaxLen(IntersectionIntervals);
+			OutMaxSharpContact = FMath::Min(IntervalsMaxLen(HullAIntervals), IntervalsMaxLen(HullBIntervals));
+		}
+	};
+
+	static float ComputeHullArea(const Chaos::FConvex& Hull)
+	{
+		float Area = 0;
+		const Chaos::FConvexStructureData& HullData = Hull.GetStructureData();
+		int32 NumPlanes = Hull.NumPlanes();
+		for (int PlaneIdx = 0; PlaneIdx < NumPlanes; PlaneIdx++)
+		{
+			int32 NumPlaneVerts = HullData.NumPlaneVertices(PlaneIdx);
+			Chaos::FVec3f V0 = Hull.GetVertex(HullData.GetPlaneVertex(PlaneIdx, 0));
+			for (int32 PlaneVertexIdx = 1; PlaneVertexIdx + 1 < NumPlaneVerts; PlaneVertexIdx++)
+			{
+				Chaos::FVec3f V1 = Hull.GetVertex(HullData.GetPlaneVertex(PlaneIdx, PlaneVertexIdx));
+				Chaos::FVec3f V2 = Hull.GetVertex(HullData.GetPlaneVertex(PlaneIdx, PlaneVertexIdx + 1));
+				Area += UE::Geometry::VectorUtil::Area<float>(V0, V1, V2);
+			}
+		}
+		return Area;
+	}
+}
+
+
+
+namespace UE::GeometryCollectionConvexUtility
+{
+
+
+	void HullIntersectionStats(const Chaos::FConvex* HullA, const Chaos::FConvex* HullB, float HullBExpansion, float& OutArea, float& OutMaxArea, float& OutSharpContact, float& OutMaxSharpContact)
+	{
+		FHullPolygons HullPolygons(*HullA);
+		HullPolygons.Intersect(*HullB, HullBExpansion);
+		OutArea = HullPolygons.ComputeArea();
+		// The maximum intersection area is ~ the minimum of the two hull areas
+		float MaxIntersectionArea = FMath::Min(ComputeHullArea(*HullA), ComputeHullArea(*HullB));
+		OutMaxArea = MaxIntersectionArea;
+		HullPolygons.EstimateSharpContact(HullA, HullB, OutSharpContact, OutMaxSharpContact);
+	}
+
+	void IntersectConvexHulls(Chaos::FConvex* ResultHull, const Chaos::FConvex* ClipHull, float ClipHullOffset, const Chaos::FConvex* UpdateHull, const FTransform* ClipHullTransform, const FTransform* UpdateHullTransform, const FTransform* ResultTransform)
+	{
+		FHullPolygons HullPolygons(*ClipHull);
+		bool bNeedRecomputeBounds = false;
+		if (ClipHullTransform)
+		{
+			bNeedRecomputeBounds = true;
+			for (Chaos::FVec3f& V : HullPolygons.Vertices)
+			{
+				V = (Chaos::FVec3f)ClipHullTransform->TransformPosition((FVector)V);
+			}
+		}
+		if (UpdateHullTransform)
+		{
+			bNeedRecomputeBounds = true;
+			for (Chaos::FVec3f& V : HullPolygons.Vertices)
+			{
+				V = (Chaos::FVec3f)UpdateHullTransform->InverseTransformPosition((FVector)V);
+			}
+		}
+		if (bNeedRecomputeBounds)
+		{
+			HullPolygons.Bounds = Chaos::FAABB3f::EmptyAABB();
+			for (Chaos::FVec3f& V : HullPolygons.Vertices)
+			{
+				HullPolygons.Bounds.GrowToInclude(V);
+			}
+		}
+		HullPolygons.Intersect(*UpdateHull, ClipHullOffset);
+		if (ResultTransform)
+		{
+			for (Chaos::FVec3f& V : HullPolygons.Vertices)
+			{
+				V = (Chaos::FVec3f)ResultTransform->TransformPosition((FVector)V);
+			}
+		}
+		*ResultHull = Chaos::FConvex(HullPolygons.Vertices, UpdateHull->GetMargin());
+	}
+}
diff --git a/Engine/Source/Runtime/Experimental/Chaos/Private/GeometryCollection/GeometryCollectionProximityUtility.cpp b/Engine/Source/Runtime/Experimental/Chaos/Private/GeometryCollection/GeometryCollectionProximityUtility.cpp
--- a/Engine/Source/Runtime/Experimental/Chaos/Private/GeometryCollection/GeometryCollectionProximityUtility.cpp
+++ b/Engine/Source/Runtime/Experimental/Chaos/Private/GeometryCollection/GeometryCollectionProximityUtility.cpp
@@ -790,453 +790,10 @@ void FGeometryCollectionProximityUtility::UpdateProximity(UE::GeometryCollection
 	}
 }
 
-namespace
-{
-	// Helpful struct to represent a convex hull with a local editable representation, to perform intersections / clipping on it
-	// TODO: Consider exposing this in a more generally useful place, or at least the functions to compute intersections and volume/area of intersection
-	struct FHullPolygons
-	{
-		// simple packed representation for convex hull faces, where each polygon's indices are listed sequentially,
-		// and negative values indicate the number of vertices in the next polygon. If no negative value is listed, polygon is a triangle.
-		TArray<int32> PackedPolygons;
-
-		// Copy of hull vertices, to be refined through plane cuts
-		TArray<Chaos::FVec3f> Vertices;
-		Chaos::FAABB3f Bounds;
-
-		FHullPolygons(const Chaos::FConvex& HullIn)
-		{
-			Vertices = HullIn.GetVertices();
-			const Chaos::FConvexStructureData& HullData = HullIn.GetStructureData();
-			int32 NumPlanes = HullIn.NumPlanes();
-			PackedPolygons.Reserve(NumPlanes * 3);
-			for (int PlaneIdx = 0; PlaneIdx < NumPlanes; PlaneIdx++)
-			{
-				int32 NumPlaneVerts = HullData.NumPlaneVertices(PlaneIdx);
-				if (NumPlaneVerts > 3)
-				{
-					PackedPolygons.Add(-NumPlaneVerts);
-				}
-				for (int32 PlaneVertexIdx = 0; PlaneVertexIdx < NumPlaneVerts; PlaneVertexIdx++)
-				{
-					PackedPolygons.Add(HullData.GetPlaneVertex(PlaneIdx, PlaneVertexIdx));
-				}
-			}
-			Bounds = HullIn.GetLocalBoundingBox();
-		}
-
-		void Intersect(const Chaos::FConvex& OtherHull, float ExpandAmount)
-		{
-			// Arrays to store intermediate plane cut data
-			TArray<int32> NewPolygons; 
-			TMap<FIntVector2, int> NewVertices; // mapping from edges to new vertices
-			TArray<float> SignedDist; // signed distance from vertices to a cutting plane
-			TArray<int32> VertexRemap;
-			TMap<int32, int32> OpenEdgeVertMap;
-
-			// Cut by the plane through PlanePt with PlaneNormal. Note: Must pre-apply ExpandAmount; we do not assume the plane is shifted here.
-			// Uses the intermediate data above as temp storage
-			auto PlaneCut = [&NewPolygons, &NewVertices, &SignedDist, &VertexRemap, &OpenEdgeVertMap, this](Chaos::FVec3f PlanePt, Chaos::FVec3f PlaneNormal)
-			{
-				NewPolygons.Reset(PackedPolygons.Num());
-				NewVertices.Reset();
-				OpenEdgeVertMap.Reset();
-				SignedDist.SetNum(Vertices.Num(), false);
-				VertexRemap.SetNum(Vertices.Num(), false);
-				int32 OpenEdgeStart = -1;
-
-				// Possible optimization: if many vertices, check plane vs AABB corners first?
-
-				int32 OutTotal = 0;
-				for (int32 VertIdx = 0; VertIdx < Vertices.Num(); ++VertIdx)
-				{
-					float SD = (float)(Vertices[VertIdx] - PlanePt).Dot(PlaneNormal);
-					SignedDist[VertIdx] = SD;
-					OutTotal += int32(SD > 0);
-				}
-				// hull is fully outside plane's half-space
-				if (OutTotal == Vertices.Num())
-				{
-					Vertices.Empty();
-					PackedPolygons.Empty();
-					return;
-				}
-				// hull is fully inside plane's half-space
-				if (OutTotal == 0)
-				{
-					return;
-				}
-
-				for (int32 Idx = 0, PolyLen = 3; Idx < PackedPolygons.Num(); Idx += PolyLen)
-				{
-					// extract length of current polygon
-					PolyLen = 3;
-					int32 OrigStart = Idx;
-					if (PackedPolygons[Idx] < 0)
-					{
-						PolyLen = -PackedPolygons[Idx];
-						Idx++;
-					}
-					int32 Start = Idx;
-
-					// helper to convert index within polygon to index within vertices array
-					auto ToV = [this, Start, PolyLen](int32 SubIdx) -> int32
-					{
-						checkSlow(SubIdx >= 0 && SubIdx < PolyLen);
-						int32 VertIdx = PackedPolygons[Start + SubIdx];
-						checkSlow(VertIdx >= 0);
-						return VertIdx;
-					};
-					// track where the polygon crosses the plane to decide what to do with it
-					int32 OutCount = 0;
-					int32 FirstIn = -1, FirstOut = -1, LastIn = -1;
-					for (int32 SubIdx = 0; SubIdx < PolyLen; ++SubIdx)
-					{
-						float SD = SignedDist[ToV(SubIdx)];
-						bool IsOut = SD > 0;
-						if (FirstIn == -1)
-						{
-							if (!IsOut)
-							{
-								FirstIn = SubIdx;
-							}
-						}
-						else if (FirstOut == -1 && IsOut)
-						{
-							LastIn = SubIdx - 1;
-							FirstOut = SubIdx;
-						}
-						OutCount += int32(IsOut);
-					}
-					if (FirstOut == -1)
-					{
-						FirstOut = 0;
-						LastIn = PolyLen - 1;
-					}
-					if (OutCount == PolyLen)
-					{
-						continue;
-					}
-					if (OutCount == 0)
-					{
-						// copy original polygon data
-						for (int32 CopyIdx = OrigStart; CopyIdx < Start + PolyLen; ++CopyIdx)
-						{
-							NewPolygons.Add(PackedPolygons[CopyIdx]);
-						}
-						continue;
-					}
-					int32 NewPolyLen = LastIn + 1 - FirstIn;
-					if (FirstIn == 0)
-					{
-						int32 WalkBack = PolyLen - 1;
-						while (WalkBack > 0 && SignedDist[ToV(WalkBack)] <= 0)
-						{
-							FirstIn = WalkBack--;
-							NewPolyLen++;
-						}
-					}
-					auto GetCrossVertIdx = [&NewVertices, &SignedDist, this](int32 InsideVertIdx, int32 OutsideVertIdx) -> int32
-					{
-						checkSlow(InsideVertIdx != OutsideVertIdx);
-						// If within zero tolerance of plane, snap to plane
-						float InsideSD = SignedDist[InsideVertIdx];
-						checkSlow(InsideSD <= 0);
-						if (InsideSD > -FMathf::ZeroTolerance)
-						{
-							return -1;
-						}
-						FIntVector2 Key(InsideVertIdx, OutsideVertIdx);
-						int32* FoundVert = NewVertices.Find(Key);
-						if (!FoundVert)
-						{
-							float OutsideSD = SignedDist[OutsideVertIdx];
-							checkSlow(OutsideSD >= 0);
-							Chaos::FVec3f NewVert = FMath::Lerp(Vertices[InsideVertIdx], Vertices[OutsideVertIdx], InsideSD / (InsideSD - OutsideSD));
-							int32 NewVertIdx = Vertices.Add(NewVert);
-							NewVertices.Add(Key, NewVertIdx);
-							return NewVertIdx;
-						}
-						return *FoundVert;
-					};
-					int32 FirstCross = GetCrossVertIdx(ToV(FirstIn), ToV((FirstIn + PolyLen - 1) % PolyLen));
-					int32 LastCross = GetCrossVertIdx(ToV(LastIn), ToV(FirstOut));
-					int32 OpenPlaneEdgeVA = FirstCross;
-					int32 OpenPlaneEdgeVB = LastCross;
-					if (FirstCross != -1)
-					{
-						NewPolyLen++;
-					}
-					else
-					{
-						OpenPlaneEdgeVA = ToV(FirstIn);
-					}
-					if (LastCross != -1)
-					{
-						NewPolyLen++;
-					}
-					else
-					{
-						OpenPlaneEdgeVB = ToV(LastIn);
-					}
-
-					if (NewPolyLen < 2) // single co-incident vertex; no open edge here
-					{
-						continue;
-					}
-					OpenEdgeStart = OpenPlaneEdgeVB;
-					OpenEdgeVertMap.Add(OpenPlaneEdgeVA, OpenPlaneEdgeVB);
-					if (NewPolyLen == 2)
-					{
-						continue;
-					}
-					if (NewPolyLen > 3)
-					{
-						NewPolygons.Add(-NewPolyLen);
-					}
-					int32 NewPolygonStart = NewPolygons.Num();
-					if (FirstCross != -1)
-					{
-						NewPolygons.Add(FirstCross);
-					}
-					int32 AddStart = FirstIn;
-					if (FirstIn > LastIn)
-					{
-						for (int32 SubIdx = FirstIn; SubIdx < PolyLen; ++SubIdx)
-						{
-							NewPolygons.Add(ToV(SubIdx));
-						}
-						AddStart = 0;
-					}
-					for (int32 SubIdx = AddStart; SubIdx <= LastIn; ++SubIdx)
-					{
-						NewPolygons.Add(ToV(SubIdx));
-					}
-					if (LastCross != -1)
-					{
-						NewPolygons.Add(LastCross);
-					}
 
-					check(NewPolygons.Num() - NewPolygonStart == NewPolyLen);
-				}
 
-				// add the closing polygon
-				if (OpenEdgeStart != -1 && OpenEdgeVertMap.Num() > 2)
-				{
-					int32 OrigEnd = NewPolygons.Num();
-					int32 PolyEdges = OpenEdgeVertMap.Num();
-					if (PolyEdges > 3)
-					{
-						NewPolygons.Add(-PolyEdges);
-					}
-					int32 TraverseIdx = OpenEdgeStart;
-					int32 Added = 0;
-					do
-					{
-						NewPolygons.Add(TraverseIdx);
-						int32* FoundNext = OpenEdgeVertMap.Find(TraverseIdx);
-						if (!FoundNext)
-						{
-							break;
-						}
-						TraverseIdx = *FoundNext;
-						Added++;
-					} while (TraverseIdx != OpenEdgeStart && Added < PolyEdges);
-					if (Added != PolyEdges || TraverseIdx != OpenEdgeStart)
-					{
-						// failsafe if we didn't find a closed loop covering all edges:
-						// add a triangle fan closing off the edges that we did find
-						NewPolygons.SetNum(OrigEnd, false);
-						Chaos::FVec3f Center(0, 0, 0);
-						float CenterWt = 0;
-						int32 CenterIdx = Vertices.Num();
-						for (TPair<int32, int32> KV : OpenEdgeVertMap)
-						{
-							NewPolygons.Add(KV.Key);
-							NewPolygons.Add(KV.Value);
-							NewPolygons.Add(CenterIdx);
-							Center += Vertices[KV.Key];
-							CenterWt += 1.f;
-						}
-						Center /= CenterWt;
-						Vertices.Add(Center);
-					}
-				}
 
-				// NewPolygons now contains the updated polygon data
-				Swap(PackedPolygons, NewPolygons);
-				// Compress the vertex array to only include the vertices that weren't outside
-				// and track how the indices were remapped
-				int32 NumKept = 0;
-				const int32 OldVertCount = SignedDist.Num();
-				for (int32 OldV = 0; OldV < OldVertCount; ++OldV)
-				{
-					if (SignedDist[OldV] <= 0)
-					{
-						int32 UseNewV = NumKept++;
-						VertexRemap[OldV] = UseNewV;
-						checkSlow(OldV >= UseNewV);
-						Vertices[UseNewV] = Vertices[OldV];
-					}
-				}
-				// Translate back the new vertices
-				if (NumKept < OldVertCount)
-				{
-					for (int32 OldIdx = OldVertCount, AddedIdx = 0; OldIdx < Vertices.Num(); ++OldIdx, ++AddedIdx)
-					{
-						Vertices[NumKept + AddedIdx] = Vertices[OldIdx];
-					}
-					int32 NumNew = Vertices.Num() - OldVertCount;
-					Vertices.SetNum(NumKept + NumNew, false);
-				}
-				// Update the polygons w/ the compressed vertex indices
-				for (int32& VIdx : PackedPolygons)
-				{
-					if (VIdx >= 0) // Only remap vertices, not polygon sizes
-					{
-						if (VIdx < OldVertCount)
-						{
-							VIdx = VertexRemap[VIdx];
-						}
-						else // newly-created vertices are kept in the same order at the end of the array
-						{
-							VIdx = NumKept + (VIdx - OldVertCount);
-						}
-					}
-				}
-			};
 
-			// TODO: For performance, consider also pre-cutting with (some of) OtherHull's (expanded) bounding box planes
-			//Chaos::FConvex::FAABB3Type OtherBounds = OtherHull.GetLocalBoundingBox();
-			//OtherBounds.Thicken(ExpandAmount);
-
-			// Cut with each convex plane
-			const int32 NumPlanes = OtherHull.NumPlanes();
-			for (int32 PlaneIdx = 0; PlaneIdx < NumPlanes; ++PlaneIdx)
-			{
-				Chaos::TPlaneConcrete<float, 3> Plane = OtherHull.GetPlaneRaw(PlaneIdx);
-				Chaos::FVec3f N = Plane.Normal();
-				Chaos::FVec3f X = Plane.X();
-				X += N * ExpandAmount;
-				PlaneCut(X, N);
-			}
-
-			// When ExpandAmount is positive, also clip the hull at offsets of average edge planes for 'sharp' edges
-			if (ExpandAmount > 0)
-			{
-				const int32 NumEdges = OtherHull.NumEdges();
-				for (int32 EdgeIdx = 0; EdgeIdx < NumEdges; ++EdgeIdx)
-				{
-					Chaos::TPlaneConcrete<float, 3> EPlane0 = OtherHull.GetPlaneRaw(OtherHull.GetEdgePlane(EdgeIdx, 0));
-					Chaos::TPlaneConcrete<float, 3> EPlane1 = OtherHull.GetPlaneRaw(OtherHull.GetEdgePlane(EdgeIdx, 1));
-					float NormalDot = EPlane0.Normal().Dot(EPlane1.Normal());
-					if (NormalDot < -.1) // add an extra plane when not doing so would leave ~1.5x more space than the expected offset across from the edge, due to the miter
-					{
-						Chaos::FVec3f AvgNormal = EPlane0.Normal() + EPlane1.Normal();
-						if (AvgNormal.Normalize())
-						{
-							Chaos::FVec3f EdgeVert = OtherHull.GetVertex(OtherHull.GetEdgeVertex(EdgeIdx, 0));
-							PlaneCut(EdgeVert + AvgNormal * ExpandAmount, AvgNormal);
-						}
-					}
-
-				}
-			}
-		}
-
-		float ComputeArea()
-		{
-			float Area = 0;
-			for (int32 Idx = 0, PolyLen = 3; Idx < PackedPolygons.Num(); Idx += PolyLen)
-			{
-				// extract length of current polygon
-				PolyLen = 3;
-				if (PackedPolygons[Idx] < 0)
-				{
-					PolyLen = -PackedPolygons[Idx];
-					Idx++;
-				}
-				int32 Start = Idx;
-
-				// Add area of triangle fan covering the polygon
-				Chaos::FVec3f V0 = Vertices[PackedPolygons[Start]];
-				for (int32 SubIdx = 1; SubIdx + 1 < PolyLen; ++SubIdx)
-				{
-					Chaos::FVec3f V1 = Vertices[PackedPolygons[Start + SubIdx]];
-					Chaos::FVec3f V2 = Vertices[PackedPolygons[Start + SubIdx + 1]];
-					Area += UE::Geometry::VectorUtil::Area<float>(V0, V1, V2);
-				}
-			}
-			return Area;
-		}
-
-		void EstimateSharpContact(const Chaos::FConvex* HullA, const Chaos::FConvex* HullB, float& OutSharpContact, float& OutMaxSharpContact)
-		{
-			UE::Geometry::FExtremePoints3f ExtremePts(Vertices.Num(), [this](int32 Idx) {return Vertices[Idx];});
-			if (ExtremePts.Dimension < 1)
-			{
-				OutSharpContact = 0;
-				OutMaxSharpContact = 1;
-				return; // degenerate/empty contact
-			}
-			UE::Geometry::FInterval1f IntersectionIntervals[2];
-			UE::Geometry::FInterval1f HullAIntervals[2], HullBIntervals[2];
-			if (ExtremePts.Dimension > 1)
-			{
-				auto SetIntervals = [&ExtremePts](const TArray<Chaos::FVec3f>& UseVertices, UE::Geometry::FInterval1f* Intervals) -> void
-				{
-					for (FVector3f Vertex : UseVertices)
-					{
-						Intervals[0].Contain(Vertex.Dot(ExtremePts.Basis[1]));
-						if (ExtremePts.Dimension > 2)
-						{
-							Intervals[1].Contain(Vertex.Dot(ExtremePts.Basis[2]));
-						}
-					}
-				};
-				SetIntervals(Vertices, IntersectionIntervals);
-				SetIntervals(HullA->GetVertices(), HullAIntervals);
-				SetIntervals(HullB->GetVertices(), HullBIntervals);
-			}
-			auto IntervalsMaxLen = [](UE::Geometry::FInterval1f* Intervals)
-			{
-				return FMath::Max(Intervals[0].Length(), Intervals[1].Length());
-			};
-			OutSharpContact = IntervalsMaxLen(IntersectionIntervals);
-			OutMaxSharpContact = FMath::Min(IntervalsMaxLen(HullAIntervals), IntervalsMaxLen(HullBIntervals));
-		}
-	};
-
-	static float ComputeHullArea(const Chaos::FConvex& Hull)
-	{
-		float Area = 0;
-		const Chaos::FConvexStructureData& HullData = Hull.GetStructureData();
-		int32 NumPlanes = Hull.NumPlanes();
-		for (int PlaneIdx = 0; PlaneIdx < NumPlanes; PlaneIdx++)
-		{
-			int32 NumPlaneVerts = HullData.NumPlaneVertices(PlaneIdx);
-			Chaos::FVec3f V0 = Hull.GetVertex(HullData.GetPlaneVertex(PlaneIdx, 0));
-			for (int32 PlaneVertexIdx = 1; PlaneVertexIdx + 1 < NumPlaneVerts; PlaneVertexIdx++)
-			{
-				Chaos::FVec3f V1 = Hull.GetVertex(HullData.GetPlaneVertex(PlaneIdx, PlaneVertexIdx));
-				Chaos::FVec3f V2 = Hull.GetVertex(HullData.GetPlaneVertex(PlaneIdx, PlaneVertexIdx + 1));
-				Area += UE::Geometry::VectorUtil::Area<float>(V0, V1, V2);
-			}
-		}
-		return Area;
-	}
-
-	static void HullIntersectionStats(int32 DbgA, int32 DbgB, const Chaos::FConvex* HullA, const Chaos::FConvex* HullB, float HullBExpansion, float& OutArea, float& OutMaxArea, float& OutSharpContact, float& OutMaxSharpContact)
-	{
-		FHullPolygons HullPolygons(*HullA);
-		HullPolygons.Intersect(*HullB, HullBExpansion);
-		OutArea = HullPolygons.ComputeArea();
-		// The maximum intersection area is ~ the minimum of the two hull areas
-		float MaxIntersectionArea = FMath::Min(ComputeHullArea(*HullA), ComputeHullArea(*HullB));
-		OutMaxArea = MaxIntersectionArea;
-		HullPolygons.EstimateSharpContact(HullA, HullB, OutSharpContact, OutMaxSharpContact);
-	}
-}
 
 
 TArray<FGeometryCollectionProximityUtility::FGeometryContactEdge>
@@ -1271,7 +828,7 @@ FGeometryCollectionProximityUtility::ComputeConvexGeometryContactFromProximity(
 				for (int32 NbrHullIdx : NbrGeoHulls)
 				{
 					float Area, MaxArea, SharpContact, MaxSharpContact;
-					HullIntersectionStats(TransformIdx, NbrTransformIdx, LocalHulls.Hulls[GeoHullIdx].Get(), LocalHulls.Hulls[NbrHullIdx].Get(), DistanceTolerance, Area, MaxArea, SharpContact, MaxSharpContact);
+					UE::GeometryCollectionConvexUtility::HullIntersectionStats(LocalHulls.Hulls[GeoHullIdx].Get(), LocalHulls.Hulls[NbrHullIdx].Get(), DistanceTolerance, Area, MaxArea, SharpContact, MaxSharpContact);
 					// Note: The sharp contact concept would be most accurately implemented by combining the intersection/overlaps of all involved hulls
 					// but taking the max values should be a reasonable conservative approximation in most cases
 					OverlapSharpContact = FMath::Max(OverlapSharpContact, SharpContact);
diff --git a/Engine/Source/Runtime/Experimental/Chaos/Public/GeometryCollection/GeometryCollectionConvexUtility.h b/Engine/Source/Runtime/Experimental/Chaos/Public/GeometryCollection/GeometryCollectionConvexUtility.h
--- a/Engine/Source/Runtime/Experimental/Chaos/Public/GeometryCollection/GeometryCollectionConvexUtility.h
+++ b/Engine/Source/Runtime/Experimental/Chaos/Public/GeometryCollection/GeometryCollectionConvexUtility.h
@@ -27,6 +27,17 @@ enum class EConvexOverlapRemoval : int32
 	OnlyClustersVsClusters = 3
 };
 
+UENUM()
+enum class EGenerateConvexMethod : uint8
+{
+	// Convert from external collision shapes (if available)
+	ExternalCollision,
+	// Compute all convex hulls from geometry
+	ComputedFromGeometry,
+	// Intersect external collision shapes with computed convex hulls
+	IntersectExternalWithComputed
+};
+
 namespace UE::GeometryCollectionConvexUtility
 {
 	// Used to pass computed convex hulls along with metadata
@@ -44,6 +55,15 @@ namespace UE::GeometryCollectionConvexUtility
 		// The pivots used for scaling by OverlapRemovalShrinkPercent
 		TArray<FVector> Pivots;
 	};
+
+	// Compute surface area and contact estimates for two convex hulls in the same local space, with the second hull optionally expanded
+	void CHAOS_API HullIntersectionStats(const Chaos::FConvex* HullA, const Chaos::FConvex* HullB, float HullBExpansion, float& OutArea, float& OutMaxArea, float& OutSharpContact, float& OutMaxSharpContact);
+
+	// Compute the intersection of ClipHull and UpdateHull, with ClipHull optionally offset, and place the result in ResultHull. Note ResultHull can point to the same FConvex as UpdateHull.
+	// If optional transforms are provided, ClipHull will be transformed into the local space of UpdateHull, and finally transformed by UpdateToResultTransform. Otherwise, all FConvex are assumed to be in the same coordinate space.
+	void CHAOS_API IntersectConvexHulls(Chaos::FConvex* ResultHull, const Chaos::FConvex* ClipHull, float ClipHullOffset, const Chaos::FConvex* UpdateHull,
+		const FTransform* ClipHullTransform = nullptr, const FTransform* UpdateHullTransform = nullptr, const FTransform* UpdateToResultTransform = nullptr);
+
 }
 
 class CHAOS_API FGeometryCollectionConvexUtility
@@ -79,9 +99,21 @@ public:
 	static FGeometryCollectionConvexData CreateNonOverlappingConvexHullData(FGeometryCollection* GeometryCollection, double FractionAllowRemove = .3, double SimplificationDistanceThreshold = 0.0, double CanExceedFraction = .5, 
 		EConvexOverlapRemoval OverlapRemovalMethod = EConvexOverlapRemoval::All, double OverlapRemovalShrinkPercent = 0.0, UE::GeometryCollectionConvexUtility::FConvexHulls* ComputedLeafHullsToModify = nullptr);
 
-	static void GenerateClusterConvexHullsFromChildrenHulls(FGeometryCollection& Collection, int32 ConvexCount, double ErrorToleranceInCm, const TArrayView<const int32> TransformSubset);
-	static void GenerateClusterConvexHullsFromLeafHulls(FGeometryCollection& Collection, int32 ConvexCount, double ErrorToleranceInCm, const TArrayView<const int32> OptionalTransformSubset);
-	static void GenerateClusterConvexHullsFromLeafHulls(FGeometryCollection& Collection, int32 ConvexCount, double ErrorToleranceInCm);
+	static void GenerateClusterConvexHullsFromChildrenHulls(FGeometryCollection& Collection, int32 ConvexCount, double ErrorToleranceInCm, const TArrayView<const int32> TransformSubset, bool bUseExternalCollisionIfAvailable = true);
+	static void GenerateClusterConvexHullsFromChildrenHulls(FGeometryCollection& Collection, int32 ConvexCount, double ErrorToleranceInCm, bool bUseExternalCollisionIfAvailable = true);
+	static void GenerateClusterConvexHullsFromLeafHulls(FGeometryCollection& Collection, int32 ConvexCount, double ErrorToleranceInCm, const TArrayView<const int32> OptionalTransformSubset, bool bUseExternalCollisionIfAvailable = true);
+	static void GenerateClusterConvexHullsFromLeafHulls(FGeometryCollection& Collection, int32 ConvexCount, double ErrorToleranceInCm, bool bUseExternalCollisionIfAvailable = true);
+
+	// Additional settings for filtering when the EGenerateConvexMethod::IntersectExternalWithComputed is applied
+	struct FIntersectionFilters
+	{
+		FIntersectionFilters() : OnlyIntersectIfComputedIsSmallerFactor(1.0), MinExternalVolumeToIntersect(0.0)
+		{}
+		double OnlyIntersectIfComputedIsSmallerFactor;
+		double MinExternalVolumeToIntersect;
+	};
+	
+	static void GenerateLeafConvexHulls(FGeometryCollection& Collection, bool bRestrictToSelection, const TArrayView<const int32> TransformSubset, double SimplificationDistanceThreshold, EGenerateConvexMethod GenerateMethod = EGenerateConvexMethod::ExternalCollision, FIntersectionFilters IntersectFilters = FIntersectionFilters());
 
 	/** Returns the convex hull of the vertices contained in the specified geometry. */
 	static TUniquePtr<Chaos::FConvex> FindConvexHull(const FGeometryCollection* GeometryCollection, int32 GeometryIndex);
@@ -121,7 +153,8 @@ public:
 	// This is an initial step of several algorithms: The CreateNonOverlappingConvexHullData function as well as convex-based proximity detection (TODO: and the auto-embed algorithm?)
 	// (TODO: Make auto-embed use this instead of the full hulls?)
 	// @param GlobalTransformArray		GeometryCollection's transforms to global space, as computed by GeometryCollectionAlgo::GlobalMatrices
-	static UE::GeometryCollectionConvexUtility::FConvexHulls ComputeLeafHulls(FGeometryCollection* GeometryCollection, const TArray<FTransform>& GlobalTransformArray, double SimplificationDistanceThreshold = 0.0, double OverlapRemovalShrinkPercent = 0.0);
+	static UE::GeometryCollectionConvexUtility::FConvexHulls ComputeLeafHulls(FGeometryCollection* GeometryCollection, const TArray<FTransform>& GlobalTransformArray, double SimplificationDistanceThreshold = 0.0, double OverlapRemovalShrinkPercent = 0.0,
+		TFunction<bool(int32)> SkipBoneFn = nullptr);
 
 	struct FTransformedConvex
 	{
@@ -149,6 +182,6 @@ private:
 	static void CreateConvexHullAttributesIfNeeded(FManagedArrayCollection& GeometryCollection);
 
 	// Implementation for GenerateClusterConvexHullsFromLeafHulls, supporting the full-collection and subset cases
-	static void GenerateClusterConvexHullsFromLeafOrChildrenHullsInternal(FGeometryCollection& Collection, int32 ConvexCount, double ErrorToleranceInCm, bool bOnlySubset, bool bUseDirectChildren, const TArrayView<const int32> OptionalTransformSubset);
+	static void GenerateClusterConvexHullsFromLeafOrChildrenHullsInternal(FGeometryCollection& Collection, int32 ConvexCount, double ErrorToleranceInCm, bool bOnlySubset, bool bUseDirectChildren, const TArrayView<const int32> OptionalTransformSubset, bool bUseExternalCollisionIfAvailable = true);
 };
 
