cis fix non-unity build error for geometry collection convex utility

diff --git a/Engine/Source/Runtime/Experimental/Chaos/Public/GeometryCollection/GeometryCollectionConvexUtility.h b/Engine/Source/Runtime/Experimental/Chaos/Public/GeometryCollection/GeometryCollectionConvexUtility.h
--- a/Engine/Source/Runtime/Experimental/Chaos/Public/GeometryCollection/GeometryCollectionConvexUtility.h
+++ b/Engine/Source/Runtime/Experimental/Chaos/Public/GeometryCollection/GeometryCollectionConvexUtility.h
@@ -4,16 +4,12 @@
 
 #include "CoreMinimal.h"
 #include "GeometryCollection/ManagedArrayCollection.h"
+#include "Chaos/Convex.h"
 
 #include "GeometryCollectionConvexUtility.generated.h"
 
 class FGeometryCollection;
 
-namespace Chaos
-{
-	class FConvex;
-}
-
 UENUM()
 enum class EConvexOverlapRemoval : int32
 {
