Change UpdateClusterMassProperties to use MassProperties.h helper functions

diff --git a/Engine/Source/Runtime/Experimental/Chaos/Private/Chaos/PBDRigidClusteringAlgo.cpp b/Engine/Source/Runtime/Experimental/Chaos/Private/Chaos/PBDRigidClusteringAlgo.cpp
--- a/Engine/Source/Runtime/Experimental/Chaos/Private/Chaos/PBDRigidClusteringAlgo.cpp
+++ b/Engine/Source/Runtime/Experimental/Chaos/Private/Chaos/PBDRigidClusteringAlgo.cpp
@@ -7,6 +7,7 @@
 #include "Chaos/Utilities.h"
 #include "Chaos/PBDRigidClusteringCollisionParticleAlgo.h"
 #include "Chaos/PBDRigidsEvolutionGBF.h"
+#include "Chaos/MassProperties.h"
 
 namespace Chaos
 {
@@ -44,177 +45,100 @@ namespace Chaos
 	{
 		SCOPE_CYCLE_COUNTER(STAT_UpdateClusterMassProperties);
 
-		Parent->SetX(FVec3(0));
-		Parent->SetR(FRotation3(FQuat::MakeFromEuler(FVec3(0))));
-		Parent->SetV(FVec3(0));
-		Parent->SetW(FVec3(0));
-		Parent->SetM(0);
-		Parent->SetI(TVec3<FRealSingle>(0)); // Diagonal should be updated from ParentInertia by caller
-		ParentInertia = FMatrix33(0);
-
-		bool bHasChild = false;
-		bool bHasProxyChild = false;
-		for (FPBDRigidParticleHandle* OriginalChild : Children)
-		{
-			FPBDRigidParticleHandle* Child;
-			FVec3 ChildPosition;
-			FRotation3 ChildRotation;
-
-			Child = OriginalChild;
-			ChildPosition = Child->X();
-			ChildRotation = Child->R();
-
-			const FReal ChildMass = Child->M();
-			const FMatrix33 ChildWorldSpaceI = Chaos::Utilities::ComputeWorldSpaceInertia(ChildRotation, Child->I());
-			if (ChildWorldSpaceI.ContainsNaN())
-			{
-				continue;
-			}
-			bHasProxyChild = true;
-			bHasChild = true;
-			bHasProxyChild = true;
-
-			ParentInertia += ChildWorldSpaceI;
-			Parent->M() += ChildMass;
-			Parent->X() += ChildPosition * ChildMass;
-			Parent->V() += OriginalChild->V() * ChildMass; // Use orig child for vel because we don't sim the proxy
-			Parent->W() += OriginalChild->W() * ChildMass;
-		}
-		if (!bHasProxyChild)
-		{
-			for (FPBDRigidParticleHandle* OriginalChild : Children)
-			{
-				FPBDRigidParticleHandle* Child = OriginalChild;
-				const FVec3& ChildPosition = Child->X();
-				const FRotation3& ChildRotation = Child->R();
-				const FReal ChildMass = Child->M();
-
-				const FMatrix33 ChildWorldSpaceI = Chaos::Utilities::ComputeWorldSpaceInertia(ChildRotation, Child->I());
-				if (ChildWorldSpaceI.ContainsNaN())
-				{
-					continue;
-				}
-				bHasChild = true;
-				ParentInertia += ChildWorldSpaceI;
-				Parent->M() += ChildMass;
-				Parent->X() += ChildPosition * ChildMass;
-				Parent->V() += OriginalChild->V() * ChildMass; // Use orig child for vel because we don't sim the proxy
-				Parent->W() += OriginalChild->W() * ChildMass;
-			}
-		}
-		for (int32 i = 0; i < 3; i++)
+		// Initialize parent
+		Parent->M() = FReal(0);
+		Parent->InvM() = FReal(0);
+		Parent->I() = FVec3(0);
+		Parent->InvI() = FVec3(0);
+		Parent->SetCenterOfMass(FVec3(0));
+		Parent->SetRotationOfMass(FQuat::Identity);
+		if (ForceMassOrientation)
 		{
-			const FVec3& InertiaTensor = Parent->I();
-			if (InertiaTensor[i] < UE_SMALL_NUMBER)
-			{
-				Parent->SetI(TVec3<FRealSingle>(1.f, 1.f, 1.f));
-				break;
-			}
+			Parent->X() = ForceMassOrientation->GetLocation();
+			Parent->R() = ForceMassOrientation->GetRotation();
 		}
+		Parent->P() = Parent->X();
+		Parent->Q() = Parent->R();
 
-		if (!bHasChild || !(Parent->M() > UE_SMALL_NUMBER))
+		if (Children.Num() == 0)
 		{
-			Parent->M() = 1.0;
-			if (ForceMassOrientation)
-			{
-				Parent->X() = ForceMassOrientation->GetLocation();
-				Parent->R() = ForceMassOrientation->GetRotation();
-			}
-			else
-			{
-				Parent->X() = FVec3(0);
-				Parent->R() = FRotation3(FMatrix::Identity);
-			}
-			Parent->V() = FVec3(0);
-			Parent->PreV() = Parent->V();
-			Parent->InvM() = 1;
-			Parent->P() = Parent->X();
-			Parent->W() = FVec3(0);
-			Parent->PreW() = Parent->W();
-			Parent->Q() = Parent->R();
-			Parent->I() = TVec3<FRealSingle>(1);
-			Parent->InvI() = TVec3<FRealSingle>(1);
-			ParentInertia = FMatrix33(1,1,1);
 			return;
 		}
 
-		check(Parent->M() > UE_SMALL_NUMBER);
+		//
+		// Step 1: Compute the world CoM and total mass of the parent
+		//
 
-		Parent->X() /= Parent->M();
-		Parent->V() /= Parent->M();
-		Parent->PreV() = Parent->V();
-		Parent->InvM() = static_cast<FReal>(1.0) / Parent->M();
-		if (ForceMassOrientation)
+		FVec3 WorldCoM = FVec3::ZeroVector;
+		for (const FPBDRigidParticleHandle* Child : Children)
 		{
-			Parent->X() = ForceMassOrientation->GetLocation();
+			WorldCoM += Child->M() * Child->XCom();
+			Parent->M() += Child->M();
 		}
-		Parent->P() = Parent->X();
-		for (FPBDRigidParticleHandle* OriginalChild : Children)
+		if (FMath::IsNearlyZero(Parent->M()))
 		{
-			FPBDRigidParticleHandle* Child;
-			FVec3 ChildPosition;
-
-			Child = OriginalChild;
-			ChildPosition = Child->X();
+			return;
+		}
+		Parent->InvM() = FReal(1) / Parent->M();
+		WorldCoM *= Parent->InvM();
 
-			FVec3 ParentToChild = ChildPosition - Parent->X();
 
-			const FReal ChildMass = Child->M();
-			// taking v from original child since we are not actually simulating the proxy child
-			Parent->W() +=
-				FVec3::CrossProduct(ParentToChild,
-					OriginalChild->V() * ChildMass);
-			{
-				const FReal& p0 = ParentToChild[0];
-				const FReal& p1 = ParentToChild[1];
-				const FReal& p2 = ParentToChild[2];
-				const FReal& m = ChildMass;
-				// TODO: We are computing inertia twice. Need to decide if we prefer this to loop above.
-				ParentInertia += FMatrix33(
-						m * (p1 * p1 + p2 * p2), -m * p1 * p0, -m * p2 * p0,
-						m * (p2 * p2 + p0 * p0), -m * p2 * p1, m * (p1 * p1 + p0 * p0));
-			}
-		}
+		//
+		// Step 2: Pick the parent's orientation and location.
+		//
+		// If we have a ForceMassOrientation transform, then use that, otherwise
+		// default to X = CoM. To do this, we need to compute the world CoM of
+		// the children.
+		//
 
-		if (ForceMassOrientation)
+		if (ForceMassOrientation == nullptr)
 		{
-			ParentInertia = ForceMassOrientation->GetRotation().ToMatrix() * ParentInertia * ForceMassOrientation->GetRotation().ToMatrix().GetTransposed();
+			Parent->X() = WorldCoM;
+			Parent->R() = FQuat::Identity;
 		}
+		Parent->P() = Parent->X();
+		Parent->Q() = Parent->R();
+		const FRigidTransform3 ParentTM = Parent->GetTransformXR();
+		const FRigidTransform3 InvParentTM = ParentTM.Inverse();
 
-		if (ParentInertia.ContainsNaN())
-		{
-			ParentInertia = FMatrix33(1,1,1);
-		}
-		else
-		{
-			for (int32 i = 0; i < 3; i++)
-			{
-				if (ParentInertia.GetColumn(i)[i] < UE_SMALL_NUMBER)
-				{
-					ParentInertia = FMatrix33(1,1,1);
-					break;
-				}
-			}
-		}
-		Parent->W() /= Parent->M();
-		Parent->PreW() = Parent->W();
-		if (ForceMassOrientation)
-		{
-			Parent->R() = ForceMassOrientation->GetRotation();
-		}
 
-		if (Parent->R().ContainsNaN())
+		//
+		// Step 3: Compute mass properties of each particle & store them in a list
+		//
+
+		TArray<FMassProperties> ChildMasses;
+		ChildMasses.Reserve(Children.Num());
+		for (const FPBDRigidParticleHandle* Child : Children)
 		{
-			ParentInertia = FMatrix33(1,1,1);
-			Parent->R() = TRotation<FReal, 3>(FMatrix::Identity);
+			// Get the child's transform relative to the parent
+			const FRigidTransform3 ChildTM = Child->GetTransformXR();
+			const FRigidTransform3 LocalTM = ChildTM * InvParentTM;
+
+			// Get the child's mass properties
+			FMassProperties ChildMass;
+			ChildMass.Mass = Child->M();
+			ChildMass.InertiaTensor = FMatrix33(Child->I());
+			ChildMass.CenterOfMass = LocalTM.TransformPosition(Child->CenterOfMass());
+			ChildMass.RotationOfMass = LocalTM.GetRotation() * Child->RotationOfMass();
+			ChildMasses.Add(ChildMass);
 		}
 
-		Parent->Q() = Parent->R();
-		
-		// TODO: This should be computed from rotated tensor by caller.
-		// Initializing to keep current behaviour until fixed.
-		Parent->I() = ParentInertia.GetDiagonal(); 
-		Parent->InvI() = TVec3<FRealSingle>(1.f / Parent->I()[0], 1.f / Parent->I()[1], 1.f / Parent->I()[2]);
+		//
+		// Step 4: Combine mass properties of sub particles & store
+		// them in the parent particle
+		//
+
+		FMassProperties ParentMass = Chaos::Combine(ChildMasses);
+		// NOTE: The combine method will have diagonalized the inertia.
+		ParentInertia = ParentMass.InertiaTensor;
+		const FVec3 Inertia = ParentInertia.GetDiagonal();
+		Parent->SetCenterOfMass(ParentMass.CenterOfMass);
+		Parent->SetRotationOfMass(ParentMass.RotationOfMass);
+		Parent->I() = Inertia;
+		Parent->InvI()
+			= (FMath::IsNearlyZero(Inertia[0]) || FMath::IsNearlyZero(Inertia[1]) || FMath::IsNearlyZero(Inertia[2]))
+			? FVec3::ZeroVector
+			: FReal(1) / Inertia;
 	}
 
 
