Rename variables in XPBDStretchBiasElementConstraints to be compliant with coding standards. No other changes.

diff --git a/Engine/Source/Runtime/Experimental/Chaos/Private/Chaos/XPBDStretchBiasElementConstraints.cpp b/Engine/Source/Runtime/Experimental/Chaos/Private/Chaos/XPBDStretchBiasElementConstraints.cpp
--- a/Engine/Source/Runtime/Experimental/Chaos/Private/Chaos/XPBDStretchBiasElementConstraints.cpp
+++ b/Engine/Source/Runtime/Experimental/Chaos/Private/Chaos/XPBDStretchBiasElementConstraints.cpp
@@ -205,9 +205,9 @@ void FXPBDStretchBiasElementConstraints::InitConstraintsAndRestData(const FSolve
 			const FSolverVec3 X01 = X1 - X0;
 			const FSolverVec3 X02 = X2 - X0;
 
-			const FSolverVec3 dX_dU = X01 * DeltaUVInv.M[0] + X02 * DeltaUVInv.M[1];
-			const FSolverVec3 dX_dV = X01 * DeltaUVInv.M[2] + X02 * DeltaUVInv.M[3];
-			const FSolverVec2& RestStretch = RestStretchLengths.Add_GetRef(FSolverVec2(dX_dU.Length(), dX_dV.Length()));
+			const FSolverVec3 DXDu = X01 * DeltaUVInv.M[0] + X02 * DeltaUVInv.M[1];
+			const FSolverVec3 DXDv = X01 * DeltaUVInv.M[2] + X02 * DeltaUVInv.M[3];
+			const FSolverVec2& RestStretch = RestStretchLengths.Add_GetRef(FSolverVec2(DXDu.Length(), DXDv.Length()));
 
 			const FSolverReal StiffnessScaleU = RestStretch[0] < UE_SMALL_NUMBER ? 1.f : 1.f / RestStretch[0];
 			const FSolverReal StiffnessScaleV = RestStretch[1] < UE_SMALL_NUMBER ? 1.f : 1.f / RestStretch[1];
@@ -646,14 +646,14 @@ void FXPBDStretchBiasElementConstraints::Apply(FSolverParticles& Particles, cons
 	}
 }
 
-void FXPBDStretchBiasElementConstraints::CalculateUVStretch(const int32 ConstraintIndex, const FSolverVec3& P0, const FSolverVec3& P1, const FSolverVec3& P2, FSolverVec3& dX_dU, FSolverVec3& dX_dV) const
+void FXPBDStretchBiasElementConstraints::CalculateUVStretch(const int32 ConstraintIndex, const FSolverVec3& P0, const FSolverVec3& P1, const FSolverVec3& P2, FSolverVec3& DXDu, FSolverVec3& DXDv) const
 {
 	const FSolverMatrix22& DeltaUVInv = DeltaUVInverse[ConstraintIndex];
 
 	const FSolverVec3 P01 = P1 - P0;
 	const FSolverVec3 P02 = P2 - P0;
-	dX_dU = P01 * DeltaUVInv.M[0] + P02 * DeltaUVInv.M[1];
-	dX_dV = P01 * DeltaUVInv.M[2] + P02 * DeltaUVInv.M[3];
+	DXDu = P01 * DeltaUVInv.M[0] + P02 * DeltaUVInv.M[1];
+	DXDv = P01 * DeltaUVInv.M[2] + P02 * DeltaUVInv.M[3];
 }
 
 static FSolverReal CalcDLambda(const FSolverReal StiffnessValue, const FSolverReal Dt, const FSolverReal Lambda, const FSolverReal Damping,
@@ -685,11 +685,11 @@ void FXPBDStretchBiasElementConstraints::ApplyHelper(FSolverParticles& Particles
 
 	const FSolverReal CombinedInvMass = Particles.InvM(i0) + Particles.InvM(i1) + Particles.InvM(i2);
 
-	FSolverVec3 dX_dU, dX_dV;
-	CalculateUVStretch(ConstraintIndex, Particles.P(i0), Particles.P(i1), Particles.P(i2), dX_dU, dX_dV);
+	FSolverVec3 DXDu, DXDv;
+	CalculateUVStretch(ConstraintIndex, Particles.P(i0), Particles.P(i1), Particles.P(i2), DXDu, DXDv);
 
-	const FSolverReal dX_dU_length = dX_dU.Length();
-	const FSolverReal dX_dV_length = dX_dV.Length();
+	const FSolverReal DXDuLength = DXDu.Length();
+	const FSolverReal DXDvLength = DXDv.Length();
 	auto SafeRecip = [](const FSolverReal Len, const FSolverReal Fallback)
 	{
 		if (Len > UE_SMALL_NUMBER)
@@ -698,29 +698,29 @@ void FXPBDStretchBiasElementConstraints::ApplyHelper(FSolverParticles& Particles
 		}
 		return Fallback;
 	};
-	const FSolverReal OneOverDxDuLen = SafeRecip(dX_dU_length, 0.f);
-	const FSolverReal OneOverDxDvLen = SafeRecip(dX_dV_length, 0.f);
+	const FSolverReal OneOverDXDuLen = SafeRecip(DXDuLength, 0.f);
+	const FSolverReal OneOverDXDvLen = SafeRecip(DXDvLength, 0.f);
 
-	const FSolverVec3 dX_dU_normalized = dX_dU * OneOverDxDuLen;
-	const FSolverVec3 dX_dV_normalized = dX_dV * OneOverDxDvLen;
+	const FSolverVec3 DXDuNormalized = DXDu * OneOverDXDuLen;
+	const FSolverVec3 DXDvNormalized = DXDv * OneOverDXDvLen;
 
 	// constraints
-	const FSolverReal Cu = (dX_dU_length - RestStretchLengths[ConstraintIndex][0]*WarpScaleValue); // stretch in warp direction
-	const FSolverReal Cv = (dX_dV_length - RestStretchLengths[ConstraintIndex][1]*WeftScaleValue); // stretch in weft direction
-	const FSolverReal Cs = FSolverVec3::DotProduct(dX_dU_normalized, dX_dV_normalized);
+	const FSolverReal Cu = (DXDuLength - RestStretchLengths[ConstraintIndex][0]*WarpScaleValue); // stretch in warp direction
+	const FSolverReal Cv = (DXDvLength - RestStretchLengths[ConstraintIndex][1]*WeftScaleValue); // stretch in weft direction
+	const FSolverReal Cs = FSolverVec3::DotProduct(DXDuNormalized, DXDvNormalized);
 
 
 	const FSolverMatrix22& DeltaUVInv = DeltaUVInverse[ConstraintIndex];
 
-	const FSolverVec3 dCu_dX1 = dX_dU_normalized * DeltaUVInv.M[0];
-	const FSolverVec3 dCu_dX2 = dX_dU_normalized * DeltaUVInv.M[1];
-	const FSolverVec3 dCu_dX0 = -dCu_dX1 - dCu_dX2;
-	const FSolverVec3 dCv_dX1 = dX_dV_normalized * DeltaUVInv.M[2];
-	const FSolverVec3 dCv_dX2 = dX_dV_normalized * DeltaUVInv.M[3];
-	const FSolverVec3 dCv_dX0 = -dCv_dX1 - dCv_dX2;
-	const FSolverVec3 dCs_dX1 = (dX_dV_normalized * OneOverDxDuLen * DeltaUVInv.M[0] + dX_dU_normalized * OneOverDxDvLen * DeltaUVInv.M[2]);
-	const FSolverVec3 dCs_dX2 = (dX_dV_normalized * OneOverDxDuLen * DeltaUVInv.M[1] + dX_dU_normalized * OneOverDxDvLen * DeltaUVInv.M[3]);
-	const FSolverVec3 dCs_dX0 = -dCs_dX1 - dCs_dX2;
+	const FSolverVec3 GradCu1 = DXDuNormalized * DeltaUVInv.M[0];
+	const FSolverVec3 GradCu2 = DXDuNormalized * DeltaUVInv.M[1];
+	const FSolverVec3 GradCu0 = -GradCu1 - GradCu2;
+	const FSolverVec3 GradCv1 = DXDvNormalized * DeltaUVInv.M[2];
+	const FSolverVec3 GradCv2 = DXDvNormalized * DeltaUVInv.M[3];
+	const FSolverVec3 GradCv0 = -GradCv1 - GradCv2;
+	const FSolverVec3 GradCs1 = (DXDvNormalized * OneOverDXDuLen * DeltaUVInv.M[0] + DXDuNormalized * OneOverDXDvLen * DeltaUVInv.M[2]);
+	const FSolverVec3 GradCs2 = (DXDvNormalized * OneOverDXDuLen * DeltaUVInv.M[1] + DXDuNormalized * OneOverDXDvLen * DeltaUVInv.M[3]);
+	const FSolverVec3 GradCs0 = -GradCs1 - GradCs2;
 
 	const FSolverVec3 V0TimesDt = Particles.P(i0) - Particles.X(i0);
 	const FSolverVec3 V1TimesDt = Particles.P(i1) - Particles.X(i1);
@@ -735,15 +735,15 @@ void FXPBDStretchBiasElementConstraints::ApplyHelper(FSolverParticles& Particles
 	FSolverVec3& Lambda = Lambdas[ConstraintIndex];
 	FSolverVec3 DLambda;
 	DLambda[0] = CalcDLambda(FinalStiffnesses[0], Dt, Lambda[0], Damping[0], Particles.InvM(i0), Particles.InvM(i1), Particles.InvM(i2),
-		V0TimesDt, V1TimesDt, V2TimesDt, Cu, dCu_dX0, dCu_dX1, dCu_dX2);
+		V0TimesDt, V1TimesDt, V2TimesDt, Cu, GradCu0, GradCu1, GradCu2);
 	DLambda[1] = CalcDLambda(FinalStiffnesses[1], Dt, Lambda[1], Damping[1], Particles.InvM(i0), Particles.InvM(i1), Particles.InvM(i2),
-		V0TimesDt, V1TimesDt, V2TimesDt, Cv, dCv_dX0, dCv_dX1, dCv_dX2);
+		V0TimesDt, V1TimesDt, V2TimesDt, Cv, GradCv0, GradCv1, GradCv2);
 	DLambda[2] = CalcDLambda(FinalStiffnesses[2], Dt, Lambda[2], Damping[2], Particles.InvM(i0), Particles.InvM(i1), Particles.InvM(i2),
-		V0TimesDt, V1TimesDt, V2TimesDt, Cs, dCs_dX0, dCs_dX1, dCs_dX2);
+		V0TimesDt, V1TimesDt, V2TimesDt, Cs, GradCs0, GradCs1, GradCs2);
 
-	Particles.P(i0) -= Particles.InvM(i0) * (DLambda[0] * dCu_dX0 + DLambda[1] * dCv_dX0 + DLambda[2] * dCs_dX0);
-	Particles.P(i1) -= Particles.InvM(i1) * (DLambda[0] * dCu_dX1 + DLambda[1] * dCv_dX1 + DLambda[2] * dCs_dX1);
-	Particles.P(i2) -= Particles.InvM(i2) * (DLambda[0] * dCu_dX2 + DLambda[1] * dCv_dX2 + DLambda[2] * dCs_dX2);
+	Particles.P(i0) -= Particles.InvM(i0) * (DLambda[0] * GradCu0 + DLambda[1] * GradCv0 + DLambda[2] * GradCs0);
+	Particles.P(i1) -= Particles.InvM(i1) * (DLambda[0] * GradCu1 + DLambda[1] * GradCv1 + DLambda[2] * GradCs1);
+	Particles.P(i2) -= Particles.InvM(i2) * (DLambda[0] * GradCu2 + DLambda[1] * GradCv2 + DLambda[2] * GradCs2);
 
 	Lambda += DLambda;
 }
diff --git a/Engine/Source/Runtime/Experimental/Chaos/Private/Chaos/XPBDStretchBiasElementConstraints.ispc b/Engine/Source/Runtime/Experimental/Chaos/Private/Chaos/XPBDStretchBiasElementConstraints.ispc
--- a/Engine/Source/Runtime/Experimental/Chaos/Private/Chaos/XPBDStretchBiasElementConstraints.ispc
+++ b/Engine/Source/Runtime/Experimental/Chaos/Private/Chaos/XPBDStretchBiasElementConstraints.ispc
@@ -13,40 +13,40 @@ static inline varying float SafeDivide(const varying float Numerator, const vary
 static void CalcConstraintsAndGrads(const FVector3f& P0, const FVector3f& P1, const FVector3f& P2, const FVector4f& DeltaUVInv,
 	const FVector2f& RestStretchLengths, const float WarpScale, const float WeftScale,
 	float& Cu, float& Cv, float& Cs, 
-	FVector3f& dCu_dX1, FVector3f& dCu_dX2, FVector3f& dCv_dX1, FVector3f& dCv_dX2, FVector3f& dCs_dX1, FVector3f& dCs_dX2)
+	FVector3f& GradCu1, FVector3f& GradCu2, FVector3f& GradCv1, FVector3f& GradCv2, FVector3f& GradCs1, FVector3f& GradCs2)
 {
 	const FVector3f P01 = P1 - P0;
 	const FVector3f P02 = P2 - P0;
-	const FVector3f dX_dU = DeltaUVInv.V[0] * P01 + DeltaUVInv.V[1] * P02;
-	const FVector3f dX_dV = DeltaUVInv.V[2] * P01 + DeltaUVInv.V[3] * P02;
+	const FVector3f DXDu = DeltaUVInv.V[0] * P01 + DeltaUVInv.V[1] * P02;
+	const FVector3f DXDv = DeltaUVInv.V[2] * P01 + DeltaUVInv.V[3] * P02;
 
-	const float dX_dU_length = VectorSize(dX_dU);
-	const float dX_dV_length = VectorSize(dX_dV);
+	const float DXDuLength = VectorSize(DXDu);
+	const float DXDvLength = VectorSize(DXDv);
 
-	const float OneOverDxDuLen = SafeDivide(1.f, dX_dU_length);
-	const float OneOverDxDvLen = SafeDivide(1.f, dX_dV_length);
+	const float OneOverDXDuLen = SafeDivide(1.f, DXDuLength);
+	const float OneOverDXDvLen = SafeDivide(1.f, DXDvLength);
 
-	const FVector3f dX_dU_normalized = OneOverDxDuLen * dX_dU;
-	const FVector3f dX_dV_normalized = OneOverDxDvLen * dX_dV;
+	const FVector3f DXDuNormalized = OneOverDXDuLen * DXDu;
+	const FVector3f DXDvNormalized = OneOverDXDvLen * DXDv;
 
-	Cu = dX_dU_length - RestStretchLengths.V[0] * WarpScale;
-	Cv = dX_dV_length - RestStretchLengths.V[1] * WeftScale;
-	Cs = VectorDot(dX_dU_normalized, dX_dV_normalized);
+	Cu = DXDuLength - RestStretchLengths.V[0] * WarpScale;
+	Cv = DXDvLength - RestStretchLengths.V[1] * WeftScale;
+	Cs = VectorDot(DXDuNormalized, DXDvNormalized);
 
-	dCu_dX1 = dX_dU_normalized * DeltaUVInv.V[0];
-	dCu_dX2 = dX_dU_normalized * DeltaUVInv.V[1];
-	dCv_dX1 = dX_dV_normalized * DeltaUVInv.V[2];
-	dCv_dX2 = dX_dV_normalized * DeltaUVInv.V[3];
-	dCs_dX1 = (dX_dV_normalized * OneOverDxDuLen * DeltaUVInv.V[0] + dX_dU_normalized * OneOverDxDvLen * DeltaUVInv.V[2]);
-	dCs_dX2 = (dX_dV_normalized * OneOverDxDuLen * DeltaUVInv.V[1] + dX_dU_normalized * OneOverDxDvLen * DeltaUVInv.V[3]);
+	GradCu1 = DXDuNormalized * DeltaUVInv.V[0];
+	GradCu2 = DXDuNormalized * DeltaUVInv.V[1];
+	GradCv1 = DXDvNormalized * DeltaUVInv.V[2];
+	GradCv2 = DXDvNormalized * DeltaUVInv.V[3];
+	GradCs1 = (DXDvNormalized * OneOverDXDuLen * DeltaUVInv.V[0] + DXDuNormalized * OneOverDXDvLen * DeltaUVInv.V[2]);
+	GradCs2 = (DXDvNormalized * OneOverDXDuLen * DeltaUVInv.V[1] + DXDuNormalized * OneOverDXDvLen * DeltaUVInv.V[3]);
 }
 
 static inline float CalcDLambda(const float StiffnessValue, const float Dt, const float MinStiffness, const float Lambda,
 	const float InvM0, const float InvM1, const float InvM2,
-	const float C, const FVector3f& dC_dX1, const FVector3f& dC_dX2)
+	const float C, const FVector3f& GradC1, const FVector3f& GradC2)
 {
 	const float Alpha = 1.f / (StiffnessValue  * Dt * Dt);
-	const float Denom = InvM0 * VectorSizeSquared(dC_dX1 + dC_dX2) + InvM1 * VectorSizeSquared(dC_dX1) + InvM2 * VectorSizeSquared(dC_dX2) + Alpha;
+	const float Denom = InvM0 * VectorSizeSquared(GradC1 + GradC2) + InvM1 * VectorSizeSquared(GradC1) + InvM2 * VectorSizeSquared(GradC2) + Alpha;
 	const float DLambda = (C - Alpha * Lambda) / Denom;
 	return select(StiffnessValue < MinStiffness, (uniform float)0, DLambda);
 }
@@ -59,12 +59,12 @@ static inline float CalcDamping(const float StiffnessValue, const float Combined
 static inline float CalcDLambda(const float StiffnessValue, const float Dt, const float MinStiffness, const float Lambda, const float Damping,
 	const float InvM0, const float InvM1, const float InvM2,
 	const FVector3f& V0TimesDt, const FVector3f& V1TimesDt, const FVector3f& V2TimesDt, 
-	const float C, const FVector3f& dC_dX1, const FVector3f& dC_dX2)
+	const float C, const FVector3f& GradC1, const FVector3f& GradC2)
 {
 	const float Alpha = 1.f / (StiffnessValue  * Dt * Dt);
 	const float Gamma = Alpha * Damping * Dt;
-	const float DampingTerm = Gamma * (VectorDot(V1TimesDt, dC_dX1) + VectorDot(V2TimesDt, dC_dX2) - VectorDot(V0TimesDt, dC_dX1 + dC_dX2));
-	const float Denom = (1. + Gamma) * (InvM0 * VectorSizeSquared(dC_dX1 + dC_dX2) + InvM1 * VectorSizeSquared(dC_dX1) + InvM2 * VectorSizeSquared(dC_dX2)) + Alpha;
+	const float DampingTerm = Gamma * (VectorDot(V1TimesDt, GradC1) + VectorDot(V2TimesDt, GradC2) - VectorDot(V0TimesDt, GradC1 + GradC2));
+	const float Denom = (1. + Gamma) * (InvM0 * VectorSizeSquared(GradC1 + GradC2) + InvM1 * VectorSizeSquared(GradC1) + InvM2 * VectorSizeSquared(GradC2)) + Alpha;
 	const float DLambda = (C - Alpha * Lambda + DampingTerm) / Denom;
 	return select(StiffnessValue < MinStiffness, (uniform float)0, DLambda);
 }
@@ -106,24 +106,24 @@ export void ApplyXPBDStretchBiasConstraints(uniform FVector4f PandInvM[],
 		const varying FVector2f RestStretchLength = VectorLoad(&RestStretchLengths[extract(i,0)]);
 
 		varying float Cu, Cv, Cs;
-		varying FVector3f dCu_dX1, dCu_dX2, dCv_dX1, dCv_dX2, dCs_dX1, dCs_dX2;
+		varying FVector3f GradCu1, GradCu2, GradCv1, GradCv2, GradCs1, GradCs2;
 		CalcConstraintsAndGrads(P0, P1, P2, DeltaUVInverse, RestStretchLength, WarpScale, WeftScale, Cu, Cv, Cs,
-			dCu_dX1, dCu_dX2, dCv_dX1, dCv_dX2, dCs_dX1, dCs_dX2);
+			GradCu1, GradCu2, GradCv1, GradCv2, GradCs1, GradCs2);
 		
 		FVector3f Lambda = VectorLoad(&Lambdas[extract(i,0)]);
 		const varying FVector3f StiffnessScale = VectorLoad(&StiffnessScales[extract(i,0)]);
 		const varying FVector3f FinalStiffness = Stiffness * StiffnessScale;
 
-		const varying float DLambdaU = CalcDLambda(FinalStiffness.V[0], Dt, MinStiffness, Lambda.V[0], M0, M1, M2, Cu, dCu_dX1, dCu_dX2);
-		const varying float DLambdaV = CalcDLambda(FinalStiffness.V[1], Dt, MinStiffness, Lambda.V[1], M0, M1, M2, Cv, dCv_dX1, dCv_dX2);
-		const varying float DLambdaS = CalcDLambda(FinalStiffness.V[2], Dt, MinStiffness, Lambda.V[2], M0, M1, M2, Cs, dCs_dX1, dCs_dX2);
+		const varying float DLambdaU = CalcDLambda(FinalStiffness.V[0], Dt, MinStiffness, Lambda.V[0], M0, M1, M2, Cu, GradCu1, GradCu2);
+		const varying float DLambdaV = CalcDLambda(FinalStiffness.V[1], Dt, MinStiffness, Lambda.V[1], M0, M1, M2, Cv, GradCv1, GradCv2);
+		const varying float DLambdaS = CalcDLambda(FinalStiffness.V[2], Dt, MinStiffness, Lambda.V[2], M0, M1, M2, Cs, GradCs1, GradCs2);
 		const varying FVector3f DLambda = {{DLambdaU, DLambdaV, DLambdaS}};
 
 		Lambda = Lambda + DLambda;
 		VectorStore(&Lambdas[extract(i,0)], Lambda);
 		
-		const varying FVector3f DP1 = DLambdaU * dCu_dX1 + DLambdaV * dCv_dX1 + DLambdaS * dCs_dX1;
-		const varying FVector3f DP2 = DLambdaU * dCu_dX2 + DLambdaV * dCv_dX2 + DLambdaS * dCs_dX2;
+		const varying FVector3f DP1 = DLambdaU * GradCu1 + DLambdaV * GradCv1 + DLambdaS * GradCs1;
+		const varying FVector3f DP2 = DLambdaU * GradCu2 + DLambdaV * GradCv2 + DLambdaS * GradCs2;
 
 		if (M0 > 0)
 		{
@@ -179,9 +179,9 @@ export void ApplyXPBDStretchBiasConstraintsWithDamping(uniform FVector4f PandInv
 		const varying FVector2f RestStretchLength = VectorLoad(&RestStretchLengths[extract(i,0)]);
 
 		varying float Cu, Cv, Cs;
-		varying FVector3f dCu_dX1, dCu_dX2, dCv_dX1, dCv_dX2, dCs_dX1, dCs_dX2;
+		varying FVector3f GradCu1, GradCu2, GradCv1, GradCv2, GradCs1, GradCs2;
 		CalcConstraintsAndGrads(P0, P1, P2, DeltaUVInverse, RestStretchLength, WarpScale, WeftScale, Cu, Cv, Cs,
-			dCu_dX1, dCu_dX2, dCv_dX1, dCv_dX2, dCs_dX1, dCs_dX2);
+			GradCu1, GradCu2, GradCv1, GradCv2, GradCs1, GradCs2);
 		
 		const varying FVector3f X0 = VectorGather(&X[i0]);
 		const varying FVector3f X1 = VectorGather(&X[i1]);
@@ -196,20 +196,20 @@ export void ApplyXPBDStretchBiasConstraintsWithDamping(uniform FVector4f PandInv
 		const varying FVector3f FinalStiffness = Stiffness * StiffnessScale;
 		const varying float DLambdaU = CalcDLambda(FinalStiffness.V[0], Dt, MinStiffness, Lambda.V[0], 
 			CalcDamping(FinalStiffness.V[0], CombinedInvMass, DampingRatio),
-			M0, M1, M2, V0TimesDt, V1TimesDt, V2TimesDt, Cu, dCu_dX1, dCu_dX2);
+			M0, M1, M2, V0TimesDt, V1TimesDt, V2TimesDt, Cu, GradCu1, GradCu2);
 		const varying float DLambdaV = CalcDLambda(FinalStiffness.V[1], Dt, MinStiffness, Lambda.V[1], 
 			CalcDamping(FinalStiffness.V[1], CombinedInvMass, DampingRatio),
-			M0, M1, M2, V0TimesDt, V1TimesDt, V2TimesDt, Cv, dCv_dX1, dCv_dX2);
+			M0, M1, M2, V0TimesDt, V1TimesDt, V2TimesDt, Cv, GradCv1, GradCv2);
 		const varying float DLambdaS = CalcDLambda(FinalStiffness.V[2], Dt, MinStiffness, Lambda.V[2], 
 			CalcDamping(FinalStiffness.V[2], CombinedInvMass, DampingRatio),
-			M0, M1, M2, V0TimesDt, V1TimesDt, V2TimesDt, Cs, dCs_dX1, dCs_dX2);
+			M0, M1, M2, V0TimesDt, V1TimesDt, V2TimesDt, Cs, GradCs1, GradCs2);
 		const varying FVector3f DLambda = {{DLambdaU, DLambdaV, DLambdaS}};
 
 		Lambda = Lambda + DLambda;
 		VectorStore(&Lambdas[extract(i,0)], Lambda);
 		
-		const varying FVector3f DP1 = DLambdaU * dCu_dX1 + DLambdaV * dCv_dX1 + DLambdaS * dCs_dX1;
-		const varying FVector3f DP2 = DLambdaU * dCu_dX2 + DLambdaV * dCv_dX2 + DLambdaS * dCs_dX2;
+		const varying FVector3f DP1 = DLambdaU * GradCu1 + DLambdaV * GradCv1 + DLambdaS * GradCs1;
+		const varying FVector3f DP2 = DLambdaU * GradCu2 + DLambdaV * GradCv2 + DLambdaS * GradCs2;
 
 		if (M0 > 0)
 		{
@@ -285,9 +285,9 @@ export void ApplyXPBDStretchBiasConstraintsWithMaps(uniform FVector4f PandInvM[]
 		const varying float WeftScale = WeftScaleTable[WeftScaleIndex];
 
 		varying float Cu, Cv, Cs;
-		varying FVector3f dCu_dX1, dCu_dX2, dCv_dX1, dCv_dX2, dCs_dX1, dCs_dX2;
+		varying FVector3f GradCu1, GradCu2, GradCv1, GradCv2, GradCs1, GradCs2;
 		CalcConstraintsAndGrads(P0, P1, P2, DeltaUVInverse, RestStretchLength, WarpScale, WeftScale, Cu, Cv, Cs,
-			dCu_dX1, dCu_dX2, dCv_dX1, dCv_dX2, dCs_dX1, dCs_dX2);		
+			GradCu1, GradCu2, GradCv1, GradCv2, GradCs1, GradCs2);		
 		
 		#pragma ignore warning(perf)
 		const varying float StiffnessWarp = StiffnessWarpTable[StiffnessWarpIndex];
@@ -301,16 +301,16 @@ export void ApplyXPBDStretchBiasConstraintsWithMaps(uniform FVector4f PandInvM[]
 		const varying FVector3f StiffnessScale = VectorLoad(&StiffnessScales[extract(i,0)]);
 		const varying FVector3f FinalStiffness = Stiffness * StiffnessScale;
 
-		const varying float DLambdaU = CalcDLambda(FinalStiffness.V[0], Dt, MinStiffness, Lambda.V[0], M0, M1, M2, Cu, dCu_dX1, dCu_dX2);
-		const varying float DLambdaV = CalcDLambda(FinalStiffness.V[1], Dt, MinStiffness, Lambda.V[1], M0, M1, M2, Cv, dCv_dX1, dCv_dX2);
-		const varying float DLambdaS = CalcDLambda(FinalStiffness.V[2], Dt, MinStiffness, Lambda.V[2], M0, M1, M2, Cs, dCs_dX1, dCs_dX2);
+		const varying float DLambdaU = CalcDLambda(FinalStiffness.V[0], Dt, MinStiffness, Lambda.V[0], M0, M1, M2, Cu, GradCu1, GradCu2);
+		const varying float DLambdaV = CalcDLambda(FinalStiffness.V[1], Dt, MinStiffness, Lambda.V[1], M0, M1, M2, Cv, GradCv1, GradCv2);
+		const varying float DLambdaS = CalcDLambda(FinalStiffness.V[2], Dt, MinStiffness, Lambda.V[2], M0, M1, M2, Cs, GradCs1, GradCs2);
 		const varying FVector3f DLambda = {{DLambdaU, DLambdaV, DLambdaS}};
 
 		Lambda = Lambda + DLambda;
 		VectorStore(&Lambdas[extract(i,0)], Lambda);
 		
-		const varying FVector3f DP1 = DLambdaU * dCu_dX1 + DLambdaV * dCv_dX1 + DLambdaS * dCs_dX1;
-		const varying FVector3f DP2 = DLambdaU * dCu_dX2 + DLambdaV * dCv_dX2 + DLambdaS * dCs_dX2;
+		const varying FVector3f DP1 = DLambdaU * GradCu1 + DLambdaV * GradCv1 + DLambdaS * GradCs1;
+		const varying FVector3f DP2 = DLambdaU * GradCu2 + DLambdaV * GradCv2 + DLambdaS * GradCs2;
 
 		if (M0 > 0)
 		{
@@ -391,9 +391,9 @@ export void ApplyXPBDStretchBiasConstraintsWithDampingAndMaps(uniform FVector4f
 		const varying float WeftScale = WeftScaleTable[WeftScaleIndex];
 
 		varying float Cu, Cv, Cs;
-		varying FVector3f dCu_dX1, dCu_dX2, dCv_dX1, dCv_dX2, dCs_dX1, dCs_dX2;
+		varying FVector3f GradCu1, GradCu2, GradCv1, GradCv2, GradCs1, GradCs2;
 		CalcConstraintsAndGrads(P0, P1, P2, DeltaUVInverse, RestStretchLength, WarpScale, WeftScale, Cu, Cv, Cs,
-			dCu_dX1, dCu_dX2, dCv_dX1, dCv_dX2, dCs_dX1, dCs_dX2);
+			GradCu1, GradCu2, GradCv1, GradCv2, GradCs1, GradCs2);
 		
 		const varying FVector3f X0 = VectorGather(&X[i0]);
 		const varying FVector3f X1 = VectorGather(&X[i1]);
@@ -419,20 +419,20 @@ export void ApplyXPBDStretchBiasConstraintsWithDampingAndMaps(uniform FVector4f
 		const varying FVector3f FinalStiffness = Stiffness * StiffnessScale;
 		const varying float DLambdaU = CalcDLambda(FinalStiffness.V[0], Dt, MinStiffness, Lambda.V[0], 
 			CalcDamping(FinalStiffness.V[0], CombinedInvMass, DampingRatio),
-			M0, M1, M2, V0TimesDt, V1TimesDt, V2TimesDt, Cu, dCu_dX1, dCu_dX2);
+			M0, M1, M2, V0TimesDt, V1TimesDt, V2TimesDt, Cu, GradCu1, GradCu2);
 		const varying float DLambdaV = CalcDLambda(FinalStiffness.V[1], Dt, MinStiffness, Lambda.V[1], 
 			CalcDamping(FinalStiffness.V[1], CombinedInvMass, DampingRatio),
-			M0, M1, M2, V0TimesDt, V1TimesDt, V2TimesDt, Cv, dCv_dX1, dCv_dX2);
+			M0, M1, M2, V0TimesDt, V1TimesDt, V2TimesDt, Cv, GradCv1, GradCv2);
 		const varying float DLambdaS = CalcDLambda(FinalStiffness.V[2], Dt, MinStiffness, Lambda.V[2], 
 			CalcDamping(FinalStiffness.V[2], CombinedInvMass, DampingRatio),
-			M0, M1, M2, V0TimesDt, V1TimesDt, V2TimesDt, Cs, dCs_dX1, dCs_dX2);
+			M0, M1, M2, V0TimesDt, V1TimesDt, V2TimesDt, Cs, GradCs1, GradCs2);
 		const varying FVector3f DLambda = {{DLambdaU, DLambdaV, DLambdaS}};
 
 		Lambda = Lambda + DLambda;
 		VectorStore(&Lambdas[extract(i,0)], Lambda);
 		
-		const varying FVector3f DP1 = DLambdaU * dCu_dX1 + DLambdaV * dCv_dX1 + DLambdaS * dCs_dX1;
-		const varying FVector3f DP2 = DLambdaU * dCu_dX2 + DLambdaV * dCv_dX2 + DLambdaS * dCs_dX2;
+		const varying FVector3f DP1 = DLambdaU * GradCu1 + DLambdaV * GradCv1 + DLambdaS * GradCs1;
+		const varying FVector3f DP2 = DLambdaU * GradCu2 + DLambdaV * GradCv2 + DLambdaS * GradCs2;
 
 		if (M0 > 0)
 		{
diff --git a/Engine/Source/Runtime/Experimental/Chaos/Public/Chaos/XPBDStretchBiasElementConstraints.h b/Engine/Source/Runtime/Experimental/Chaos/Public/Chaos/XPBDStretchBiasElementConstraints.h
--- a/Engine/Source/Runtime/Experimental/Chaos/Public/Chaos/XPBDStretchBiasElementConstraints.h
+++ b/Engine/Source/Runtime/Experimental/Chaos/Public/Chaos/XPBDStretchBiasElementConstraints.h
@@ -97,7 +97,7 @@ public:
 	}
 
 	void Apply(FSolverParticles& Particles, const FSolverReal Dt) const;
-	void CalculateUVStretch(const int32 ConstraintIndex, const FSolverVec3& P0, const FSolverVec3& P1, const FSolverVec3& P2, FSolverVec3& dX_dU, FSolverVec3& dX_dV) const;
+	void CalculateUVStretch(const int32 ConstraintIndex, const FSolverVec3& P0, const FSolverVec3& P1, const FSolverVec3& P2, FSolverVec3& DXDu, FSolverVec3& DXDv) const;
 	
 	const TArray<TVec3<int32>>& GetConstraints() const { return Constraints; }
 	const TArray<FSolverVec2> GetRestStretchLengths() const { return RestStretchLengths; }
