Iris
- Fixed bug where we did not immediately resend lost state for properties utilizing multiple changemask bits (arrays) until next time root property bit was modfied.
- Added testcase to catch this.

diff --git a/Engine/Plugins/Runtime/ReplicationSystemTestPlugin/Source/Private/Tests/ReplicationState/TestFastArrayReplicationState.cpp b/Engine/Plugins/Runtime/ReplicationSystemTestPlugin/Source/Private/Tests/ReplicationState/TestFastArrayReplicationState.cpp
--- a/Engine/Plugins/Runtime/ReplicationSystemTestPlugin/Source/Private/Tests/ReplicationState/TestFastArrayReplicationState.cpp
+++ b/Engine/Plugins/Runtime/ReplicationSystemTestPlugin/Source/Private/Tests/ReplicationState/TestFastArrayReplicationState.cpp
@@ -938,5 +938,59 @@ UE_NET_TEST_FIXTURE(FReplicationSystemServerClientTestFixture, TestFastArrayWith
 	UE_NET_ASSERT_EQ(ServerObject->FastArray.ExtraInt, ClientObject->FastArray.ExtraInt);
 }
 
+UE_NET_TEST_FIXTURE(FReplicationSystemServerClientTestFixture, TestLostDataDirtiesPropertyBitWithDataInFlight)
+{
+	UTestFastArrayReplicationState_FastArray_TestClassFastArray* ServerObject = Server->CreateObject<UTestFastArrayReplicationState_FastArray_TestClassFastArray>();
+
+	// Add a client
+	FReplicationSystemTestClient* Client = CreateClient();
+
+	// Send and deliver packet
+	Server->PreSendUpdate();
+	Server->SendAndDeliverTo(Client, true);
+	Server->PostSendUpdate();
+
+	// Store Pointer to objects
+	UTestFastArrayReplicationState_FastArray_TestClassFastArray* ClientObject = Cast<UTestFastArrayReplicationState_FastArray_TestClassFastArray>(Client->GetReplicationBridge()->GetReplicatedObject(ServerObject->NetRefHandle));
+
+	// Add entry to array
+	FTestFastArrayReplicationState_FastArrayItem NewEntry;
+	NewEntry.bRepBool = false;
+	NewEntry.RepInt32 = 0U;
+	ServerObject->FastArray.Edit().Add(NewEntry);
+	NewEntry.bRepBool = false;
+	NewEntry.RepInt32 = 3U;
+	ServerObject->FastArray.Edit().Add(NewEntry);
+
+	// Put data on wire, but hold off delivery
+	Server->PreSendUpdate();
+	Server->SendTo(Client);
+	Server->PostSendUpdate();
+
+	// Modify entry to trigger new update
+	++ServerObject->FastArray.Edit()[0].RepInt32;
+
+	// Put latest state on the wire
+	Server->PreSendUpdate();
+	Server->SendTo(Client);
+	Server->PostSendUpdate();
+
+	// Notify that we dropped the first packet, which should mark property dirty even if we already have the latest state in flight
+	Server->DeliverTo(Client, false);
+
+	// Notify that we delivered second packet
+	Server->DeliverTo(Client, true);
+
+	// Send and deliver packet
+	Server->PreSendUpdate();
+	Server->SendAndDeliverTo(Client, true);
+	Server->PostSendUpdate();
+
+	// Verify that we received the expected data, for both elements
+	UE_NET_ASSERT_EQ(ServerObject->FastArray.GetItemArray().Num(), ClientObject->FastArray.GetItemArray().Num());
+	UE_NET_ASSERT_EQ(ServerObject->FastArray.GetItemArray()[0].RepInt32, ClientObject->FastArray.GetItemArray()[0].RepInt32);
+	UE_NET_ASSERT_EQ(ServerObject->FastArray.GetItemArray()[1].RepInt32, ClientObject->FastArray.GetItemArray()[1].RepInt32);
+}
+
 
 }
diff --git a/Engine/Source/Runtime/Experimental/Iris/Core/Private/Iris/ReplicationSystem/ReplicationOperations.cpp b/Engine/Source/Runtime/Experimental/Iris/Core/Private/Iris/ReplicationSystem/ReplicationOperations.cpp
--- a/Engine/Source/Runtime/Experimental/Iris/Core/Private/Iris/ReplicationSystem/ReplicationOperations.cpp
+++ b/Engine/Source/Runtime/Experimental/Iris/Core/Private/Iris/ReplicationSystem/ReplicationOperations.cpp
@@ -301,7 +301,7 @@ void FReplicationStateOperations::SerializeWithMask(FNetSerializationContext& Co
 		const FReplicationStateMemberChangeMaskDescriptor& MemberChangeMaskDescriptor = Descriptor->MemberChangeMaskDescriptors[MemberIt];
 		const uint32 MemberChangeMaskOffset = ChangeMaskOffset + MemberChangeMaskDescriptor.BitOffset;
 
-		if (ChangeMask.GetBit(MemberChangeMaskOffset))
+		if (ChangeMask.IsAnyBitSet(MemberChangeMaskOffset, MemberChangeMaskDescriptor.BitCount))
 		{
 #if UE_NET_TRACE_ENABLED
 			const int8 Condition = (MemberLifetimeConditionDescriptors != nullptr) ? MemberLifetimeConditionDescriptors[MemberIt].Condition : int8(ELifetimeCondition::COND_None);
@@ -346,7 +346,7 @@ void FReplicationStateOperations::DeserializeWithMask(FNetSerializationContext&
 		const FReplicationStateMemberChangeMaskDescriptor& MemberChangeMaskDescriptor = Descriptor->MemberChangeMaskDescriptors[MemberIt];
 		const uint32 MemberChangeMaskOffset = ChangeMaskOffset + MemberChangeMaskDescriptor.BitOffset;
 
-		if (ChangeMask.GetBit(MemberChangeMaskOffset))
+		if (ChangeMask.IsAnyBitSet(MemberChangeMaskOffset, MemberChangeMaskDescriptor.BitCount))
 		{
 #if UE_NET_TRACE_ENABLED
 			const int8 Condition = (MemberLifetimeConditionDescriptors != nullptr) ? MemberLifetimeConditionDescriptors[MemberIt].Condition : int8(ELifetimeCondition::COND_None);
@@ -392,7 +392,7 @@ void FReplicationStateOperations::SerializeDeltaWithMask(FNetSerializationContex
 		const FReplicationStateMemberChangeMaskDescriptor& MemberChangeMaskDescriptor = Descriptor->MemberChangeMaskDescriptors[MemberIt];
 		const uint32 MemberChangeMaskOffset = ChangeMaskOffset + MemberChangeMaskDescriptor.BitOffset;
 
-		if (ChangeMask.GetBit(MemberChangeMaskOffset))
+		if (ChangeMask.IsAnyBitSet(MemberChangeMaskOffset, MemberChangeMaskDescriptor.BitCount))
 		{
 #if UE_NET_TRACE_ENABLED
 			const int8 Condition = (MemberLifetimeConditionDescriptors != nullptr) ? MemberLifetimeConditionDescriptors[MemberIt].Condition : int8(ELifetimeCondition::COND_None);
@@ -438,7 +438,7 @@ void FReplicationStateOperations::DeserializeDeltaWithMask(FNetSerializationCont
 		const FReplicationStateMemberChangeMaskDescriptor& MemberChangeMaskDescriptor = Descriptor->MemberChangeMaskDescriptors[MemberIt];
 		const uint32 MemberChangeMaskOffset = ChangeMaskOffset + MemberChangeMaskDescriptor.BitOffset;
 
-		if (ChangeMask.GetBit(MemberChangeMaskOffset))
+		if (ChangeMask.IsAnyBitSet(MemberChangeMaskOffset, MemberChangeMaskDescriptor.BitCount))
 		{
 #if UE_NET_TRACE_ENABLED
 			const int8 Condition = (MemberLifetimeConditionDescriptors != nullptr) ? MemberLifetimeConditionDescriptors[MemberIt].Condition : int8(ELifetimeCondition::COND_None);
diff --git a/Engine/Source/Runtime/Net/Core/Public/Net/Core/NetBitArray.h b/Engine/Source/Runtime/Net/Core/Public/Net/Core/NetBitArray.h
--- a/Engine/Source/Runtime/Net/Core/Public/Net/Core/NetBitArray.h
+++ b/Engine/Source/Runtime/Net/Core/Public/Net/Core/NetBitArray.h
@@ -902,7 +902,7 @@ inline bool FNetBitArray::IsAnyBitSet() const
 
 inline bool FNetBitArray::IsAnyBitSet(uint32 StartIndex, uint32 Count) const
 {
-	return FNetBitArrayHelper::IsAnyBitSet(Storage.GetData(), BitCount, StartIndex, Count);
+	return ((Count == 1U) ? GetBit(StartIndex) : FNetBitArrayHelper::IsAnyBitSet(Storage.GetData(), BitCount, StartIndex, Count));
 }
 
 inline bool FNetBitArray::IsNoBitSet() const
@@ -1096,7 +1096,7 @@ bool FNetBitArrayView::IsAnyBitSet() const
 
 bool FNetBitArrayView::IsAnyBitSet(uint32 StartIndex, uint32 Count) const
 {
-	return FNetBitArrayHelper::IsAnyBitSet(Storage, BitCount, StartIndex, Count);
+	return ((Count == 1U) ? GetBit(StartIndex) : FNetBitArrayHelper::IsAnyBitSet(GetData(), BitCount, StartIndex, Count));
 }
 
 bool FNetBitArrayView::IsNoBitSet() const
