Add UPROPERTY markup to UObject pointer to prevent object from being GCd.

diff --git a/Engine/Source/Runtime/Experimental/GeometryCollectionEngine/Public/GeometryCollection/GeometryCollectionISMPoolSubSystem.h b/Engine/Source/Runtime/Experimental/GeometryCollectionEngine/Public/GeometryCollection/GeometryCollectionISMPoolSubSystem.h
--- a/Engine/Source/Runtime/Experimental/GeometryCollectionEngine/Public/GeometryCollection/GeometryCollectionISMPoolSubSystem.h
+++ b/Engine/Source/Runtime/Experimental/GeometryCollectionEngine/Public/GeometryCollection/GeometryCollectionISMPoolSubSystem.h
@@ -29,5 +29,6 @@ public:
 protected:
 
 	/** for now we only use one ISMPool actor per world, but we could extend the system to manage many more and return the right one based on  search criteria */
+	UPROPERTY(Transient)
 	TObjectPtr<AGeometryCollectionISMPoolActor> ISMPoolActor;
 };
