Engine physics changes to support RigidBodyWithControl - support setting kinematic etc dynamically

diff --git a/Engine/Source/Runtime/Engine/Private/PhysicsEngine/ImmediatePhysics/ImmediatePhysicsChaos/ImmediatePhysicsActorHandle_Chaos.cpp b/Engine/Source/Runtime/Engine/Private/PhysicsEngine/ImmediatePhysics/ImmediatePhysicsChaos/ImmediatePhysicsActorHandle_Chaos.cpp
--- a/Engine/Source/Runtime/Engine/Private/PhysicsEngine/ImmediatePhysics/ImmediatePhysicsChaos/ImmediatePhysicsActorHandle_Chaos.cpp
+++ b/Engine/Source/Runtime/Engine/Private/PhysicsEngine/ImmediatePhysics/ImmediatePhysicsChaos/ImmediatePhysicsActorHandle_Chaos.cpp
@@ -190,6 +190,84 @@ namespace ImmediatePhysics_Chaos
 		return true;
 	}
 
+	void FActorHandle::CreateParticleHandle(
+		FBodyInstance*                 BodyInstance,
+		const EActorType               ActorType,
+		const FTransform&              WorldTransform,
+		const FReal                    Mass,
+		const Chaos::FVec3             Inertia,
+		const Chaos::FRigidTransform3& CoMTransform)
+	{
+		using namespace Chaos;
+
+		switch (ActorType)
+		{
+		case EActorType::StaticActor:
+			ParticleHandle = Particles.CreateStaticParticles(1, nullptr, FGeometryParticleParameters())[0];
+			break;
+		case EActorType::KinematicActor:
+			ParticleHandle = Particles.CreateKinematicParticles(1, nullptr, FKinematicGeometryParticleParameters())[0];
+			break;
+		case EActorType::DynamicActor:
+			ParticleHandle = Particles.CreateDynamicParticles(1, nullptr, FPBDRigidParticleParameters())[0];
+			break;
+		}
+
+		if (ParticleHandle != nullptr)
+		{
+			SetWorldTransform(WorldTransform);
+
+			ParticleHandle->SetGeometry(MakeSerializable(Geometry));
+
+			// Set the collision filter data for the shapes to collide with everything.
+			// Even though we already tried to do this when we created the original shapes array, 
+			// that gets thrown away and we need to do it here. This is not a good API
+			FCollisionData CollisionData;
+			CollisionData.SimData.Word1 = 0xFFFFF;
+			CollisionData.SimData.Word3 = 0xFFFFF;
+			CollisionData.bSimCollision = 1;
+			for (const TUniquePtr<FPerShapeData>& Shape : ParticleHandle->ShapesArray())
+			{
+				Shape->SetCollisionData(CollisionData);
+			}
+
+			if (Geometry && Geometry->HasBoundingBox())
+			{
+				ParticleHandle->SetHasBounds(true);
+				ParticleHandle->SetLocalBounds(Geometry->BoundingBox());
+				ParticleHandle->UpdateWorldSpaceState(FRigidTransform3(ParticleHandle->X(), ParticleHandle->R()), FVec3(0));
+			}
+
+			if (FKinematicGeometryParticleHandle* Kinematic = ParticleHandle->CastToKinematicParticle())
+			{
+				Kinematic->SetV(FVector3f::ZeroVector);
+				Kinematic->SetW(FVector3f::ZeroVector);
+			}
+
+			FPBDRigidParticleHandle* Dynamic = ParticleHandle->CastToRigidParticle();
+			if (Dynamic && Dynamic->ObjectState() == EObjectStateType::Dynamic)
+			{
+				FReal MassInv = (Mass > 0.0f) ? 1.0f / Mass : 0.0f;
+				FVec3 InertiaInv = (Mass > 0.0f) ? Inertia.Reciprocal() : FVec3::ZeroVector;
+				Dynamic->SetM(Mass);
+				Dynamic->SetInvM(MassInv);
+				Dynamic->SetCenterOfMass(CoMTransform.GetTranslation());
+				Dynamic->SetRotationOfMass(CoMTransform.GetRotation());
+				Dynamic->SetI(TVec3<FRealSingle>(Inertia.X, Inertia.Y, Inertia.Z));
+				Dynamic->SetInvI(TVec3<FRealSingle>(InertiaInv.X, InertiaInv.Y, InertiaInv.Z));
+				if (BodyInstance != nullptr)
+				{
+					Dynamic->SetInertiaConditioningEnabled(BodyInstance->IsInertiaConditioningEnabled());
+					Dynamic->SetLinearEtherDrag(BodyInstance->LinearDamping);
+					Dynamic->SetAngularEtherDrag(BodyInstance->AngularDamping);
+					Dynamic->SetGravityEnabled(BodyInstance->bEnableGravity);
+					Dynamic->SetUpdateKinematicFromSimulation(BodyInstance->bUpdateKinematicFromSimulation);
+				}
+				Dynamic->SetDisabled(true);
+			}
+		}
+	}
+
 	//
 	// Actor Handle
 	//
@@ -216,72 +294,7 @@ namespace ImmediatePhysics_Chaos
 		FRigidTransform3 CoMTransform = FRigidTransform3::Identity;
 		if (CreateGeometry(BodyInstance, ActorType, Scale, Mass, Inertia, CoMTransform, Geometry, Shapes))
 		{
-			switch (ActorType)
-			{
-			case EActorType::StaticActor:
-				ParticleHandle = Particles.CreateStaticParticles(1, nullptr, FGeometryParticleParameters())[0];
-				break;
-			case EActorType::KinematicActor:
-				ParticleHandle = Particles.CreateKinematicParticles(1, nullptr, FKinematicGeometryParticleParameters())[0];
-				break;
-			case EActorType::DynamicActor:
-				ParticleHandle = Particles.CreateDynamicParticles(1, nullptr, FPBDRigidParticleParameters())[0];
-				break;
-			}
-
-			if (ParticleHandle != nullptr)
-			{
-				SetWorldTransform(Transform);
-
-				ParticleHandle->SetGeometry(MakeSerializable(Geometry));
-
-				// Set the collision filter data for the shapes to collide with everything.
-				// Even though we already tried to do this when we created the original shapes array, 
-				// that gets thrown away and we need to do it here. This is not a good API
-				FCollisionData CollisionData;
-				CollisionData.SimData.Word1 = 0xFFFFF;
-				CollisionData.SimData.Word3 = 0xFFFFF;
-				CollisionData.bSimCollision = 1;
-				for (const TUniquePtr<FPerShapeData>& Shape : ParticleHandle->ShapesArray())
-				{
-					Shape->SetCollisionData(CollisionData);
-				}
-
-				if (Geometry && Geometry->HasBoundingBox())
-				{
-					ParticleHandle->SetHasBounds(true);
-					ParticleHandle->SetLocalBounds(Geometry->BoundingBox());
-					ParticleHandle->UpdateWorldSpaceState(FRigidTransform3(ParticleHandle->X(), ParticleHandle->R()), FVec3(0));
-				}
-
-				if (FKinematicGeometryParticleHandle* Kinematic = ParticleHandle->CastToKinematicParticle())
-				{
-					Kinematic->SetV(FVector3f::ZeroVector);
-					Kinematic->SetW(FVector3f::ZeroVector);
-				}
-
-				FPBDRigidParticleHandle* Dynamic = ParticleHandle->CastToRigidParticle();
-				if (Dynamic && Dynamic->ObjectState() == EObjectStateType::Dynamic)
-				{
-					FReal MassInv = (Mass > 0.0f) ? 1.0f / Mass : 0.0f;
-					FVec3 InertiaInv = (Mass > 0.0f) ? Inertia.Reciprocal() : FVec3::ZeroVector;
-					Dynamic->SetM(Mass);
-					Dynamic->SetInvM(MassInv);
-					Dynamic->SetCenterOfMass(CoMTransform.GetTranslation());
-					Dynamic->SetRotationOfMass(CoMTransform.GetRotation());
-					Dynamic->SetI(TVec3<FRealSingle>( Inertia.X, Inertia.Y, Inertia.Z ));
-					Dynamic->SetInvI(TVec3<FRealSingle>(InertiaInv.X, InertiaInv.Y, InertiaInv.Z ));
-					if (BodyInstance != nullptr)
-					{
-						Dynamic->SetInertiaConditioningEnabled(BodyInstance->IsInertiaConditioningEnabled());
-						Dynamic->SetLinearEtherDrag(BodyInstance->LinearDamping);
-						Dynamic->SetAngularEtherDrag(BodyInstance->AngularDamping);
-						Dynamic->SetGravityEnabled(BodyInstance->bEnableGravity);
-						Dynamic->SetUpdateKinematicFromSimulation(BodyInstance->bUpdateKinematicFromSimulation);
-					}
-					Dynamic->SetDisabled(true);
-				}
-			}
+			CreateParticleHandle(BodyInstance, ActorType, Transform, Mass, Inertia, CoMTransform);
 		}
 	}
 
@@ -357,11 +370,40 @@ namespace ImmediatePhysics_Chaos
 		}
 	}
 
-	void FActorHandle::SetIsKinematic(bool bKinematic)
+	bool FActorHandle::SetIsKinematic(bool bKinematic)
 	{
-#if IMMEDIATEPHYSICS_CHAOS_TODO
-		// This needs to destroy and recreate the particle
-#endif
+		using namespace Chaos;
+
+		if (ParticleHandle == nullptr)
+		{
+			return false;
+		}
+
+		EParticleType CurrentParticleType = ParticleHandle->GetParticleType();
+
+		if (CurrentParticleType == EParticleType::Kinematic && bKinematic)
+		{
+			return true;
+		}
+		if (CurrentParticleType == EParticleType::Rigid)
+		{
+			if (FPBDRigidParticleHandle* Dynamic = ParticleHandle->CastToRigidParticle())
+			{
+				// Note that the state might be dynamic, sleeping, or kinematic
+				if (Dynamic->ObjectState() != EObjectStateType::Kinematic && bKinematic)
+				{
+					Dynamic->SetObjectStateLowLevel(EObjectStateType::Kinematic);
+					return true;
+				}
+				else if (Dynamic->ObjectState() == EObjectStateType::Kinematic && !bKinematic)
+				{
+					Dynamic->SetObjectStateLowLevel(EObjectStateType::Dynamic);
+					return true;
+				}
+				return true;
+			}
+		}
+		return false;
 	}
 
 	bool FActorHandle::GetIsKinematic() const
diff --git a/Engine/Source/Runtime/Engine/Public/Physics/ImmediatePhysics/ImmediatePhysicsChaos/ImmediatePhysicsActorHandle_Chaos.h b/Engine/Source/Runtime/Engine/Public/Physics/ImmediatePhysics/ImmediatePhysicsChaos/ImmediatePhysicsActorHandle_Chaos.h
--- a/Engine/Source/Runtime/Engine/Public/Physics/ImmediatePhysics/ImmediatePhysicsChaos/ImmediatePhysicsActorHandle_Chaos.h
+++ b/Engine/Source/Runtime/Engine/Public/Physics/ImmediatePhysics/ImmediatePhysicsChaos/ImmediatePhysicsActorHandle_Chaos.h
@@ -28,8 +28,13 @@ namespace ImmediatePhysics_Chaos
 		/** Sets the world transform, maintains velocity etc.*/
 		void SetWorldTransform(const FTransform& WorldTM);
 
-		/** Make a body kinematic, or non-kinematic */
-		void SetIsKinematic(bool bKinematic);
+		/** 
+		 * Normally used to modify a EParticleType::Rigid between kinematic and dynamic. Returns true if that 
+		 * was possible (even if nothing changed). Returns false if it fails - e.g. because there was no current 
+		 * ParticleHandle, or if the particle type was not EParticleType::Rigid (unless the current type was
+		 * EParticleType::Kinematic and bKinematic was true)
+		 */
+		bool SetIsKinematic(bool bKinematic);
 
 		/** Is the actor kinematic */
 		bool GetIsKinematic() const;
@@ -134,6 +139,14 @@ namespace ImmediatePhysics_Chaos
 	private:
 		FKinematicTarget& GetKinematicTarget();
 
+		void CreateParticleHandle(
+			FBodyInstance*                 BodyInstance,
+			const EActorType               ActorType,
+			const FTransform&              WorldTransform,
+			const FReal                    Mass,
+			const Chaos::FVec3             Inertia,
+			const Chaos::FRigidTransform3& CoMTransform);
+
 		friend struct FSimulation;
 		friend struct FJointHandle;
 
diff --git a/Engine/Source/Runtime/Experimental/Chaos/Private/Chaos/Evolution/PBDMinEvolution.cpp b/Engine/Source/Runtime/Experimental/Chaos/Private/Chaos/Evolution/PBDMinEvolution.cpp
--- a/Engine/Source/Runtime/Experimental/Chaos/Private/Chaos/Evolution/PBDMinEvolution.cpp
+++ b/Engine/Source/Runtime/Experimental/Chaos/Private/Chaos/Evolution/PBDMinEvolution.cpp
@@ -281,15 +281,24 @@ namespace Chaos
 		check(StepFraction > (FReal)0);
 		check(StepFraction <= (FReal)1);
 
-		// @todo(ccaulfield): optimize. Depending on the number of kinematics relative to the number that have 
-		// targets set, it may be faster to process a command list rather than iterate over them all each frame. 
-		const FReal MinDt = 1e-6f;
-		for (auto& Particle : Particles.GetActiveKinematicParticlesView())
+		const bool IsLastStep = (FMath::IsNearlyEqual(StepFraction, (FReal)1, (FReal)UE_KINDA_SMALL_NUMBER));
+
+		const auto& ApplyDynamicParticleKinematicTarget = 
+		[Dt, StepFraction, IsLastStep]
+		(FTransientPBDRigidParticleHandle& Particle, const int32 ParticleIndex)
+		-> void
 		{
+			if (!Particle.IsKinematic())
+			{
+				return;
+			}
+
 			TKinematicTarget<FReal, 3>& KinematicTarget = Particle.KinematicTarget();
 			const FVec3 CurrentX = Particle.X();
 			const FRotation3 CurrentR = Particle.R();
+			constexpr FReal MinDt = 1e-6f;
 
+			bool bMoved = false;
 			switch (KinematicTarget.GetMode())
 			{
 			case EKinematicTargetMode::None:
@@ -299,8 +308,8 @@ namespace Chaos
 			case EKinematicTargetMode::Reset:
 			{
 				// Reset velocity and then switch to do-nothing mode
-				Particle.V() = FVec3(0);
-				Particle.W() = FVec3(0);
+				Particle.V() = FVec3(0, 0, 0);
+				Particle.W() = FVec3(0, 0, 0);
 				KinematicTarget.SetMode(EKinematicTargetMode::None);
 				break;
 			}
@@ -311,7 +320,7 @@ namespace Chaos
 				// Target positions only need to be processed once, and we reset the velocity next frame (if no new target is set)
 				FVec3 NewX;
 				FRotation3 NewR;
-				if (FMath::IsNearlyEqual(StepFraction, (FReal)1, (FReal)UE_KINDA_SMALL_NUMBER))
+				if (IsLastStep)
 				{
 					NewX = KinematicTarget.GetTarget().GetLocation();
 					NewR = KinematicTarget.GetTarget().GetRotation();
@@ -322,30 +331,168 @@ namespace Chaos
 					// as a reminder, stepfraction is the remaing fraction of the step from the remaining steps
 					// for total of 4 steps and current step of 2, this will be 1/3 ( 1 step passed, 3 steps remains )
 					NewX = FVec3::Lerp(CurrentX, KinematicTarget.GetTarget().GetLocation(), StepFraction);
-					NewR = FRotation3::Slerp(CurrentR, KinematicTarget.GetTarget().GetRotation(), (decltype(FQuat::X))StepFraction);		// LWC_TODO: Remove decltype cast once FQuat supports variants
+					NewR = FRotation3::Slerp(CurrentR, KinematicTarget.GetTarget().GetRotation(), decltype(FQuat::X)(StepFraction));
 				}
+
+				const bool bPositionChanged = !FVec3::IsNearlyEqual(NewX, CurrentX, UE_SMALL_NUMBER);
+				const bool bRotationChanged = !FRotation3::IsNearlyEqual(NewR, CurrentR, UE_SMALL_NUMBER);
+				bMoved = bPositionChanged || bRotationChanged;
+				FVec3 NewV = FVec3(0);
+				FVec3 NewW = FVec3(0);
 				if (Dt > MinDt)
 				{
-					Particle.V() = FVec3::CalculateVelocity(CurrentX, NewX, Dt);
-					Particle.W() = FRotation3::CalculateAngularVelocity(CurrentR, NewR, Dt);
+					if (bPositionChanged)
+					{
+						NewV = FVec3::CalculateVelocity(CurrentX, NewX, Dt);
+					}
+					if (bRotationChanged)
+					{
+						NewW = FRotation3::CalculateAngularVelocity(CurrentR, NewR, Dt);
+					}
 				}
 				Particle.X() = NewX;
 				Particle.R() = NewR;
+				Particle.V() = NewV;
+				Particle.W() = NewW;
+
 				break;
 			}
 
 			case EKinematicTargetMode::Velocity:
 			{
 				// Move based on velocity
+				bMoved = true;
 				Particle.X() = Particle.X() + Particle.V() * Dt;
-				FRotation3::IntegrateRotationWithAngularVelocity(Particle.R(), Particle.W(), Dt);
+				Particle.R() = FRotation3::IntegrateRotationWithAngularVelocity(Particle.R(), Particle.W(), Dt);
 				break;
 			}
 			}
+			
+			// Set positions and previous velocities if we can
+			// Note: At present kinematics are in fact rigid bodies
+			Particle.P() = Particle.X();
+			Particle.Q() = Particle.R();
+			Particle.PreV() = Particle.V();
+			Particle.PreW() = Particle.W();
+
+			if (bMoved)
+			{
+				if (!Particle.CCDEnabled())
+				{
+					Particle.UpdateWorldSpaceState(FRigidTransform3(Particle.P(), Particle.Q()), FVec3(0));
+				}
+				else
+				{
+					Particle.UpdateWorldSpaceStateSwept(FRigidTransform3(Particle.P(), Particle.Q()), FVec3(0), -Particle.V() * Dt);
+				}
+			}
+		};
+
+		const auto& ApplyKinematicParticleKinematicTarget =
+		[Dt, StepFraction, IsLastStep]
+		(FTransientKinematicGeometryParticleHandle& Particle, const int32 ParticleIndex)
+		-> void
+		{
+			TKinematicTarget<FReal, 3>& KinematicTarget = Particle.KinematicTarget();
+			const FVec3 CurrentX = Particle.X();
+			const FRotation3 CurrentR = Particle.R();
+			constexpr FReal MinDt = 1e-6f;
+
+			bool bMoved = false;
+			switch (KinematicTarget.GetMode())
+			{
+			case EKinematicTargetMode::None:
+				// Nothing to do
+				break;
+
+			case EKinematicTargetMode::Reset:
+			{
+				// Reset velocity and then switch to do-nothing mode
+				Particle.V() = FVec3(0, 0, 0);
+				Particle.W() = FVec3(0, 0, 0);
+				KinematicTarget.SetMode(EKinematicTargetMode::None);
+				break;
+			}
 
-			// NOTE: we do not expand the bounds of kinematics, only dynamics
-			Particle.UpdateWorldSpaceState(FRigidTransform3(Particle.X(), Particle.R()), FVec3(0));
+			case EKinematicTargetMode::Position:
+			{
+				// Move to kinematic target and update velocities to match
+				// Target positions only need to be processed once, and we reset the velocity next frame (if no new target is set)
+				FVec3 NewX;
+				FRotation3 NewR;
+				if (IsLastStep)
+				{
+					NewX = KinematicTarget.GetTarget().GetLocation();
+					NewR = KinematicTarget.GetTarget().GetRotation();
+					KinematicTarget.SetMode(EKinematicTargetMode::Reset);
+				}
+				else
+				{
+					// as a reminder, stepfraction is the remaing fraction of the step from the remaining steps
+					// for total of 4 steps and current step of 2, this will be 1/3 ( 1 step passed, 3 steps remains )
+					NewX = FVec3::Lerp(CurrentX, KinematicTarget.GetTarget().GetLocation(), StepFraction);
+					NewR = FRotation3::Slerp(CurrentR, KinematicTarget.GetTarget().GetRotation(), decltype(FQuat::X)(StepFraction));
+				}
+
+				const bool bPositionChanged = !FVec3::IsNearlyEqual(NewX, CurrentX, UE_SMALL_NUMBER);
+				const bool bRotationChanged = !FRotation3::IsNearlyEqual(NewR, CurrentR, UE_SMALL_NUMBER);
+				bMoved = bPositionChanged || bRotationChanged;
+				FVec3 NewV = FVec3(0);
+				FVec3 NewW = FVec3(0);
+				if (Dt > MinDt)
+				{
+					if (bPositionChanged)
+					{
+						NewV = FVec3::CalculateVelocity(CurrentX, NewX, Dt);
+					}
+					if (bRotationChanged)
+					{
+						NewW = FRotation3::CalculateAngularVelocity(CurrentR, NewR, Dt);
+					}
+				}
+				Particle.X() = NewX;
+				Particle.R() = NewR;
+				Particle.V() = NewV;
+				Particle.W() = NewW;
+
+				break;
+			}
+
+			case EKinematicTargetMode::Velocity:
+			{
+				// Move based on velocity
+				bMoved = true;
+				Particle.X() = Particle.X() + Particle.V() * Dt;
+				Particle.R() = FRotation3::IntegrateRotationWithAngularVelocity(Particle.R(), Particle.W(), Dt);
+				break;
+			}
+			}
+
+			if (bMoved)
+			{
+				Particle.UpdateWorldSpaceState(FRigidTransform3(Particle.X(), Particle.R()), FVec3(0));
+			}
+		};
+
+		// Apply kinematic targets
+		// 
+		// We could run the updates in parallel, but in practice we're more likely to get parallelism benefits
+		// from running multiple characters in parallel, especially as often we'll only have a small number of 
+		// objects to process here.
+		bool bForceSingleThreaded = true;
+		// All the real kinematic particles
+		Particles.GetActiveKinematicParticlesView().ParallelFor(ApplyKinematicParticleKinematicTarget, bForceSingleThreaded);
+		// All the dynamic particles plus kinematic ones.
+		Particles.GetActiveDynamicMovingKinematicParticlesView().ParallelFor(ApplyDynamicParticleKinematicTarget, bForceSingleThreaded);
+
+		// done with update, let's clear the tracking structures
+		if (IsLastStep)
+		{
+			Particles.UpdateAllMovingKinematic();
 		}
+
+		// If we changed any particle state, the views need to be refreshed
+		Particles.UpdateDirtyViews();
 	}
 
 	void FPBDMinEvolution::DetectCollisions(FReal Dt)
diff --git a/Engine/Source/Runtime/Experimental/Chaos/Private/Chaos/PBDJointConstraints.cpp b/Engine/Source/Runtime/Experimental/Chaos/Private/Chaos/PBDJointConstraints.cpp
--- a/Engine/Source/Runtime/Experimental/Chaos/Private/Chaos/PBDJointConstraints.cpp
+++ b/Engine/Source/Runtime/Experimental/Chaos/Private/Chaos/PBDJointConstraints.cpp
@@ -139,7 +139,7 @@ namespace Chaos
 
 	void FPBDJointConstraintHandle::SetSettings(const FPBDJointSettings& InSettings)
 	{
-		// buffer the previous targets so plasticiy can reuse them
+		// buffer the previous targets so plasticity can reuse them
 		FVec3 LinearTarget = GetSettings().LinearDrivePositionTarget;
 		FRotation3 AngularTarget = GetSettings().AngularDrivePositionTarget;
 		if (!bLinearPlasticityInitialized && !FMath::IsNearlyEqual(InSettings.LinearPlasticityLimit, FLT_MAX))
@@ -151,11 +151,10 @@ namespace Chaos
 			bAngularPlasticityInitialized = true;
 		}
 
-
 		ConcreteContainer()->SetConstraintSettings(ConstraintIndex, InSettings);
 
 
-		// transfer the previous targets when controlled by plasticiy
+		// transfer the previous targets when controlled by plasticity
 		if (bLinearPlasticityInitialized)
 		{
 			ConcreteContainer()->SetLinearDrivePositionTarget(ConstraintIndex,LinearTarget);
@@ -758,7 +757,6 @@ namespace Chaos
 		return ConstraintStates[ConstraintIndex].Color;
 	}
 
-
 	FVec3 FPBDJointConstraints::GetConstraintLinearImpulse(int32 ConstraintIndex) const
 	{
 		return ConstraintStates[ConstraintIndex].LinearImpulse;
diff --git a/Engine/Source/Runtime/Experimental/Chaos/Public/Chaos/PBDJointConstraints.h b/Engine/Source/Runtime/Experimental/Chaos/Public/Chaos/PBDJointConstraints.h
--- a/Engine/Source/Runtime/Experimental/Chaos/Public/Chaos/PBDJointConstraints.h
+++ b/Engine/Source/Runtime/Experimental/Chaos/Public/Chaos/PBDJointConstraints.h
@@ -47,8 +47,9 @@ namespace Chaos
 		const FPBDJointSettings& GetSettings() const;
 		const FPBDJointSettings& GetJointSettings() const { return GetSettings(); }	//needed for property macros
 
+		// Note that setting drive targets etc can be done through the settings
 		void SetSettings(const FPBDJointSettings& Settings);
-		
+
 		TVec2<FGeometryParticleHandle*> GetConstrainedParticles() const override final;
 
 		static const FConstraintHandleTypeID& StaticType()
