Fix assert in setting dirty bit mask when a single object is polled

#rnx

diff --git a/Engine/Source/Runtime/Experimental/Iris/Core/Private/Iris/ReplicationSystem/ObjectReplicationBridge.cpp b/Engine/Source/Runtime/Experimental/Iris/Core/Private/Iris/ReplicationSystem/ObjectReplicationBridge.cpp
--- a/Engine/Source/Runtime/Experimental/Iris/Core/Private/Iris/ReplicationSystem/ObjectReplicationBridge.cpp
+++ b/Engine/Source/Runtime/Experimental/Iris/Core/Private/Iris/ReplicationSystem/ObjectReplicationBridge.cpp
@@ -856,6 +856,9 @@ void UObjectReplicationBridge::PreUpdateAndPollImpl(FNetRefHandle Handle)
 	{
 		if (uint32 InternalObjectIndex = LocalNetRefHandleManager.GetInternalIndex(Handle))
 		{
+			// From here we call into user code via PreUpdateInstanceFunction, so allow external code to set dirty flags since DirtyObjects is not read anymore.
+			ReplicationSystemInternal->GetDirtyNetObjectTracker().AllowExternalAccess();
+
 			UpdateAndPollFunction(InternalObjectIndex);
 		}
 
