make dataflow pins use Display Name when available

+ fix dataflow pin tooltips

diff --git a/Engine/Source/Runtime/Experimental/Dataflow/Core/Private/Dataflow/DataflowNode.cpp b/Engine/Source/Runtime/Experimental/Dataflow/Core/Private/Dataflow/DataflowNode.cpp
--- a/Engine/Source/Runtime/Experimental/Dataflow/Core/Private/Dataflow/DataflowNode.cpp
+++ b/Engine/Source/Runtime/Experimental/Dataflow/Core/Private/Dataflow/DataflowNode.cpp
@@ -477,6 +477,30 @@ FString FDataflowNode::GetToolTip()
 	return "";
 }
 
+FText FDataflowNode::GetPinDisplayName(const FName& PropertyName)
+{
+	// GetDisplayNameText() is only available if WITH_EDITORONLY_DATA
+#if WITH_EDITORONLY_DATA
+	if (const TUniquePtr<FStructOnScope> ScriptOnStruct = TUniquePtr<FStructOnScope>(NewStructOnScope()))
+	{
+		if (const UStruct* Struct = ScriptOnStruct->GetStruct())
+		{
+			for (TFieldIterator<FProperty> PropertyIt(Struct); PropertyIt; ++PropertyIt)
+			{
+				const FProperty* Property = *PropertyIt;
+
+				if (Property->GetName() == PropertyName.ToString())
+				{
+					return Property->GetDisplayNameText();
+				}
+			}
+		}
+	}
+#endif
+
+	return FText();
+}
+
 FString FDataflowNode::GetPinToolTip(const FName& PropertyName)
 {
 	if (const TUniquePtr<FStructOnScope> ScriptOnStruct = TUniquePtr<FStructOnScope>(NewStructOnScope()))
diff --git a/Engine/Source/Runtime/Experimental/Dataflow/Core/Public/Dataflow/DataflowNode.h b/Engine/Source/Runtime/Experimental/Dataflow/Core/Public/Dataflow/DataflowNode.h
--- a/Engine/Source/Runtime/Experimental/Dataflow/Core/Public/Dataflow/DataflowNode.h
+++ b/Engine/Source/Runtime/Experimental/Dataflow/Core/Public/Dataflow/DataflowNode.h
@@ -77,6 +77,7 @@ struct DATAFLOWCORE_API FDataflowNode
 	virtual FString GetTags() const { return ""; }
 	virtual FString GetToolTip();
 	FString GetPinToolTip(const FName& PropertyName);
+	FText GetPinDisplayName(const FName& PropertyName);
 	TArray<FString> GetPinMetaData(const FName& PropertyName);
 	virtual TArray<Dataflow::FRenderingParameter> GetRenderParameters() const { return GetRenderParametersImpl(); }
 
diff --git a/Engine/Source/Runtime/Experimental/Dataflow/Engine/Private/Dataflow/DataflowEdNode.cpp b/Engine/Source/Runtime/Experimental/Dataflow/Engine/Private/Dataflow/DataflowEdNode.cpp
--- a/Engine/Source/Runtime/Experimental/Dataflow/Engine/Private/Dataflow/DataflowEdNode.cpp
+++ b/Engine/Source/Runtime/Experimental/Dataflow/Engine/Private/Dataflow/DataflowEdNode.cpp
@@ -354,6 +354,25 @@ FText UDataflowEdNode::GetTooltipText() const
 
 }
 
+FText UDataflowEdNode::GetPinDisplayName(const UEdGraphPin* Pin) const
+{
+	if (Pin && DataflowGraph)
+	{
+		if (DataflowNodeGuid.IsValid())
+		{
+			if (TSharedPtr<FDataflowNode> DataflowNode = DataflowGraph->FindBaseNode(DataflowNodeGuid))
+			{
+				FText DisplayName = DataflowNode->GetPinDisplayName(Pin->PinName);
+				if (!DisplayName.IsEmpty())
+				{
+					return DisplayName;
+				}
+			}
+		}
+	}
+	return Super::GetPinDisplayName(Pin);
+}
+
 void UDataflowEdNode::GetPinHoverText(const UEdGraphPin& Pin, FString& HoverTextOut) const
 {
 	if (DataflowGraph)
diff --git a/Engine/Source/Runtime/Experimental/Dataflow/Engine/Public/Dataflow/DataflowEdNode.h b/Engine/Source/Runtime/Experimental/Dataflow/Engine/Public/Dataflow/DataflowEdNode.h
--- a/Engine/Source/Runtime/Experimental/Dataflow/Engine/Public/Dataflow/DataflowEdNode.h
+++ b/Engine/Source/Runtime/Experimental/Dataflow/Engine/Public/Dataflow/DataflowEdNode.h
@@ -44,6 +44,7 @@ public:
 	virtual FLinearColor GetNodeBodyTintColor() const override;
 	virtual FText GetTooltipText() const override;
 	virtual void GetPinHoverText(const UEdGraphPin& Pin, FString& HoverTextOut) const override;
+	virtual FText GetPinDisplayName(const UEdGraphPin* Pin) const override;
 #endif // WITH_EDITOR
 	// End of UEdGraphNode interface
 
