Add ways to get and set the rest transform on a geometry collection

diff --git a/Engine/Source/Runtime/Experimental/GeometryCollectionEngine/Private/GeometryCollection/GeometryCollectionComponent.cpp b/Engine/Source/Runtime/Experimental/GeometryCollectionEngine/Private/GeometryCollection/GeometryCollectionComponent.cpp
--- a/Engine/Source/Runtime/Experimental/GeometryCollectionEngine/Private/GeometryCollection/GeometryCollectionComponent.cpp
+++ b/Engine/Source/Runtime/Experimental/GeometryCollectionEngine/Private/GeometryCollection/GeometryCollectionComponent.cpp
@@ -4841,15 +4841,17 @@ int32 UGeometryCollectionComponent::GetRootIndex() const
 	return INDEX_NONE;
 }
 
-
-
 FTransform UGeometryCollectionComponent::GetRootInitialTransform() const
 {
 	FTransform RootInitialTransform{ FTransform::Identity };
-	const int32 RootIndex = RestCollection->GetRootIndex();
-	if (RestTransforms.IsValidIndex(RootIndex))
+	if (RestCollection && RestCollection->GetGeometryCollection())
 	{
-		RootInitialTransform = RestTransforms[RootIndex] * GetComponentTransform();
+		const int32 RootIndex = RestCollection->GetRootIndex();
+		if (RestTransforms.IsValidIndex(RootIndex))
+		{
+			const FTransform LocalSpaceTransform = GeometryCollectionAlgo::GlobalMatrix(RestTransforms, RestCollection->GetGeometryCollection()->Parent, RootIndex);
+			RootInitialTransform = LocalSpaceTransform * GetComponentTransform();
+		}
 	}
 	return RootInitialTransform;
 }
@@ -4857,14 +4859,85 @@ FTransform UGeometryCollectionComponent::GetRootInitialTransform() const
 FTransform UGeometryCollectionComponent::GetRootCurrentTransform() const
 {
 	FTransform RootInitialTransform{ FTransform::Identity };
-	const int32 RootIndex = RestCollection->GetRootIndex();
-	if (GlobalMatrices.IsValidIndex(RootIndex))
+	if (RestCollection)
 	{
-		RootInitialTransform = FTransform(GlobalMatrices[RootIndex]) * GetComponentTransform();
+		const int32 RootIndex = RestCollection->GetRootIndex();
+		if (GlobalMatrices.IsValidIndex(RootIndex))
+		{
+			RootInitialTransform = FTransform(GlobalMatrices[RootIndex]) * GetComponentTransform();
+		}
 	}
 	return RootInitialTransform;
 }
 
+TArray<FTransform> UGeometryCollectionComponent::GetInitialLocalRestTransforms() const
+{
+	TArray<FTransform> InitialLocalTransforms;
+
+	if (RestCollection && RestCollection->GetGeometryCollection())
+	{
+		FGeometryCollection& GeometryCollection = *RestCollection->GetGeometryCollection();
+
+		GeometryCollectionAlgo::GlobalMatrices(RestTransforms, RestCollection->GetGeometryCollection()->Parent, InitialLocalTransforms);
+
+		const TManagedArray<FTransform>* MassToLocal = GeometryCollection.FindAttribute<FTransform>("MassToLocal", FGeometryCollection::TransformGroup);
+		if (MassToLocal && InitialLocalTransforms.Num() == MassToLocal->Num())
+		{
+			for (int32 TransformIndex = 0; TransformIndex < InitialLocalTransforms.Num(); TransformIndex++)
+			{
+				InitialLocalTransforms[TransformIndex] = (*MassToLocal)[TransformIndex] * InitialLocalTransforms[TransformIndex];
+			}
+		}
+	}
+	return InitialLocalTransforms;
+}
+
+void UGeometryCollectionComponent::SetLocalRestTransforms(const TArray<FTransform>& NewTransforms, bool bOnlyLeaves)
+{
+	if (RestCollection && RestCollection->GetGeometryCollection())
+	{
+		// get the current transform and will use this array to update
+		TArray<FTransform> LocalTransforms = GetInitialLocalRestTransforms();
+		if (LocalTransforms.Num() == NewTransforms.Num())
+		{
+			FGeometryCollection& GeometryCollection = *RestCollection->GetGeometryCollection();
+			const TManagedArray<FTransform>* MassToLocal = GeometryCollection.FindAttribute<FTransform>("MassToLocal", FGeometryCollection::TransformGroup);
+			const int32 NumTransforms = GeometryCollection.Parent.Num();
+
+			TArray<int32> TransformsToProcess;
+			TransformsToProcess.Reserve(NumTransforms);
+			TransformsToProcess.Emplace(RestCollection->GetRootIndex());
+			for (int32 ProcessIndex = 0; ProcessIndex < TransformsToProcess.Num(); ProcessIndex++)
+			{
+				const int32 TransformIndex = TransformsToProcess[ProcessIndex];
+				const bool bIsLeaf = (GeometryCollection.SimulationType[TransformIndex] == FGeometryCollection::ESimulationTypes::FST_Rigid);
+				const bool bNeedUpdate = !bOnlyLeaves || bIsLeaf;
+				if (bNeedUpdate)
+				{
+					const FTransform InvMassToLocal = MassToLocal ? (*MassToLocal)[TransformIndex].Inverse() : FTransform::Identity;
+
+					LocalTransforms[TransformIndex] = InvMassToLocal * NewTransforms[TransformIndex];
+
+					// because we update top-down, the parent up-to-date transform is stored in LocalTransforms
+					const int32 ParentTransformIndex = GeometryCollection.Parent[TransformIndex];
+					if (ParentTransformIndex != INDEX_NONE)
+					{
+						LocalTransforms[TransformIndex] = LocalTransforms[TransformIndex].GetRelativeTransform(LocalTransforms[ParentTransformIndex]);
+					}
+				}
+
+				for (const int32 ChildTransformIndex : GeometryCollection.Children[TransformIndex])
+				{
+					TransformsToProcess.Emplace(ChildTransformIndex);
+				}
+			}
+
+			// send the transforms for update
+			SetRestState(MoveTemp(LocalTransforms));
+		}
+	}
+}
+
 void UGeometryCollectionComponent::GetMassAndExtents(int32 ItemIndex, float& OutMass, FBox& OutExtents)
 {
 	using FGeometryCollectionPtr = const TSharedPtr<FGeometryCollection, ESPMode::ThreadSafe>;
diff --git a/Engine/Source/Runtime/Experimental/GeometryCollectionEngine/Public/GeometryCollection/GeometryCollectionComponent.h b/Engine/Source/Runtime/Experimental/GeometryCollectionEngine/Public/GeometryCollection/GeometryCollectionComponent.h
--- a/Engine/Source/Runtime/Experimental/GeometryCollectionEngine/Public/GeometryCollection/GeometryCollectionComponent.h
+++ b/Engine/Source/Runtime/Experimental/GeometryCollectionEngine/Public/GeometryCollection/GeometryCollectionComponent.h
@@ -631,14 +631,28 @@ public:
 	UFUNCTION(BlueprintCallable, Category = "ChaosPhysics")
 	int32 GetRootIndex() const;
 
-	/** Get the root item initial transform */
+	/** Get the root item initial transform in world space */
 	UFUNCTION(BlueprintCallable, Category = "ChaosPhysics")
 	FTransform GetRootInitialTransform() const;
 
-	/** Get the root item initial transform */
+	/** Get the root item current world transform */
 	UFUNCTION(BlueprintCallable, Category = "ChaosPhysics")
 	FTransform GetRootCurrentTransform() const;
 
+	/** 
+	* Get the initial rest transforms in component (local) space  space, 
+	* they are the transforms as defined in the rest collection asset 
+	*/
+	UFUNCTION(BlueprintCallable, Category = "ChaosPhysics")
+	TArray<FTransform> GetInitialLocalRestTransforms() const;
+
+	/** 
+	* Set the local rest transform, this may be different from the rest collection 
+	* If the geometry collection is already simulating those matrices will be overriden by the physics state updates
+	*/
+	UFUNCTION(BlueprintCallable, Category = "ChaosPhysics")
+	void SetLocalRestTransforms(const TArray<FTransform>& Transforms, bool bOnlyLeaves);
+
 	/**
 	* Get mass and extent of a specific piece
 	* @param ItemIndex item index ( from HitResult) of the cluster to get level from
