Iris
- Initial support for async loading of packages when resolving networked object references similar to how it works for the existing replication system.
- Can be disabled either by using net.AllowAsyncLoading or specifically for iris with net.iris.AllowAsyncLoading

diff --git a/Engine/Source/Runtime/Experimental/Iris/Core/Private/Iris/ReplicationSystem/ObjectReferenceCache.cpp b/Engine/Source/Runtime/Experimental/Iris/Core/Private/Iris/ReplicationSystem/ObjectReferenceCache.cpp
--- a/Engine/Source/Runtime/Experimental/Iris/Core/Private/Iris/ReplicationSystem/ObjectReferenceCache.cpp
+++ b/Engine/Source/Runtime/Experimental/Iris/Core/Private/Iris/ReplicationSystem/ObjectReferenceCache.cpp
@@ -20,6 +20,7 @@
 #include "Net/Core/Trace/NetDebugName.h"
 #include "Net/Core/Trace/NetTrace.h"
 #include "UObject/Package.h"
+#include "HAL/IConsoleManager.h"
 
 DEFINE_LOG_CATEGORY_STATIC(LogIrisReferences, Log, All);
 
@@ -46,6 +47,28 @@ DEFINE_LOG_CATEGORY_STATIC(LogIrisReferences, Log, All);
 namespace UE::Net::Private
 {
 
+static bool bIrisAllowAsyncLoading = true;
+FAutoConsoleVariableRef CVarIrisAllowAsyncLoading(TEXT("net.iris.AllowAsyncLoading"), bIrisAllowAsyncLoading, TEXT("Flag to allow or disallow async loading when using iris replication. Note: net.allowAsyncLoading must also be enabled."), ECVF_Default);
+
+FObjectReferenceCache::FPendingAsyncLoadRequest::FPendingAsyncLoadRequest(FNetRefHandle InNetRefHandle, double InRequestStartTime)
+: RequestStartTime(InRequestStartTime)
+{
+	NetRefHandles.Add(InNetRefHandle);
+}
+
+void FObjectReferenceCache::FPendingAsyncLoadRequest::Merge(const FPendingAsyncLoadRequest& Other)
+{
+	for (FNetRefHandle OtherNetRefHandle : Other.NetRefHandles)
+	{
+		NetRefHandles.AddUnique(OtherNetRefHandle);
+	}
+}
+
+void FObjectReferenceCache::FPendingAsyncLoadRequest::Merge(FNetRefHandle InNetRefHandle)
+{
+	NetRefHandles.AddUnique(InNetRefHandle);
+}
+
 /**
  * Don't allow infinite recursion of InternalLoadObject - an attacker could
  * send malicious packets that cause a stack overflow on the server.
@@ -59,7 +82,10 @@ FObjectReferenceCache::FObjectReferenceCache()
 , StringTokenStore(nullptr)
 , NetRefHandleManager(nullptr)
 , bIsAuthority(false)
+, AsyncLoadMode(EAsyncLoadMode::UseCVar)
+, bCachedCVarAllowAsyncLoading(false)
 {
+	SetAsyncLoadMode(EAsyncLoadMode::UseCVar);
 }
 
 void FObjectReferenceCache::Init(UReplicationSystem* InReplicationSystem)
@@ -533,6 +559,12 @@ UObject* FObjectReferenceCache::GetObjectFromReferenceHandle(FNetRefHandle RefHa
 	return ResolvedObject;
 }
 
+bool FObjectReferenceCache::IsNetRefHandleBroken(FNetRefHandle RefHandle, bool bMustBeRegistered) const
+{
+	const FCachedNetObjectReference* CacheObjectPtr = ReferenceHandleToCachedReference.Find(RefHandle);	
+	return CacheObjectPtr ? CacheObjectPtr->bIsBroken : bMustBeRegistered;
+}
+
 // $IRIS: $TODO: Most of the logic comes from GUIDCache::GetObjectFromNetGUID so we want to keep this in sync
 UObject* FObjectReferenceCache::ResolveObjectReferenceHandleInternal(FNetRefHandle RefHandle, const FNetObjectResolveContext& ResolveContext, bool& bOutMustBeMapped)
 {
@@ -543,8 +575,6 @@ UObject* FObjectReferenceCache::ResolveObjectReferenceHandleInternal(FNetRefHand
 		bOutMustBeMapped = false;
 		return nullptr;
 	}
-
-	check(!CacheObjectPtr->bIsPending);
 	
 	UObject* Object = CacheObjectPtr->Object.Get();
 
@@ -638,7 +668,6 @@ UObject* FObjectReferenceCache::ResolveObjectReferenceHandleInternal(FNetRefHand
 
 	// At this point, we either have an outer, or we are a package
 	const uint32 TreatAsLoadedFlags = EPackageFlags::PKG_CompiledIn | EPackageFlags::PKG_PlayInEditor;
-
 	check(!CacheObjectPtr->bIsPending);
 
 	if (!ensure(ObjOuter == nullptr || ObjOuter->GetPackage()->IsFullyLoaded() || ObjOuter->GetPackage()->HasAnyPackageFlags(TreatAsLoadedFlags)))
@@ -656,8 +685,10 @@ UObject* FObjectReferenceCache::ResolveObjectReferenceHandleInternal(FNetRefHand
 	constexpr bool bReading = true;
 	RenamePathForPie(ResolveContext.ConnectionId, ObjectPath, bReading);
 
+	const FName ObjectPathName(ObjectPath);
+
 	// See if this object is in memory
-	Object = StaticFindObject(UObject::StaticClass(), ObjOuter, *ObjectPath, false);
+	Object = FindObjectFast<UObject>(ObjOuter, ObjectPathName);
 #if WITH_EDITOR
 	// Object must be null if the package is a dynamic PIE package with pending external objects still loading, as it would normally while object is async loading
 	if (Object && Object->GetPackage()->IsDynamicPIEPackagePending())
@@ -666,18 +697,17 @@ UObject* FObjectReferenceCache::ResolveObjectReferenceHandleInternal(FNetRefHand
 	}
 #endif
 
-	// Assume this is a package if the outer is invalid and this is a static guid
+	// Assume this is a package if the outer is invalid and this is a static reference
 	const bool bIsPackage = CacheObjectPtr->bIsPackage;
 
 	if (Object == nullptr && !CacheObjectPtr->bNoLoad)
 	{
 		if (IsAuthority())
 		{
-			// Log when the server needs to re-load an object, it's probably due to a GC after initially loading as default guid
+			// Log when the server needs to re-load an object, it's probably due to a GC after initially loading as default reference
 			UE_LOG(LogIris, Error, TEXT("GetObjectFromRefHandle: Server re-loading object (might have been GC'd). FullPath: %s"), ToCStr(FullPath(RefHandle, ResolveContext)));
 		}
 
-		// $IRIS: $TODO: Implement Async loading support https://jira.it.epicgames.com/browse/UE-123487
 		if (bIsPackage)
 		{
 			// Async load the package if:
@@ -686,27 +716,28 @@ UObject* FObjectReferenceCache::ResolveObjectReferenceHandleInternal(FNetRefHand
 			//	3. We're actually suppose to load (levels don't load here for example)
 			//		(Refer to CanClientLoadObject, which is where we protect clients from trying to load levels)
 
-			//if (ShouldAsyncLoad())
-			//{
-			//	if (!PendingAsyncLoadRequests.Contains(Path))
-			//	{
-			//		CacheObjectPtr->bIsPending = true;
-
-			//		StartAsyncLoadingPackage(NetGUID, false);
-			//		UE_LOG(LogIris, Error, TEXT("GetObjectFromRefHandle: Async loading package. Path: %s, %s"), *ObjectPath, *RefHandle.ToString());
-			//	}
-			//	else
-			//	{
-			//		ValidateAsyncLoadingPackage(NetGUID);
-			//	}
-
-			//	// There is nothing else to do except wait on the delegate to tell us this package is done loading
-			//	return NULL;
-			//}
-			//else
+			if (ShouldAsyncLoad() && !ResolveContext.bForceSyncLoad)
+			{
+				if (!PendingAsyncLoadRequests.Contains(ObjectPathName))
+				{
+					CacheObjectPtr->bIsPending = true;
+
+					StartAsyncLoadingPackage(*CacheObjectPtr, ObjectPathName, RefHandle, false);
+					UE_LOG_REFERENCECACHE(Log, TEXT("GetObjectFromRefHandle: Async loading package. Path: %s, %s"), *ObjectPath, *RefHandle.ToString());
+				}
+				else
+				{
+					ValidateAsyncLoadingPackage(*CacheObjectPtr, ObjectPathName, RefHandle);
+				}
+
+				// There is nothing else to do except wait on the delegate to tell us this package is done loading
+				return NULL;
+			}
+			else
 			{
 				// Async loading disabled
 				Object = LoadPackage(nullptr, *ObjectPath, LOAD_None);
+				//SyncLoadedGUIDs.AddUnique(NetGUID);
 			}
 		}
 		else
@@ -717,10 +748,10 @@ UObject* FObjectReferenceCache::ResolveObjectReferenceHandleInternal(FNetRefHand
 			//	2. Someone else started async loading the outer package, and it's not fully loaded yet
 			Object = StaticLoadObject(UObject::StaticClass(), ObjOuter, *ObjectPath, nullptr, LOAD_NoWarn);
 
-			//if (ShouldAsyncLoad())
-			//{
-			//	UE_LOG( LogNetPackageMap, Error, TEXT( "GetObjectFromRefHandle: Forced blocking load. Path: %s, %s"), *ObjectPath, *RefHandle.ToString());
-			//}
+			if (ShouldAsyncLoad() && !ResolveContext.bForceSyncLoad)
+			{
+				UE_LOG(LogIris, Error, TEXT( "GetObjectFromRefHandle: Forced blocking load. Path: %s, %s"), *ObjectPath, *RefHandle.ToString());
+			}
 		}
 	}
 
@@ -750,19 +781,22 @@ UObject* FObjectReferenceCache::ResolveObjectReferenceHandleInternal(FNetRefHand
 		if (!Package->IsFullyLoaded() 
 			&& !Package->HasAnyPackageFlags(TreatAsLoadedFlags)) //TODO: dependencies of CompiledIn could still be loaded asynchronously. Are they necessary at this point??
 		{
-			//if (ShouldAsyncLoad() && Package->HasAnyInternalFlags(EInternalObjectFlags::AsyncLoading))
-			//{
-			//	CacheObjectPtr->bIsPending = true;
-
-			//	// Something else is already async loading this package, calling load again will add our callback to the existing load request
-			//	StartAsyncLoadingPackage(NetGUID, true);
-			//	UE_LOG(LogNetPackageMap, Log, TEXT("GetObjectFromNetGUID: Listening to existing async load. Path: %s, NetGUID: %s"), *Path.ToString(), *NetGUID.ToString());
-			//}
-			//else
+			if (ShouldAsyncLoad()  && !ResolveContext.bForceSyncLoad && Package->HasAnyInternalFlags(EInternalObjectFlags::AsyncLoading))
+			{
+				// Something else is already async loading this package, calling load again will add our callback to the existing load request
+				StartAsyncLoadingPackage(*CacheObjectPtr, ObjectPathName, RefHandle, true);
+				UE_LOG(LogNetPackageMap, Log, TEXT("GetObjectFromRefHandle: Listening to existing async load. Path: %s, NetRefHandle: %s"), *ObjectPath, *RefHandle.ToString());
+			}
+			else if (/*!GbNetCheckNoLoadPackages ||*/ !CacheObjectPtr->bNoLoad)
 			{
 				// If package isn't fully loaded, load it now
 				UE_LOG_REFERENCECACHE(Warning, TEXT("ObjectReferenceCache::GetObjectFromRefHandle Blocking load of %s, %s"), ToCStr(ObjectPath), *RefHandle.ToString());
 				Object = LoadPackage(nullptr, ToCStr(ObjectPath), LOAD_None);
+				//SyncLoadedGUIDs.AddUnique(NetGUID);
+			}
+			else
+			{
+				return nullptr;
 			}
 		}
 	}
@@ -876,6 +910,9 @@ UObject* FObjectReferenceCache::ResolveObjectReferenceHandleInternal(FNetRefHand
 	
 	UE_LOG_REFERENCECACHE(Verbose, TEXT("ObjectReferenceCache::ResolveObjectRefererenceHandle %s for Object: %s Pointer: %p"), *RefHandle.ToString(), ToCStr(Object->GetPathName()), Object);
 
+	// Update our QueuedObjectReference if one exists.
+	UpdateTrackedQueuedBatchObjectReference(CompleteRefHandle, Object);
+
 	return Object;
 }
 
@@ -906,6 +943,12 @@ ENetObjectReferenceResolveResult FObjectReferenceCache::ResolveObjectReference(c
 			{
 				ResolvedObject = StaticFindObject(UObject::StaticClass(), nullptr, *ObjectPath, false);
 			}
+
+			UE_LOG_REFERENCECACHE_WARNING(TEXT("ResolveObjectReference Failed to resolve clientassigned ref: %s due to not finding object with path %s."), *Reference.ToString(), *ObjectPath);	
+		}
+		else
+		{
+			UE_LOG_REFERENCECACHE_WARNING(TEXT("ResolveObjectReference Failed to resolve clientassigned ref: %s due to missing token:"), *Reference.ToString());	
 		}
 	}
 	else
@@ -1085,13 +1128,7 @@ void FObjectReferenceCache::WriteFullReferenceInternal(FNetSerializationContext&
 	UE_LOG_REFERENCECACHE(VeryVerbose, TEXT("ObjectReferenceCache::WriteFullReferenceInternal Auth: %u %s, Outer: %s Export: %u"), IsAuthority() ? 1U : 0U, *CachedObject->NetRefHandle.ToString(), *CachedObject->OuterNetRefHandle.ToString(), bMustExport ? 1U : 0U);
 
 	if (Writer->WriteBool(bMustExport))
-	{
-		// $IRIS: $TODO: Implement Async loading support https://jira.it.epicgames.com/browse/UE-123487	
-		//if (ShouldAsyncLoad() && IsAuthority() && !CachedObject->bNoLoad)
-		//{
-		//	MustBeMappedGuidsInLastBunch.AddUnique( NetGUID );
-		//}
-	
+	{	
 		// Write the data, bNoLoad is implicit true for dynamic objects
 		Writer->WriteBits(CachedObject->bNoLoad, 1U);
 
@@ -1500,23 +1537,22 @@ bool FObjectReferenceCache::WriteMustBeMappedExports(FNetSerializationContext& C
 	FNetBitStreamWriter& Writer = *Context.GetBitStreamWriter();
 	FNetExportContext* ExportContext = Context.GetExportContext();
 
-	// $TODO: Will be enabled in later changelist
-	//if (IsAuthority() && ShouldAsyncLoad() && ExportContext)
-	//{
-	//	UE_NET_TRACE_SCOPE(MustBeMappedExports, Writer, Context.GetTraceCollector(), ENetTraceVerbosity::Verbose);
+	if (IsAuthority() && ShouldAsyncLoad() && ExportContext)
+	{
+		UE_NET_TRACE_SCOPE(MustBeMappedExports, Writer, Context.GetTraceCollector(), ENetTraceVerbosity::Verbose);
 
-	//	for (const FNetObjectReference& Reference : ExportsView)
-	//	{
-	//		const FNetRefHandle Handle = Reference.GetRefHandle();
-	//		const FCachedNetObjectReference* CachedObject = Reference.CanBeExported() ? ReferenceHandleToCachedReference.Find(Handle) : nullptr;
-	//		if (CachedObject && !CachedObject->bNoLoad)
-	//		{
-	//			UE_NET_TRACE_OBJECT_SCOPE(Handle, *Context.GetBitStreamWriter(), Context.GetTraceCollector(), ENetTraceVerbosity::Trace);
-	//			Writer.WriteBool(true);
-	//			WriteNetRefHandle(Context, Handle);
-	//		}
-	//	}
-	//}
+		for (const FNetObjectReference& Reference : ExportsView)
+		{
+			const FNetRefHandle Handle = Reference.GetRefHandle();
+			const FCachedNetObjectReference* CachedObject = Reference.CanBeExported() ? ReferenceHandleToCachedReference.Find(Handle) : nullptr;
+			if (CachedObject && !CachedObject->bNoLoad)
+			{
+				UE_NET_TRACE_OBJECT_SCOPE(Handle, *Context.GetBitStreamWriter(), Context.GetTraceCollector(), ENetTraceVerbosity::Verbose);
+				Writer.WriteBool(true);
+				WriteNetRefHandle(Context, Handle);
+			}
+		}
+	}
 
 	// Write stop bit
 	Writer.WriteBool(false);
@@ -1533,7 +1569,10 @@ void FObjectReferenceCache::ReadMustBeMappedExports(FNetSerializationContext& Co
 
 	while (bHasExportsToRead && !Context.HasErrorOrOverflow())
 	{
+		UE_NET_TRACE_NAMED_OBJECT_SCOPE(ReferenceScope, FNetRefHandle(), *Context.GetBitStreamReader(), Context.GetTraceCollector(), ENetTraceVerbosity::Verbose);
 		const FNetRefHandle MustBeMappedHandle = ReadNetRefHandle(Context);
+		UE_NET_TRACE_SET_SCOPE_OBJECTID(ReferenceScope, MustBeMappedHandle);	
+
 		bHasExportsToRead = Reader.ReadBool();
 
 		if (MustBeMappedExports)
@@ -1605,4 +1644,215 @@ void FObjectReferenceCache::GenerateFullPath_r(FNetRefHandle RefHandle, const FN
 	}
 }
 
+
+void FObjectReferenceCache::ValidateAsyncLoadingPackage(FCachedNetObjectReference& CacheObject, FName PackagePath, const FNetRefHandle RefHandle)
+{
+	// With level streaming support we may end up trying to load the same package with a different
+	// RefHandle during replay fast-forwarding. This is because if a package was unloaded, and later
+	// re-loaded, it will likely be assigned a new RefHandle (since the TWeakObjectPtr to the old package
+	// in the cache object would have gone stale). During replay fast-forward, it's possible
+	// to see the new RefHandle before the previous one has finished loading, so here we fix up
+	// PendingAsyncLoadRequests to refer to the new NetRefHandle. Also keep track of all the NetRefHandles referring
+	// to the same package so their CacheObjects can be properly updated later.
+	FPendingAsyncLoadRequest& PendingLoadRequest = PendingAsyncLoadRequests[PackagePath];
+
+	PendingLoadRequest.Merge(RefHandle);
+	CacheObject.bIsPending = true;
+
+	if (PendingLoadRequest.NetRefHandles.Last() != RefHandle)
+	{
+		UE_LOG_REFERENCECACHE(Log, TEXT("ValidateAsyncLoadingPackage: Already async loading package with a different NetRefHandle. Path: %s, original RefHandle: %s, new RefHandle: %s"),
+			*PackagePath.ToString(), *PendingLoadRequest.NetRefHandles.Last().ToString(), *RefHandle.ToString());
+	}
+	else
+	{
+		UE_LOG_REFERENCECACHE(Log, TEXT("ValidateAsyncLoadingPackage: Already async loading package. Path: %s, RefHandle: %s"), *PackagePath.ToString(), *RefHandle.ToString());
+	}
+	
+//#if CSV_PROFILER
+//	PendingLoadRequest.bWasRequestedByOwnerOrPawn |= IsTrackingOwnerOrPawn();
+//#endif
+}
+
+void FObjectReferenceCache::StartAsyncLoadingPackage(FCachedNetObjectReference& CacheObject, FName PackagePath, const FNetRefHandle RefHandle, const bool bWasAlreadyAsyncLoading)
+{
+	LLM_SCOPE_BYTAG(Iris);
+
+	// Need timer?
+	FPendingAsyncLoadRequest LoadRequest(RefHandle, ReplicationSystem->GetElapsedTime());
+	
+//#if CSV_PROFILER
+//	LoadRequest.bWasRequestedByOwnerOrPawn = IsTrackingOwnerOrPawn();
+//#endif
+
+	CacheObject.bIsPending = true;
+
+	FPendingAsyncLoadRequest* ExistingRequest = PendingAsyncLoadRequests.Find(PackagePath);
+	if (ExistingRequest)
+	{
+		// Same package name but a possibly different net GUID. Note down the GUID and wait for the async load completion callback
+		ExistingRequest->Merge(LoadRequest);
+		return;
+	}
+
+	PendingAsyncLoadRequests.Emplace(PackagePath, MoveTemp(LoadRequest));
+
+	//DelinquentAsyncLoads.MaxConcurrentAsyncLoads = FMath::Max<uint32>(DelinquentAsyncLoads.MaxConcurrentAsyncLoads, PendingAsyncLoadRequests.Num());
+
+	LoadPackageAsync(PackagePath.ToString(), FLoadPackageAsyncDelegate::CreateRaw(this, &FObjectReferenceCache::AsyncPackageCallback));
+}
+
+void FObjectReferenceCache::AsyncPackageCallback(const FName& PackageName, UPackage* Package, EAsyncLoadingResult::Type Result)
+{
+	// $TODO: $IRIS
+	// We probably want to deffer this to a safe time to process, we do not want any callbacks when we do not expect them
+	// Probably want to put this in a queue that we can guard to be processed when we are ready.
+	
+	LLM_SCOPE_BYTAG(Iris);
+
+	check(Package == nullptr || Package->IsFullyLoaded());
+
+	if (FPendingAsyncLoadRequest const* const PendingLoadRequest = PendingAsyncLoadRequests.Find(PackageName))
+	{
+		const bool bIsBroken = (Package == nullptr);
+
+		for (FNetRefHandle RefHandleToProcess : PendingLoadRequest->NetRefHandles)
+		{
+			if (FCachedNetObjectReference* CacheObject = ReferenceHandleToCachedReference.Find(RefHandleToProcess))
+			{
+				if (!CacheObject->bIsPending)
+				{
+					UE_LOG_REFERENCECACHE(Error, TEXT("AsyncPackageCallback: Package wasn't pending. Path: %s, RefHandle: %s"), *PackageName.ToString(), *RefHandleToProcess.ToString());
+				}
+
+				CacheObject->bIsPending = false;
+
+				if (bIsBroken)
+				{
+					CacheObject->bIsBroken = true;
+					UE_LOG_REFERENCECACHE(Error, TEXT("AsyncPackageCallback: Package FAILED to load. Path: %s, RefHandle: %s"), *PackageName.ToString(), *RefHandleToProcess.ToString());
+				}
+
+				if (UObject* Object = CacheObject->Object.Get())
+				{
+					// Something is loaded, we can now resolve queued data
+					UpdateTrackedQueuedBatchObjectReference(RefHandleToProcess, Object);
+
+					// $TODO: $IRIS
+					// This is odd, can we deprecate?, Does not really belong here IMHO
+					// If we need it, we can expose and bind a delegate from NetDriver
+					//if (UWorld* World = Object->GetWorld())
+					//{
+					//	if (AGameStateBase* GS = World->GetGameState())
+					//	{
+					//		GS->AsyncPackageLoaded(Object);
+					//	}
+					//}
+				}
+			}
+			else
+			{
+				UE_LOG_REFERENCECACHE(Error, TEXT("AsyncPackageCallback: Could not find net guid. Path: %s, RefHandle: %s"), *PackageName.ToString(), *RefHandleToProcess.ToString());
+			}
+		}
+
+		// This won't be the exact amount of time that we spent loading the package, but should
+		// give us a close enough estimate (within a frame time).
+		const double LoadTime = (ReplicationSystem->GetElapsedTime() - PendingLoadRequest->RequestStartTime);
+		//if (GGuidCacheTrackAsyncLoadingGUIDThreshold > 0.f &&
+		//	LoadTime >= GGuidCacheTrackAsyncLoadingGUIDThreshold)
+		//{
+		//	DelinquentAsyncLoads.DelinquentAsyncLoads.Emplace(PackageName, LoadTime);
+		//}
+
+//#if CSV_PROFILER
+//		if (PendingLoadRequest->bWasRequestedByOwnerOrPawn &&
+//			GGuidCacheTrackAsyncLoadingGUIDThresholdOwner > 0.f &&
+//			LoadTime >= GGuidCacheTrackAsyncLoadingGUIDThresholdOwner &&
+//			Driver->ServerConnection)
+//		{
+//			CSV_EVENT(PackageMap, TEXT("Owner Net Stall Async Load (Package=%s|LoadTime=%.2f)"), *PackageName.ToString(), LoadTime);
+//		}
+//#endif
+
+		PendingAsyncLoadRequests.Remove(PackageName);
+	}
+	else
+	{
+		UE_LOG_REFERENCECACHE(Error, TEXT( "AsyncPackageCallback: Could not find package. Path: %s" ), *PackageName.ToString());
+	}
+}
+
+void FObjectReferenceCache::SetAsyncLoadMode(const EAsyncLoadMode NewMode)
+{
+	AsyncLoadMode = NewMode;
+	if (const IConsoleVariable* CVarAllowAsyncLoading = IConsoleManager::Get().FindConsoleVariable(TEXT("net.AllowAsyncLoading"), false /* bTrackFrequentCalls */))
+	{
+		bCachedCVarAllowAsyncLoading = CVarAllowAsyncLoading->GetInt() > 0;
+	}
+}
+
+bool FObjectReferenceCache::ShouldAsyncLoad() const
+{
+	if (!bIrisAllowAsyncLoading)
+	{
+		return false;
+	}
+
+	switch (AsyncLoadMode)
+	{
+		case EAsyncLoadMode::UseCVar:		return bCachedCVarAllowAsyncLoading;
+		case EAsyncLoadMode::ForceDisable:	return false;
+		case EAsyncLoadMode::ForceEnable:	return true;
+		default: ensureMsgf( false, TEXT( "Invalid AsyncLoadMode: %i" ), (int32)AsyncLoadMode ); return false;
+	}
+}
+void FObjectReferenceCache::AddReferencedObjects(FReferenceCollector& ReferenceCollector)
+{
+	for (auto& It : QueuedBatchObjectReferences)
+	{
+		FQueuedBatchObjectReference& Ref = It.Value;
+
+		if (Ref.Object)
+		{
+			// AddReferencedObject will set our reference to nullptr if the object is pending kill.
+			ReferenceCollector.AddReferencedObject(Ref.Object, ReplicationSystem);
+
+			if (!Ref.Object)
+			{
+				UE_LOG_REFERENCECACHE(Warning, TEXT("FObjectReferenceCache::CollectReferences: QueuedBatchObjectReference was killed by GC. %s"), *(It.Key.ToString()));
+			}			
+		}
+	}
+}
+
+void FObjectReferenceCache::AddTrackedQueuedBatchObjectReference(const FNetRefHandle InHandle, const UObject* InObject)
+{
+	FQueuedBatchObjectReference& Ref = QueuedBatchObjectReferences.FindOrAdd(InHandle);
+	Ref.Object = InObject;
+	++Ref.RefCount;
+}
+
+void FObjectReferenceCache::UpdateTrackedQueuedBatchObjectReference(const FNetRefHandle InHandle, const UObject* NewObject)
+{
+	if (FQueuedBatchObjectReference* Ref = QueuedBatchObjectReferences.Find(InHandle))
+	{
+		Ref->Object = NewObject;
+	}
+}
+
+void FObjectReferenceCache::RemoveTrackedQueuedBatchObjectReference(const FNetRefHandle InHandle)
+{
+	if (auto It = QueuedBatchObjectReferences.CreateKeyIterator(InHandle))
+	{
+		FQueuedBatchObjectReference& Ref = It.Value();
+		check(Ref.RefCount > 0);
+		if (--Ref.RefCount == 0)
+		{
+			It.RemoveCurrent();
+		}
+	}
+}
+
+
 }
diff --git a/Engine/Source/Runtime/Experimental/Iris/Core/Private/Iris/ReplicationSystem/ObjectReferenceCache.h b/Engine/Source/Runtime/Experimental/Iris/Core/Private/Iris/ReplicationSystem/ObjectReferenceCache.h
--- a/Engine/Source/Runtime/Experimental/Iris/Core/Private/Iris/ReplicationSystem/ObjectReferenceCache.h
+++ b/Engine/Source/Runtime/Experimental/Iris/Core/Private/Iris/ReplicationSystem/ObjectReferenceCache.h
@@ -30,6 +30,7 @@ namespace UE::Net::Private
 class FObjectReferenceCache
 {
 public:
+
 	FObjectReferenceCache();
 
 	void Init(UReplicationSystem* ReplicationSystem);
@@ -55,6 +56,9 @@ public:
 	UObject* ResolveObjectReference(const FNetObjectReference& ObjectRef, const FNetObjectResolveContext& ResolveContext);
 	ENetObjectReferenceResolveResult ResolveObjectReference(const FNetObjectReference& ObjectRef, const FNetObjectResolveContext& ResolveContext, UObject*& OutResolvedObject);
 
+	// Returns true of this NetRefHandle is marked as broken
+	bool IsNetRefHandleBroken(FNetRefHandle Handle, bool bMustBeRegistered) const;
+
 	// Find replicated outer
 	FNetObjectReference GetReplicatedOuter(const FNetObjectReference& Reference) const;
 
@@ -108,6 +112,24 @@ public:
 
 	static FNetObjectReference MakeNetObjectReference(FNetRefHandle Handle);
 
+	// Async interface, kept as close to possible to FNetGuidCache/PackageMapClient
+	enum class EAsyncLoadMode : uint8
+	{
+		UseCVar			= 0,		// Use CVar (net.AllowAsyncLoading) to determine if we should async load
+		ForceDisable	= 1,		// Disable async loading
+		ForceEnable		= 2,		// Force enable async loading
+	};
+
+	void SetAsyncLoadMode(const EAsyncLoadMode NewMode);
+	bool ShouldAsyncLoad() const;
+	void AsyncPackageCallback(const FName& PackageName, UPackage* Package, EAsyncLoadingResult::Type Result);
+
+	// While async loading of pending must be mapped references we need to maintain references to already resolved objects as there will be no instance referencing them
+	void AddReferencedObjects(FReferenceCollector& ReferenceCollector);
+	void AddTrackedQueuedBatchObjectReference(const FNetRefHandle InHandle, const UObject* InObject);
+	void UpdateTrackedQueuedBatchObjectReference(const FNetRefHandle InHandle, const UObject* NewObject);
+	void RemoveTrackedQueuedBatchObjectReference(const FNetRefHandle InHandle);
+
 private:
 
 	struct FCachedNetObjectReference
@@ -132,6 +154,23 @@ private:
 		uint8 bIsPending : 1;
 	};
 
+	struct FQueuedBatchObjectReference
+	{
+		const UObject* Object = nullptr;
+		uint32 RefCount = 0U;
+	};
+
+	struct FPendingAsyncLoadRequest
+	{
+		FPendingAsyncLoadRequest(FNetRefHandle InNetRefHandle, double InRequestStartTime);
+		void Merge(const FPendingAsyncLoadRequest& Other);
+		void Merge(FNetRefHandle InNetRefHandle);
+
+		// NetRefHandles that requested loading for the same UPackage
+		TArray<FNetRefHandle, TInlineAllocator<4>> NetRefHandles;
+		double RequestStartTime;
+	};
+
 	bool CreateObjectReferenceInternal(const UObject* Object, FNetObjectReference& OutReference);
 
 	void ConditionalWriteNetTokenData(FNetSerializationContext& Context, Private::FNetExportContext* ExportContext, const FNetToken& NetToken) const;
@@ -139,6 +178,7 @@ private:
 
 	void ReadFullReferenceInternal(FNetSerializationContext& Context, FNetObjectReference& OutRef, uint32 RecursionCount);
 	void WriteFullReferenceInternal(FNetSerializationContext& Context, const FNetObjectReference& Ref) const;
+
 	UObject* ResolveObjectReferenceHandleInternal(FNetRefHandle RefHandle, const FNetObjectResolveContext& ResolveContext, bool& bOutMustBeMapped);
 	bool IsDynamicInternal(const UObject* Object) const;
 	bool SupportsObjectInternal(const UObject* Object) const;
@@ -161,6 +201,9 @@ private:
 	bool WriteMustBeMappedExports(FNetSerializationContext& Context, TArrayView<const FNetObjectReference> ExportsView) const;
 	void ReadMustBeMappedExports(FNetSerializationContext& Context, TArray<FNetRefHandle>* MustBeMappedExports);
 
+	void StartAsyncLoadingPackage(FCachedNetObjectReference& Object, FName PackagePath, const FNetRefHandle RefHandle, const bool bWasAlreadyAsyncLoading);
+	void ValidateAsyncLoadingPackage(FCachedNetObjectReference& Object, FName PackagePath, const FNetRefHandle RefHandle);
+
 private:
 
 	// Map raw UObject pointer -> Handle
@@ -181,6 +224,38 @@ private:
 	
 	// Do we have authority to create references?
 	uint32 bIsAuthority : 1;
+
+	/**
+	 * Set of all current Objects that we've been requested to be referenced while we are doing async loading.
+	 * This is used to prevent objects (especially async load objects,
+	 * which may have no other references) from being GC'd while a the object is waiting for more
+	 * pending references
+	 */
+	TMap<FNetRefHandle, FQueuedBatchObjectReference> QueuedBatchObjectReferences;
+
+	EAsyncLoadMode AsyncLoadMode;
+	bool bCachedCVarAllowAsyncLoading;
+
+	/** Set of packages that are currently pending async loads, referenced by package name. */
+	TMap<FName, FPendingAsyncLoadRequest> PendingAsyncLoadRequests;
+
+	// $TODO: $IRIS: Stats support
+#if 0
+	/** Store all GUIDs that caused the sync loading of a package, for debugging & logging with LogNetSyncLoads */
+	//TArray<FNetRefHandle> SyncLoadedGUIDs;
+	//FNetAsyncLoadDelinquencyAnalytics DelinquentAsyncLoads;
+	//void ConsumeAsyncLoadDelinquencyAnalytics(FNetAsyncLoadDelinquencyAnalytics& Out);
+	//const FNetAsyncLoadDelinquencyAnalytics& GetAsyncLoadDelinquencyAnalytics() const;
+	//void ResetAsyncLoadDelinquencyAnalytics();	
+	//bool WasGUIDSyncLoaded(FNetworkGUID NetGUID) const { return SyncLoadedGUIDs.Contains(NetGUID); }
+	//void ClearSyncLoadedGUID(FNetworkGUID NetGUID) { SyncLoadedGUIDs.Remove(NetGUID); }
+	/**
+	 * If LogNetSyncLoads is enabled, log all objects that caused a sync load that haven't been otherwise reported
+	 * by the package map yet, and clear that list.
+	 */
+	//void ReportSyncLoadedGUIDs();
+#endif
+
 };
 
 inline UObject* FObjectReferenceCache::ResolveObjectReferenceHandle(FNetRefHandle RefHandle, const FNetObjectResolveContext& ResolveContext)
diff --git a/Engine/Source/Runtime/Experimental/Iris/Core/Private/Iris/ReplicationSystem/ObjectReferenceCacheFwd.h b/Engine/Source/Runtime/Experimental/Iris/Core/Private/Iris/ReplicationSystem/ObjectReferenceCacheFwd.h
--- a/Engine/Source/Runtime/Experimental/Iris/Core/Private/Iris/ReplicationSystem/ObjectReferenceCacheFwd.h
+++ b/Engine/Source/Runtime/Experimental/Iris/Core/Private/Iris/ReplicationSystem/ObjectReferenceCacheFwd.h
@@ -16,6 +16,7 @@ namespace UE::Net
 	{
 		FNetTokenStoreState* RemoteNetTokenStoreState = nullptr;
 		uint32 ConnectionId = InvalidConnectionId;
+		bool bForceSyncLoad = false;
 	};
 
 	namespace Private
diff --git a/Engine/Source/Runtime/Experimental/Iris/Core/Private/Iris/ReplicationSystem/ReplicationReader.cpp b/Engine/Source/Runtime/Experimental/Iris/Core/Private/Iris/ReplicationSystem/ReplicationReader.cpp
--- a/Engine/Source/Runtime/Experimental/Iris/Core/Private/Iris/ReplicationSystem/ReplicationReader.cpp
+++ b/Engine/Source/Runtime/Experimental/Iris/Core/Private/Iris/ReplicationSystem/ReplicationReader.cpp
@@ -107,6 +107,58 @@ private:
 	bool bIncludeInitState = false;
 };
 
+// Helper class to deal with management of ObjectsToDispatch allocations from our temporary allocator
+class FReplicationReader::FObjectsToDispatchArray
+{
+public:
+
+	FObjectsToDispatchArray(uint32 InitialCapacity, FMemStackBase& Allocator)
+	: ObjectsToDispatchCount(0U)
+	, Capacity(InitialCapacity + ObjectsToDispatchSlackCount)
+	{			
+		ObjectsToDispatch = new (Allocator) FDispatchObjectInfo[Capacity];
+	}
+
+	void Grow(uint32 Count, FMemStackBase& Allocator)
+	{
+		if (Capacity < (ObjectsToDispatchCount + Count))
+		{
+			Capacity = ObjectsToDispatchCount + Count + ObjectsToDispatchSlackCount;
+			FDispatchObjectInfo* NewObjectsToDispatch = new (Allocator) FDispatchObjectInfo[Capacity];
+
+			// If we already had allocated data we need to copy the old elements
+			if (ObjectsToDispatchCount)
+			{
+				FPlatformMemory::Memcpy(NewObjectsToDispatch, ObjectsToDispatch, ObjectsToDispatchCount*sizeof(FDispatchObjectInfo));
+			}
+			ObjectsToDispatch = NewObjectsToDispatch;
+		}		
+	}
+
+	FDispatchObjectInfo& AddPendingDispatchObjectInfo(FMemStackBase& Allocator)
+	{
+		Grow(1, Allocator);
+		ObjectsToDispatch[ObjectsToDispatchCount] = FDispatchObjectInfo();
+
+		return ObjectsToDispatch[ObjectsToDispatchCount];
+	}
+
+	void CommitPendingDispatchObjectInfo()
+	{
+		checkSlow(ObjectsToDispatchCount < Capacity);
+		++ObjectsToDispatchCount;
+	}
+
+	uint32 Num() const { return ObjectsToDispatchCount; }
+	TArrayView<FDispatchObjectInfo> GetObjectsToDispatch() { return MakeArrayView(ObjectsToDispatch, ObjectsToDispatchCount); }
+
+private:
+
+	FDispatchObjectInfo* ObjectsToDispatch;
+	uint32 ObjectsToDispatchCount;
+	uint32 Capacity;		
+};
+
 FReplicationReader::FReplicatedObjectInfo::FReplicatedObjectInfo()
 : InternalIndex(FNetRefHandleManager::InvalidInternalIndex)
 , Value(0U)
@@ -122,9 +174,7 @@ FReplicationReader::FReplicationReader()
 , ReplicationSystemInternal(nullptr)
 , NetRefHandleManager(nullptr)
 , StateStorage(nullptr)
-, ObjectsToDispatch(nullptr)
-, ObjectsToDispatchCount(0U)
-, ObjectsToDispatchCapacity(0U)
+, ObjectsToDispatchArray(nullptr)
 , NetBlobHandlerManager(nullptr)
 , NetObjectBlobType(InvalidNetBlobType)
 , DelayAttachmentsWithUnresolvedReferences(IConsoleManager::Get().FindConsoleVariable(TEXT("net.DelayUnmappedRPCs"), false /* bTrackFrequentCalls */))
@@ -170,6 +220,20 @@ void FReplicationReader::Init(const FReplicationParameters& InParameters)
 
 void FReplicationReader::Deinit()
 {
+	for (FPendingBatchData& PendingBatchData : PendingBatches)
+	{
+		UE_LOG_REPLICATIONREADER_WARNING(TEXT("FReplicationReader::Deinit NetHandle %s has %d unprocessed data batches"), *PendingBatchData.Handle.ToString(), PendingBatchData.QueuedDataChunks.Num());
+
+		// Make sure to release all references that we are holding on to
+		if (ObjectReferenceCache)
+		{
+			for (const FNetRefHandle& RefHandle : PendingBatchData.ResolvedReferences)
+			{
+				ObjectReferenceCache->RemoveTrackedQueuedBatchObjectReference(RefHandle);
+			}
+		}		
+	}
+
 	// Cleanup any allocation stored in the per object info
 	for (auto& ObjectIt : ReplicatedObjects)
 	{
@@ -181,8 +245,13 @@ void FReplicationReader::Deinit()
 // Read incomplete handle
 FNetRefHandle FReplicationReader::ReadNetRefHandleId(FNetSerializationContext& Context, FNetBitStreamReader& Reader) const
 {
+	UE_NET_TRACE_NAMED_OBJECT_SCOPE(ReferenceScope, FNetRefHandle(), *Context.GetBitStreamReader(), Context.GetTraceCollector(), ENetTraceVerbosity::Verbose);
+
 	const uint64 NetId = ReadPackedUint64(&Reader);
 	FNetRefHandle RefHandle = FNetRefHandleManager::MakeNetRefHandleFromId(NetId);
+
+	UE_NET_TRACE_SET_SCOPE_OBJECTID(ReferenceScope, RefHandle);
+
 	if (RefHandle.GetId() != NetId)
 	{
 		Context.SetError(GNetError_InvalidNetHandle);
@@ -203,17 +272,30 @@ uint16 FReplicationReader::ReadObjectsPendingDestroy(FNetSerializationContext& C
 	
 	if (!Context.HasErrorOrOverflow())
 	{
+		const bool bHasPendingBatches = !PendingBatches.IsEmpty();
+	
 		for (uint32 It = 0; It < ObjectsToRead; ++It)
 		{
 			UE_NET_TRACE_NAMED_OBJECT_SCOPE(DestroyedObjectScope, FNetRefHandle(), Reader, Context.GetTraceCollector(), ENetTraceVerbosity::Trace);
 
 			FNetRefHandle IncompleteHandle = ReadNetRefHandleId(Context, Reader);
+			FNetRefHandle SubObjectRootOrHandle = IncompleteHandle;
+			if (Reader.ReadBool())
+			{
+				SubObjectRootOrHandle = ReadNetRefHandleId(Context, Reader);
+			}
 			const bool bShouldDestroyInstance = Reader.ReadBool();
 			if (Context.HasErrorOrOverflow())
 			{
 				break;
 			}
 
+			if (FPendingBatchData* PendingBatchData = bHasPendingBatches ? PendingBatches.FindByPredicate([&SubObjectRootOrHandle](const FPendingBatchData& Entry) { return Entry.Handle == SubObjectRootOrHandle; }) : nullptr)
+			{
+				EnqueueEndReplication(PendingBatchData, bShouldDestroyInstance, IncompleteHandle);
+				continue;
+			}
+
 			{
 				// Resolve handle and destroy using bridge
 				const uint32 InternalIndex = NetRefHandleManager->GetInternalIndex(IncompleteHandle);
@@ -224,8 +306,7 @@ uint16 FReplicationReader::ReadObjectsPendingDestroy(FNetSerializationContext& C
 					// Defer EndReplication until after applying state data
 					if (bDeferEndReplication)
 					{
-						FDispatchObjectInfo& Info = ObjectsToDispatch[ObjectsToDispatchCount];
-						Info = FDispatchObjectInfo();
+						FDispatchObjectInfo& Info = ObjectsToDispatchArray->AddPendingDispatchObjectInfo(TempLinearAllocator);
 
 						Info.bDestroy = bShouldDestroyInstance;
 						Info.bTearOff = false;
@@ -236,7 +317,7 @@ uint16 FReplicationReader::ReadObjectsPendingDestroy(FNetSerializationContext& C
 						Info.bHasAttachments = 0U;
 
 						// Mark for dispatch
-						++ObjectsToDispatchCount;
+						ObjectsToDispatchArray->CommitPendingDispatchObjectInfo();
 					}
 					else
 					{
@@ -245,9 +326,8 @@ uint16 FReplicationReader::ReadObjectsPendingDestroy(FNetSerializationContext& C
 				}
 				else
 				{
-
 					// If we did not find the object or associated bridge, the packet that would have created the object may have been lost.
-					UE_LOG(LogIris, Log, TEXT("FReplicationReader::Read Tried to destroy object with %s (This can occur if the server sends destroy for an object that has not yet been confirmed as created)"), *IncompleteHandle.ToString());
+					UE_LOG(LogIris, Verbose, TEXT("FReplicationReader::Read Tried to destroy object with %s (This can occur if the server sends destroy for an object that has not yet been confirmed as created)"), *IncompleteHandle.ToString());
 				}
 			}
 		}
@@ -376,10 +456,123 @@ void FReplicationReader::DeserializeObjectStateDelta(FNetSerializationContext& C
 	}
 }
 
+FReplicationReader::FPendingBatchData* FReplicationReader::UpdateUnresolvedMustBeMappedReferences(FNetRefHandle InHandle, TArray<FNetRefHandle>& MustBeMappedReferences)
+{
+	FPendingBatchData* PendingBatch = PendingBatches.FindByPredicate([&InHandle](const FPendingBatchData& Entry) { return Entry.Handle == InHandle; });
+	// If we already have a pending batch we append any new must be mapped references to it.
+	if (PendingBatch)
+	{
+		for (const FNetRefHandle Ref : PendingBatch->PendingMustBeMappedReferences)
+		{
+			MustBeMappedReferences.AddUnique(Ref);
+		}
+	}
+
+	// Resolve
+	TArray<FNetRefHandle, TInlineAllocator<4>> Unresolved;
+	TArray<TPair<FNetRefHandle, UObject*>, TInlineAllocator<4>> QueuedObjectsToTrack;
+
+	Unresolved.Reserve(MustBeMappedReferences.Num());
+	QueuedObjectsToTrack.SetNum(MustBeMappedReferences.Num());
+	
+	for (FNetRefHandle Handle : MustBeMappedReferences)
+	{
+		UObject* ResolvedObject = nullptr;
+		// TODO: Report broken status in same call to avoid map lookup
+		ENetObjectReferenceResolveResult ResolveResult = ObjectReferenceCache->ResolveObjectReference(FObjectReferenceCache::MakeNetObjectReference(Handle), ResolveContext, ResolvedObject);
+		if (EnumHasAnyFlags(ResolveResult, ENetObjectReferenceResolveResult::HasUnresolvedMustBeMappedReferences) && !ObjectReferenceCache->IsNetRefHandleBroken(Handle, true))
+		{
+			Unresolved.Add(Handle);
+		}
+		else if (ResolveResult == ENetObjectReferenceResolveResult::None)
+		{
+			QueuedObjectsToTrack.Emplace(Handle, ResolvedObject);
+		}
+	}
+
+	if (Unresolved.Num())
+	{
+		FPendingBatchData* Batch = PendingBatch;
+		// We must create a new batch
+		if (!Batch)
+		{
+			Batch = &PendingBatches.AddDefaulted_GetRef();
+			Batch->Handle = InHandle;
+		}
+
+		// Update
+		Batch->PendingMustBeMappedReferences = MoveTemp(Unresolved);
+	
+		// If we resolved more references, add them to tracking list
+		for (TPair<FNetRefHandle, UObject*>& NetRefHandleObjectPair : QueuedObjectsToTrack)
+		{
+			if (!Batch->ResolvedReferences.Contains(NetRefHandleObjectPair.Key))
+			{
+				Batch->ResolvedReferences.Add(NetRefHandleObjectPair.Key);
+				ObjectReferenceCache->AddTrackedQueuedBatchObjectReference(NetRefHandleObjectPair.Key, NetRefHandleObjectPair.Value);
+			}
+		}
+
+		return Batch;
+	}
+	else if (PendingBatch)
+	{
+		PendingBatch->PendingMustBeMappedReferences.Reset();
+
+		// If we resolved more references, add them to tracking list
+		for (TPair<FNetRefHandle, UObject*>& NetRefHandleObjectPair : QueuedObjectsToTrack)
+		{
+			if (!PendingBatch->ResolvedReferences.Contains(NetRefHandleObjectPair.Key))
+			{
+				PendingBatch->ResolvedReferences.Add(NetRefHandleObjectPair.Key);
+				ObjectReferenceCache->AddTrackedQueuedBatchObjectReference(NetRefHandleObjectPair.Key, NetRefHandleObjectPair.Value);
+			}
+		}
+		
+		return PendingBatch;
+	}
+	
+	return nullptr;
+}
+
+uint32 FReplicationReader::ReadObjectsInBatch(FNetSerializationContext& Context, FNetRefHandle IncompleteHandle, bool bHasBatchOwnerData, uint32 BatchEndBitPosition)
+{
+	uint32 ReadObjectCount = 0;
+	FNetBitStreamReader& Reader = *Context.GetBitStreamReader();
+
+	// If the batch owner had state, we read it now
+	if (bHasBatchOwnerData)
+	{
+		ReadObjectInBatch(Context, IncompleteHandle, false);
+		if (Context.HasErrorOrOverflow())
+		{
+			return 0U;
+		}
+		++ReadObjectCount;
+	}
+
+	ensureAlways(Reader.GetPosBits() <= BatchEndBitPosition);
+
+	// ReadSubObjects 
+	while (Reader.GetPosBits() < BatchEndBitPosition)
+	{
+		ReadObjectInBatch(Context, IncompleteHandle, true);
+		if (Context.HasErrorOrOverflow())
+		{
+			return 0U;
+		}
+		++ReadObjectCount;
+	}
+
+	return ReadObjectCount;
+}
+
 uint32 FReplicationReader::ReadObjectBatch(FNetSerializationContext& Context)
 {
 	FNetBitStreamReader& Reader = *Context.GetBitStreamReader();
 
+	UE_NET_TRACE_SCOPE(Batch, Reader, Context.GetTraceCollector(), ENetTraceVerbosity::Trace);
+
 	// Special handling for destruction infos
 	if (const bool bIsDestructionInfo = Reader.ReadBool())
 	{
@@ -424,16 +617,18 @@ uint32 FReplicationReader::ReadObjectBatch(FNetSerializationContext& Context)
 	// This either marks the end of the data associated with this batch or the offset in the stream where exports are stored.
 	const uint32 BatchEndOrStartOfExportsPos = Reader.GetPosBits() + BatchSize;
 
-	uint32 ReadObjectCount = 0;
-
 	// Do we have state data or attachments for the owner of the batch?
 	const bool bHasBatchOwnerData = Reader.ReadBool();
 
 	// Do we have exports or not?
 	const bool bHasExports = Reader.ReadBool();
+	
+	uint32 ReadObjectCount = 0U;
 
 	// First we need to read exports, they are stored at the end of the batch
 	uint32 BatchEndPos = BatchEndOrStartOfExportsPos;
+
+	TempMustBeMappedReferences.Reset();
 	if (bHasExports)
 	{
 		const uint32 ReturnPos = Reader.GetPosBits();
@@ -441,7 +636,8 @@ uint32 FReplicationReader::ReadObjectBatch(FNetSerializationContext& Context)
 		// Seek to the export section
 		Reader.Seek(BatchEndPos);
 
-		ObjectReferenceCache->ReadExports(Context, nullptr /* Currently we do not process MustBeMappedReferences */);
+		// Read exports and any must be mapped references
+		ObjectReferenceCache->ReadExports(Context, &TempMustBeMappedReferences);
 		if (Context.HasErrorOrOverflow())
 		{
 			return 0U;
@@ -454,26 +650,47 @@ uint32 FReplicationReader::ReadObjectBatch(FNetSerializationContext& Context)
 		Reader.Seek(ReturnPos);
 	}
 
-	// If the batch owner had state, we read it now
-	if (bHasBatchOwnerData)
+	// This object has pending must be mapped references that must be resolved before we can process the data.
+	FPendingBatchData* PendingBatchData = ObjectReferenceCache->ShouldAsyncLoad() ? UpdateUnresolvedMustBeMappedReferences(IncompleteHandle, TempMustBeMappedReferences) : nullptr;
+	if (PendingBatchData)
 	{
-		ReadObjectInBatch(Context, IncompleteHandle, false);
+		UE_LOG(LogIris, Verbose, TEXT("FReplicationReader::ReadObjectBatch Handle %s will be defered as it has unresolved must be mapped references"), *IncompleteHandle.ToString());				
+
+		UE_NET_TRACE_OBJECT_SCOPE(IncompleteHandle, Reader, Context.GetTraceCollector(), ENetTraceVerbosity::Trace);
+		UE_NET_TRACE_SCOPE(QueuedBatch, Reader, Context.GetTraceCollector(), ENetTraceVerbosity::Trace);
+
+		// Enqueue BatchData
+		FQueuedDataChunk DataChunk;
+		
+		const uint32 NumDataBits = BatchEndOrStartOfExportsPos - Reader.GetPosBits();
+		const uint32 NumDataWords = (NumDataBits + 31U) / 32U;
+
+		DataChunk.NumBits = NumDataBits;
+		DataChunk.StorageOffset = PendingBatchData->DataChunkStorage.Num();
+		DataChunk.bHasBatchOwnerData = bHasBatchOwnerData;
+		DataChunk.bIsEndReplicationChunk = false;
+
+		// Make sure we have space
+		PendingBatchData->DataChunkStorage.AddUninitialized(NumDataWords);
+	
+		// Store batch data
+		Reader.ReadBitStream(PendingBatchData->DataChunkStorage.GetData() + DataChunk.StorageOffset, DataChunk.NumBits);
+
 		if (Context.HasErrorOrOverflow())
 		{
 			return 0U;
 		}
-		++ReadObjectCount;
-	}
 
-	// ReadSubObjects 
-	while (Reader.GetPosBits() < BatchEndOrStartOfExportsPos)
+		PendingBatchData->QueuedDataChunks.Add(MoveTemp(DataChunk));
+	}
+	else
 	{
-		ReadObjectInBatch(Context, IncompleteHandle, true);
+		ReadObjectCount = ReadObjectsInBatch(Context, IncompleteHandle, bHasBatchOwnerData, BatchEndOrStartOfExportsPos);
+
 		if (Context.HasErrorOrOverflow())
 		{
 			return 0U;
 		}
-		++ReadObjectCount;
 	}
 
 	// Skip to the end as we already have read any exports
@@ -501,7 +718,7 @@ void FReplicationReader::ReadObjectInBatch(FNetSerializationContext& Context, FN
 	if (bHasState)
 	{
 #if UE_NET_USE_READER_WRITER_SENTINEL
-		ReadAndVerifySentinelBits(&Reader, TEXT("Exports"), 8);
+		ReadAndVerifySentinelBits(&Reader, TEXT("HasState"), 8);
 #endif
 	}
 
@@ -609,8 +826,7 @@ void FReplicationReader::ReadObjectInBatch(FNetSerializationContext& Context, FN
 		const FNetRefHandleManager::FReplicatedObjectData& ObjectData = NetRefHandleManager->GetReplicatedObjectDataNoCheck(InternalIndex);
 
 		// Add entry in our received data as we postpone state application until we have received all data in order to be able to properly resolve references
-		FDispatchObjectInfo& Info = ObjectsToDispatch[ObjectsToDispatchCount];
-		Info = FDispatchObjectInfo();
+		FDispatchObjectInfo& Info = ObjectsToDispatchArray->AddPendingDispatchObjectInfo(TempLinearAllocator);
 
 		// Update info based on ReplicatedDestroyHeader
 		Info.bDestroy = !!(ReplicatedDestroyHeaderFlags & (ReplicatedDestroyHeaderFlags_TearOff | ReplicatedDestroyHeaderFlags_DestroyInstance));
@@ -650,6 +866,9 @@ void FReplicationReader::ReadObjectInBatch(FNetSerializationContext& Context, FN
 					FReplicationProtocolOperations::DeserializeWithMask(Context, Info.ChangeMaskOrPointer.GetPointer(ChangeMaskBitCount), ObjectData.ReceiveStateBuffer, ObjectData.Protocol);
 				}
 			}
+		#if UE_NET_USE_READER_WRITER_SENTINEL
+				ReadAndVerifySentinelBits(&Reader, TEXT("HasStateEnd"), 8);
+		#endif
 
 			// Should we store a new baseline?
 			if (NewBaselineIndex != FDeltaCompressionBaselineManager::InvalidBaselineIndex)
@@ -717,8 +936,7 @@ void FReplicationReader::ReadObjectInBatch(FNetSerializationContext& Context, FN
 			Info.bHasState = bHasState ? 1U : 0U;
 			Info.bHasAttachments = bHasAttachments ? 1U : 0U;
 
-			// Mark for dispatch
-			++ObjectsToDispatchCount;
+			ObjectsToDispatchArray->CommitPendingDispatchObjectInfo();
 		}
 	}
 	
@@ -1116,7 +1334,7 @@ void FReplicationReader::DispatchStateData(FNetSerializationContext& Context)
 	};
 
 	// Allocate temporary space for post dispatch
-	FPostDispatchObjectInfo* PostDispatchObjectInfos = new (TempLinearAllocator) FPostDispatchObjectInfo[ObjectsToDispatchCount];
+	FPostDispatchObjectInfo* PostDispatchObjectInfos = new (TempLinearAllocator) FPostDispatchObjectInfo[ObjectsToDispatchArray->Num()];
 	uint32 NumObjectsPendingPostDistpatch = 0U;
 
 	// Function to flush all objects pending post dispatch
@@ -1176,7 +1394,7 @@ void FReplicationReader::DispatchStateData(FNetSerializationContext& Context)
 	FInternalNetRefIndex LastDispatchedRootInternalIndex = 0U;
 	
 	// Dispatch and apply received state data
-	for (FDispatchObjectInfo& Info : MakeArrayView(ObjectsToDispatch, ObjectsToDispatchCount))
+	for (FDispatchObjectInfo& Info : ObjectsToDispatchArray->GetObjectsToDispatch())
 	{
 		FReplicatedObjectInfo* ReplicationInfo = GetReplicatedObjectInfo(Info.InternalIndex);
 
@@ -1256,8 +1474,6 @@ void FReplicationReader::DispatchStateData(FNetSerializationContext& Context)
 	}
 
 	FlushPostDispatchForBatch();
-
-	ObjectsToDispatchCapacity = 0U;
 }
 
 void FReplicationReader::ResolveAndDispatchUnresolvedReferences()
@@ -1375,7 +1591,7 @@ void FReplicationReader::UpdateUnresolvableReferenceTracking()
 
 void FReplicationReader::DispatchEndReplication(FNetSerializationContext& Context)
 {
-	for (FDispatchObjectInfo& Info : MakeArrayView(ObjectsToDispatch, ObjectsToDispatchCount))
+	for (FDispatchObjectInfo& Info : ObjectsToDispatchArray->GetObjectsToDispatch())
 	{
 		if (Info.bDeferredEndReplication)
 		{
@@ -1386,22 +1602,16 @@ void FReplicationReader::DispatchEndReplication(FNetSerializationContext& Contex
 }
 
 
-void FReplicationReader::ReadObjects(FNetSerializationContext& Context, uint32 ObjectCountToRead)
+void FReplicationReader::ReadObjects(FNetSerializationContext& Context, uint32 ObjectBatchCountToRead)
 {
 	IRIS_PROFILER_SCOPE(ReplicationReader_ReadObjects);
 
 	FNetBitStreamReader& Reader = *Context.GetBitStreamReader();
 	
-	while (ObjectCountToRead && !Context.HasErrorOrOverflow())
+	while (ObjectBatchCountToRead && !Context.HasErrorOrOverflow())
 	{
-		const uint32 ReadObjectCount = ReadObjectBatch(Context);
-		if (ReadObjectCount > ObjectCountToRead)
-		{			
-			Context.SetError(GNetError_BitStreamError);
-			break;
-		}
-
-		ObjectCountToRead -= ReadObjectCount;
+		ReadObjectBatch(Context);
+		--ObjectBatchCountToRead;
 	}
 
 	ensureAlwaysMsgf(!Context.HasErrorOrOverflow(), TEXT("Overflow: %c Error: %s Bit stream bits left: %u position: %u"), "YN"[Context.HasError()], ToCStr(Context.GetError().ToString()), Reader.GetBitsLeft(), Reader.GetPosBits());
@@ -1442,14 +1652,8 @@ void FReplicationReader::ProcessHugeObjectAttachment(FNetSerializationContext& C
 		}
 	}
 
-	if (ObjectsToDispatchCapacity < (ObjectsToDispatchCount + HugeObjectHeader.ObjectCount))
-	{
-		// Need to reallocate ObjectDispatchInfos and copy the old data. We allocate a bit extra just in case.
-		ObjectsToDispatchCapacity = ObjectsToDispatchCount + HugeObjectHeader.ObjectCount + ObjectsToDispatchSlackCount;
-		FDispatchObjectInfo* NewObjectsToDispatch = new (TempLinearAllocator) FDispatchObjectInfo[ObjectsToDispatchCapacity];
-		FPlatformMemory::Memcpy(NewObjectsToDispatch, ObjectsToDispatch, ObjectsToDispatchCount*sizeof(FDispatchObjectInfo));
-		ObjectsToDispatch = NewObjectsToDispatch;
-	}
+	// Reserve space for more dispatch infos as needed, we allocate some extra to account for subobjects etc
+	ObjectsToDispatchArray->Grow(HugeObjectHeader.ObjectCount + ObjectsToDispatchSlackCount, TempLinearAllocator);
 
 	ReadObjects(HugeObjectSerializationContext, HugeObjectHeader.ObjectCount);
 	if (HugeObjectSerializationContext.HasErrorOrOverflow())
@@ -1471,6 +1675,149 @@ void FReplicationReader::ProcessHugeObjectAttachment(FNetSerializationContext& C
 #endif
 }
 
+bool FReplicationReader::EnqueueEndReplication(FPendingBatchData* PendingBatchData, bool bShouldDestroyInstance, FNetRefHandle NetRefHandleToEndReplication)
+{
+	UE_LOG(LogIris, Verbose, TEXT("FReplicationReader::EnqueueEndReplication for %s since %s has queued batches"), *NetRefHandleToEndReplication.ToString(), *PendingBatchData->Handle.ToString());
+		
+	const uint32 MaxNumDataBits = 65U;
+	const uint32 NumDataWords = (MaxNumDataBits + 31U) / 32U;
+
+	// Enqueue BatchData
+	FQueuedDataChunk DataChunk;
+
+	DataChunk.NumBits = MaxNumDataBits;
+	DataChunk.StorageOffset = PendingBatchData->DataChunkStorage.Num();
+	DataChunk.bHasBatchOwnerData = false;
+	DataChunk.bIsEndReplicationChunk = true;
+
+	// Make sure we have space
+	PendingBatchData->DataChunkStorage.AddUninitialized(NumDataWords);
+
+	FNetBitStreamWriter Writer;
+	Writer.InitBytes(PendingBatchData->DataChunkStorage.GetData() + DataChunk.StorageOffset, NumDataWords*sizeof(uint32));
+
+	// Write data
+	WriteUint64(&Writer, NetRefHandleToEndReplication.GetId());
+	Writer.WriteBool(bShouldDestroyInstance);
+	Writer.CommitWrites();
+
+	if (Writer.IsOverflown())
+	{
+		ensureAlwaysMsgf(false, TEXT("Failed to EnqueueEndReplication for %s, Should never occur unless size of NetRefHandle has been increased."), *NetRefHandleToEndReplication.ToString());
+		return false;
+	}
+
+	PendingBatchData->QueuedDataChunks.Add(MoveTemp(DataChunk));
+
+	return true;
+}
+
+void FReplicationReader::ProcessQueuedBatches()
+{
+	UE_NET_TRACE_FRAME_STATSCOUNTER(Parameters.ReplicationSystem->GetId(), ReplicationReader.PendingQueuedBatches, PendingBatches.Num(), ENetTraceVerbosity::Trace);
+
+	if (PendingBatches.IsEmpty())
+	{
+		//Nothing to do.
+		return;
+	}
+
+	// Setup context for dispatch
+	FInternalNetSerializationContext InternalContext;
+	FInternalNetSerializationContext::FInitParameters InternalContextInitParams;
+	InternalContextInitParams.ReplicationSystem = Parameters.ReplicationSystem;
+	InternalContextInitParams.ObjectResolveContext = ResolveContext;
+	InternalContext.Init(InternalContextInitParams);
+
+	FNetBitStreamReader Reader;
+	FNetSerializationContext Context(&Reader);
+	Context.SetLocalConnectionId(ResolveContext.ConnectionId);
+	Context.SetInternalContext(&InternalContext);
+	Context.SetNetBlobReceiver(&ReplicationSystemInternal->GetNetBlobHandlerManager());
+
+	for (int BatchIt = 0; BatchIt < PendingBatches.Num(); )
+	{
+		FPendingBatchData& PendingBatchData = PendingBatches[BatchIt];
+
+		// Try to resolve remaining must be mapped references
+		TempMustBeMappedReferences.Reset();
+		UpdateUnresolvedMustBeMappedReferences(PendingBatchData.Handle, TempMustBeMappedReferences);
+
+		// If we have no more pending must be referenes we can apply the received state
+		if (PendingBatchData.PendingMustBeMappedReferences.IsEmpty())		
+		{
+			UE_LOG(LogIris, Verbose, TEXT("ProcessQueuedBatches processing %d queued batches for Handle %s "), PendingBatchData.QueuedDataChunks.Num(), *PendingBatchData.Handle.ToString());
+
+			// Process batched data & dispatch data
+			for (const FQueuedDataChunk& CurrentChunk : PendingBatchData.QueuedDataChunks)
+			{
+				Reader.InitBits(PendingBatchData.DataChunkStorage.GetData() + CurrentChunk.StorageOffset, CurrentChunk.NumBits);
+
+				// Chunks marked as bIsEndReplicationChunk are dispatched immediately as we do not know if the next chunk tries to re-create the instance
+				if (CurrentChunk.bIsEndReplicationChunk)
+				{
+					// Read data stored for objects ending replication, 
+					// this can be the batch root or a subobject owned by the batched root.
+					const uint64 NetRefHandleIdToEndReplication = ReadUint64(&Reader);
+					const bool bShouldDestroyInstance = Reader.ReadBool();
+					
+					const FNetRefHandle NetRefHandleToEndReplication = FNetRefHandleManager::MakeNetRefHandleFromId(NetRefHandleIdToEndReplication);
+					
+					// End replication for object
+					const uint32 InternalIndex = NetRefHandleManager->GetInternalIndex(NetRefHandleToEndReplication);
+					if (InternalIndex != FNetRefHandleManager::InvalidInternalIndex)
+					{
+						EndReplication(InternalIndex, false, bShouldDestroyInstance);
+					}
+
+					UE_LOG(LogIris, Verbose, TEXT("FReplicationReader::ProcessQueuedBatches EndReplication for %s while processing queued batches for %s"), *NetRefHandleToEndReplication.ToString(), *PendingBatchData.Handle.ToString());
+					continue;
+				}
+
+				// Read and process chunk as it was a received packet for now at least
+				FMemMark TempAllocatorScope(TempLinearAllocator);
+
+				// We need to set this up to store temporary dispatch data, the array will grow if needed
+				FObjectsToDispatchArray TempObjectsToDispatchArray(ObjectsToDispatchSlackCount, TempLinearAllocator);
+
+				// Need to set this pointer as we are dealing with temporary linear allocations
+				ObjectsToDispatchArray = &TempObjectsToDispatchArray;
+
+				// $IRIS: $TODO: Implement special dispatch to defer RepNotifies if we are processing multiple batches for the same object.
+				ReadObjectsInBatch(Context, PendingBatchData.Handle, CurrentChunk.bHasBatchOwnerData, CurrentChunk.NumBits);
+
+				// $IRIS: $TODO: What to do if we fail to process this batch? Just delete it? Might need to report this to server as a broken object and build logic to reset replication of the object
+				ensureAlwaysMsgf(!Context.HasErrorOrOverflow(), TEXT("FReplicationReader::ProcessQueuedBatches - Failed to process enqueued batch for %s - %s"), *PendingBatchData.Handle.ToString(), *Context.GetError().ToString());
+			
+				// Apply received data and resolve dependencies
+				DispatchStateData(Context);
+
+				// Resolve
+				ResolveAndDispatchUnresolvedReferences();
+
+				// EndReplication for all objects in the batch that should no longer replicate
+				DispatchEndReplication(Context);
+
+				// Drop temporary data
+				ObjectsToDispatchArray = nullptr;
+			}
+	
+			// Make sure to release all references that we hold on to
+			for (const FNetRefHandle& RefHandle : PendingBatchData.ResolvedReferences)
+			{
+				ObjectReferenceCache->RemoveTrackedQueuedBatchObjectReference(RefHandle);
+			}
+
+			// Not optimal, but we want to preserve the order if we can as there might be batches waiting for the same reference
+			PendingBatches.RemoveAt(BatchIt);
+		}
+		else
+		{
+			++BatchIt;
+		}
+	}
+}
+
 void FReplicationReader::ProcessHugeObject(FNetSerializationContext& Context)
 {
 	if (!Attachments.HasUnprocessedAttachments(ENetObjectAttachmentType::HugeObject, ObjectIndexForOOBAttachment))
@@ -1518,12 +1865,12 @@ void FReplicationReader::Read(FNetSerializationContext& Context)
 
 	FMemMark TempAllocatorScope(TempLinearAllocator);
 
-	// Sanity check received obejct count
-	const uint32 MaxObjectCountToRead = 8192U;
-	const uint32 ReceivedObjectCountToRead = Reader.ReadBits(16);
-	uint32 ObjectCountToRead = ReceivedObjectCountToRead;
+	// Sanity check received object count
+	const uint32 MaxObjectBatchCountToRead = 8192U;
+	const uint32 ReceivedObjectBatchCountToRead = Reader.ReadBits(16);
+	uint32 ObjectBatchCountToRead = ReceivedObjectBatchCountToRead;
 
-	if (Reader.IsOverflown() || ObjectCountToRead >= MaxObjectCountToRead)
+	if (Reader.IsOverflown() || ObjectBatchCountToRead >= MaxObjectBatchCountToRead)
 	{
 		const FName& NetError = (Reader.IsOverflown() ? GNetError_BitStreamOverflow : GNetError_BitStreamError);
 		Context.SetError(NetError);
@@ -1531,27 +1878,29 @@ void FReplicationReader::Read(FNetSerializationContext& Context)
 		return;
 	}
 
-	if (ObjectCountToRead == 0)
+	if (ObjectBatchCountToRead == 0)
 	{
 		return;
 	}
 
-	// Allocate tracking info for objects we receive this packet
-	ObjectsToDispatchCapacity = ObjectCountToRead + ObjectsToDispatchSlackCount;
-	ObjectsToDispatchCount = 0U;
-	ObjectsToDispatch = new (TempLinearAllocator) FDispatchObjectInfo[ObjectsToDispatchCapacity];
+	// Allocate tracking info for objects we receive this packet from temporary allocator
+	// We need to set this up to store temporary dispatch data, the array will grow if needed
+	FObjectsToDispatchArray TempObjectsToDispatchArray(ObjectBatchCountToRead + ObjectsToDispatchSlackCount, TempLinearAllocator);
+
+	// Need to set this pointer as we are dealing with temporary linear allocations
+	ObjectsToDispatchArray = &TempObjectsToDispatchArray;
 
 	uint32 DestroyedObjectCount = ReadObjectsPendingDestroy(Context);
 
-	ObjectCountToRead -= DestroyedObjectCount;
+	ObjectBatchCountToRead -= DestroyedObjectCount;
 
 	// Nothing more to do or we failed and should disconnect
-	if (Context.HasErrorOrOverflow() || (ObjectCountToRead == 0 && ObjectsToDispatchCount == 0))
+	if (Context.HasErrorOrOverflow() || (ObjectBatchCountToRead == 0 && ObjectsToDispatchArray->Num() == 0))
 	{
 		return;
 	}
 
-	ReadObjects(Context, ObjectCountToRead);
+	ReadObjects(Context, ObjectBatchCountToRead);
 	if (Context.HasErrorOrOverflow())
 	{
 		return;
@@ -1564,9 +1913,14 @@ void FReplicationReader::Read(FNetSerializationContext& Context)
 		return;
 	}
 
+	if (Context.HasErrorOrOverflow())
+	{
+		return;
+	}
+
 	// Stats
-	UE_NET_TRACE_PACKET_STATSCOUNTER(Parameters.ReplicationSystem->GetId(), Parameters.ConnectionId, ReplicationReader.ReadObjectCOunt, ReceivedObjectCountToRead, ENetTraceVerbosity::Trace);
-	UE_NET_TRACE_PACKET_STATSCOUNTER(Parameters.ReplicationSystem->GetId(), Parameters.ConnectionId, ReplicationReader.ObjectsToDispatchCount, ObjectsToDispatchCount, ENetTraceVerbosity::Trace);
+	UE_NET_TRACE_PACKET_STATSCOUNTER(Parameters.ReplicationSystem->GetId(), Parameters.ConnectionId, ReplicationReader.ReadObjectBatchCount, ReceivedObjectBatchCountToRead, ENetTraceVerbosity::Trace);
+	UE_NET_TRACE_PACKET_STATSCOUNTER(Parameters.ReplicationSystem->GetId(), Parameters.ConnectionId, ReplicationReader.ReadObjectsAndSubObjectsToDispatchCount, ObjectsToDispatchArray->Num(), ENetTraceVerbosity::Trace);
 
 	// Apply received data and resolve dependencies
 	DispatchStateData(Context);
@@ -1581,9 +1935,8 @@ void FReplicationReader::Read(FNetSerializationContext& Context)
 	// EndReplication for all objects that should no longer replicate
 	DispatchEndReplication(Context);
 
-	// Drop temporary data
-	ObjectsToDispatch = nullptr;
-	ObjectsToDispatchCount = 0U;
+	// Drop temporary dispatch data
+	ObjectsToDispatchArray = nullptr;
 }
 
 void FReplicationReader::ResolveAndDispatchAttachments(FNetSerializationContext& Context, FReplicatedObjectInfo* ReplicationInfo, ENetObjectAttachmentDispatchFlags DispatchFlags)
diff --git a/Engine/Source/Runtime/Experimental/Iris/Core/Private/Iris/ReplicationSystem/ReplicationReader.h b/Engine/Source/Runtime/Experimental/Iris/Core/Private/Iris/ReplicationSystem/ReplicationReader.h
--- a/Engine/Source/Runtime/Experimental/Iris/Core/Private/Iris/ReplicationSystem/ReplicationReader.h
+++ b/Engine/Source/Runtime/Experimental/Iris/Core/Private/Iris/ReplicationSystem/ReplicationReader.h
@@ -57,6 +57,9 @@ public:
 	// Mark objects pending destroy as unresolvable.
 	void UpdateUnresolvableReferenceTracking();
 
+	// Process queued batches
+	void ProcessQueuedBatches();
+
 private:
 	// ChangeMaskOffset -> FNetRefHandle
 	enum EConstants : uint32
@@ -113,11 +116,45 @@ private:
 		uint32 bDeferredEndReplication : 1;
 	};
 
+	// Queued data chunk
+	struct FQueuedDataChunk
+	{
+		FQueuedDataChunk()
+		: StorageOffset(0U)
+		, NumBits(0U)
+		, bHasBatchOwnerData(0U)
+		, bIsEndReplicationChunk(0U)
+		{
+		}
+
+		uint32 StorageOffset;
+		uint32 NumBits : 30;
+		uint32 bHasBatchOwnerData : 1;
+		uint32 bIsEndReplicationChunk : 1;
+	};
+
+	// Struct to contain storage and required data for queued batches pending must be mapped references
+	struct FPendingBatchData
+	{
+		// We use a single array to store the actual data, it will grow if required.
+		TArray<uint32, TInlineAllocator<32>> DataChunkStorage;		
+		TArray<FQueuedDataChunk, TInlineAllocator<4>> QueuedDataChunks;
+
+		// Must be mapped references pending resolve
+		TArray<FNetRefHandle, TInlineAllocator<4>> PendingMustBeMappedReferences;
+
+		// Resolved references for which we have are holding on to references to avoid GC
+		TArray<FNetRefHandle, TInlineAllocator<4>> ResolvedReferences;
+
+		// Batch owner with queued data chunks
+		FNetRefHandle Handle;
+	};
+
 	enum : uint32
 	{
 		ObjectIndexForOOBAttachment = 0U,
 		// Try to avoid reallocations for dispatch in the case we need to process a huge object
-		ObjectsToDispatchSlackCount = 16U,
+		ObjectsToDispatchSlackCount = 32U,
 	};
 
 	bool IsObjectIndexForOOBAttachment(uint32 InternalIndex) const { return InternalIndex == ObjectIndexForOOBAttachment; }
@@ -131,6 +168,8 @@ private:
 	// Read object or subobject
 	void ReadObjectInBatch(FNetSerializationContext& Context, FNetRefHandle BatchHandle, bool bIsSubObject);
 
+	uint32 ReadObjectsInBatch(FNetSerializationContext& Context, FNetRefHandle InCompleteHandle, bool bHasBatchOwnerData, uint32 BatchEndBitPosition);
+
 	// Read all objects pending destroy
 	uint16 ReadObjectsPendingDestroy(FNetSerializationContext& Context);
 
@@ -188,7 +227,17 @@ private:
 
 	void DeserializeObjectStateDelta(FNetSerializationContext& Context, uint32 InternalIndex, FDispatchObjectInfo& Info, FReplicatedObjectInfo& ObjectInfo, const FNetRefHandleManager::FReplicatedObjectData& ObjectData, uint32& OutNewBaselineIndex);
 
+	// If async loading is enabled this function will verify if we can resolve all PendingMustBeMappedReferences
+	FReplicationReader::FPendingBatchData* UpdateUnresolvedMustBeMappedReferences(FNetRefHandle Handle, TArray<FNetRefHandle>& MustBeMappedReferences);
+
+	// If we are queuing data for a batch we must also defer calls to EndReplication
+	// This method writes this method in the form of a QueuedChunk
+	bool EnqueueEndReplication(FPendingBatchData* PendingBatchData, bool bShouldDestroyInstance, FNetRefHandle NetRefHandleToEndReplication);
+	
 private:
+
+	class FObjectsToDispatchArray;
+
 	FReplicationParameters Parameters;
 
 	FMemStackBase TempLinearAllocator;
@@ -206,9 +255,7 @@ private:
 	TMap<uint32, FReplicatedObjectInfo> ReplicatedObjects;
 
 	// Temporary data valid during receive
-	FDispatchObjectInfo* ObjectsToDispatch;
-	uint32 ObjectsToDispatchCount;
-	uint32 ObjectsToDispatchCapacity;
+	FObjectsToDispatchArray* ObjectsToDispatchArray;
 
 	// We need to keep some data around for objects with pending dependencies
 	// For now just use array and brute force the updates
@@ -220,6 +267,12 @@ private:
 	// Track all objects with a dynamic handle in case it becomes unresolvable
 	TMultiMap<FNetRefHandle, uint32> ResolvedDynamicHandleToDependents;
 
+	// We do not expect to have many objects in this state
+	TArray<FPendingBatchData> PendingBatches;
+
+	// Used during receive and processing of pending batches
+	TArray<FNetRefHandle> TempMustBeMappedReferences;
+
 	FNetBlobHandlerManager* NetBlobHandlerManager;
 	FNetBlobType NetObjectBlobType;
 	FNetObjectAttachmentsReader Attachments;
diff --git a/Engine/Source/Runtime/Experimental/Iris/Core/Private/Iris/ReplicationSystem/ReplicationSystem.cpp b/Engine/Source/Runtime/Experimental/Iris/Core/Private/Iris/ReplicationSystem/ReplicationSystem.cpp
--- a/Engine/Source/Runtime/Experimental/Iris/Core/Private/Iris/ReplicationSystem/ReplicationSystem.cpp
+++ b/Engine/Source/Runtime/Experimental/Iris/Core/Private/Iris/ReplicationSystem/ReplicationSystem.cpp
@@ -476,6 +476,7 @@ public:
 		auto UpdateUnresolvableReferenceTracking = [&Connections](uint32 ConnectionId)
 		{
 			FReplicationConnection* Conn = Connections.GetConnection(ConnectionId);
+			Conn->ReplicationReader->ProcessQueuedBatches();
 			Conn->ReplicationReader->UpdateUnresolvableReferenceTracking();
 		};
 		const FNetBitArray& ValidConnections = Connections.GetValidConnections();
@@ -544,6 +545,8 @@ void UReplicationSystem::PreSendUpdate(float DeltaSeconds)
 
 	IRIS_PROFILER_SCOPE(FReplicationSystem_PreSendUpdate);
 
+	ElapsedTime += DeltaSeconds;
+
 #if !UE_BUILD_SHIPPING
 	// Force a integrity check of all replicated instances
 	if (bDoCollectGarbage || ReplicationSystemCVars::bForcePruneBeforeUpdate)
@@ -1221,6 +1224,7 @@ void UReplicationSystem::AddReferencedObjects(UObject* InThis, FReferenceCollect
 	if (This->Impl.IsValid())
 	{
 		This->Impl->ReplicationSystemInternal.GetNetRefHandleManager().AddReferencedObjects(Collector);
+		This->Impl->ReplicationSystemInternal.GetObjectReferenceCache().AddReferencedObjects(Collector);
 	}
 	Super::AddReferencedObjects(InThis, Collector);
 }
diff --git a/Engine/Source/Runtime/Experimental/Iris/Core/Private/Iris/ReplicationSystem/ReplicationWriter.cpp b/Engine/Source/Runtime/Experimental/Iris/Core/Private/Iris/ReplicationSystem/ReplicationWriter.cpp
--- a/Engine/Source/Runtime/Experimental/Iris/Core/Private/Iris/ReplicationSystem/ReplicationWriter.cpp
+++ b/Engine/Source/Runtime/Experimental/Iris/Core/Private/Iris/ReplicationSystem/ReplicationWriter.cpp
@@ -1646,6 +1646,12 @@ uint32 FReplicationWriter::WriteObjectsPendingDestroy(FNetSerializationContext&
 		// Write handle with the needed bitCount
 		WriteNetRefHandleId(Context, ObjectData.RefHandle);
 
+		// If this is a explicit subobject we must also send the owner handle in the case that the client is currently queued data due to async loading.
+		if (Writer.WriteBool(ObjectData.IsSubObject()))
+		{			
+			WriteNetRefHandleId(Context, NetRefHandleManager->GetReplicatedObjectDataNoCheck(ObjectData.SubObjectRootIndex).RefHandle);
+		}
+
 		// Write bit indicating if the static instance should be destroyed or not (could skip the bit for dynamic objects)
 		const bool bShouldDestroyInstance = ObjectData.RefHandle.IsDynamic() || NetRefHandleManager->GetIsDestroyedStartupObject(InternalIndex);
 		Writer.WriteBool(bShouldDestroyInstance);
@@ -2013,6 +2019,9 @@ FReplicationWriter::EWriteObjectStatus FReplicationWriter::WriteObjectAndSubObje
 					FReplicationProtocolOperations::SerializeWithMask(Context, Info.GetChangeMaskStoragePointer(), ReplicatedObjectStateBuffer, ObjectData.Protocol);
 				}
 			}
+#if UE_NET_USE_READER_WRITER_SENTINEL
+			WriteSentinelBits(&Writer, 8);
+#endif
 		}
 
 		{
@@ -2188,7 +2197,7 @@ FReplicationWriter::EWriteObjectStatus FReplicationWriter::WriteObjectAndSubObje
 	// We include the size of the data written so we can skip it if needed.
 	if (OutBatchInfo.ParentInternalIndex == InternalIndex)
 	{
-		const FBatchObjectInfo& ParentBatchEntry = OutBatchInfo.ObjectInfos[ParentBatchEntryIndex];
+		FBatchObjectInfo& ParentBatchEntry = OutBatchInfo.ObjectInfos[ParentBatchEntryIndex];
 
 		const uint32 WrittenBitsInBatch = (Writer.GetPosBits() - InitialStateHeaderPos) - NumBitsUsedForBatchSize;
 		
@@ -2212,6 +2221,8 @@ FReplicationWriter::EWriteObjectStatus FReplicationWriter::WriteObjectAndSubObje
 				Writer.WriteBool(bWroteData);
 				Writer.WriteBool(bWroteExports);
 			}
+
+			ParentBatchEntry.bSentBatchData = 1U;
 		}
 		// If we did not write any data we rollback any written headers and report a success
 		else
@@ -2347,7 +2358,9 @@ int FReplicationWriter::PrepareAndSendHugeObjectPayload(FNetSerializationContext
 	HugeObjectContext.SendStatus = EHugeObjectSendStatus::Sending;
 	HugeObjectContext.StartSendingTime = FPlatformTime::Cycles64();
 	// Store batch record for later processing once the whole state is acked.
-	HugeObjectHeader.ObjectCount = HandleObjectBatchSuccess(BatchInfo, HugeObjectContext.BatchRecord);
+	HandleObjectBatchSuccess(BatchInfo, HugeObjectContext.BatchRecord);
+	// We want to track the number of Batches
+	HugeObjectHeader.ObjectCount = HugeObjectContext.BatchRecord.BatchCount;
 
 	// Write huge object header
 	{
@@ -2459,6 +2472,9 @@ int FReplicationWriter::WriteObjectBatch(FNetSerializationContext& Context, uint
 			FBatchRecord BatchRecord;
 			const int WrittenObjectCount = HandleObjectBatchSuccess(BatchInfo, BatchRecord);
 
+			// As a single batch also might include dependent objects which are treated as separate batches on the receiving end we need to account for this when tracking the written batch count.
+			WriteContext.WrittenBatchCount += BatchRecord.BatchCount;
+
 			CommitBatchRecord(BatchRecord);
 			return WrittenObjectCount;
 		}
@@ -2499,6 +2515,13 @@ int FReplicationWriter::WriteObjectBatch(FNetSerializationContext& Context, uint
 	// Try #2 - Object will be serialized to a temporary buffer of maximum supported size and split into multiple chunks.
 	{
 		const int SendHugeObjectStatus = PrepareAndSendHugeObjectPayload(Context, InternalIndex);
+
+		// If the huge object wrote data it will be tracked as a single batch
+		if (SendHugeObjectStatus == 1)
+		{
+			++WriteContext.WrittenBatchCount;
+		}
+				
 		return SendHugeObjectStatus;
 	}
 }
@@ -2563,6 +2586,9 @@ int FReplicationWriter::WriteDestructionInfo(FNetSerializationContext& Context,
 		SchedulingPriorities[InternalIndex] = 0.0f;
 
 		WriteContext.Stats.AddNumberOfReplicatedDestructionInfos(1U);
+
+		// We count this as an object batch
+		++WriteContext.WrittenBatchCount;
 	}
 
 	return Writer.IsOverflown() ? -1 :  1;
@@ -2708,6 +2734,7 @@ uint32 FReplicationWriter::WriteObjects(FNetSerializationContext& Context)
 int FReplicationWriter::HandleObjectBatchSuccess(const FBatchInfo& BatchInfo, FReplicationWriter::FBatchRecord& OutRecord)
 {
 	int WrittenObjectCount = 0;
+	int WrittenBatchCount = 0;
 
 	const bool bTrackObjectStats = BatchInfo.Type != EBatchInfoType::Internal;
 	uint32 ObjectCount = 0;
@@ -2784,7 +2811,12 @@ int FReplicationWriter::HandleObjectBatchSuccess(const FBatchInfo& BatchInfo, FR
 			}
 		}			
 
-		WrittenObjectCount += (BatchObjectInfo.bSentState | BatchObjectInfo.bSentAttachments | BatchObjectInfo.bSentTearOff | BatchObjectInfo.bSentDestroySubObject) ? 1 : 0;
+		if (BatchObjectInfo.bSentState | BatchObjectInfo.bSentAttachments | BatchObjectInfo.bSentTearOff | BatchObjectInfo.bSentDestroySubObject)
+		{
+			++WrittenObjectCount;
+		}
+
+		WrittenBatchCount += BatchObjectInfo.bSentBatchData;
 
 		Info.HasDirtyChangeMask = 0U;
 		Info.HasDirtySubObjects = BatchObjectInfo.bHasDirtySubObjects;
@@ -2831,6 +2863,8 @@ int FReplicationWriter::HandleObjectBatchSuccess(const FBatchInfo& BatchInfo, FR
 	}
 #endif
 
+	OutRecord.BatchCount = WrittenBatchCount;
+
 	return WrittenObjectCount;
 }
 
@@ -2993,6 +3027,10 @@ UDataStream::EWriteResult FReplicationWriter::Write(FNetSerializationContext& Co
 		return UDataStream::EWriteResult::NoData;
 	}
 
+	// We have somethings in the WriteContext that must be reset each packet
+	WriteContext.FailedToWriteSmallObjectCount = 0U;
+	WriteContext.WrittenBatchCount = 0U;
+
 	FNetBitStreamWriter& Writer = *Context.GetBitStreamWriter();
 
 	// Setup internal context
@@ -3012,12 +3050,14 @@ UDataStream::EWriteResult FReplicationWriter::Write(FNetSerializationContext& Co
 	uint16 WrittenObjectCount = 0;
 	const uint32 OldReplicationInfoCount = ReplicationRecord.GetInfoCount();
 
-	Writer.WriteBits(WrittenObjectCount, 16);
+	// Written batch count
+	Writer.WriteBits(0U, 16);
 
 	// Write timestamps etc? Or do we do this in header.
 	// WriteReplicationFrameData();
 
-	WrittenObjectCount += WriteObjectsPendingDestroy(Context);
+	const uint16 WrittenObjectsPendingDestroyCount = WriteObjectsPendingDestroy(Context);
+	WrittenObjectCount += WrittenObjectsPendingDestroyCount;
 
 	// Only reason for overflow here is if we did not fit header
 	if (Writer.IsOverflown())
@@ -3051,7 +3091,8 @@ UDataStream::EWriteResult FReplicationWriter::Write(FNetSerializationContext& Co
 		{
 			// Seek back to HeaderPos and update the header
 			FNetBitStreamWriteScope WriteScope(Writer, HeaderPos);
-			Writer.WriteBits(WrittenObjectCount, 16);
+			const uint16 TotalWrittenBatchCount = WriteContext.WrittenBatchCount + WrittenObjectsPendingDestroyCount;
+			Writer.WriteBits(TotalWrittenBatchCount, 16);
 		}
 
 		//UE_LOG_REPLICATIONWRITER(TEXT("FReplicationWriter::Write() Wrote %u Objects for ConnectionId:%u, ReplicationSystemId: %u."), WrittenObjectCount, Parameters.ConnectionId, Parameters.ReplicationSystem->GetId());	
@@ -3082,6 +3123,7 @@ UDataStream::EWriteResult FReplicationWriter::Write(FNetSerializationContext& Co
 
 	// Report packet stats
 	UE_NET_TRACE_PACKET_STATSCOUNTER(Parameters.ReplicationSystem->GetId(), Parameters.ConnectionId, ReplicationWriter.WrittenObjectCount, WrittenObjectCount, ENetTraceVerbosity::Trace);
+	UE_NET_TRACE_PACKET_STATSCOUNTER(Parameters.ReplicationSystem->GetId(), Parameters.ConnectionId, ReplicationWriter.WrittenBatchCount, WriteContext.WrittenBatchCount, ENetTraceVerbosity::Trace);
 	UE_NET_TRACE_PACKET_STATSCOUNTER(Parameters.ReplicationSystem->GetId(), Parameters.ConnectionId, ReplicationWriter.FailedToWriteSmallObjectCount, WriteContext.FailedToWriteSmallObjectCount, ENetTraceVerbosity::Trace);
 	UE_NET_TRACE_PACKET_STATSCOUNTER(Parameters.ReplicationSystem->GetId(), Parameters.ConnectionId, ReplicationWriter.RemainingObjectsPendingWriteCount, WriteContext.ScheduledObjectCount - WriteContext.CurrentIndex, ENetTraceVerbosity::Trace);
 	UE_NET_TRACE_PACKET_STATSCOUNTER(Parameters.ReplicationSystem->GetId(), Parameters.ConnectionId, ReplicationWriter.ScheduledObjectCount, WriteContext.ScheduledObjectCount, ENetTraceVerbosity::Trace);
diff --git a/Engine/Source/Runtime/Experimental/Iris/Core/Private/Iris/ReplicationSystem/ReplicationWriter.h b/Engine/Source/Runtime/Experimental/Iris/Core/Private/Iris/ReplicationSystem/ReplicationWriter.h
--- a/Engine/Source/Runtime/Experimental/Iris/Core/Private/Iris/ReplicationSystem/ReplicationWriter.h
+++ b/Engine/Source/Runtime/Experimental/Iris/Core/Private/Iris/ReplicationSystem/ReplicationWriter.h
@@ -219,6 +219,9 @@ private:
 		// The index into the scheduled objects array to attempt to replicate next.
 		uint32 CurrentIndex;
 
+		// Written batch count, which excludes objects pending destroy and subobjects
+		uint32 WrittenBatchCount;
+
 		// How many objects that were attempted to be replicated but which ultimately didn't fit in the packet.
 		uint32 FailedToWriteSmallObjectCount;
 
@@ -245,6 +248,7 @@ private:
 		uint32 bHasDirtySubObjects : 1;
 		uint32 bSentTearOff : 1;
 		uint32 bSentDestroySubObject : 1;
+		uint32 bSentBatchData : 1;
 	};
 
 	enum class EBatchInfoType : uint32
@@ -273,6 +277,7 @@ private:
 	struct FBatchRecord
 	{
 		TArray<FObjectRecord, TInlineAllocator<16>> ObjectReplicationRecords;
+		uint32 BatchCount = 0U;
 	};
 
 	struct FBitStreamInfo
diff --git a/Engine/Source/Runtime/Experimental/Iris/Core/Public/Iris/ReplicationSystem/ReplicationSystem.h b/Engine/Source/Runtime/Experimental/Iris/Core/Public/Iris/ReplicationSystem/ReplicationSystem.h
--- a/Engine/Source/Runtime/Experimental/Iris/Core/Public/Iris/ReplicationSystem/ReplicationSystem.h
+++ b/Engine/Source/Runtime/Experimental/Iris/Core/Public/Iris/ReplicationSystem/ReplicationSystem.h
@@ -512,6 +512,9 @@ public:
 	/** Returns the previously set squared cull distance for an object, or DefaultValue if it wasn't or the Handle isn't valid. */
 	IRISCORE_API float GetCullDistanceSqrOverride(FNetRefHandle Handle, float DefaultValue = -1.0f) const;
 
+	/** Returns elapsed time in seconds since ReplicatonSystem was created */
+	double GetElapsedTime() const { return ElapsedTime; }
+
 public:
 	// For internal use and not exported.
 
@@ -550,6 +553,7 @@ private:
 	UPROPERTY(transient)
 	TObjectPtr<UReplicationBridge> ReplicationBridge;
 
+	double ElapsedTime = 0;
 	uint32 Id;
 	int32 PIEInstanceID;
 	uint32 bIsServer : 1;
