Iris
- Utilize changemask when quantizing statedata to only quantize dirty members.
- Added cvar to be able to disable if it causes issues net.Iris.OnlyQuantizeDirtyMembers

diff --git a/Engine/Source/Runtime/Experimental/Iris/Core/Private/Iris/ReplicationSystem/ReplicationOperations.cpp b/Engine/Source/Runtime/Experimental/Iris/Core/Private/Iris/ReplicationSystem/ReplicationOperations.cpp
--- a/Engine/Source/Runtime/Experimental/Iris/Core/Private/Iris/ReplicationSystem/ReplicationOperations.cpp
+++ b/Engine/Source/Runtime/Experimental/Iris/Core/Private/Iris/ReplicationSystem/ReplicationOperations.cpp
@@ -19,6 +19,7 @@
 #include "Iris/Serialization/NetBitStreamUtil.h"
 #include "Iris/Serialization/NetSerializer.h"
 #include "Iris/Serialization/ObjectNetSerializer.h"
+#include "Iris/ReplicationState/ReplicationStateUtil.h"
 #include "Misc/CoreMiscDefines.h"
 #include "Net/Core/Trace/NetTrace.h"
 #include "Net/Core/Trace/NetDebugName.h"
@@ -67,6 +68,13 @@ static FAutoConsoleVariableRef CVarbDeltaCompressInitialState(
 	TEXT("if true we compare with default state when serializing inital state."
 	));
 
+static bool bOnlyQuantizeDirtyMembers = true;
+static FAutoConsoleVariableRef CVarbOnlyQuantizeDirtyMembers(
+	TEXT("net.Iris.OnlyQuantizeDirtyMembers"),
+	bOnlyQuantizeDirtyMembers,
+	TEXT("if true we only quantize members marked as dirty unless this is a new object."
+	));
+
 void FReplicationStateOperations::Quantize(FNetSerializationContext& Context, uint8* RESTRICT DstInternalBuffer, const uint8* RESTRICT SrcExternalBuffer, const FReplicationStateDescriptor* Descriptor)
 {
 	check(IsAligned(DstInternalBuffer, Descriptor->InternalAlignment) && IsAligned(SrcExternalBuffer, Descriptor->ExternalAlignment));
@@ -90,6 +98,41 @@ void FReplicationStateOperations::Quantize(FNetSerializationContext& Context, ui
 	}
 }
 
+void FReplicationStateOperations::QuantizeWithMask(FNetSerializationContext& Context, const FNetBitArrayView& ChangeMask, const uint32 ChangeMaskOffset, uint8* RESTRICT DstInternalBuffer, const uint8* RESTRICT SrcExternalBuffer, const FReplicationStateDescriptor* Descriptor)
+{
+#if !UE_BUILD_SHIPPING && !UE_BUILD_TEST
+	check(!Descriptor->IsInitState());
+#endif
+
+	check(IsAligned(DstInternalBuffer, Descriptor->InternalAlignment) && IsAligned(SrcExternalBuffer, Descriptor->ExternalAlignment));
+
+	const FReplicationStateMemberDescriptor* MemberDescriptors = Descriptor->MemberDescriptors;
+	const FReplicationStateMemberSerializerDescriptor* MemberSerializerDescriptors = Descriptor->MemberSerializerDescriptors;
+	const uint32 MemberCount = Descriptor->MemberCount;
+
+	for (uint32 MemberIt = 0; MemberIt < MemberCount; ++MemberIt)
+	{
+		const FReplicationStateMemberDescriptor& MemberDescriptor = MemberDescriptors[MemberIt];
+		const FReplicationStateMemberSerializerDescriptor& MemberSerializerDescriptor = MemberSerializerDescriptors[MemberIt];
+		const FReplicationStateMemberChangeMaskDescriptor& MemberChangeMaskDescriptor = Descriptor->MemberChangeMaskDescriptors[MemberIt];
+		const uint32 MemberChangeMaskOffset = ChangeMaskOffset + MemberChangeMaskDescriptor.BitOffset;
+
+		if (ChangeMask.IsAnyBitSet(MemberChangeMaskOffset, MemberChangeMaskDescriptor.BitCount))
+		{
+			FNetQuantizeArgs Args;
+			Args.Version = 0;
+			Args.NetSerializerConfig = MemberSerializerDescriptor.SerializerConfig;
+			Args.Source = reinterpret_cast<NetSerializerValuePointer>(SrcExternalBuffer + MemberDescriptor.ExternalMemberOffset);
+			Args.Target = reinterpret_cast<NetSerializerValuePointer>(DstInternalBuffer + MemberDescriptor.InternalMemberOffset);
+
+			Args.ChangeMaskInfo.BitCount = MemberChangeMaskDescriptor.BitCount;
+			Args.ChangeMaskInfo.BitOffset = MemberChangeMaskOffset;
+
+			MemberSerializerDescriptor.Serializer->Quantize(Context, Args);
+		}
+	}
+}
+
 void FReplicationStateOperations::Dequantize(FNetSerializationContext& Context, uint8* RESTRICT DstExternalBuffer, const uint8* RESTRICT SrcInternalBuffer, const FReplicationStateDescriptor* Descriptor)
 {
 	check(IsAligned(SrcInternalBuffer, Descriptor->InternalAlignment) && IsAligned(DstExternalBuffer, Descriptor->ExternalAlignment));
@@ -577,7 +620,15 @@ void FReplicationInstanceOperations::CopyAndQuantizeIfDirty(FNetSerializationCon
 		if (Private::AppendMemberChangeMasks(OutChangeMaskWriter, OutConditionalChangeMaskWriter, FragmentData[StateIt].ExternalSrcBuffer, CurrentDescriptor))
 		{
 			IRIS_PROFILER_PROTOCOL_NAME(CurrentDescriptor->DebugName->Name);
-			FReplicationStateOperations::Quantize(Context, CurrentInternalStateBuffer, FragmentData[StateIt].ExternalSrcBuffer, CurrentDescriptor);
+
+			if (CurrentDescriptor->IsInitState() || !bOnlyQuantizeDirtyMembers)
+			{
+				FReplicationStateOperations::Quantize(Context, CurrentInternalStateBuffer, FragmentData[StateIt].ExternalSrcBuffer, CurrentDescriptor);
+			}
+			else
+			{
+				FReplicationStateOperations::QuantizeWithMask(Context, UE::Net::Private::GetMemberChangeMask(FragmentData[StateIt].ExternalSrcBuffer, CurrentDescriptor), 0, CurrentInternalStateBuffer, FragmentData[StateIt].ExternalSrcBuffer, CurrentDescriptor);
+			}
 		}
 		
 		CurrentInternalStateBuffer += CurrentDescriptor->InternalSize;
diff --git a/Engine/Source/Runtime/Experimental/Iris/Core/Private/Iris/Serialization/ArrayPropertyNetSerializer.cpp b/Engine/Source/Runtime/Experimental/Iris/Core/Private/Iris/Serialization/ArrayPropertyNetSerializer.cpp
--- a/Engine/Source/Runtime/Experimental/Iris/Core/Private/Iris/Serialization/ArrayPropertyNetSerializer.cpp
+++ b/Engine/Source/Runtime/Experimental/Iris/Core/Private/Iris/Serialization/ArrayPropertyNetSerializer.cpp
@@ -482,13 +482,37 @@ void FArrayPropertyNetSerializer::Quantize(FNetSerializationContext& Context, co
 	ElementArgs.Source = 0;
 	ElementArgs.Target = NetSerializerValuePointer(TargetArray.ElementStorage);
 
+	const FNetBitArrayView* ChangeMask = Args.ChangeMaskInfo.BitCount > 1U ? Context.GetChangeMask() : nullptr;
+
 	const FNetSerializer* ElementSerializer = ElementSerializerDescriptor.Serializer;
 	const uint32 ElementSize = ElementStateDescriptor->InternalSize;
-	for (uint32 ElementIt = 0, ElementEndIt = ElementCount; ElementIt < ElementEndIt; ++ElementIt)
+
+	if (!ChangeMask)
 	{
-		ElementArgs.Source = NetSerializerValuePointer(static_cast<const void*>(ScriptArrayHelper.GetRawPtr(ElementIt)));
-		ElementSerializer->Quantize(Context, ElementArgs);
-		ElementArgs.Target += ElementSize;
+		for (uint32 ElementIt = 0, ElementEndIt = ElementCount; ElementIt < ElementEndIt; ++ElementIt)
+		{
+			ElementArgs.Source = NetSerializerValuePointer(static_cast<const void*>(ScriptArrayHelper.GetRawPtr(ElementIt)));
+			ElementSerializer->Quantize(Context, ElementArgs);
+			ElementArgs.Target += ElementSize;
+		}
+	}
+	else
+	{
+		// We currently use a simple modulo scheme for bits in the changemask, if we have more elements in the array then is covered by the changemask
+		// several entries in the array will be considered dirty and be serialized
+		// As the first bit is used by the owning property we need to offset by one and deduct one from the usable bits
+		const uint32 ChangeMaskBitOffset = ChangeMask ? Args.ChangeMaskInfo.BitOffset + 1U : 0U;
+		const uint32 ChangeMaskBitCount = ChangeMask ? Args.ChangeMaskInfo.BitCount - 1U : 0U;
+
+		for (uint32 ElementIt = 0, ElementEndIt = ElementCount; ElementIt < ElementEndIt; ++ElementIt)
+		{
+			if (ChangeMask->GetBit(ChangeMaskBitOffset + (ElementIt % ChangeMaskBitCount)))
+			{
+				ElementArgs.Source = NetSerializerValuePointer(static_cast<const void*>(ScriptArrayHelper.GetRawPtr(ElementIt)));
+				ElementSerializer->Quantize(Context, ElementArgs);
+			}
+			ElementArgs.Target += ElementSize;
+		}
 	}
 }
 
diff --git a/Engine/Source/Runtime/Experimental/Iris/Core/Private/Iris/Serialization/StructNetSerializer.cpp b/Engine/Source/Runtime/Experimental/Iris/Core/Private/Iris/Serialization/StructNetSerializer.cpp
--- a/Engine/Source/Runtime/Experimental/Iris/Core/Private/Iris/Serialization/StructNetSerializer.cpp
+++ b/Engine/Source/Runtime/Experimental/Iris/Core/Private/Iris/Serialization/StructNetSerializer.cpp
@@ -240,6 +240,9 @@ void FStructNetSerializer::Quantize(FNetSerializationContext& Context, const FNe
 		MemberArgs.Source = Args.Source + MemberDescriptor.ExternalMemberOffset;
 		MemberArgs.Target = Args.Target + MemberDescriptor.InternalMemberOffset;
 
+		// Currently we pass on the changemask info unmodified to support fastarrays, but if we decide to support other serializers utilizing additional changemask we need to extend this
+		MemberArgs.ChangeMaskInfo = Args.ChangeMaskInfo;
+
 		Serializer->Quantize(Context, MemberArgs);
 	}
 }
diff --git a/Engine/Source/Runtime/Experimental/Iris/Core/Public/Iris/ReplicationSystem/ReplicationOperations.h b/Engine/Source/Runtime/Experimental/Iris/Core/Public/Iris/ReplicationSystem/ReplicationOperations.h
--- a/Engine/Source/Runtime/Experimental/Iris/Core/Public/Iris/ReplicationSystem/ReplicationOperations.h
+++ b/Engine/Source/Runtime/Experimental/Iris/Core/Public/Iris/ReplicationSystem/ReplicationOperations.h
@@ -32,6 +32,9 @@ struct FReplicationStateOperations
 	/** Dequantize a Replication state from internal buffer to already constructed ExternalBuffer */
 	static IRISCORE_API void Dequantize(FNetSerializationContext& Context, uint8* RESTRICT DstExternalBuffer, const uint8* RESTRICT SrcInternalBuffer, const FReplicationStateDescriptor* Descriptor);
 
+	/** Quantize a Replication state from ExternalBuffer to internal buffer, DstInternalBuffer does not need to be initialized */
+	static IRISCORE_API void QuantizeWithMask(FNetSerializationContext& Context, const FNetBitArrayView& ChangeMask, const uint32 ChangeMaskOffset, uint8* RESTRICT DstInternalBuffer, const uint8* RESTRICT SrcExternalBuffer, const FReplicationStateDescriptor* Descriptor);
+
 	/** Validate a ReplicationState in external format */
 	static IRISCORE_API bool Validate(FNetSerializationContext& Context, const uint8* RESTRICT SrcExternalBuffer, const FReplicationStateDescriptor* Descriptor);
 
