Iris:
* Fix pulling objects that were WantsToBeDormant but not dirty

#rnx

diff --git a/Engine/Source/Runtime/Experimental/Iris/Core/Private/Iris/ReplicationSystem/ObjectReplicationBridge.cpp b/Engine/Source/Runtime/Experimental/Iris/Core/Private/Iris/ReplicationSystem/ObjectReplicationBridge.cpp
--- a/Engine/Source/Runtime/Experimental/Iris/Core/Private/Iris/ReplicationSystem/ObjectReplicationBridge.cpp
+++ b/Engine/Source/Runtime/Experimental/Iris/Core/Private/Iris/ReplicationSystem/ObjectReplicationBridge.cpp
@@ -889,7 +889,7 @@ void UObjectReplicationBridge::PreUpdateAndPollImpl(FNetRefHandle Handle)
 		IRIS_PROFILER_SCOPE(PreUpdateAndPollImpl_Dormancy);
 
 		// Mask off objects pending dormancy that are not dirty
-		FrequencyBasedObjectsToPollView.CombineMultiple(FNetBitArrayView::AndNotOp, MakeNetBitArrayView(WantToBeDormantObjects), FNetBitArrayView::AndOp, DirtyObjects);
+		FrequencyBasedObjectsToPollView.CombineMultiple(FNetBitArrayView::AndNotOp, MakeNetBitArrayView(WantToBeDormantObjects), FNetBitArrayView::AndNotOp, DirtyObjects);
 
 		// Add objects that have requested to flush dormancy
 		for (FNetRefHandle HandlePendingFlush : MakeArrayView(DormantHandlesPendingFlush))
