Geometry collection:  fix crash when using compiling a blueprint based on a geometry collection actor
- added null check for the event dispatcher when destroying the physics state

diff --git a/Engine/Source/Runtime/Experimental/GeometryCollectionEngine/Private/GeometryCollection/GeometryCollectionComponent.cpp b/Engine/Source/Runtime/Experimental/GeometryCollectionEngine/Private/GeometryCollection/GeometryCollectionComponent.cpp
--- a/Engine/Source/Runtime/Experimental/GeometryCollectionEngine/Private/GeometryCollection/GeometryCollectionComponent.cpp
+++ b/Engine/Source/Runtime/Experimental/GeometryCollectionEngine/Private/GeometryCollection/GeometryCollectionComponent.cpp
@@ -3139,7 +3139,10 @@ void UGeometryCollectionComponent::OnDestroyPhysicsState()
 	}
 
 	// we need to unregister the events because it relies on the proxy internally 
-	EventDispatcher->UnregisterChaosEvents();
+	if (EventDispatcher)
+	{
+		EventDispatcher->UnregisterChaosEvents();
+	}
 
 	if(PhysicsProxy)
 	{
