Dataflow : fix when evaluating dataflow terminal node from blueprint with a name that does not exists in the graph

diff --git a/Engine/Source/Runtime/Experimental/Dataflow/Engine/Private/Dataflow/DataflowBlueprintLibrary.cpp b/Engine/Source/Runtime/Experimental/Dataflow/Engine/Private/Dataflow/DataflowBlueprintLibrary.cpp
--- a/Engine/Source/Runtime/Experimental/Dataflow/Engine/Private/Dataflow/DataflowBlueprintLibrary.cpp
+++ b/Engine/Source/Runtime/Experimental/Dataflow/Engine/Private/Dataflow/DataflowBlueprintLibrary.cpp
@@ -11,15 +11,22 @@ void UDataflowBlueprintLibrary::EvaluateTerminalNodeByName(UDataflow* Dataflow,
 	if (Dataflow && Dataflow->Dataflow)
 	{
 		TSharedPtr<FDataflowNode> Node = Dataflow->Dataflow->FindTerminalNode(TerminalNodeName);
-		if (const FDataflowTerminalNode* TerminalNode = Node->AsType<const FDataflowTerminalNode>())
+		if (Node)
 		{
-			Dataflow::FEngineContext Context(ResultAsset, Dataflow, FPlatformTime::Cycles64());
-			TerminalNode->Evaluate(Context);
-			if (ResultAsset)
+			if (const FDataflowTerminalNode* TerminalNode = Node->AsType<const FDataflowTerminalNode>())
 			{
-				TerminalNode->SetAssetValue(ResultAsset, Context);
+				Dataflow::FEngineContext Context(ResultAsset, Dataflow, FPlatformTime::Cycles64());
+				TerminalNode->Evaluate(Context);
+				if (ResultAsset)
+				{
+					TerminalNode->SetAssetValue(ResultAsset, Context);
+				}
 			}
 		}
+		else
+		{
+			UE_LOG(LogChaos, Warning, TEXT("EvaluateTerminalNodeByName : Could not find terminal node : [%s], skipping evaluation"), *TerminalNodeName.ToString());
+		}
 	}
 }
 
