Fix clang compile error and renable FScopedRPCContext

#rnx

diff --git a/Engine/Source/Runtime/Engine/Private/DataReplication.cpp b/Engine/Source/Runtime/Engine/Private/DataReplication.cpp
--- a/Engine/Source/Runtime/Engine/Private/DataReplication.cpp
+++ b/Engine/Source/Runtime/Engine/Private/DataReplication.cpp
@@ -1376,7 +1376,7 @@ bool FObjectReplicator::ReceivedRPC(FNetBitReader& Reader, const FReplicationFla
 
 
 			{
-				//UE::Net::FScopedNetContextRPC CallingRPC;
+				UE::Net::FScopedNetContextRPC CallingRPC;
 				// Call the function.
 				Object->ProcessEvent(Function, Parms);
 			}
diff --git a/Engine/Source/Runtime/Experimental/Iris/Core/Private/Iris/ReplicationSystem/NetBlob/NetRPC.cpp b/Engine/Source/Runtime/Experimental/Iris/Core/Private/Iris/ReplicationSystem/NetBlob/NetRPC.cpp
--- a/Engine/Source/Runtime/Experimental/Iris/Core/Private/Iris/ReplicationSystem/NetBlob/NetRPC.cpp
+++ b/Engine/Source/Runtime/Experimental/Iris/Core/Private/Iris/ReplicationSystem/NetBlob/NetRPC.cpp
@@ -477,7 +477,7 @@ void FNetRPC::CallFunction(FNetSerializationContext& Context)
 	// Call the function
 	if (Function->ParmsSize == 0)
 	{
-		//UE::Net::FScopedNetContextRPC CallingRPC;
+		UE::Net::FScopedNetContextRPC CallingRPC;
 		Object->ProcessEvent(const_cast<UFunction*>(Function), nullptr);
 	}
 	else
@@ -512,7 +512,7 @@ void FNetRPC::CallFunction(FNetSerializationContext& Context)
 		
 		FReplicationStateOperations::Dequantize(Context, FunctionParameters, QuantizedBlobState.Get(), BlobDescriptor);
 		{
-			//UE::Net::FScopedNetContextRPC CallingRPC;
+			UE::Net::FScopedNetContextRPC CallingRPC;
 			Object->ProcessEvent(const_cast<UFunction*>(Function), FunctionParameters);
 		}
 
diff --git a/Engine/Source/Runtime/Net/Core/Public/Net/Core/Misc/NetContext.h b/Engine/Source/Runtime/Net/Core/Public/Net/Core/Misc/NetContext.h
--- a/Engine/Source/Runtime/Net/Core/Public/Net/Core/Misc/NetContext.h
+++ b/Engine/Source/Runtime/Net/Core/Public/Net/Core/Misc/NetContext.h
@@ -2,10 +2,12 @@
 
 #pragma once
 
+// Forward declarations
 namespace UE::Net::Private
 {
 	class FNetRPC;
 }
+class FObjectReplicator;
 
 namespace UE::Net
 {
@@ -38,8 +40,8 @@ class NETCORE_API FScopedNetContextRPC
 private:
 
 	// Only friend class are allowed to push the net context
-	friend Private::FNetRPC;
-	friend class FObjectReplicator;
+	friend UE::Net::Private::FNetRPC;
+	friend FObjectReplicator;
 	
 	FScopedNetContextRPC();
 	~FScopedNetContextRPC();
