Dataflow : fix stack overflow in invalidation code when connecting nodes in a loop

diff --git a/Engine/Source/Runtime/Experimental/Dataflow/Core/Private/Dataflow/DataflowInputOutput.cpp b/Engine/Source/Runtime/Experimental/Dataflow/Core/Private/Dataflow/DataflowInputOutput.cpp
--- a/Engine/Source/Runtime/Experimental/Dataflow/Core/Private/Dataflow/DataflowInputOutput.cpp
+++ b/Engine/Source/Runtime/Experimental/Dataflow/Core/Private/Dataflow/DataflowInputOutput.cpp
@@ -63,9 +63,9 @@ const TArray< const FDataflowOutput* > FDataflowInput::GetConnectedOutputs() con
 	return RetList;
 }
 
-void FDataflowInput::Invalidate()
+void FDataflowInput::Invalidate(const Dataflow::FTimestamp& ModifiedTimestamp)
 {
-	OwningNode->Invalidate();
+	OwningNode->Invalidate(ModifiedTimestamp);
 }
 
 //
@@ -122,11 +122,11 @@ bool FDataflowOutput::RemoveConnection(FDataflowConnection* InInput)
 	Connections.RemoveSwap((FDataflowInput*)InInput); return true;
 }
 
-void FDataflowOutput::Invalidate()
+void FDataflowOutput::Invalidate(const Dataflow::FTimestamp& ModifiedTimestamp)
 {
 	for (FDataflowConnection* Con : GetConnections())
 	{
-		Con->Invalidate();
+		Con->Invalidate(ModifiedTimestamp);
 	}
 }
 
diff --git a/Engine/Source/Runtime/Experimental/Dataflow/Core/Private/Dataflow/DataflowNode.cpp b/Engine/Source/Runtime/Experimental/Dataflow/Core/Private/Dataflow/DataflowNode.cpp
--- a/Engine/Source/Runtime/Experimental/Dataflow/Core/Private/Dataflow/DataflowNode.cpp
+++ b/Engine/Source/Runtime/Experimental/Dataflow/Core/Private/Dataflow/DataflowNode.cpp
@@ -218,17 +218,22 @@ TArray<Dataflow::FPin> FDataflowNode::GetPins() const
 	return RetVal;
 }
 
-void FDataflowNode::Invalidate()
+void FDataflowNode::Invalidate(const Dataflow::FTimestamp& InModifiedTimestamp)
 {
-	LastModifiedTimestamp = Dataflow::FTimestamp(FPlatformTime::Cycles64());
-
-	for (TPair<uint32, FDataflowOutput*> Elem : Outputs)
+	if (LastModifiedTimestamp < InModifiedTimestamp)
 	{
-		FDataflowOutput* Con = Elem.Value;
-		Con->Invalidate();
-	}
+		LastModifiedTimestamp = InModifiedTimestamp;
+
+		for (TPair<uint32, FDataflowOutput*> Elem : Outputs)
+		{
+			FDataflowOutput* Con = Elem.Value;
+			Con->Invalidate(InModifiedTimestamp);
+		}
 
-	OnNodeInvalidatedDelegate.Broadcast(this);
+		OnInvalidate();
+
+		OnNodeInvalidatedDelegate.Broadcast(this);
+	}
 }
 
 void FDataflowNode::RegisterInputConnection(const void* InProperty)
@@ -319,7 +324,7 @@ void FDataflowNode::RegisterOutputConnection(const void* InProperty, const void*
 
 bool FDataflowNode::ValidateConnections()
 {
-	bValid = true;
+	bHasValidConnections = true;
 	if (const TUniquePtr<FStructOnScope> ScriptOnStruct = TUniquePtr<FStructOnScope>(NewStructOnScope()))
 	{
 		if (const UStruct* Struct = ScriptOnStruct->GetStruct())
@@ -334,7 +339,7 @@ bool FDataflowNode::ValidateConnections()
 					if (!FindInput(PropName))
 					{
 						UE_LOG(LogChaos, Warning, TEXT("Missing dataflow RegisterInputConnection in constructor for (%s:%s)"), *GetName().ToString(), *PropName.ToString())
-						bValid = false;
+							bHasValidConnections = false;
 					}
 				}
 				if (Property->HasMetaData(FDataflowNode::DataflowOutput))
@@ -343,7 +348,7 @@ bool FDataflowNode::ValidateConnections()
 					if(!OutputConnection)
 					{
 						UE_LOG(LogChaos, Warning, TEXT("Missing dataflow RegisterOutputConnection in constructor for (%s:%s)"), *GetName().ToString(),*PropName.ToString());
-						bValid = false;
+						bHasValidConnections = false;
 					}
 
 					// Validate passthrough connections if they exist
@@ -353,7 +358,7 @@ bool FDataflowNode::ValidateConnections()
 						if(PassthroughConnectionAddress == nullptr)
 						{
 							UE_LOG(LogChaos, Warning, TEXT("Missing DataflowPassthrough registration for (%s:%s)"), *GetName().ToString(),*PropName.ToString());
-							bValid = false;
+							bHasValidConnections = false;
 						}
 
 						const FDataflowInput* PassthroughConnectionInput = FindInput(FName(*PassthroughName));
@@ -362,32 +367,32 @@ bool FDataflowNode::ValidateConnections()
 						if(PassthroughConnectionInputFromArg != PassthroughConnectionInput)
 						{
 							UE_LOG(LogChaos, Warning, TEXT("Mismatch in declared and registered DataflowPassthrough connection; (%s:%s vs %s)"), *GetName().ToString(), *PropName.ToString(), *PassthroughConnectionInputFromArg->GetName().ToString());
-							bValid = false;
+							bHasValidConnections = false;
 						}
 
 						if(!PassthroughConnectionInput)
 						{
 							UE_LOG(LogChaos, Warning, TEXT("Incorrect DataflowPassthrough Connection set for (%s:%s)"), *GetName().ToString(),*PropName.ToString());
-							bValid = false;
+							bHasValidConnections = false;
 						}
 
 						else if(OutputConnection->GetType() != PassthroughConnectionInput->GetType())
 						{
 							UE_LOG(LogChaos, Warning, TEXT("DataflowPassthrough connection types mismatch for (%s:%s)"), *GetName().ToString(),*PropName.ToString());
-							bValid = false;
+							bHasValidConnections = false;
 						}
 					}
 					else if(OutputConnection->GetPassthroughRealAddress()) 
 					{
 						UE_LOG(LogChaos, Warning, TEXT("Missing DataflowPassthrough declaration for (%s:%s)"), *GetName().ToString(),*PropName.ToString());
-						bValid = false;
+						bHasValidConnections = false;
 					}
 				}
 #endif
 			}
 		}
 	}
-	return bValid;
+	return bHasValidConnections;
 }
 
 FString FDataflowNode::GetToolTip()
diff --git a/Engine/Source/Runtime/Experimental/Dataflow/Core/Private/Dataflow/DataflowNodeFactory.cpp b/Engine/Source/Runtime/Experimental/Dataflow/Core/Private/Dataflow/DataflowNodeFactory.cpp
--- a/Engine/Source/Runtime/Experimental/Dataflow/Core/Private/Dataflow/DataflowNodeFactory.cpp
+++ b/Engine/Source/Runtime/Experimental/Dataflow/Core/Private/Dataflow/DataflowNodeFactory.cpp
@@ -14,7 +14,7 @@ namespace Dataflow
 		if (ClassMap.Contains(Param.Type))
 		{
 			TUniquePtr<FDataflowNode> Node = ClassMap[Param.Type](Param);
-			if(Node->IsValid())
+			if(Node->HasValidConnections())
 			{
 				ParametersMap[Param.Type].ToolTip = Node->GetToolTip();
 
diff --git a/Engine/Source/Runtime/Experimental/Dataflow/Core/Public/Dataflow/DataflowConnection.h b/Engine/Source/Runtime/Experimental/Dataflow/Core/Public/Dataflow/DataflowConnection.h
--- a/Engine/Source/Runtime/Experimental/Dataflow/Core/Public/Dataflow/DataflowConnection.h
+++ b/Engine/Source/Runtime/Experimental/Dataflow/Core/Public/Dataflow/DataflowConnection.h
@@ -84,6 +84,6 @@ public:
 		return (size_t)OwningNode + (size_t)GetOffset() == (size_t)InVar;
 	}
 
-	virtual void Invalidate() {};
+	virtual void Invalidate(const Dataflow::FTimestamp& ModifiedTimestamp = Dataflow::FTimestamp::Current()) {};
 
 };
diff --git a/Engine/Source/Runtime/Experimental/Dataflow/Core/Public/Dataflow/DataflowInputOutput.h b/Engine/Source/Runtime/Experimental/Dataflow/Core/Public/Dataflow/DataflowInputOutput.h
--- a/Engine/Source/Runtime/Experimental/Dataflow/Core/Public/Dataflow/DataflowInputOutput.h
+++ b/Engine/Source/Runtime/Experimental/Dataflow/Core/Public/Dataflow/DataflowInputOutput.h
@@ -61,7 +61,7 @@ public:
 	template<class T>
 	TFuture<const T&> GetValueParallel(Dataflow::FContext& Context, const T& Default) const;
 
-	virtual void Invalidate() override;
+	virtual void Invalidate(const Dataflow::FTimestamp& ModifiedTimestamp = Dataflow::FTimestamp::Current()) override;
 };
 
 //
@@ -157,7 +157,7 @@ public:
 	template<class T>
 	TFuture<bool> EvaluateParallel(Dataflow::FContext& Context) const;
 
-	virtual void Invalidate() override;
+	virtual void Invalidate(const Dataflow::FTimestamp& ModifiedTimestamp = Dataflow::FTimestamp::Current()) override;
 
 };
  
diff --git a/Engine/Source/Runtime/Experimental/Dataflow/Core/Public/Dataflow/DataflowNode.h b/Engine/Source/Runtime/Experimental/Dataflow/Core/Public/Dataflow/DataflowNode.h
--- a/Engine/Source/Runtime/Experimental/Dataflow/Core/Public/Dataflow/DataflowNode.h
+++ b/Engine/Source/Runtime/Experimental/Dataflow/Core/Public/Dataflow/DataflowNode.h
@@ -212,11 +212,13 @@ struct DATAFLOWCORE_API FDataflowNode
 		return (FindInput(Reference)->GetConnection() != nullptr);
 	}
 
-	void Invalidate();
+	void Invalidate(const Dataflow::FTimestamp& ModifiedTimestamp = Dataflow::FTimestamp::Current());
+
+	virtual void OnInvalidate() {}
 
 	virtual bool ValidateConnections();
 
-	bool IsValid() const { return bValid; }
+	bool HasValidConnections() const { return bHasValidConnections; }
 
 	virtual bool IsA(FName InType) const 
 	{ 
@@ -252,7 +254,7 @@ private:
 	virtual TArray<Dataflow::FRenderingParameter> GetRenderParametersImpl() const { return TArray<Dataflow::FRenderingParameter>(); }
 
 
-	bool bValid = true;
+	bool bHasValidConnections = true;
 
 protected:
 	FOnNodeInvalidated OnNodeInvalidatedDelegate;
diff --git a/Engine/Source/Runtime/Experimental/Dataflow/Core/Public/Dataflow/DataflowNodeParameters.h b/Engine/Source/Runtime/Experimental/Dataflow/Core/Public/Dataflow/DataflowNodeParameters.h
--- a/Engine/Source/Runtime/Experimental/Dataflow/Core/Public/Dataflow/DataflowNodeParameters.h
+++ b/Engine/Source/Runtime/Experimental/Dataflow/Core/Public/Dataflow/DataflowNodeParameters.h
@@ -21,6 +21,7 @@ namespace Dataflow
 
 		FTimestamp(Type InValue) : Value(InValue) {}
 		bool operator>=(const FTimestamp& InTimestamp) const { return Value >= InTimestamp.Value; }
+		bool operator<(const FTimestamp& InTimestamp) const { return Value < InTimestamp.Value; }
 		bool IsInvalid() { return Value == Invalid; }
 
 		static Type Current();
