Dataflow : add the concept of callstack to the context to allow detection of loops

diff --git a/Engine/Source/Runtime/Experimental/Dataflow/Core/Private/Dataflow/DataflowInputOutput.cpp b/Engine/Source/Runtime/Experimental/Dataflow/Core/Private/Dataflow/DataflowInputOutput.cpp
--- a/Engine/Source/Runtime/Experimental/Dataflow/Core/Private/Dataflow/DataflowInputOutput.cpp
+++ b/Engine/Source/Runtime/Experimental/Dataflow/Core/Private/Dataflow/DataflowInputOutput.cpp
@@ -132,6 +132,13 @@ void FDataflowOutput::Invalidate(const Dataflow::FTimestamp& ModifiedTimestamp)
 
 bool FDataflowOutput::EvaluateImpl(Dataflow::FContext& Context) const
 {
+	Dataflow::FContextScopedCallstack Callstack(Context, this);
+	if (Callstack.IsLoopDetected())
+	{ 
+		ensureMsgf(false, TEXT("Connection %s is already in the callstack, this is certainly because of a loop in the graph"), *GetName().ToString());
+		return false;
+	}
+
 	// check if the cache has a valid version
 	if(Context.HasData(CacheKey(), OwningNode->LastModifiedTimestamp))
 	{
@@ -145,5 +152,6 @@ bool FDataflowOutput::EvaluateImpl(Dataflow::FContext& Context) const
 		ensureMsgf(false, TEXT("Failed to evaluate output (%s:%s)"), *OwningNode->GetName().ToString(), *GetName().ToString());
 		return false;
 	}
+
 	return true;
 }
diff --git a/Engine/Source/Runtime/Experimental/Dataflow/Core/Private/Dataflow/DataflowNodeParameters.cpp b/Engine/Source/Runtime/Experimental/Dataflow/Core/Private/Dataflow/DataflowNodeParameters.cpp
--- a/Engine/Source/Runtime/Experimental/Dataflow/Core/Private/Dataflow/DataflowNodeParameters.cpp
+++ b/Engine/Source/Runtime/Experimental/Dataflow/Core/Private/Dataflow/DataflowNodeParameters.cpp
@@ -5,11 +5,49 @@
 #include "Dataflow/DataflowInputOutput.h"
 #include "Dataflow/DataflowNode.h"
 
+
 namespace Dataflow
 {
 	uint64 FTimestamp::Invalid = 0;
 	uint64 FTimestamp::Current() { return FPlatformTime::Cycles64();  }
 
+	void FContext::PushToCallstack(const FDataflowConnection* Connection) 
+	{
+#if DATAFLOW_EDITOR_EVALUATION
+		Callstack.Push(Connection); 
+#endif
+	}
+
+	void FContext::PopFromCallstack(const FDataflowConnection* Connection)
+	{ 
+#if DATAFLOW_EDITOR_EVALUATION
+		ensure(Connection == Callstack.Top());
+		Callstack.Pop();
+#endif
+	}
+
+	bool FContext::IsInCallstack(const FDataflowConnection* Connection) const
+	{
+#if DATAFLOW_EDITOR_EVALUATION
+		return Callstack.Contains(Connection);
+#else
+		return false;
+#endif
+	}
+
+	FContextScopedCallstack::FContextScopedCallstack(FContext& InContext, const FDataflowConnection* InConnection)
+		: Context(InContext)
+		, Connection(InConnection)
+	{
+		bLoopDetected = Context.IsInCallstack(Connection);
+		Context.PushToCallstack(Connection);
+	}
+
+	FContextScopedCallstack::~FContextScopedCallstack()
+	{
+		Context.PopFromCallstack(Connection);
+	}
+
 	void BeginContextEvaluation(FContext& Context, const FDataflowNode* Node, const FDataflowOutput* Output)
 	{
 		if (Node != nullptr)
diff --git a/Engine/Source/Runtime/Experimental/Dataflow/Core/Public/Dataflow/DataflowNodeParameters.h b/Engine/Source/Runtime/Experimental/Dataflow/Core/Public/Dataflow/DataflowNodeParameters.h
--- a/Engine/Source/Runtime/Experimental/Dataflow/Core/Public/Dataflow/DataflowNodeParameters.h
+++ b/Engine/Source/Runtime/Experimental/Dataflow/Core/Public/Dataflow/DataflowNodeParameters.h
@@ -11,6 +11,9 @@
 class  UDataflow;
 struct FDataflowNode;
 struct FDataflowOutput;
+struct FDataflowConnection;
+
+#define DATAFLOW_EDITOR_EVALUATION WITH_EDITOR
 
 namespace Dataflow
 {
@@ -172,6 +175,29 @@ namespace Dataflow
 		FTimestamp GetTimestamp() const { return Timestamp; }
 		virtual void Evaluate(const FDataflowNode* Node, const FDataflowOutput* Output) = 0;
 		virtual bool Evaluate(const FDataflowOutput& Connection) = 0;
+
+		void PushToCallstack(const FDataflowConnection* Connection);
+		void PopFromCallstack(const FDataflowConnection* Connection);
+		bool IsInCallstack(const FDataflowConnection* Connection) const;
+
+	private:
+#if DATAFLOW_EDITOR_EVALUATION
+		TArray<const FDataflowConnection*> Callstack;
+#endif
+	};
+
+	struct DATAFLOWCORE_API FContextScopedCallstack
+	{
+	public:
+		FContextScopedCallstack(FContext& InContext, const FDataflowConnection* InConnection);
+		~FContextScopedCallstack();
+
+		bool IsLoopDetected() const { return bLoopDetected; }
+
+	private:
+		bool bLoopDetected;
+		FContext& Context;
+		const FDataflowConnection* Connection;
 	};
 
 #define DATAFLOW_CONTEXT_INTERNAL(PARENTTYPE, TYPENAME)														\
