Bug fix:

- Make sure that NaN bounds don't make it into the the spatial acceleration structures.
- Recover one of the NaN cases
- Check on non recoverable error
- The Dynamic Dirty Tree and Dirty Grid can be swapped out at will during runtime

diff --git a/Engine/Source/Runtime/Experimental/Chaos/Public/Chaos/AABBTree.h b/Engine/Source/Runtime/Experimental/Chaos/Public/Chaos/AABBTree.h
--- a/Engine/Source/Runtime/Experimental/Chaos/Public/Chaos/AABBTree.h
+++ b/Engine/Source/Runtime/Experimental/Chaos/Public/Chaos/AABBTree.h
@@ -17,6 +17,10 @@
 #include "Math/VectorRegister.h"
 #include <type_traits>
 
+#ifndef CHAOS_DEBUG_NAME
+#define CHAOS_DEBUG_NAME 0
+#endif
+
 CSV_DECLARE_CATEGORY_EXTERN(ChaosPhysicsTimers);
 
 struct CHAOS_API FAABBTreeCVars
@@ -1447,19 +1451,26 @@ public:
 
 			int32 MotherOfBestGrandChild = Nodes[BestGrandChildToSwap].ParentNode;
 
-			// Modify NodeIdx 
-			Nodes[NodeIdx].ChildrenNodes[AuntLocalChildIdx] = BestGrandChildToSwap;
-			// Modify BestGrandChildToSwap
 			if (UNLIKELY(!ensure(BestGrandChildToSwap != NodeIdx)))
 			{
 				UE_LOG(LogChaos, Warning, TEXT("AABBTree: 1: Rotate Node loop detected"));
+				return;
 			}
-			Nodes[BestGrandChildToSwap].ParentNode = NodeIdx;
-			// Modify BestAuntToSwap
+
 			if (UNLIKELY(!ensure(BestAuntToSwap != MotherOfBestGrandChild)))
 			{
+				// This really should not happen, but was due to NaNs entering the bounds (FIXED)
 				UE_LOG(LogChaos, Warning, TEXT("AABBTree: 2: Rotate Node loop detected"));
+				return;
 			}
+
+			// Modify NodeIdx 
+			Nodes[NodeIdx].ChildrenNodes[AuntLocalChildIdx] = BestGrandChildToSwap;
+			
+			// Modify BestGrandChildToSwap
+			Nodes[BestGrandChildToSwap].ParentNode = NodeIdx;
+			
+			// Modify BestAuntToSwap
 			Nodes[BestAuntToSwap].ParentNode = MotherOfBestGrandChild;
 			// Modify MotherOfBestGrandChild
 			Nodes[MotherOfBestGrandChild].ChildrenNodes[GrandChildLocalChildIdx] = BestAuntToSwap;
@@ -1569,6 +1580,8 @@ public:
 			if (UNLIKELY(!ensure(CurrentNodeIdx != ParentNodeIdx)))
 			{
 				UE_LOG(LogChaos, Warning, TEXT("AABBTree: 1: UpdateAncestorBounds loop detected"));
+				check(false); // Crash here, this is not recoverable
+				return;
 			}
 			int32 ChildIndex = WhichChildAmI(CurrentNodeIdx);
 			Nodes[ParentNodeIdx].ChildrenBounds[ChildIndex] = Nodes[CurrentNodeIdx].ChildrenBounds[0];
@@ -1727,10 +1740,20 @@ public:
 
 		bool bHasBounds = bInHasBounds;
 		// If bounds are bad, use global
-		if (bDynamicTree && bHasBounds && ValidateBounds(NewBounds) == false)
+		if (bHasBounds && ValidateBounds(NewBounds) == false)
 		{
 			bHasBounds = false;
-			ensureMsgf(false, TEXT("AABBTree encountered invalid bounds input. Forcing element to global payload. Min: %s Max: %s."),
+#if CHAOS_DEBUG_NAME			
+			if constexpr (std::is_same_v<TPayloadType, FAccelerationStructureHandle>)
+			{
+				if (IsInPhysicsThreadContext())
+				{
+					FString DebugStr = *(Payload.GetGeometryParticleHandle_PhysicsThread()->DebugName());
+					UE_LOG(LogChaos, Warning, TEXT("AABBTree encountered invalid bounds input : %s"), *DebugStr);
+				}
+			}
+#endif
+			UE_LOG(LogChaos, Warning, TEXT("AABBTree encountered invalid bounds input.Forcing element to global payload. Min: %s Max: %s"),
 				*NewBounds.Min().ToString(), *NewBounds.Max().ToString());
 		}
 
@@ -2034,16 +2057,17 @@ public:
 
 		this->SetAsyncTimeSlicingComplete(false);
 
-		if (DirtyElementTree)
+		if (From.DirtyElementTree)
 		{
-			if (From.DirtyElementTree)
+			if (!DirtyElementTree)
 			{
-				DirtyElementTree->PrepareCopyTimeSliced(*(From.DirtyElementTree));
+				DirtyElementTree = TUniquePtr<TAABBTree<TPayloadType, TLeafType, bMutable, T>>(new TAABBTree<TPayloadType, TLeafType, bMutable, T>());
 			}
-			else
-			{
-				DirtyElementTree = nullptr;
-			}			
+			DirtyElementTree->PrepareCopyTimeSliced(*(From.DirtyElementTree));
+		}
+		else
+		{
+			DirtyElementTree = nullptr;
 		}
 	}
 	
@@ -3599,6 +3623,10 @@ private:
 				}				
 				*DirtyElementTree = *Rhs.DirtyElementTree;
 			}
+			else
+			{
+				DirtyElementTree = nullptr;
+			}
 		}
 
 		return *this;
