add an option (p.Chaos.DebugDraw.ConnectionGraphLevelOffset) to change what level of connection graph is being drawn in chaos debug draw

diff --git a/Engine/Source/Runtime/Experimental/Chaos/Private/Chaos/ChaosDebugDraw.cpp b/Engine/Source/Runtime/Experimental/Chaos/Private/Chaos/ChaosDebugDraw.cpp
--- a/Engine/Source/Runtime/Experimental/Chaos/Private/Chaos/ChaosDebugDraw.cpp
+++ b/Engine/Source/Runtime/Experimental/Chaos/Private/Chaos/ChaosDebugDraw.cpp
@@ -89,6 +89,8 @@ namespace Chaos
 		float ChaosDebugDrawCollisionDuration = 0.0f;
 		FAutoConsoleVariableRef CVarChaosDebugDrawCollisionDuration(TEXT("p.Chaos.DebugDraw.CollisionDuration"), ChaosDebugDrawCollisionDuration, TEXT("How long Collision debug draw should remain on screen in seconds. 0 for 1 frame."));
 
+		int32 ChaosConnectionGraphDrawLevelOffset = 0;
+		FAutoConsoleVariableRef CVarChaosConnectionGraphDrawLevelOffset(TEXT("p.Chaos.DebugDraw.ConnectionGraphLevelOffset"), ChaosConnectionGraphDrawLevelOffset, TEXT("If 0, draws the connection graph between children particles of active particles. If 1, draws the connection grpah between grand-children particles of active particles, etc."));
 
 		// NOTE: These settings should never really be used - they are the fallback defaults
 		// if the user does not specify settings in the debug draw call.
@@ -1390,24 +1392,21 @@ namespace Chaos
 		}
 
 		struct FOrderedClusterConnection {
-			const FRigidClustering::FClusterHandle Parent;
 			const FRigidClustering::FClusterHandle Handle0;
 			const FRigidClustering::FClusterHandle Handle1;
 			FReal Strain;
 
 			FOrderedClusterConnection()
-				: Parent(nullptr)
-				, Handle0(nullptr)
+				: Handle0(nullptr)
 				, Handle1(nullptr)
 				, Strain(0)
 			{}
 
-			FOrderedClusterConnection(const FRigidClustering::FClusterHandle InParent,
+			FOrderedClusterConnection(
 				const FRigidClustering::FClusterHandle In0,
 				const FRigidClustering::FClusterHandle In1,
 				FReal InStrain)
-				: Parent(InParent)
-				, Handle0(In0)
+				: Handle0(In0)
 				, Handle1(In1)
 				, Strain(InStrain)
 			{
@@ -1432,7 +1431,32 @@ namespace Chaos
 
 		void DrawConnectionGraphImpl(const FRigidClustering& Clustering, const FChaosDebugDrawSettings& Settings)
 		{
-			auto DrawConnections = [&](TSet<FOrderedClusterConnection>& Connections, FReal MaxStrain, const FChaosDebugDrawSettings& Settings)
+			auto GetClusterTransform = [](FPBDRigidParticleHandle* Particle) -> FTransform
+			{
+				auto GetClusterTransformImpl = [](FPBDRigidParticleHandle* Particle, auto& Helper)->FTransform
+				{
+					if (!Particle)
+					{
+						return FTransform{};
+					}
+
+					if (FPBDRigidClusteredParticleHandle* Cluster = Particle->CastToClustered())
+					{
+						const ClusterId& Id = Cluster->ClusterIds();
+						if (Id.Id)
+						{
+							return Cluster->ChildToParent() * Helper(Id.Id, Helper);
+						}
+					}
+
+					return FTransform{ Particle->R(), Particle->X() };
+				};
+
+				return GetClusterTransformImpl(Particle, GetClusterTransformImpl);
+			};
+
+
+			auto DrawConnections = [&GetClusterTransform](TSet<FOrderedClusterConnection>& Connections, FReal MaxStrain, const FChaosDebugDrawSettings& Settings)
 			{
 				if (FMath::IsNearlyZero(MaxStrain))
 				{
@@ -1441,9 +1465,8 @@ namespace Chaos
 
 				for (auto& Connection : Connections)
 				{
-					FRigidTransform3 ClusterTransform(Connection.Parent->X(), Connection.Parent->R());
-					FVec3 Pos0 = ClusterTransform.TransformPosition(Connection.Handle0->ChildToParent().GetLocation());
-					FVec3 Pos1 = ClusterTransform.TransformPosition(Connection.Handle1->ChildToParent().GetLocation());
+					FVec3 Pos0 = GetClusterTransform(Connection.Handle0).GetLocation();
+					FVec3 Pos1 = GetClusterTransform(Connection.Handle1).GetLocation();
 					const FVec3 Normal = (Pos1 - Pos0).GetSafeNormal();
 					const FReal Length = (Pos1 - Pos0).Size();
 					constexpr FReal ShrinkPercentage = 0.05;
@@ -1451,69 +1474,101 @@ namespace Chaos
 					Pos0 += SkrinkVector;
 					Pos1 -= SkrinkVector;
 
-					FColor Color = FColor::Green; // FMath::Lerp(FColor::Green, FColor::Red, Connection.Strain / MaxStrain);
+					FColor Color = (Connection.Handle0->Parent() == Connection.Handle1->Parent()) ? FColor::Green : FColor::Blue; // FMath::Lerp(FColor::Green, FColor::Red, Connection.Strain / MaxStrain);
 
 					FDebugDrawQueue::GetInstance().DrawDebugLine(Pos0, Pos1, Color, false, UE_KINDA_SMALL_NUMBER, Settings.DrawPriority, Settings.LineThickness);
 				}
-
 			};
 
-			//if (bChaosDebugDrawClusterConnectionGraph)
+			auto FindChildParticlesWithLevelOffset = [&Clustering](const FRigidClustering::FClusterHandle& Cluster)
 			{
-				FRealSingle MaxStrain = TNumericLimits<FRealSingle>::Lowest();
-				int32 NumConnections = 0;
-				const TArrayCollectionArray<FRealSingle>& Strain = Clustering.GetStrainArray();
+				FRigidClustering::FRigidHandleArray Children;
+
+				TQueue<TPair<FRigidClustering::FClusterHandle, int32>> Process;
+				Process.Enqueue({ Cluster, 0 });
+
+				TPair<FRigidClustering::FClusterHandle, int32> Next;
+				while (Process.Dequeue(Next))
+				{
+					if (const FRigidClustering::FRigidHandleArray* CurrentChildren = Clustering.GetChildrenMap().Find(Next.Key))
+					{
+						if (Next.Value == ChaosConnectionGraphDrawLevelOffset)
+						{
+							Children.Append(*CurrentChildren);
+						}
+						else
+						{
+							for (FPBDRigidParticleHandle* Child : *CurrentChildren)
+							{
+								if (FPBDRigidClusteredParticleHandle* ClusteredChild = Child->CastToClustered())
+								{
+									Process.Enqueue({ClusteredChild, Next.Value + 1});
+								}
+							}
+						}
+					}
+				}
+
+				return Children;
+			};
 
-				for (auto& ActiveCluster : Clustering.GetChildrenMap())
+			auto VisitEveryConnectionToDraw = [&Clustering, &FindChildParticlesWithLevelOffset](auto&& Lambda)
+			{
+				for (const TPair<FRigidClustering::FClusterHandle, FRigidClustering::FRigidHandleArray>& ActiveCluster : Clustering.GetChildrenMap())
 				{
 					if (!ActiveCluster.Key->Disabled())
 					{
-						for (const FRigidClustering::FRigidHandle& Rigid : ActiveCluster.Value)
+						for (const FRigidClustering::FRigidHandle& Rigid : FindChildParticlesWithLevelOffset(ActiveCluster.Key))
 						{
 							if (const FRigidClustering::FClusterHandle& Child = Rigid->CastToClustered())
 							{
-								MaxStrain = FMath::Max(MaxStrain, Child->GetInternalStrains());
+								Lambda(Child);
 							}
 						}
-						NumConnections += ActiveCluster.Value.Num();
 					}
 				}
+			};
 
-				if (NumConnections)
+
+			FRealSingle MaxStrain = TNumericLimits<FRealSingle>::Lowest();
+			int32 NumConnections = 0;
+			const TArrayCollectionArray<FRealSingle>& Strain = Clustering.GetStrainArray();
+
+			VisitEveryConnectionToDraw(
+				[&MaxStrain, &NumConnections](const FRigidClustering::FClusterHandle& Child)
 				{
-					TSet<FOrderedClusterConnection> Connections;
-					Connections.Reserve(NumConnections);
+					MaxStrain = FMath::Max(MaxStrain, Child->GetInternalStrains());
+					++NumConnections;
+				}
+			);
+
+			if (NumConnections)
+			{
+				TSet<FOrderedClusterConnection> Connections;
+				Connections.Reserve(NumConnections);
 
-					for (auto& ActiveCluster : Clustering.GetChildrenMap())
+				VisitEveryConnectionToDraw(
+					[&Settings, &Connections, &GetClusterTransform](const FRigidClustering::FClusterHandle& Child)
 					{
-						if (!ActiveCluster.Key->Disabled())
+						FColor PointColor = Settings.ShapesColorsPerState.GetColorFromState(Child->ObjectState());
+						if (Child->IsAnchored())
 						{
-							for (const FRigidClustering::FRigidHandle& Rigid : ActiveCluster.Value)
-							{
-								if (const FRigidClustering::FClusterHandle& Child = Rigid->CastToClustered())
-								{
+							PointColor = Settings.ShapesColorsPerState.KinematicColor;
+						}
+						FDebugDrawQueue::GetInstance().DrawDebugPoint(GetClusterTransform(Child).GetLocation(), PointColor, false, UE_KINDA_SMALL_NUMBER, Settings.DrawPriority, Settings.LineThickness * 10);
 
-									FColor PointColor = Settings.ShapesColorsPerState.GetColorFromState(Child->ObjectState());
-									if (Child->IsAnchored())
-									{ 
-										PointColor = Settings.ShapesColorsPerState.KinematicColor;
-									}
-									FDebugDrawQueue::GetInstance().DrawDebugPoint(Child->X(), PointColor, false, UE_KINDA_SMALL_NUMBER, Settings.DrawPriority, Settings.LineThickness * 10);
-
-
-									const FConnectivityEdgeArray& Edges = Child->ConnectivityEdges();
-									for (const TConnectivityEdge<FReal>& Edge : Edges)
-									{
-										const FRigidClustering::FRigidHandle A = (Child < Edge.Sibling) ? Child : Edge.Sibling;
-										const FRigidClustering::FRigidHandle B = (Child < Edge.Sibling) ? Edge.Sibling : Child;
-										Connections.Add(FOrderedClusterConnection(ActiveCluster.Key, A->CastToClustered(), B->CastToClustered(), Edge.Strain));
-									}
-								}
-							}
+
+						const FConnectivityEdgeArray& Edges = Child->ConnectivityEdges();
+						for (const TConnectivityEdge<FReal>& Edge : Edges)
+						{
+							const FRigidClustering::FRigidHandle A = (Child < Edge.Sibling) ? Child : Edge.Sibling;
+							const FRigidClustering::FRigidHandle B = (Child < Edge.Sibling) ? Edge.Sibling : Child;
+							Connections.Add(FOrderedClusterConnection(A->CastToClustered(), B->CastToClustered(), Edge.Strain));
 						}
 					}
-					DrawConnections(Connections, MaxStrain, Settings);
-				}
+				);
+
+				DrawConnections(Connections, MaxStrain, Settings);
 			}
 		}
 
