Add support for custom primitive data to ISM Pool

diff --git a/Engine/Source/Runtime/Experimental/GeometryCollectionEngine/Private/GeometryCollection/GeometryCollectionISMPoolComponent.cpp b/Engine/Source/Runtime/Experimental/GeometryCollectionEngine/Private/GeometryCollection/GeometryCollectionISMPoolComponent.cpp
--- a/Engine/Source/Runtime/Experimental/GeometryCollectionEngine/Private/GeometryCollection/GeometryCollectionISMPoolComponent.cpp
+++ b/Engine/Source/Runtime/Experimental/GeometryCollectionEngine/Private/GeometryCollection/GeometryCollectionISMPoolComponent.cpp
@@ -83,6 +83,10 @@ FGeometryCollectionISM::FGeometryCollectionISM(AActor* InOwningActor, const FGeo
 	{
 		ISMC->SetMaterial(MaterialIndex, MeshInstance.MaterialsOverrides[MaterialIndex]);
 	}
+	for (int32 DataIndex = 0; DataIndex < MeshInstance.CustomPrimitiveData.Num(); DataIndex++)
+	{
+		ISMC->SetCustomPrimitiveDataFloat(DataIndex, MeshInstance.CustomPrimitiveData[DataIndex]);
+	}
 
 	ISMC->SetRemoveSwap();
 	ISMC->NumCustomDataFloats = MeshInstance.Desc.NumCustomDataFloats;
diff --git a/Engine/Source/Runtime/Experimental/GeometryCollectionEngine/Public/GeometryCollection/GeometryCollectionISMPoolComponent.h b/Engine/Source/Runtime/Experimental/GeometryCollectionEngine/Public/GeometryCollection/GeometryCollectionISMPoolComponent.h
--- a/Engine/Source/Runtime/Experimental/GeometryCollectionEngine/Public/GeometryCollection/GeometryCollectionISMPoolComponent.h
+++ b/Engine/Source/Runtime/Experimental/GeometryCollectionEngine/Public/GeometryCollection/GeometryCollectionISMPoolComponent.h
@@ -145,27 +145,41 @@ struct FGeometryCollectionStaticMeshInstance
 {
 	UStaticMesh* StaticMesh = nullptr;
 	TArray<UMaterialInterface*> MaterialsOverrides;
+	TArray<float> CustomPrimitiveData;
 	FISMComponentDescription Desc;
 
 	bool operator==(const FGeometryCollectionStaticMeshInstance& Other) const 
 	{
-		if (StaticMesh == Other.StaticMesh && Desc == Other.Desc)
+		if (StaticMesh != Other.StaticMesh || !(Desc == Other.Desc))
 		{
-			if (MaterialsOverrides.Num() == Other.MaterialsOverrides.Num())
+			return false;
+		}
+		if (MaterialsOverrides.Num() != Other.MaterialsOverrides.Num())
+		{
+			return false;
+		}
+		for (int32 MatIndex = 0; MatIndex < MaterialsOverrides.Num(); MatIndex++)
+		{
+			const FName MatName = MaterialsOverrides[MatIndex] ? MaterialsOverrides[MatIndex]->GetFName() : NAME_None;
+			const FName OtherName = Other.MaterialsOverrides[MatIndex] ? Other.MaterialsOverrides[MatIndex]->GetFName() : NAME_None;
+			if (MatName != OtherName)
 			{
-				for (int32 MatIndex = 0; MatIndex < MaterialsOverrides.Num(); MatIndex++)
-				{
-					const FName MatName = MaterialsOverrides[MatIndex] ? MaterialsOverrides[MatIndex]->GetFName() : NAME_None;
-					const FName OtherName = Other.MaterialsOverrides[MatIndex] ? Other.MaterialsOverrides[MatIndex]->GetFName() : NAME_None;
-					if (MatName != OtherName)
-					{
-						return false;
-					}
-				}
-				return true;
+				return false;
 			}
 		}
-		return false;
+		if (CustomPrimitiveData.Num() != Other.CustomPrimitiveData.Num())
+		{
+			return false;
+		}
+		for (int32 DataIndex = 0; DataIndex < CustomPrimitiveData.Num(); DataIndex++)
+		{
+			if (CustomPrimitiveData[DataIndex] != Other.CustomPrimitiveData[DataIndex])
+			{
+				return false;
+			}
+		}
+
+		return true;
 	}
 };
 
@@ -177,6 +191,10 @@ FORCEINLINE uint32 GetTypeHash(const FGeometryCollectionStaticMeshInstance& Mesh
 	{
 		CombinedHash = HashCombine(CombinedHash, GetTypeHash(Material));
 	}
+	for (const float CustomFloat : MeshInstance.CustomPrimitiveData)
+	{
+		CombinedHash = HashCombine(CombinedHash, GetTypeHash(CustomFloat));
+	}
 	CombinedHash = HashCombine(CombinedHash, GetTypeHash(MeshInstance.Desc));
 	return CombinedHash;
 }
