add functions to the physics object interface to query the hierarchy level of a particle and generate connectivity edges between particles

- also modified the various connectivity graph function APIs to be more flexible and to support filtering based off a "from" and "to" set of particles

diff --git a/Engine/Source/Runtime/Engine/Classes/Interfaces/IPhysicsComponent.h b/Engine/Source/Runtime/Engine/Classes/Interfaces/IPhysicsComponent.h
--- a/Engine/Source/Runtime/Engine/Classes/Interfaces/IPhysicsComponent.h
+++ b/Engine/Source/Runtime/Engine/Classes/Interfaces/IPhysicsComponent.h
@@ -4,6 +4,7 @@
 
 #include "Chaos/ParticleHandleFwd.h"
 #include "Chaos/PhysicsObject.h"
+#include "Containers/Array.h"
 #include "UObject/ObjectMacros.h"
 #include "UObject/Interface.h"
 #include "IPhysicsComponent.generated.h"
@@ -25,4 +26,34 @@ public:
 	virtual TArray<Chaos::FPhysicsObject*> GetAllPhysicsObjects() const = 0;
 	
 	virtual Chaos::FPhysicsObjectId GetIdFromGTParticle(Chaos::FGeometryParticle* Particle) const = 0;
-};
\ No newline at end of file
+};
+
+template<typename TId>
+TArray<Chaos::FPhysicsObjectHandle> GetAllPhysicsObjectsById(IPhysicsComponent* Component, const TArray<TId>& AllIds)
+{
+	constexpr bool bIsId = std::is_same_v<TId, Chaos::FPhysicsObjectId>;
+	constexpr bool bIsName = std::is_same_v<TId, FName>;
+
+	static_assert(bIsId || bIsName, "Invalid ID type passed to GetAllPhysicsObjectsById");
+	if (!Component)
+	{
+		return {};
+	}
+
+	TArray<Chaos::FPhysicsObjectHandle> Objects;
+	Objects.Reserve(AllIds.Num());
+
+	for (const TId& Id : AllIds)
+	{
+		if constexpr (bIsId)
+		{
+			Objects.Add(Component->GetPhysicsObjectById(Id));
+		}
+		else if constexpr (bIsName)
+		{
+			Objects.Add(Component->GetPhysicsObjectByName(Id));
+		}
+	}
+
+	return Objects;
+}
\ No newline at end of file
diff --git a/Engine/Source/Runtime/Experimental/Chaos/Private/Chaos/PBDRigidClustering.cpp b/Engine/Source/Runtime/Experimental/Chaos/Private/Chaos/PBDRigidClustering.cpp
--- a/Engine/Source/Runtime/Experimental/Chaos/Private/Chaos/PBDRigidClustering.cpp
+++ b/Engine/Source/Runtime/Experimental/Chaos/Private/Chaos/PBDRigidClustering.cpp
@@ -1430,14 +1430,14 @@ namespace Chaos
 	}
 	
 	DECLARE_CYCLE_STAT(TEXT("TPBDRigidClustering<>::GenerateConnectionGraph"), STAT_GenerateConnectionGraph, STATGROUP_Chaos);
-	void 
+	void
 	FRigidClustering::GenerateConnectionGraph(
-		Chaos::FPBDRigidClusteredParticleHandle* Parent,
-		const FClusterCreationParameters& Parameters)
+		TArray<FPBDRigidParticleHandle*> Particles,
+		const FClusterCreationParameters& Parameters,
+		const TSet<FPBDRigidParticleHandle*>* FromParticles,
+		const TSet<FPBDRigidParticleHandle*>* ToParticles)
 	{
 		SCOPE_CYCLE_COUNTER(STAT_GenerateConnectionGraph);
-		if (!MChildren.Contains(Parent)) 
-			return;
 
 		// Connectivity Graph
 		//    Build a connectivity graph for the cluster. If the PointImplicit is specified
@@ -1448,9 +1448,7 @@ namespace Chaos
 		{
 			FClusterCreationParameters::EConnectionMethod LocalConnectionMethod = Parameters.ConnectionMethod;
 
-			if (LocalConnectionMethod == FClusterCreationParameters::EConnectionMethod::None ||
-				(LocalConnectionMethod == FClusterCreationParameters::EConnectionMethod::PointImplicit && 
-				 !Parent->CollisionParticles()))
+			if (LocalConnectionMethod == FClusterCreationParameters::EConnectionMethod::None)
 			{
 				LocalConnectionMethod = FClusterCreationParameters::EConnectionMethod::MinimalSpanningSubsetDelaunayTriangulation; // default method
 			}
@@ -1458,27 +1456,44 @@ namespace Chaos
 			if (LocalConnectionMethod == FClusterCreationParameters::EConnectionMethod::PointImplicit ||
 				LocalConnectionMethod == FClusterCreationParameters::EConnectionMethod::PointImplicitAugmentedWithMinimalDelaunay)
 			{
-				UpdateConnectivityGraphUsingPointImplicit(Parent, Parameters);
+				UpdateConnectivityGraphUsingPointImplicit(Particles, Parameters.CoillisionThicknessPercent, FromParticles, ToParticles);
 			}
 
 			if (LocalConnectionMethod == FClusterCreationParameters::EConnectionMethod::DelaunayTriangulation)
 			{
-				UpdateConnectivityGraphUsingDelaunayTriangulation(Parent, Parameters); // not thread safe
+				UpdateConnectivityGraphUsingDelaunayTriangulation(Particles, Parameters, FromParticles, ToParticles); // not thread safe
 			}
 
 			if (LocalConnectionMethod == FClusterCreationParameters::EConnectionMethod::BoundsOverlapFilteredDelaunayTriangulation)
 			{
-				UpdateConnectivityGraphUsingDelaunayTriangulationWithBoundsOverlaps(Parent, Parameters);
+				UpdateConnectivityGraphUsingDelaunayTriangulationWithBoundsOverlaps(Particles, Parameters, FromParticles, ToParticles);
 			}
 
 			if (LocalConnectionMethod == FClusterCreationParameters::EConnectionMethod::PointImplicitAugmentedWithMinimalDelaunay ||
 				LocalConnectionMethod == FClusterCreationParameters::EConnectionMethod::MinimalSpanningSubsetDelaunayTriangulation)
 			{
-				FixConnectivityGraphUsingDelaunayTriangulation(Parent, Parameters);
+				FixConnectivityGraphUsingDelaunayTriangulation(Particles, Parameters, FromParticles, ToParticles);
 			}
 		}
 	}
 
+	void 
+	FRigidClustering::GenerateConnectionGraph(
+		Chaos::FPBDRigidClusteredParticleHandle* Parent,
+		const FClusterCreationParameters& Parameters)
+	{
+		if (!MChildren.Contains(Parent))
+			return;
+
+		FClusterCreationParameters FinalParameters = Parameters;
+		if (Parameters.ConnectionMethod == FClusterCreationParameters::EConnectionMethod::PointImplicit && !Parent->CollisionParticles())
+		{
+			FinalParameters.ConnectionMethod = FClusterCreationParameters::EConnectionMethod::MinimalSpanningSubsetDelaunayTriangulation;
+		}
+
+		GenerateConnectionGraph(MChildren[Parent], FinalParameters);
+	}
+
 	DECLARE_CYCLE_STAT(TEXT("TPBDRigidClustering<>::ClearConnectionGraph"), STAT_ClearConnectionGraph, STATGROUP_Chaos);
 	void
 	FRigidClustering::ClearConnectionGraph(FPBDRigidClusteredParticleHandle* Parent)
@@ -1869,11 +1884,26 @@ namespace Chaos
 		return bCrumbledAnyCluster;
 	}
 
-	static void ConnectClusteredNodes(FPBDRigidClusteredParticleHandle* ClusteredChild1, FPBDRigidClusteredParticleHandle* ClusteredChild2)
+	static void ConnectClusteredNodes(FPBDRigidClusteredParticleHandle* ClusteredChild1, FPBDRigidClusteredParticleHandle* ClusteredChild2, const TSet<FPBDRigidParticleHandle*>* FromParticles, const TSet<FPBDRigidParticleHandle*>* ToParticles)
 	{
 		check(ClusteredChild1 && ClusteredChild2);
 		if (ClusteredChild1 == ClusteredChild2)
+		{
 			return;
+		}
+
+		if (FromParticles && ToParticles)
+		{
+			// We are enforcing that the edge we create crosses from the "from" set to the "to" set (or vice versa).
+			const bool bInFrom = FromParticles->Contains(ClusteredChild1) || FromParticles->Contains(ClusteredChild2);
+			const bool bInTo = ToParticles->Contains(ClusteredChild1) || ToParticles->Contains(ClusteredChild2);
+			if (!bInFrom || !bInTo)
+			{
+				return;
+			}
+		}
+
+
 		const FRealSingle AvgStrain = (ClusteredChild1->GetInternalStrains() + ClusteredChild2->GetInternalStrains()) * 0.5f;
 		TArray<TConnectivityEdge<FReal>>& Edges1 = ClusteredChild1->ConnectivityEdges();
 		TArray<TConnectivityEdge<FReal>>& Edges2 = ClusteredChild2->ConnectivityEdges();
@@ -1889,19 +1919,21 @@ namespace Chaos
 		}
 	}
 	
-	static void ConnectNodes(FPBDRigidParticleHandle* Child1, FPBDRigidParticleHandle* Child2)
+	static void ConnectNodes(FPBDRigidParticleHandle* Child1, FPBDRigidParticleHandle* Child2, const TSet<FPBDRigidParticleHandle*>* FromParticles, const TSet<FPBDRigidParticleHandle*>* ToParticles)
 	{
 		check(Child1 != Child2);
 		FPBDRigidClusteredParticleHandle* ClusteredChild1 = Child1->CastToClustered();
 		FPBDRigidClusteredParticleHandle* ClusteredChild2 = Child2->CastToClustered();
-		ConnectClusteredNodes(ClusteredChild1, ClusteredChild2);
+		ConnectClusteredNodes(ClusteredChild1, ClusteredChild2, FromParticles, ToParticles);
 	}
 
 	DECLARE_CYCLE_STAT(TEXT("TPBDRigidClustering<>::UpdateConnectivityGraphUsingPointImplicit"), STAT_UpdateConnectivityGraphUsingPointImplicit, STATGROUP_Chaos);
 	void 
 	FRigidClustering::UpdateConnectivityGraphUsingPointImplicit(
-		Chaos::FPBDRigidClusteredParticleHandle* Parent,
-		const FClusterCreationParameters& Parameters)
+		const TArray<FPBDRigidParticleHandle*>& Particles,
+		FReal CollisionThicknessPercent,
+		const TSet<FPBDRigidParticleHandle*>* FromParticles,
+		const TSet<FPBDRigidParticleHandle*>* ToParticles)
 	{
 		SCOPE_CYCLE_COUNTER(STAT_UpdateConnectivityGraphUsingPointImplicit);
 
@@ -1910,18 +1942,17 @@ namespace Chaos
 			return;
 		}
 
-		const FReal Delta = FMath::Min(FMath::Max(Parameters.CoillisionThicknessPercent, FReal(0)), FReal(1));
-		const TArray<FPBDRigidParticleHandle*>& Children = MChildren[Parent];
+		const FReal Delta = FMath::Min(FMath::Max(CollisionThicknessPercent, FReal(0)), FReal(1));
 
 		typedef TPair<FPBDRigidParticleHandle*, FPBDRigidParticleHandle*> ParticlePair;
 		typedef TSet<ParticlePair> ParticlePairArray;
 
 		TArray<ParticlePairArray> Connections;
-		Connections.Init(ParticlePairArray(), Children.Num());
+		Connections.Init(ParticlePairArray(), Particles.Num());
 
-		PhysicsParallelFor(Children.Num(), [&](int32 i)
+		PhysicsParallelFor(Particles.Num(), [&](int32 i)
 			{
-				FPBDRigidParticleHandle* Child1 = Children[i];
+				FPBDRigidParticleHandle* Child1 = Particles[i];
 				if (Child1->Geometry() && Child1->Geometry()->HasBoundingBox())
 				{
 					ParticlePairArray& ConnectionList = Connections[i];
@@ -1930,12 +1961,12 @@ namespace Chaos
 					FRigidTransform3 TM1 = FRigidTransform3(Child1X, Child1->R());
 
 					const int32 Offset = i + 1;
-					const int32 NumRemainingChildren = Children.Num() - Offset;
+					const int32 NumRemainingParticles = Particles.Num() - Offset;
 
-					for (int32 Idx = 0; Idx < NumRemainingChildren; ++Idx)
+					for (int32 Idx = 0; Idx < NumRemainingParticles; ++Idx)
 					{
-						const int32 ChildrenIdx = Offset + Idx;
-						FPBDRigidParticleHandle* Child2 = Children[ChildrenIdx];
+						const int32 ParticlesIdx = Offset + Idx;
+						FPBDRigidParticleHandle* Child2 = Particles[ParticlesIdx];
 						if (Child2->CollisionParticles())
 						{
 
@@ -1964,53 +1995,62 @@ namespace Chaos
 		{
 			for (const ParticlePair& Edge : ConnectionList)
 			{
-				ConnectNodes(Edge.Key, Edge.Value);
+				ConnectNodes(Edge.Key, Edge.Value, FromParticles, ToParticles);
 			}
 		}
 
 	}
 
-	DECLARE_CYCLE_STAT(TEXT("TPBDRigidClustering<>::FixConnectivityGraphUsingDelaunayTriangulation"), STAT_FixConnectivityGraphUsingDelaunayTriangulation, STATGROUP_Chaos);
-	void 
-	FRigidClustering::FixConnectivityGraphUsingDelaunayTriangulation(
+	void
+	FRigidClustering::UpdateConnectivityGraphUsingPointImplicit(
 		Chaos::FPBDRigidClusteredParticleHandle* Parent,
 		const FClusterCreationParameters& Parameters)
 	{
-		SCOPE_CYCLE_COUNTER(STAT_FixConnectivityGraphUsingDelaunayTriangulation);
-
 		const TArray<FPBDRigidParticleHandle*>& Children = MChildren[Parent];
+		UpdateConnectivityGraphUsingPointImplicit(Children, Parameters.CoillisionThicknessPercent);
+	}
 
-		// Compute Delaunay neighbor graph on children centers
+	DECLARE_CYCLE_STAT(TEXT("TPBDRigidClustering<>::FixConnectivityGraphUsingDelaunayTriangulation"), STAT_FixConnectivityGraphUsingDelaunayTriangulation, STATGROUP_Chaos);
+	void
+	FRigidClustering::FixConnectivityGraphUsingDelaunayTriangulation(
+		const TArray<FPBDRigidParticleHandle*>& Particles,
+		const FClusterCreationParameters& Parameters,
+		const TSet<FPBDRigidParticleHandle*>* FromParticles,
+		const TSet<FPBDRigidParticleHandle*>* ToParticles)
+	{
+		SCOPE_CYCLE_COUNTER(STAT_FixConnectivityGraphUsingDelaunayTriangulation);
+
+		// Compute Delaunay neighbor graph on Particles centers
 		TArray<FVector> Pts;
-		Pts.AddUninitialized(Children.Num());
-		for (int32 i = 0; i < Children.Num(); ++i)
+		Pts.AddUninitialized(Particles.Num());
+		for (int32 i = 0; i < Particles.Num(); ++i)
 		{
-			Pts[i] = Children[i]->X();
+			Pts[i] = Particles[i]->X();
 		}
-		TArray<TArray<int32>> Neighbors; // Indexes into Children
+		TArray<TArray<int32>> Neighbors; // Indexes into Particles
 		VoronoiNeighbors(Pts, Neighbors);
 
-		// Build a UnionFind graph to find (indirectly) connected children
+		// Build a UnionFind graph to find (indirectly) connected Particles
 		struct UnionFindInfo
 		{
 			FPBDRigidParticleHandle* GroupId;
 			int32 Size;
 		};
 		TMap<FPBDRigidParticleHandle*, UnionFindInfo> UnionInfo;
-		UnionInfo.Reserve(Children.Num());
+		UnionInfo.Reserve(Particles.Num());
 
 		// Initialize UnionInfo:
-		//		0: GroupId = Children[0], Size = 1
-		//		1: GroupId = Children[1], Size = 1
-		//		2: GroupId = Children[2], Size = 1
-		//		3: GroupId = Children[3], Size = 1
+		//		0: GroupId = Particles[0], Size = 1
+		//		1: GroupId = Particles[1], Size = 1
+		//		2: GroupId = Particles[2], Size = 1
+		//		3: GroupId = Particles[3], Size = 1
 
-		for(FPBDRigidParticleHandle* Child : Children)
+		for (FPBDRigidParticleHandle* Child : Particles)
 		{
 			UnionInfo.Add(Child, { Child, 1 }); // GroupId, Size
 		}
 
-		auto FindGroup = [&](FPBDRigidParticleHandle* Id) 
+		auto FindGroup = [&](FPBDRigidParticleHandle* Id)
 		{
 			FPBDRigidParticleHandle* GroupId = Id;
 			if (GroupId)
@@ -2029,13 +2069,13 @@ namespace Chaos
 			return GroupId;
 		};
 
-		// MergeGroup(Children[0], Children[1])
-		//		0: GroupId = Children[1], Size = 0
-		//		1: GroupId = Children[1], Size = 2
-		//		2: GroupId = Children[2], Size = 1
-		//		3: GroupId = Children[3], Size = 1
+		// MergeGroup(Particles[0], Particles[1])
+		//		0: GroupId = Particles[1], Size = 0
+		//		1: GroupId = Particles[1], Size = 2
+		//		2: GroupId = Particles[2], Size = 1
+		//		3: GroupId = Particles[3], Size = 1
 
-		auto MergeGroup = [&](FPBDRigidParticleHandle* A, FPBDRigidParticleHandle* B) 
+		auto MergeGroup = [&](FPBDRigidParticleHandle* A, FPBDRigidParticleHandle* B)
 		{
 			FPBDRigidParticleHandle* GroupA = FindGroup(A);
 			FPBDRigidParticleHandle* GroupB = FindGroup(B);
@@ -2055,9 +2095,9 @@ namespace Chaos
 		};
 
 		// Merge all groups with edges connecting them.
-		for (int32 i = 0; i < Children.Num(); ++i)
+		for (int32 i = 0; i < Particles.Num(); ++i)
 		{
-			FPBDRigidParticleHandle* Child = Children[i];
+			FPBDRigidParticleHandle* Child = Particles[i];
 			const TArray<TConnectivityEdge<FReal>>& Edges = Child->CastToClustered()->ConnectivityEdges();
 			for (const TConnectivityEdge<FReal>& Edge : Edges)
 			{
@@ -2079,10 +2119,10 @@ namespace Chaos
 		TArray<LinkCandidate> Candidates;
 		Candidates.Reserve(Neighbors.Num());
 
-		const FReal AlwaysAcceptBelowDistSqThreshold = 50.f*50.f*100.f*MClusterConnectionFactor;
+		const FReal AlwaysAcceptBelowDistSqThreshold = 50.f * 50.f * 100.f * MClusterConnectionFactor;
 		for (int32 i = 0; i < Neighbors.Num(); i++)
 		{
-			FPBDRigidParticleHandle* Child1 = Children[i];
+			FPBDRigidParticleHandle* Child1 = Particles[i];
 			const TArray<int32>& Child1Neighbors = Neighbors[i];
 			for (const int32 Nbr : Child1Neighbors)
 			{
@@ -2090,13 +2130,13 @@ namespace Chaos
 				{ // assume we'll get the symmetric connection; don't bother considering this one
 					continue;
 				}
-				FPBDRigidParticleHandle* Child2 = Children[Nbr];
+				FPBDRigidParticleHandle* Child2 = Particles[Nbr];
 
 				const FReal DistSq = FVector::DistSquared(Pts[i], Pts[Nbr]);
 				if (DistSq < AlwaysAcceptBelowDistSqThreshold)
 				{ // below always-accept threshold: don't bother adding to candidates array, just merge now
 					MergeGroup(Child1, Child2);
-					ConnectNodes(Child1, Child2);
+					ConnectNodes(Child1, Child2, FromParticles, ToParticles);
 					continue;
 				}
 
@@ -2119,10 +2159,19 @@ namespace Chaos
 			if (FindGroup(Child1) != FindGroup(Child2))
 			{
 				MergeGroup(Child1, Child2);
-				ConnectNodes(Child1, Child2);
+				ConnectNodes(Child1, Child2, FromParticles, ToParticles);
 			}
 		}
 	}
+
+	void 
+	FRigidClustering::FixConnectivityGraphUsingDelaunayTriangulation(
+		Chaos::FPBDRigidClusteredParticleHandle* Parent,
+		const FClusterCreationParameters& Parameters)
+	{
+		const TArray<FPBDRigidParticleHandle*>& Children = MChildren[Parent];
+		FixConnectivityGraphUsingDelaunayTriangulation(Children, Parameters);
+	}
 	
 	// connection filters
 	static bool IsAlwaysValidConnection(const FPBDRigidParticleHandle* Child1, const FPBDRigidParticleHandle* Child2)
@@ -2143,7 +2192,10 @@ namespace Chaos
 
 	template <typename FilterLambda>
 	static void UpdateConnectivityGraphUsingDelaunayTriangulationWithFiltering(
-		const TArray<FPBDRigidParticleHandle*>& Children, FilterLambda ShouldKeepConnection)
+		const TArray<FPBDRigidParticleHandle*>& Children,
+		FilterLambda ShouldKeepConnection,
+		const TSet<FPBDRigidParticleHandle*>* FromParticles,
+		const TSet<FPBDRigidParticleHandle*>* ToParticles)
 	{
 		TArray<FVector> Pts;
 		Pts.AddUninitialized(Children.Num());
@@ -2170,7 +2222,7 @@ namespace Chaos
 					if (ShouldKeepConnection(Child1, Child2))
 					{
 						// this does not use ConnectNodes because Neighbors is bi-direction : as in (1,2),(2,1)
-						ConnectNodes(Child1, Child2);
+						ConnectNodes(Child1, Child2, FromParticles, ToParticles);
 						UniqueEdges.Add(SortedPair);
 					}
 				}
@@ -2180,40 +2232,60 @@ namespace Chaos
 	
 	DECLARE_CYCLE_STAT(TEXT("TPBDRigidClustering<>::UpdateConnectivityGraphUsingDelaunayTriangulation"), STAT_UpdateConnectivityGraphUsingDelaunayTriangulation, STATGROUP_Chaos);
 	void FRigidClustering::UpdateConnectivityGraphUsingDelaunayTriangulation(
-		const Chaos::FPBDRigidClusteredParticleHandle* Parent, 
-		const FClusterCreationParameters& Parameters)
+		const TArray<FPBDRigidParticleHandle*>& Particles,
+		const FClusterCreationParameters& Parameters,
+		const TSet<FPBDRigidParticleHandle*>* FromParticles,
+		const TSet<FPBDRigidParticleHandle*>* ToParticles)
 	{
 		SCOPE_CYCLE_COUNTER(STAT_UpdateConnectivityGraphUsingDelaunayTriangulation);
 
 		if (UseBoundingBoxForConnectionGraphFiltering)
 		{
 			constexpr auto IsOverlappingConnectionUsingCVarMargin =
-				[] (const FPBDRigidParticleHandle* Child1, const FPBDRigidParticleHandle* Child2)
-				{
-					return IsOverlappingConnection(Child1, Child2, BoundingBoxMarginForConnectionGraphFiltering);
-				};
-			UpdateConnectivityGraphUsingDelaunayTriangulationWithFiltering(MChildren[Parent], IsOverlappingConnectionUsingCVarMargin);
+				[](const FPBDRigidParticleHandle* Child1, const FPBDRigidParticleHandle* Child2)
+			{
+				return IsOverlappingConnection(Child1, Child2, BoundingBoxMarginForConnectionGraphFiltering);
+			};
+			UpdateConnectivityGraphUsingDelaunayTriangulationWithFiltering(Particles, IsOverlappingConnectionUsingCVarMargin, FromParticles, ToParticles);
 		}
 		else
 		{
-			UpdateConnectivityGraphUsingDelaunayTriangulationWithFiltering(MChildren[Parent], IsAlwaysValidConnection);
+			UpdateConnectivityGraphUsingDelaunayTriangulationWithFiltering(Particles, IsAlwaysValidConnection, FromParticles, ToParticles);
 		}
 	}
 
-	DECLARE_CYCLE_STAT(TEXT("TPBDRigidClustering<>::UpdateConnectivityGraphUsingDelaunayTriangulationWithBoundsOverlaps"), STAT_UpdateConnectivityGraphUsingDelaunayTriangulationWithBoundsOverlaps, STATGROUP_Chaos);
-	void FRigidClustering::UpdateConnectivityGraphUsingDelaunayTriangulationWithBoundsOverlaps(
+	void FRigidClustering::UpdateConnectivityGraphUsingDelaunayTriangulation(
 		const Chaos::FPBDRigidClusteredParticleHandle* Parent, 
 		const FClusterCreationParameters& Parameters)
+	{
+		const TArray<FPBDRigidParticleHandle*>& Children = MChildren[Parent];
+		FixConnectivityGraphUsingDelaunayTriangulation(Children, Parameters);
+	}
+
+	DECLARE_CYCLE_STAT(TEXT("TPBDRigidClustering<>::UpdateConnectivityGraphUsingDelaunayTriangulationWithBoundsOverlaps"), STAT_UpdateConnectivityGraphUsingDelaunayTriangulationWithBoundsOverlaps, STATGROUP_Chaos);
+	void FRigidClustering::UpdateConnectivityGraphUsingDelaunayTriangulationWithBoundsOverlaps(
+		const TArray<FPBDRigidParticleHandle*>& Particles,
+		const FClusterCreationParameters& Parameters,
+		const TSet<FPBDRigidParticleHandle*>* FromParticles,
+		const TSet<FPBDRigidParticleHandle*>* ToParticles)
 	{
 		SCOPE_CYCLE_COUNTER(STAT_UpdateConnectivityGraphUsingDelaunayTriangulationWithBoundsOverlaps);
 
 		const auto IsOverlappingConnectionUsingCVarMargin =
-			[&Parameters] (const FPBDRigidParticleHandle* Child1, const FPBDRigidParticleHandle* Child2)
-			{
-				return IsOverlappingConnection(Child1, Child2, Parameters.ConnectionGraphBoundsFilteringMargin);
-			};
-		
-		UpdateConnectivityGraphUsingDelaunayTriangulationWithFiltering(MChildren[Parent], IsOverlappingConnectionUsingCVarMargin);
+			[&Parameters](const FPBDRigidParticleHandle* Child1, const FPBDRigidParticleHandle* Child2)
+		{
+			return IsOverlappingConnection(Child1, Child2, Parameters.ConnectionGraphBoundsFilteringMargin);
+		};
+
+		UpdateConnectivityGraphUsingDelaunayTriangulationWithFiltering(Particles, IsOverlappingConnectionUsingCVarMargin, FromParticles, ToParticles);
+	}
+
+	void FRigidClustering::UpdateConnectivityGraphUsingDelaunayTriangulationWithBoundsOverlaps(
+		const Chaos::FPBDRigidClusteredParticleHandle* Parent, 
+		const FClusterCreationParameters& Parameters)
+	{
+		const TArray<FPBDRigidParticleHandle*>& Children = MChildren[Parent];
+		UpdateConnectivityGraphUsingDelaunayTriangulationWithBoundsOverlaps(Children, Parameters);
 	}
 
 	void 
diff --git a/Engine/Source/Runtime/Experimental/Chaos/Private/Chaos/PhysicsObjectInterface.cpp b/Engine/Source/Runtime/Experimental/Chaos/Private/Chaos/PhysicsObjectInterface.cpp
--- a/Engine/Source/Runtime/Experimental/Chaos/Private/Chaos/PhysicsObjectInterface.cpp
+++ b/Engine/Source/Runtime/Experimental/Chaos/Private/Chaos/PhysicsObjectInterface.cpp
@@ -4,6 +4,7 @@
 #include "Chaos/CastingUtilities.h"
 #include "Chaos/CollisionFilterData.h"
 #include "Chaos/PhysicsObjectCollisionInterface.h"
+#include "Chaos/PhysicsObjectInternalInterface.h"
 #include "Chaos/PhysicsObjectInternal.h"
 #include "PhysicsProxy/SingleParticlePhysicsProxy.h"
 #include "Math/Transform.h"
@@ -95,6 +96,35 @@ namespace Chaos
 		return Object ? Object->HasChildren<Id>() : false;
 	}
 
+	template<EThreadContext Id>
+	int32 FReadPhysicsObjectInterface<Id>::GetClusterHierarchyLevel(const FConstPhysicsObjectHandle Object)
+	{
+		if (!Object)
+		{
+			return INDEX_NONE;
+		}
+
+		// TODO: I think the GC rest collection caches this information. If we're on the GT it might be
+		// better for perf to try and grab that information somehow instead of walking the hierarchy.
+
+		int32 CurrentLevel = 0;
+
+		FPhysicsObjectHandle Parent = Object->GetParentObject<Id>();
+		while (Parent)
+		{
+			// Being part of a cluster union doesn't count in terms of its default cluster hierarchy level.
+			if (Parent->PhysicsProxy() != Object->PhysicsProxy())
+			{
+				break;
+			}
+
+			Parent = Parent->GetParentObject<Id>();
+			++CurrentLevel;
+		}
+
+		return CurrentLevel;
+	}
+
 	template<EThreadContext Id>
 	FTransform FReadPhysicsObjectInterface<Id>::GetTransform(const FConstPhysicsObjectHandle Object)
 	{
@@ -726,6 +756,47 @@ namespace Chaos
 		}
 	}
 
+	template<EThreadContext Id>
+	void FWritePhysicsObjectInterface<Id>::AddConnectivityEdgesBetween(TArrayView<const FPhysicsObjectHandle> FromObjects, TArrayView<const FPhysicsObjectHandle> ToObjects, const FClusterCreationParameters& Parameters)
+	{
+		if (FromObjects.IsEmpty() || ToObjects.IsEmpty())
+		{
+			return;
+		}
+
+		// Assume that the solver for both sets of objects are the same.
+		if (FPBDRigidsSolver* Solver = FPhysicsObjectInterface::GetSolver(FromObjects))
+		{
+			if constexpr (Id == EThreadContext::External)
+			{
+				Solver->EnqueueCommandImmediate(
+					[FromObjects = TArray<FPhysicsObjectHandle>{ FromObjects }, ToObjects = TArray<FPhysicsObjectHandle>{ ToObjects }, Parameters]()
+					{
+						FWritePhysicsObjectInterface_Internal Internal = FPhysicsObjectInternalInterface::GetWrite();
+						Internal.AddConnectivityEdgesBetween(FromObjects, ToObjects, Parameters);
+					}
+				);
+			}
+			else
+			{
+				FRigidClustering& Clustering = Solver->GetEvolution()->GetRigidClustering();
+
+				TArray<FPBDRigidParticleHandle*> AllParticles;
+				AllParticles.Reserve(FromObjects.Num() + ToObjects.Num());
+
+				TArray<FPBDRigidParticleHandle*> FromParticles = this->GetAllRigidParticles(FromObjects);
+				TArray<FPBDRigidParticleHandle*> ToParticles = this->GetAllRigidParticles(ToObjects);
+				AllParticles.Append(FromParticles);
+				AllParticles.Append(ToParticles);
+
+				TSet<FPBDRigidParticleHandle*> FromSet{ FromParticles };
+				TSet<FPBDRigidParticleHandle*> ToSet{ ToParticles };
+				Clustering.GenerateConnectionGraph(AllParticles, Parameters, &FromSet, &ToSet);
+			}
+		}
+		
+	}
+
 	void FPhysicsObjectInterface::SetName(const FPhysicsObjectHandle Object, const FName& InName)
 	{
 		if (!Object)
diff --git a/Engine/Source/Runtime/Experimental/Chaos/Public/Chaos/PBDRigidClustering.h b/Engine/Source/Runtime/Experimental/Chaos/Public/Chaos/PBDRigidClustering.h
--- a/Engine/Source/Runtime/Experimental/Chaos/Public/Chaos/PBDRigidClustering.h
+++ b/Engine/Source/Runtime/Experimental/Chaos/Public/Chaos/PBDRigidClustering.h
@@ -360,6 +360,12 @@ public:
 	void SetClusterUnionConnectionType(FClusterCreationParameters::EConnectionMethod ClusterConnectionType) { MClusterUnionConnectionType = ClusterConnectionType; }
 	FClusterCreationParameters::EConnectionMethod GetClusterUnionConnectionType() const { return MClusterUnionConnectionType; }
 
+	void GenerateConnectionGraph(
+		TArray<FPBDRigidParticleHandle*> Particles,
+		const FClusterCreationParameters& Parameters = FClusterCreationParameters(),
+		const TSet<FPBDRigidParticleHandle*>* FromParticles = nullptr,
+		const TSet<FPBDRigidParticleHandle*>* ToParticles = nullptr);
+
 	void GenerateConnectionGraph(
 		Chaos::FPBDRigidClusteredParticleHandle* Parent,
 		const FClusterCreationParameters & Parameters = FClusterCreationParameters());
@@ -393,15 +399,38 @@ public:
 	/*
 	* Connectivity
 	*/
+	void UpdateConnectivityGraphUsingPointImplicit(
+		const TArray<FPBDRigidParticleHandle*>& Particles,
+		FReal CollisionThicknessPercent,
+		const TSet<FPBDRigidParticleHandle*>* FromParticles = nullptr,
+		const TSet<FPBDRigidParticleHandle*>* ToParticles = nullptr);
 	void UpdateConnectivityGraphUsingPointImplicit(
 		Chaos::FPBDRigidClusteredParticleHandle* Parent,
 		const FClusterCreationParameters& Parameters = FClusterCreationParameters());
+
+	void FixConnectivityGraphUsingDelaunayTriangulation(
+		const TArray<FPBDRigidParticleHandle*>& Particles,
+		const FClusterCreationParameters& Parameters = FClusterCreationParameters(),
+		const TSet<FPBDRigidParticleHandle*>* FromParticles = nullptr,
+		const TSet<FPBDRigidParticleHandle*>* ToParticles = nullptr);
 	void FixConnectivityGraphUsingDelaunayTriangulation(
 		Chaos::FPBDRigidClusteredParticleHandle* Parent,
 		const FClusterCreationParameters& Parameters = FClusterCreationParameters());
+
+	void UpdateConnectivityGraphUsingDelaunayTriangulation(
+		const TArray<FPBDRigidParticleHandle*>& Particles,
+		const FClusterCreationParameters& Parameters = FClusterCreationParameters(),
+		const TSet<FPBDRigidParticleHandle*>* FromParticles = nullptr,
+		const TSet<FPBDRigidParticleHandle*>* ToParticles = nullptr);
 	void UpdateConnectivityGraphUsingDelaunayTriangulation(
 		const Chaos::FPBDRigidClusteredParticleHandle* Parent,
 		const FClusterCreationParameters& Parameters = FClusterCreationParameters());
+
+	void UpdateConnectivityGraphUsingDelaunayTriangulationWithBoundsOverlaps(
+		const TArray<FPBDRigidParticleHandle*>& Particles,
+		const FClusterCreationParameters& Parameters = FClusterCreationParameters(),
+		const TSet<FPBDRigidParticleHandle*>* FromParticles = nullptr,
+		const TSet<FPBDRigidParticleHandle*>* ToParticles = nullptr);
 	void UpdateConnectivityGraphUsingDelaunayTriangulationWithBoundsOverlaps(
 		const Chaos::FPBDRigidClusteredParticleHandle* Parent,
 		const FClusterCreationParameters& Parameters = FClusterCreationParameters());
diff --git a/Engine/Source/Runtime/Experimental/Chaos/Public/Chaos/PhysicsObjectInterface.h b/Engine/Source/Runtime/Experimental/Chaos/Public/Chaos/PhysicsObjectInterface.h
--- a/Engine/Source/Runtime/Experimental/Chaos/Public/Chaos/PhysicsObjectInterface.h
+++ b/Engine/Source/Runtime/Experimental/Chaos/Public/Chaos/PhysicsObjectInterface.h
@@ -2,6 +2,7 @@
 
 #pragma once
 
+#include "Chaos/ClusterCreationParameters.h"
 #include "Chaos/ParticleHandle.h"
 #include "Chaos/PhysicsObject.h"
 #include "Chaos/ShapeInstanceFwd.h"
@@ -57,6 +58,7 @@ namespace Chaos
 	public:
 		FPhysicsObjectHandle GetRootObject(const FConstPhysicsObjectHandle Object);
 		bool HasChildren(const FConstPhysicsObjectHandle Object);
+		int32 GetClusterHierarchyLevel(const FConstPhysicsObjectHandle Object);
 
 		FTransform GetTransform(const FConstPhysicsObjectHandle Object);
 		FVector GetX(const FConstPhysicsObjectHandle Object);
@@ -151,6 +153,8 @@ namespace Chaos
 			}
 		}
 
+		void AddConnectivityEdgesBetween(TArrayView<const FPhysicsObjectHandle> FromObjects, TArrayView<const FPhysicsObjectHandle> ToObjects, const FClusterCreationParameters& Parameters);
+
 		friend class FPhysicsObjectInterface;
 	protected:
 		FWritePhysicsObjectInterface() = default;
