Add material based damage model using strain system

diff --git a/Engine/Source/Runtime/Experimental/Chaos/Private/Chaos/ChaosDebugDraw.cpp b/Engine/Source/Runtime/Experimental/Chaos/Private/Chaos/ChaosDebugDraw.cpp
--- a/Engine/Source/Runtime/Experimental/Chaos/Private/Chaos/ChaosDebugDraw.cpp
+++ b/Engine/Source/Runtime/Experimental/Chaos/Private/Chaos/ChaosDebugDraw.cpp
@@ -92,6 +92,12 @@ namespace Chaos
 		int32 ChaosConnectionGraphDrawLevelOffset = 0;
 		FAutoConsoleVariableRef CVarChaosConnectionGraphDrawLevelOffset(TEXT("p.Chaos.DebugDraw.ConnectionGraphLevelOffset"), ChaosConnectionGraphDrawLevelOffset, TEXT("If 0, draws the connection graph between children particles of active particles. If 1, draws the connection grpah between grand-children particles of active particles, etc."));
 
+		bool bChaosDebugDrawConnectionGraphShowAreas = false;
+		FAutoConsoleVariableRef CVarChaosDebugDrawConnectionGraphShowAreas(TEXT("p.Chaos.DebugDraw.ConnectionGraph.ShowAreas"), bChaosDebugDrawConnectionGraphShowAreas, TEXT("When displaying connection graphs show connection areas as disks"));
+
+		bool bChaosDebugDrawConnectionGraphShowInternalStrains = false;
+		FAutoConsoleVariableRef CVarChaosDebugDrawConnectionGraphShowInternalStrains(TEXT("p.Chaos.DebugDraw.ConnectionGraph.ShowInternalStrains"), bChaosDebugDrawConnectionGraphShowInternalStrains, TEXT("When displaying connection graphs show strain values of each node"));
+
 		// NOTE: These settings should never really be used - they are the fallback defaults
 		// if the user does not specify settings in the debug draw call.
 		// See PBDRigidsColver.cpp and ImmediatePhysicsSimulation_Chaos.cpp for example.
@@ -1477,6 +1483,15 @@ namespace Chaos
 					FColor Color = (Connection.Handle0->Parent() == Connection.Handle1->Parent()) ? FColor::Green : FColor::Blue; // FMath::Lerp(FColor::Green, FColor::Red, Connection.Strain / MaxStrain);
 
 					FDebugDrawQueue::GetInstance().DrawDebugLine(Pos0, Pos1, Color, false, UE_KINDA_SMALL_NUMBER, Settings.DrawPriority, Settings.LineThickness);
+
+					if (bChaosDebugDrawConnectionGraphShowAreas)
+					{
+						const FReal Radius = FMath::Sqrt(FMath::Abs(Connection.Strain / UE_PI)); // represent area as a circle
+						const FMatrix Axes = FRotationMatrix::MakeFromX(Normal);
+						FDebugDrawQueue::GetInstance().DrawDebugCircle((Pos0 + Pos1) * 0.5, Radius, 24, Color, false, UE_KINDA_SMALL_NUMBER, Settings.DrawPriority, Settings.LineThickness, Axes.GetUnitAxis(EAxis::Y), Axes.GetUnitAxis(EAxis::Z), false);
+					}
+
+
 				}
 			};
 
@@ -1557,6 +1572,10 @@ namespace Chaos
 						}
 						FDebugDrawQueue::GetInstance().DrawDebugPoint(GetClusterTransform(Child).GetLocation(), PointColor, false, UE_KINDA_SMALL_NUMBER, Settings.DrawPriority, Settings.LineThickness * 10);
 
+						if (bChaosDebugDrawConnectionGraphShowInternalStrains)
+						{
+							FDebugDrawQueue::GetInstance().DrawDebugString(Child->X(), FString::Printf(TEXT("%.1f"), Child->GetInternalStrains()), nullptr, FColor::White, UE_KINDA_SMALL_NUMBER, false, 1.0);
+						}
 
 						const FConnectivityEdgeArray& Edges = Child->ConnectivityEdges();
 						for (const TConnectivityEdge<FReal>& Edge : Edges)
diff --git a/Engine/Source/Runtime/Experimental/Chaos/Private/Chaos/PBDRigidClustering.cpp b/Engine/Source/Runtime/Experimental/Chaos/Private/Chaos/PBDRigidClustering.cpp
--- a/Engine/Source/Runtime/Experimental/Chaos/Private/Chaos/PBDRigidClustering.cpp
+++ b/Engine/Source/Runtime/Experimental/Chaos/Private/Chaos/PBDRigidClustering.cpp
@@ -791,12 +791,15 @@ namespace Chaos
 		bool bUseDamagePropagation = false;
 		float BreakDamagePropagationFactor = 0.0f;
 		float ShockDamagePropagationFactor = 0.0f;
-		if (const FGeometryCollectionPhysicsProxy* ConcreteProxy = GetConcreteProxy(ClusteredParticle))
+		EDamageEvaluationModel DamageEvaluationModel = EDamageEvaluationModel::StrainFromDamageThreshold;
+		const FGeometryCollectionPhysicsProxy* ConcreteGCProxy = GetConcreteProxy<FGeometryCollectionPhysicsProxy>(ClusteredParticle);
+		if (ConcreteGCProxy)
 		{
-			const FSimulationParameters& SimParams = ConcreteProxy->GetSimParameters();
+			const FSimulationParameters& SimParams = ConcreteGCProxy->GetSimParameters();
 			bUseDamagePropagation = SimParams.bUseDamagePropagation;
 			BreakDamagePropagationFactor = SimParams.BreakDamagePropagationFactor;
 			ShockDamagePropagationFactor = SimParams.ShockDamagePropagationFactor;
+			DamageEvaluationModel = SimParams.DamageEvaluationModel;
 		}
 
 		TArray<FPBDRigidParticleHandle*>& Children = MChildren[ClusteredParticle];
@@ -967,6 +970,18 @@ namespace Chaos
 							RemoveChildFromParent(Child, ClusteredParticle);
 							ActivatedChildren.Add(Child);
 						}
+
+						if (ConcreteGCProxy && DamageEvaluationModel == EDamageEvaluationModel::StrainFromMaterialStrengthAndConnectivity)
+						{
+							for (FPBDRigidParticleHandle* ChildParticle: Island)
+							{
+								// recompute ethe strain as the connection have changed
+								if (FPBDRigidClusteredParticleHandle* ClusteredChild = ChildParticle->CastToClustered())
+								{
+									ConcreteGCProxy->ComputeMaterialBasedDamageThreshold_Internal(*ClusteredChild);
+								}
+							}
+						}
 					}
 
 					if (bCreateNewClusters)
@@ -1093,7 +1108,7 @@ namespace Chaos
 			//
 			if (ComputeClusterCollisionStrains)
 			{
-				ComputeStrainFromCollision(CollisionRule);
+				ComputeStrainFromCollision(CollisionRule, Dt);
 			}
 			else
 			{
@@ -1526,11 +1541,13 @@ namespace Chaos
 	FAutoConsoleVariableRef CVarMinContactSpeedForStrainEval(TEXT("p.chaos.MinContactSpeedForStrainEval"), MinContactSpeedForStrainEval, TEXT("Minimum speed at the contact before accumulating for strain eval "));
 
 	DECLARE_CYCLE_STAT(TEXT("ComputeStrainFromCollision"), STAT_ComputeStrainFromCollision, STATGROUP_Chaos);
-	void FRigidClustering::ComputeStrainFromCollision(const FPBDCollisionConstraints& CollisionRule)
+	void FRigidClustering::ComputeStrainFromCollision(const FPBDCollisionConstraints& CollisionRule, const FReal Dt)
 	{
 		SCOPE_CYCLE_COUNTER(STAT_ComputeStrainFromCollision);
 		FClusterMap& MParentToChildren = GetChildrenMap();
 
+		const FRealSingle InvDt = (Dt > SMALL_NUMBER) ? (FRealSingle)(1.0 / Dt) : 1.0f;
+
 		ResetCollisionImpulseArray();
 
 		for (Chaos::FPBDCollisionConstraintHandle* ContactHandle : CollisionRule.GetConstraintHandles())
@@ -1572,7 +1589,7 @@ namespace Chaos
 				continue;
 			}
 
-			auto ComputeStrainLambda = [this, &ContactHandle](
+			auto ComputeStrainLambda = [this, &ContactHandle, &InvDt](
 				const FPBDRigidClusteredParticleHandle* Cluster,
 				FRealSingle& OutTotalImpulseAccumulator)
 			{
@@ -1602,16 +1619,25 @@ namespace Chaos
 
 					// gather propagation information from the parent proxy
 					bool bUseDamagePropagation = false;
+					EDamageEvaluationModel DamageModel = EDamageEvaluationModel::StrainFromDamageThreshold;
 					if (const FGeometryCollectionPhysicsProxy* GeometryCollectionProxy = GetConcreteProxy<FGeometryCollectionPhysicsProxy>(Cluster))
 					{
 						const FSimulationParameters& SimParams = GeometryCollectionProxy->GetSimParameters();
 						bUseDamagePropagation = SimParams.bUseDamagePropagation;
+						DamageModel = SimParams.DamageEvaluationModel;
 						if (!SimParams.bEnableStrainOnCollision)
 						{
 							return;
 						}
 					}
 
+					FRealSingle CollisionImpulseToApply = AccumulatedImpulse;
+					if (DamageModel == EDamageEvaluationModel::StrainFromMaterialStrengthAndConnectivity)
+					{
+						// this damage model use internal strain as the maximum breaking force (m.a), but we need to compare to impulses (m.v), let's divide this impulse by Dt
+						// technically this is no longer an impulse but we intend in the future to get rid of COllisionImpulse property and use external strain instead
+						CollisionImpulseToApply *= InvDt;
+					}
 					if (bUseDamagePropagation)
 					{
 						// propagation based breaking model start from the closest particle and propagate through the connection graph
@@ -1620,9 +1646,9 @@ namespace Chaos
 						{
 							if (TPBDRigidClusteredParticleHandle<FReal, 3>* ClusteredChild = ClosestChild->CastToClustered())
 							{
-								ClusteredChild->CollisionImpulses() += AccumulatedImpulse;
+								ClusteredChild->CollisionImpulses() += CollisionImpulseToApply;
 								UpdateTopLevelParticle(ClusteredChild);
-								OutTotalImpulseAccumulator += AccumulatedImpulse;
+								OutTotalImpulseAccumulator += CollisionImpulseToApply;
 							}
 						}
 					}
@@ -1640,9 +1666,9 @@ namespace Chaos
 							{
 								if (TPBDRigidClusteredParticleHandle<FReal, 3>*ClusteredChild = Child->CastToClustered())
 								{
-									ClusteredChild->CollisionImpulses() += AccumulatedImpulse;
+									ClusteredChild->CollisionImpulses() += CollisionImpulseToApply;
 									UpdateTopLevelParticle(ClusteredChild);
-									OutTotalImpulseAccumulator += AccumulatedImpulse;
+									OutTotalImpulseAccumulator += CollisionImpulseToApply;
 								}
 							}
 						}
diff --git a/Engine/Source/Runtime/Experimental/Chaos/Private/PhysicsProxy/GeometryCollectionPhysicsProxy.cpp b/Engine/Source/Runtime/Experimental/Chaos/Private/PhysicsProxy/GeometryCollectionPhysicsProxy.cpp
--- a/Engine/Source/Runtime/Experimental/Chaos/Private/PhysicsProxy/GeometryCollectionPhysicsProxy.cpp
+++ b/Engine/Source/Runtime/Experimental/Chaos/Private/PhysicsProxy/GeometryCollectionPhysicsProxy.cpp
@@ -182,6 +182,27 @@ TArray<int32> ComputeTransformToGeometryMap(const FGeometryCollection& Collectio
 	return TransformToGeometryMap;
 }
 
+static float AreaFromBoundingBoxOverlap(const Chaos::FAABB3& BoxA, const Chaos::FAABB3& BoxB)
+{
+	// if the two box don't overlap, we'll get an inside out box 
+	// but we are still using it to compute the area as an approximation
+	const Chaos::FAABB3 OverlapBox = BoxA.GetIntersection(BoxB);
+
+	const Chaos::FVec3 Extents = OverlapBox.Extents();
+	const Chaos::FVec3 CenterToCenterNormal = (BoxA.GetCenter() - BoxB.GetCenter()).GetSafeNormal();
+
+	const Chaos::FVec3 AreaPerAxis{
+		FMath::Abs(Extents.Y * Extents.Z),
+		FMath::Abs(Extents.X * Extents.Z),
+		FMath::Abs(Extents.X * Extents.Y)
+	};
+
+	// weight the area by the center to center normal
+	const Chaos::FReal Area = FMath::Abs(AreaPerAxis.Dot(CenterToCenterNormal));
+
+	return static_cast<float>(Area);
+}
+
 
 DECLARE_CYCLE_STAT(TEXT("FGeometryCollectionPhysicsProxy::PopulateSimulatedParticle"), STAT_PopulateSimulatedParticle, STATGROUP_Chaos);
 void PopulateSimulatedParticle(
@@ -978,9 +999,15 @@ void FGeometryCollectionPhysicsProxy::InitializeBodiesPT(Chaos::FPBDRigidsSolver
 		{
 			if (FClusterHandle* Handle = SolverParticleHandles[TransformGroupIndex])
 			{
+				const bool bIsAnchored = AnchoringFacade.IsAnchored(TransformGroupIndex);
+
 				// Mass space -> Composed parent space -> world
 				const Chaos::FReal ScaledMass = AdjustMassForScale(Mass[TransformGroupIndex]);
 				const Chaos::FVec3f ScaledInertia = AdjustInertiaForScale((Chaos::FVec3f)InertiaTensor[TransformGroupIndex]);
+				const uint8 AdjustedDynamicState = 
+					bIsAnchored 
+					? static_cast<uint8>(Chaos::EObjectStateType::Kinematic)
+					: static_cast<uint8>(DynamicState[TransformGroupIndex]);
 
 				PopulateSimulatedParticle(
 					Handle,
@@ -992,7 +1019,7 @@ void FGeometryCollectionPhysicsProxy::InitializeBodiesPT(Chaos::FPBDRigidsSolver
 					ScaledMass,
 					ScaledInertia,
 					Transforms[TransformGroupIndex],
-					static_cast<uint8>(DynamicState[TransformGroupIndex]),
+					AdjustedDynamicState,
 					static_cast<int16>(CollisionGroup[TransformGroupIndex]),
 					CollisionParticlesPerObjectFraction);
 
@@ -1006,7 +1033,7 @@ void FGeometryCollectionPhysicsProxy::InitializeBodiesPT(Chaos::FPBDRigidsSolver
 					float DamageThreshold = StrainDefault;
 					if (!Parameters.bUsePerClusterOnlyDamageThreshold)
 					{
-						DamageThreshold = ComputeDamageThreshold(DynamicCollection, TransformGroupIndex);
+						DamageThreshold = ComputeUserDefinedDamageThreshold_Internal(TransformGroupIndex);
 					}
 					Chaos::FPhysicsSolver* RBDSolver = GetSolver<Chaos::FPhysicsSolver>();
 					RBDSolver->GetEvolution()->GetRigidClustering().SetInternalStrain(Handle, DamageThreshold);
@@ -1236,40 +1263,60 @@ void FGeometryCollectionPhysicsProxy::InitializeBodiesPT(Chaos::FPBDRigidsSolver
 						}
 					}
 				}
+				// if using the material & connectivity damnage model, we need to compute the surface area of each connection
+				// since we do not have precise data from the tools, we are using an approximate method using the bounds of the two particles
+				if (Parameters.DamageEvaluationModel == Chaos::EDamageEvaluationModel::StrainFromMaterialStrengthAndConnectivity)
+				{
+					for (int32 TransformGroupIndex = 0; TransformGroupIndex < NumTransforms; ++TransformGroupIndex)
+					{
+						if (FClusterHandle* ClusteredParticle = SolverParticleHandles[TransformGroupIndex])
+						{
+							for (Chaos::FConnectivityEdge& Edge : ClusteredParticle->ConnectivityEdges())
+							{
+								if (ensure(Edge.Sibling))
+								{
+									const Chaos::FRealSingle Area = AreaFromBoundingBoxOverlap(ClusteredParticle->WorldSpaceInflatedBounds(), Edge.Sibling->WorldSpaceInflatedBounds());
+									Edge.SetArea(Area);
+								}
+							}
+						}
+					}
+				}
 			}
-			else
+			else if (ConnectionFacade.IsValid())
 			{
-				if (ConnectionFacade.IsValid())
+				// load connection graph from RestCollection
+				// TODO: is it worth computing the valence of each edge to Reserve() the particle edge arrays?
+
+				// Transfer connections from the Collection's ConnectionFacade over to the clustered particles
+				int32 NumConnections = ConnectionFacade.NumConnections();
+				for (int32 ConnectionIdx = 0; ConnectionIdx < NumConnections; ++ConnectionIdx)
 				{
-					// TODO: is it worth computing the valence of each edge to Reserve() the particle edge arrays?
+					TPair<int32,int32> Connection = ConnectionFacade.GetConnection(ConnectionIdx);
+					checkSlow(Connection.Key != Connection.Value); // Graph should never contain self-loops
 
-					// Transfer connections from the Collection's ConnectionFacade over to the clustered particles
-					int32 NumConnections = ConnectionFacade.NumConnections();
-					for (int32 ConnectionIdx = 0; ConnectionIdx < NumConnections; ++ConnectionIdx)
+					bool bHasContactAreas = ConnectionFacade.HasContactAreas();
+					float ContactArea = 0.0f; // Note: this default value should be unused
+					if (bHasContactAreas)
 					{
-						TPair<int32,int32> Connection = ConnectionFacade.GetConnection(ConnectionIdx);
-						checkSlow(Connection.Key != Connection.Value); // Graph should never contain self-loops
-
-						bool bHasContactAreas = ConnectionFacade.HasContactAreas();
-						float ContactArea = 0.0f; // Note: this default value should be unused
-						if (bHasContactAreas)
-						{
-							ContactArea = ConnectionFacade.GetConnectionContactArea(ConnectionIdx);
-						}
+						ContactArea = ConnectionFacade.GetConnectionContactArea(ConnectionIdx);
+					}
 
-						if (FClusterHandle* ClusteredParticle = SolverParticleHandles[Connection.Key])
+					if (FClusterHandle* ClusteredParticle = SolverParticleHandles[Connection.Key])
+					{
+						if (FClusterHandle* OtherClusteredParticle = SolverParticleHandles[Connection.Value])
 						{
-							if (FClusterHandle* OtherClusteredParticle = SolverParticleHandles[Connection.Value])
+							float EdgeStrainOrArea = bHasContactAreas ? ContactArea : (ClusteredParticle->GetInternalStrains() + OtherClusteredParticle->GetInternalStrains()) * 0.5f;
+							if (!bHasContactAreas && Parameters.DamageEvaluationModel == Chaos::EDamageEvaluationModel::StrainFromMaterialStrengthAndConnectivity)
 							{
-								float EdgeStrain = bHasContactAreas ? ContactArea : (ClusteredParticle->GetInternalStrains() + OtherClusteredParticle->GetInternalStrains()) * 0.5f;
-								// Add symmetric Chaos::TConnectivityEdge<Chaos::FReal> edges to each particle
-								ClusteredParticle->ConnectivityEdges().Emplace(OtherClusteredParticle, EdgeStrain);
-								OtherClusteredParticle->ConnectivityEdges().Emplace(ClusteredParticle, EdgeStrain);
+								EdgeStrainOrArea = AreaFromBoundingBoxOverlap(ClusteredParticle->WorldSpaceInflatedBounds(), OtherClusteredParticle->WorldSpaceInflatedBounds());
 							}
+							// Add symmetric Chaos::TConnectivityEdge<Chaos::FReal> edges to each particle
+							ClusteredParticle->ConnectivityEdges().Emplace(OtherClusteredParticle, EdgeStrainOrArea);
+							OtherClusteredParticle->ConnectivityEdges().Emplace(ClusteredParticle, EdgeStrainOrArea);
 						}
 					}
 				}
-				// load connection graph from RestCollection
 			}
 		} // end if EnableClustering
  
@@ -1319,6 +1366,21 @@ void FGeometryCollectionPhysicsProxy::InitializeBodiesPT(Chaos::FPBDRigidsSolver
 			++HandleIndex;
 #endif
 		}
+
+		// assign the material based damage thresholds 
+		if (Parameters.DamageModel != EDamageModelTypeEnum::Chaos_Damage_Model_UserDefined_Damage_Threshold)
+		{
+			Chaos::FRigidClustering& Clustering = RigidsSolver->GetEvolution()->GetRigidClustering();
+			for (int32 Index = 0; Index < NumTransforms; Index++)
+			{
+				if (Chaos::FPBDRigidClusteredParticleHandle* ClusteredParticle = SolverParticleHandles[Index])
+				{
+					const float DamageThreshold = ComputeMaterialBasedDamageThreshold_Internal(Index);
+					Clustering.SetInternalStrain(ClusteredParticle, DamageThreshold);
+				}
+			}
+		}
+
 		// call DirtyParticle to make sure the acceleration structure is up to date with all the changes happening here
 		DirtyAllParticles(*RigidsSolver);
 
@@ -1336,24 +1398,63 @@ FAutoConsoleVariableRef CVarGlobalMaxSimulatedLevel(TEXT("p.gc.GlobalMaxSimulate
 DECLARE_CYCLE_STAT(TEXT("FGeometryCollectionPhysicsProxy::BuildClusters"), STAT_BuildClusters, STATGROUP_Chaos);
 DECLARE_CYCLE_STAT(TEXT("FGeometryCollectionPhysicsProxy::BuildClusters:GlobalMatrices"), STAT_BuildClustersGlobalMatrices, STATGROUP_Chaos);
 
-float FGeometryCollectionPhysicsProxy::ComputeDamageThreshold(const FGeometryDynamicCollection& DynamicCollection, int32 TransformIndex) const 
+float FGeometryCollectionPhysicsProxy::ComputeMaterialBasedDamageThreshold_Internal(int32 TransformIndex) const
+{
+	float DamageThreshold = TNumericLimits<float>::Max();
+	if (Parameters.DamageModel == EDamageModelTypeEnum::Chaos_Damage_Model_Material_Strength_And_Connectivity_DamageThreshold)
+	{
+		if (ensure(SolverParticleHandles.IsValidIndex(TransformIndex)))
+		{
+			if (Chaos::FPBDRigidClusteredParticleHandle* ClusteredParticle = SolverParticleHandles[TransformIndex])
+			{
+				DamageThreshold = ComputeMaterialBasedDamageThreshold_Internal(*ClusteredParticle);
+			}
+		}
+	}
+
+	return DamageThreshold;
+}
+
+float FGeometryCollectionPhysicsProxy::ComputeMaterialBasedDamageThreshold_Internal(Chaos::FPBDRigidClusteredParticleHandle& ClusteredParticle) const
+{
+	float DamageThreshold = TNumericLimits<float>::Max();
+	float TotalConnectionArea = 0;
+	for (Chaos::FConnectivityEdge& Connection : ClusteredParticle.ConnectivityEdges())
+	{
+		// strain actually store the surface area 
+		TotalConnectionArea += Connection.GetArea();
+	}
+
+	// This model compute the total surface are from the connections and compute the maximum force 
+	if (const Chaos::FChaosPhysicsMaterial* ChaosMaterial = Parameters.PhysicalMaterialHandle.Get())
+	{
+		// force unit is Kg.cm/s2 here
+		const float ForceThreshold = ChaosMaterial->Strength.TensileStrength * TotalConnectionArea;
+		DamageThreshold = ForceThreshold;
+	}
+
+	return DamageThreshold;
+}
+
+float FGeometryCollectionPhysicsProxy::ComputeUserDefinedDamageThreshold_Internal(int32 TransformIndex) const
 {
 	float DamageThreshold = TNumericLimits<float>::Max();
 
 	const int32 LevelOffset = Parameters.bUsePerClusterOnlyDamageThreshold ? 0 : -1;
-	const int32 Level = FMath::Clamp(CalculateHierarchyLevel(DynamicCollection, TransformIndex) + LevelOffset, 0, INT_MAX);
+	const int32 Level = FMath::Clamp(CalculateHierarchyLevel(PhysicsThreadCollection, TransformIndex) + LevelOffset, 0, INT_MAX);
 	if (Level >= FMath::Min(Parameters.MaxClusterLevel, FMath::Min(GlobalMaxSimulatedLevel, Parameters.MaxSimulatedLevel)))
 	{
 		DamageThreshold = TNumericLimits<float>::Max();
 	}
-	else
+	else if (Parameters.DamageModel == EDamageModelTypeEnum::Chaos_Damage_Model_UserDefined_Damage_Threshold)
 	{
+		// we don't test for the damage model because this is used
 		if (Parameters.bUseSizeSpecificDamageThresholds)
 		{
 			// bounding box volume is used as a fallback to find specific size if the relative size if not available
 			// ( May happen with older GC )
 			FBox LocalBoundingBox;
-			const FGeometryDynamicCollection::FSharedImplicit& Implicit = DynamicCollection.Implicits[TransformIndex];
+			const FGeometryDynamicCollection::FSharedImplicit& Implicit = PhysicsThreadCollection.Implicits[TransformIndex];
 			if (Implicit && Implicit->HasBoundingBox())
 			{
 				const Chaos::FAABB3& ImplicitBoundingBox = Implicit->BoundingBox();
@@ -1470,7 +1571,7 @@ Chaos::TPBDGeometryCollectionParticleHandle<Chaos::FReal, 3>* FGeometryCollectio
 		float DamageThreshold = StrainDefault;
 		if (!Parameters.bUsePerClusterOnlyDamageThreshold)
 		{
-			DamageThreshold = ComputeDamageThreshold(DynamicCollection, CollectionClusterIndex);
+			DamageThreshold = ComputeUserDefinedDamageThreshold_Internal(CollectionClusterIndex);
 		}
 		Chaos::FPhysicsSolver* RBDSolver = GetSolver<Chaos::FPhysicsSolver>();
 		RBDSolver->GetEvolution()->GetRigidClustering().SetInternalStrain(Handle, DamageThreshold);
@@ -1616,7 +1717,7 @@ FGeometryCollectionPhysicsProxy::BuildClusters_Internal(
 	// two-way mapping
 	SolverClusterHandles[CollectionClusterIndex] = Parent;
 
-	const float DamageThreshold = ComputeDamageThreshold(DynamicCollection, CollectionClusterIndex);
+	const float DamageThreshold = ComputeUserDefinedDamageThreshold_Internal(CollectionClusterIndex);
 
 	Chaos::FRigidClustering& RigidClustering = static_cast<Chaos::FPBDRigidsSolver*>(Solver)->GetEvolution()->GetRigidClustering();
 	RigidClustering.SetInternalStrain(Parent, DamageThreshold);
diff --git a/Engine/Source/Runtime/Experimental/Chaos/Public/Chaos/Defines.h b/Engine/Source/Runtime/Experimental/Chaos/Public/Chaos/Defines.h
--- a/Engine/Source/Runtime/Experimental/Chaos/Public/Chaos/Defines.h
+++ b/Engine/Source/Runtime/Experimental/Chaos/Public/Chaos/Defines.h
@@ -40,8 +40,8 @@ typedef int32_t int32;
 namespace Chaos
 {
 	// Solver uses Kg for mass, Second for time and Cm for distance
-	// Since a MegaPascal is  force over an Area, and Force = Mass * Acc
-	// we get (Kg.cm)/(cm2.s2) => Kg/(cm.s2)
+	// Since a Pascal is 1N/m2 => (Kg.m)/(m2.s2)
+	// we get 1MPa = 10000 Kg/(cm.s2)
 	constexpr float MegaPascalToKgPerCmS2(float MegaPascals)
 	{
 		return MegaPascals * 10000;
diff --git a/Engine/Source/Runtime/Experimental/Chaos/Public/Chaos/PBDRigidClusteredParticles.h b/Engine/Source/Runtime/Experimental/Chaos/Public/Chaos/PBDRigidClusteredParticles.h
--- a/Engine/Source/Runtime/Experimental/Chaos/Public/Chaos/PBDRigidClusteredParticles.h
+++ b/Engine/Source/Runtime/Experimental/Chaos/Public/Chaos/PBDRigidClusteredParticles.h
@@ -60,8 +60,13 @@ struct TConnectivityEdge
 	bool operator==(const TPBDRigidParticleHandle<T, 3>* OtherSibling) const
 	{ return Sibling == OtherSibling; }
 
+	// getter/functions functions to make calling code less confusing 
+	void SetArea(Chaos::FRealSingle Area) { Strain = Area; }
+	Chaos::FRealSingle GetArea() const { return Strain; }
+
 	TPBDRigidParticleHandle<T, 3>* Sibling;
-	Chaos::FRealSingle Strain;
+	// this can be both strain or area based on the damage model, but cannot rename it for backward compatibility reasons
+	Chaos::FRealSingle Strain; 
 };
 typedef TConnectivityEdge<FReal> FConnectivityEdge;
 typedef TArray<FConnectivityEdge> FConnectivityEdgeArray;
diff --git a/Engine/Source/Runtime/Experimental/Chaos/Public/Chaos/PBDRigidClustering.h b/Engine/Source/Runtime/Experimental/Chaos/Public/Chaos/PBDRigidClustering.h
--- a/Engine/Source/Runtime/Experimental/Chaos/Public/Chaos/PBDRigidClustering.h
+++ b/Engine/Source/Runtime/Experimental/Chaos/Public/Chaos/PBDRigidClustering.h
@@ -11,6 +11,7 @@
 #include "Chaos/ClusterCreationParameters.h"
 #include "Chaos/GeometryParticlesfwd.h"
 #include "Framework/BufferedData.h"
+#include "Chaos/PBDRigidClusteringTypes.h"
 
 namespace Chaos
 {
@@ -391,7 +392,7 @@ public:
 
  protected:
 
-	void ComputeStrainFromCollision(const FPBDCollisionConstraints& CollisionRule);
+	void ComputeStrainFromCollision(const FPBDCollisionConstraints& CollisionRule, const FReal Dt);
 	void ResetCollisionImpulseArray();
 	void DisableCluster(FPBDRigidClusteredParticleHandle* ClusteredParticle);
 	void ApplyStrainModifiers();
diff --git a/Engine/Source/Runtime/Experimental/Chaos/Public/Chaos/PBDRigidClusteringTypes.h b/Engine/Source/Runtime/Experimental/Chaos/Public/Chaos/PBDRigidClusteringTypes.h
new file mode 100644
--- /dev/null
+++ b/Engine/Source/Runtime/Experimental/Chaos/Public/Chaos/PBDRigidClusteringTypes.h
@@ -0,0 +1,24 @@
+// Copyright Epic Games, Inc. All Rights Reserved.
+#pragma once
+
+namespace Chaos
+{
+	enum class EDamageEvaluationModel
+	{
+		/** 
+		* particles internal strains are set from user defined damage thresholds 
+		* It's a simple model, fast to evaluate , but can be set to be non realistic
+		* and may require more fine tuning 
+		*/
+		StrainFromDamageThreshold,
+
+		/**
+		* Particle internal strains are set from the strength properties of the physics material
+		* and the surface area of connection of a piece to another 
+		* this is as fast as the damage threshold to evaluate but provide a more realistic model 
+		* and allow for simpler tuning based on material
+		*/
+		StrainFromMaterialStrengthAndConnectivity,
+	};
+
+} // namespace Chaos
diff --git a/Engine/Source/Runtime/Experimental/Chaos/Public/GeometryCollection/GeometryCollectionSimulationCoreTypes.h b/Engine/Source/Runtime/Experimental/Chaos/Public/GeometryCollection/GeometryCollectionSimulationCoreTypes.h
--- a/Engine/Source/Runtime/Experimental/Chaos/Public/GeometryCollection/GeometryCollectionSimulationCoreTypes.h
+++ b/Engine/Source/Runtime/Experimental/Chaos/Public/GeometryCollection/GeometryCollectionSimulationCoreTypes.h
@@ -9,6 +9,7 @@
 #include "Chaos/ClusterCreationParameters.h"
 #include "Chaos/CollisionFilterData.h"
 #include "Chaos/PBDRigidsEvolutionFwd.h"
+#include "Chaos/PBDRigidClusteringTypes.h"
 
 class FGeometryCollection;
 class FGeometryDynamicCollection;
@@ -180,6 +181,8 @@ struct FSimulationParameters
 		, MaxClusterLevel(100)
 		, MaxSimulatedLevel(100)
 		, bUseSizeSpecificDamageThresholds(false)
+		, DamageModel(EDamageModelTypeEnum::Chaos_Damage_Model_UserDefined_Damage_Threshold)
+		, DamageEvaluationModel(Chaos::EDamageEvaluationModel::StrainFromDamageThreshold)
 		, DamageThreshold({500000.f, 50000.f, 5000.f})
 		, bUsePerClusterOnlyDamageThreshold(false)
 		, ClusterConnectionMethod(Chaos::FClusterCreationParameters::EConnectionMethod::PointImplicit)
@@ -231,6 +234,8 @@ struct FSimulationParameters
 		, MaxClusterLevel(Other.MaxClusterLevel)
 		, MaxSimulatedLevel(Other.MaxSimulatedLevel)
 		, bUseSizeSpecificDamageThresholds(Other.bUseSizeSpecificDamageThresholds)
+		, DamageModel(Other.DamageModel)
+		, DamageEvaluationModel(Other.DamageEvaluationModel)
 		, DamageThreshold(Other.DamageThreshold)
 		, bUsePerClusterOnlyDamageThreshold(Other.bUsePerClusterOnlyDamageThreshold)
 		, ClusterConnectionMethod(Other.ClusterConnectionMethod)
@@ -299,6 +304,13 @@ struct FSimulationParameters
 	int32 MaxClusterLevel;
 	int32 MaxSimulatedLevel;
 	bool bUseSizeSpecificDamageThresholds;
+
+	/** this is the user expose damage model, used for creation of the particles */
+	EDamageModelTypeEnum DamageModel; 
+
+	/** this is the lower level damage model for clustering, used at runm time */
+	Chaos::EDamageEvaluationModel DamageEvaluationModel;
+
 	TArray<float> DamageThreshold;
 	bool bUsePerClusterOnlyDamageThreshold;
 	Chaos::FClusterCreationParameters::EConnectionMethod ClusterConnectionMethod;
diff --git a/Engine/Source/Runtime/Experimental/Chaos/Public/GeometryCollection/GeometryCollectionSimulationTypes.h b/Engine/Source/Runtime/Experimental/Chaos/Public/GeometryCollection/GeometryCollectionSimulationTypes.h
--- a/Engine/Source/Runtime/Experimental/Chaos/Public/GeometryCollection/GeometryCollectionSimulationTypes.h
+++ b/Engine/Source/Runtime/Experimental/Chaos/Public/GeometryCollection/GeometryCollectionSimulationTypes.h
@@ -3,6 +3,7 @@
 
 #include "UObject/ObjectMacros.h"
 #include "Field/FieldSystemTypes.h"
+#include "Chaos/PBDRigidClusteringTypes.h"
 
 #include "GeometryCollectionSimulationTypes.generated.h"
 
@@ -103,4 +104,33 @@ enum class EEmissionPatternTypeEnum : uint8
 	Chaos_Emission_Pattern_On_Demand UMETA(DisplayName = "On Demand"),
 	//
 	Chaos_Max                UMETA(Hidden)
-};
\ No newline at end of file
+};
+
+
+
+UENUM(BlueprintType)
+enum class EDamageModelTypeEnum : uint8
+{
+	/** Using damage threshold set based on level of the cluster */
+	Chaos_Damage_Model_UserDefined_Damage_Threshold UMETA(DisplayName = "User-Defined Damage Threshold"),
+
+	/** Using damage threshold set using the physical material strength and how connected a cluster is */
+	Chaos_Damage_Model_Material_Strength_And_Connectivity_DamageThreshold UMETA(DisplayName = "Material Strength And Connectivity Damage Threshold"),
+
+	Chaos_Max UMETA(Hidden)
+};
+
+// convert user level damage model to chaos damage evaluation model
+inline CHAOS_API Chaos::EDamageEvaluationModel GetDamageEvaluationModel(const EDamageModelTypeEnum DamageModel)
+{
+	switch (DamageModel)
+	{
+	case EDamageModelTypeEnum::Chaos_Damage_Model_UserDefined_Damage_Threshold: 
+		return Chaos::EDamageEvaluationModel::StrainFromDamageThreshold;
+	case EDamageModelTypeEnum::Chaos_Damage_Model_Material_Strength_And_Connectivity_DamageThreshold: 
+		return Chaos::EDamageEvaluationModel::StrainFromMaterialStrengthAndConnectivity;
+	default:
+		ensure(false); // unexpected
+		return Chaos::EDamageEvaluationModel::StrainFromDamageThreshold;
+	}
+}
\ No newline at end of file
diff --git a/Engine/Source/Runtime/Experimental/Chaos/Public/PhysicsProxy/GeometryCollectionPhysicsProxy.h b/Engine/Source/Runtime/Experimental/Chaos/Public/PhysicsProxy/GeometryCollectionPhysicsProxy.h
--- a/Engine/Source/Runtime/Experimental/Chaos/Public/PhysicsProxy/GeometryCollectionPhysicsProxy.h
+++ b/Engine/Source/Runtime/Experimental/Chaos/Public/PhysicsProxy/GeometryCollectionPhysicsProxy.h
@@ -423,6 +423,8 @@ public:
 	void SetNotifyRemovals_External(bool bNotify);
 	void SetNotifyCrumblings_External(bool bNotify, bool bIncludeChildren);
 
+	float ComputeMaterialBasedDamageThreshold_Internal(Chaos::FPBDRigidClusteredParticleHandle& ClusteredParticle) const;
+
 	FProxyInterpolationData& GetInterpolationData() { return InterpolationData; }
 	const FProxyInterpolationData& GetInterpolationData() const { return InterpolationData; }
 
@@ -456,14 +458,17 @@ public:
 	Chaos::FPhysicsObjectHandle GetPhysicsObjectByIndex(int32 Index) const;
 	int32 GetNumParticles() const { return NumParticles; }
 protected:
+
+	float ComputeMaterialBasedDamageThreshold_Internal(int32 TransformIndex) const;
+
 	/**
-	* Compute damage threshold for a specific transform
+	* Compute user defined damage threshold for a specific transform
 	* this account for component level damage threshold as well as size specific ones
 	* @param DynamicCollection dynamic collection to use
 	* @param TransformIndex index of the transform to compute the threshold for
 	* #return damage threshold value
 	*/
-	float ComputeDamageThreshold(const FGeometryDynamicCollection& DynamicCollection, int32 TransformIndex) const;
+	float ComputeUserDefinedDamageThreshold_Internal(int32 TransformIndex) const;
 
 	/** adjust scalar mass to account for per component scale properties ( from material override and world transform scale ) */ 
 	Chaos::FReal AdjustMassForScale(Chaos::FReal Mass) const;
diff --git a/Engine/Source/Runtime/Experimental/GeometryCollectionEngine/Private/GeometryCollection/GeometryCollectionComponent.cpp b/Engine/Source/Runtime/Experimental/GeometryCollectionEngine/Private/GeometryCollection/GeometryCollectionComponent.cpp
--- a/Engine/Source/Runtime/Experimental/GeometryCollectionEngine/Private/GeometryCollection/GeometryCollectionComponent.cpp
+++ b/Engine/Source/Runtime/Experimental/GeometryCollectionEngine/Private/GeometryCollection/GeometryCollectionComponent.cpp
@@ -2728,6 +2728,8 @@ void UGeometryCollectionComponent::RegisterAndInitializePhysicsProxy()
 		SimulationParameters.ClusterGroupIndex = EnableClustering ? ClusterGroupIndex : 0;
 		SimulationParameters.MaxClusterLevel = MaxClusterLevel;
 		SimulationParameters.MaxSimulatedLevel = MaxSimulatedLevel;
+		SimulationParameters.DamageModel = DamageModel;
+		SimulationParameters.DamageEvaluationModel = GetDamageEvaluationModel(DamageModel);
 		SimulationParameters.bUseSizeSpecificDamageThresholds = bUseSizeSpecificDamageThreshold;
 		SimulationParameters.DamageThreshold = DamageThreshold;
 		SimulationParameters.bUsePerClusterOnlyDamageThreshold = RestCollection? RestCollection->PerClusterOnlyDamageThreshold: false; 
@@ -3277,6 +3279,7 @@ void UGeometryCollectionComponent::ApplyAssetDefaults()
 	if (RestCollection)
 	{
 		// initialize the component per level damage threshold from the asset defaults 
+		DamageModel = RestCollection->DamageModel;
 		DamageThreshold = RestCollection->DamageThreshold;
 		bUseSizeSpecificDamageThreshold = RestCollection->bUseSizeSpecificDamageThreshold;
 
diff --git a/Engine/Source/Runtime/Experimental/GeometryCollectionEngine/Private/GeometryCollection/GeometryCollectionObject.cpp b/Engine/Source/Runtime/Experimental/GeometryCollectionEngine/Private/GeometryCollection/GeometryCollectionObject.cpp
--- a/Engine/Source/Runtime/Experimental/GeometryCollectionEngine/Private/GeometryCollection/GeometryCollectionObject.cpp
+++ b/Engine/Source/Runtime/Experimental/GeometryCollectionEngine/Private/GeometryCollection/GeometryCollectionObject.cpp
@@ -89,6 +89,7 @@ UGeometryCollection::UGeometryCollection(const FObjectInitializer& ObjectInitial
 	, EnableClustering(true)
 	, ClusterGroupIndex(0)
 	, MaxClusterLevel(100)
+	, DamageModel(EDamageModelTypeEnum::Chaos_Damage_Model_UserDefined_Damage_Threshold)
 	, DamageThreshold({ 500000.f, 50000.f, 5000.f })
 	, bUseSizeSpecificDamageThreshold(false)
 	, PerClusterOnlyDamageThreshold(false)
diff --git a/Engine/Source/Runtime/Experimental/GeometryCollectionEngine/Public/GeometryCollection/GeometryCollectionComponent.h b/Engine/Source/Runtime/Experimental/GeometryCollectionEngine/Public/GeometryCollection/GeometryCollectionComponent.h
--- a/Engine/Source/Runtime/Experimental/GeometryCollectionEngine/Public/GeometryCollection/GeometryCollectionComponent.h
+++ b/Engine/Source/Runtime/Experimental/GeometryCollectionEngine/Public/GeometryCollection/GeometryCollectionComponent.h
@@ -20,6 +20,7 @@
 #include "GeometryCollectionEditorSelection.h"
 #include "GeometryCollection/GeometryCollectionDamagePropagationData.h"
 #include "GeometryCollection/RecordedTransformTrack.h"
+#include "GeometryCollection/GeometryCollectionSimulationTypes.h"
 #include "Templates/UniquePtr.h"
 #include "Chaos/ChaosGameplayEventDispatcher.h"
 #include "Chaos/ChaosNotifyHandlerInterface.h"
@@ -785,7 +786,11 @@ public:
 	UPROPERTY(EditAnywhere, BlueprintReadWrite, Category = "ChaosPhysics|Clustering")
 	int32 MaxSimulatedLevel;
 
-	UPROPERTY(EditAnywhere, BlueprintReadWrite, BlueprintGetter=GetDamageThreshold, BlueprintSetter=SetDamageThreshold, Category = "ChaosPhysics|Damage", meta = (EditCondition = "!bUseSizeSpecificDamageThreshold"))
+	/** Damage model to use for evaluating destruction. */
+	UPROPERTY(EditAnywhere, Category = "ChaosPhysics|Damage")
+	EDamageModelTypeEnum DamageModel;
+
+	UPROPERTY(EditAnywhere, BlueprintReadWrite, BlueprintGetter=GetDamageThreshold, BlueprintSetter=SetDamageThreshold, Category = "ChaosPhysics|Damage", meta = (EditCondition = "!bUseSizeSpecificDamageThreshold && DamageModel == EDamageModelTypeEnum::Chaos_Damage_Model_UserDefined_Damage_Threshold"))
 	TArray<float> DamageThreshold;
 
 	UFUNCTION(BlueprintPure, BlueprintInternalUseOnly)
@@ -795,7 +800,7 @@ public:
 	void SetDamageThreshold(const TArray<float>& InDamageThreshold);
 
 	/** Damage threshold for clusters at different levels. */
-	UPROPERTY(EditAnywhere, BlueprintReadWrite, Category = "ChaosPhysics|Damage")
+	UPROPERTY(EditAnywhere, BlueprintReadWrite, Category = "ChaosPhysics|Damage", meta = (EditCondition = "DamageModel == EDamageModelTypeEnum::Chaos_Damage_Model_UserDefined_Damage_Threshold"))
 	bool bUseSizeSpecificDamageThreshold;
 
 	/** Data about how damage propagation shoudl behave. */
diff --git a/Engine/Source/Runtime/Experimental/GeometryCollectionEngine/Public/GeometryCollection/GeometryCollectionObject.h b/Engine/Source/Runtime/Experimental/GeometryCollectionEngine/Public/GeometryCollection/GeometryCollectionObject.h
--- a/Engine/Source/Runtime/Experimental/GeometryCollectionEngine/Public/GeometryCollection/GeometryCollectionObject.h
+++ b/Engine/Source/Runtime/Experimental/GeometryCollectionEngine/Public/GeometryCollection/GeometryCollectionObject.h
@@ -509,12 +509,16 @@ public:
 	UPROPERTY(EditAnywhere, Category = "Clustering")
 	int32 MaxClusterLevel;
 
+	/** Damage model to use for evaluating destruction. */
+	UPROPERTY(EditAnywhere, Category = "Damage")
+	EDamageModelTypeEnum DamageModel;
+
 	/** Damage threshold for clusters at different levels. */
-	UPROPERTY(EditAnywhere, Category = "Damage", meta = (EditCondition = "!bUseSizeSpecificDamageThreshold"))
+	UPROPERTY(EditAnywhere, Category = "Damage", meta = (EditCondition = "!bUseSizeSpecificDamageThreshold && DamageModel == EDamageModelTypeEnum::Chaos_Damage_Model_UserDefined_Damage_Threshold"))
 	TArray<float> DamageThreshold;
 
 	/** whether to use size specific damage threshold instead of level based ones ( see Size Specific Data array ). */
-	UPROPERTY(EditAnywhere, Category = "Damage")
+	UPROPERTY(EditAnywhere, Category = "Damage", meta = (EditCondition = "DamageModel == EDamageModelTypeEnum::Chaos_Damage_Model_UserDefined_Damage_Threshold"))
 	bool bUseSizeSpecificDamageThreshold;
 
 	/** compatibility check, when true, only cluster compute damage from parameters and propagate to direct children
