Fix compile errors.

#rnx

diff --git a/Engine/Source/Runtime/Experimental/Iris/Core/Private/Iris/ReplicationSystem/ReplicationWriter.cpp b/Engine/Source/Runtime/Experimental/Iris/Core/Private/Iris/ReplicationSystem/ReplicationWriter.cpp
--- a/Engine/Source/Runtime/Experimental/Iris/Core/Private/Iris/ReplicationSystem/ReplicationWriter.cpp
+++ b/Engine/Source/Runtime/Experimental/Iris/Core/Private/Iris/ReplicationSystem/ReplicationWriter.cpp
@@ -464,7 +464,7 @@ void FReplicationWriter::StartReplication(uint32 InternalIndex)
 	// Currently we also do this for dependent objects to support objects with zero priority that should only replicate with parents
 	SchedulingPriorities[InternalIndex] = (Data.IsDependentObject() || Info.IsSubObject) ? 0.f : CreatePriority;
 
-	UE_LOG_REPLICATIONWRITER_CONN(TEXT("ReplicationWriter.StartReplication for ( InternalIndex: %u ) %s"), InternalIndex, *Data.Handle.ToString());
+	UE_LOG_REPLICATIONWRITER_CONN(TEXT("ReplicationWriter.StartReplication for ( InternalIndex: %u ) %s"), InternalIndex, ToCStr(Data.RefHandle.ToString()));
 
 	ObjectsWithDirtyChanges.SetBit(InternalIndex);
 }
@@ -1392,7 +1392,7 @@ void FReplicationWriter::HandleDroppedRecord(const FReplicationRecord::FRecordIn
 
 	check(CurrentState != EReplicatedObjectState::Invalid);
 
-	UE_LOG_REPLICATIONWRITER_CONN(TEXT("Handle dropped data for ( InternalIndex: %u ) %s, LostState %s, CurrentState is %s"), InternalIndex, *(NetRefHandleManager->GetReplicatedObjectData(InternalIndex).Handle.ToString()), LexToString(LostObjectState), LexToString(CurrentState));
+	UE_LOG_REPLICATIONWRITER_CONN(TEXT("Handle dropped data for ( InternalIndex: %u ) %s, LostState %s, CurrentState is %s"), InternalIndex, ToCStr(NetRefHandleManager->GetReplicatedObjectData(InternalIndex).RefHandle.ToString()), LexToString(LostObjectState), LexToString(CurrentState));
 
 	// If we loose a baseline we must notify the BaselineManager and invalidate our pendingbaselineindex
 	if (RecordInfo.NewBaselineIndex != FDeltaCompressionBaselineManager::InvalidBaselineIndex)
