Limit chaos cloth self collision hierarchy lods based on self collision thickness. Smallest hierarchy lod cell size is clamped to be the query bounds size for self collision spring queries. This prevents tiny lods from being generated as they slow down collision query time.

In testing of existing cloth assets, this improved simulation time by a very small amount, but I would expect it to be much better if you happen to have very small triangles.

diff --git a/Engine/Plugins/Experimental/ChaosCloth/Source/Chaos/Private/ChaosCloth/ChaosClothConstraints.cpp b/Engine/Plugins/Experimental/ChaosCloth/Source/Chaos/Private/ChaosCloth/ChaosClothConstraints.cpp
--- a/Engine/Plugins/Experimental/ChaosCloth/Source/Chaos/Private/ChaosCloth/ChaosClothConstraints.cpp
+++ b/Engine/Plugins/Experimental/ChaosCloth/Source/Chaos/Private/ChaosCloth/ChaosClothConstraints.cpp
@@ -800,7 +800,8 @@ void FClothConstraints::CreateRules()
 		ConstraintInits[ConstraintInitIndex++] =
 			[this](Softs::FSolverParticles& Particles, const Softs::FSolverReal /*Dt*/)
 			{
-				SelfCollisionInit->Init(Particles);
+				// Thickness * 2 to account for collision radius for both particles
+				SelfCollisionInit->Init(Particles, SelfCollisionConstraints->GetThickness() * (Softs::FSolverReal)2.f);
 				SelfCollisionConstraints->Init(Particles, SelfCollisionInit->GetSpatialHash(), SelfCollisionInit->GetVertexGIAColors(), SelfCollisionInit->GetTriangleGIAColors());
 			};
 
diff --git a/Engine/Source/Runtime/Experimental/Chaos/Private/Chaos/PBDTriangleMeshCollisions.cpp b/Engine/Source/Runtime/Experimental/Chaos/Private/Chaos/PBDTriangleMeshCollisions.cpp
--- a/Engine/Source/Runtime/Experimental/Chaos/Private/Chaos/PBDTriangleMeshCollisions.cpp
+++ b/Engine/Source/Runtime/Experimental/Chaos/Private/Chaos/PBDTriangleMeshCollisions.cpp
@@ -1153,7 +1153,7 @@ namespace ContourMinimization
 	}
 } // namespace ContourMinimization
 
-void FPBDTriangleMeshCollisions::Init(const FSolverParticles& Particles)
+void FPBDTriangleMeshCollisions::Init(const FSolverParticles& Particles, const FSolverReal MinProximityQueryRadius)
 {
 	if (TriangleMesh.GetNumElements() == 0)
 	{
@@ -1162,7 +1162,8 @@ void FPBDTriangleMeshCollisions::Init(const FSolverParticles& Particles)
 
 	{
 		TRACE_CPUPROFILER_EVENT_SCOPE(ChaosFPBDTriangleMeshCollisions_BuildSpatialHash);
-		TriangleMesh.BuildSpatialHash(static_cast<const TArrayView<const FSolverVec3>&>(Particles.XArray()), SpatialHash);
+		constexpr FSolverReal RadiusToLodSizeMultiplier = 2.; // Radius to Diameter
+		TriangleMesh.BuildSpatialHash(static_cast<const TArrayView<const FSolverVec3>&>(Particles.XArray()), SpatialHash, RadiusToLodSizeMultiplier * MinProximityQueryRadius);
 	}
 	ContourMinimizationIntersections.Reset();
 	VertexGIAColors.Reset();
diff --git a/Engine/Source/Runtime/Experimental/Chaos/Private/Chaos/TriangleMesh.cpp b/Engine/Source/Runtime/Experimental/Chaos/Private/Chaos/TriangleMesh.cpp
--- a/Engine/Source/Runtime/Experimental/Chaos/Private/Chaos/TriangleMesh.cpp
+++ b/Engine/Source/Runtime/Experimental/Chaos/Private/Chaos/TriangleMesh.cpp
@@ -1652,7 +1652,7 @@ template CHAOS_API bool FTriangleMesh::SmoothProject<FRealSingle>(const TBVHType
 template CHAOS_API bool FTriangleMesh::SmoothProject<FRealDouble>(const TBVHType<FRealDouble>& BVH, const TConstArrayView<FVec3>& Points, const TArray<FVec3>& PointNormals, const FVec3& Point, int32& TriangleIndex, FVec3& Weights, const int32 MaxIters) const;
 
 template<typename T>
-void FTriangleMesh::BuildSpatialHash(const TConstArrayView<TVec3<T>>& Points, TSpatialHashType<T>& SpatialHash) const
+void FTriangleMesh::BuildSpatialHash(const TConstArrayView<TVec3<T>>& Points, TSpatialHashType<T>& SpatialHash, const T MinSpatialLodSize) const
 {
 	TArray<TTriangleMeshBvEntry<T>> BVEntries;
 	const int32 NumTris = MElements.Num();
@@ -1662,10 +1662,10 @@ void FTriangleMesh::BuildSpatialHash(const TConstArrayView<TVec3<T>>& Points, TS
 		BVEntries.Add({ this, &Points, Tri });
 	}
 
-	SpatialHash.Initialize(BVEntries);
+	SpatialHash.Initialize(BVEntries, MinSpatialLodSize);
 }
-template void FTriangleMesh::BuildSpatialHash<FRealSingle>(const TConstArrayView<TVec3<FRealSingle>>& Points, TSpatialHashType<FRealSingle>& SpatialHash) const;
-template void FTriangleMesh::BuildSpatialHash<FRealDouble>(const TConstArrayView<TVec3<FRealDouble>>& Points, TSpatialHashType<FRealDouble>& SpatialHash) const;
+template void FTriangleMesh::BuildSpatialHash<FRealSingle>(const TConstArrayView<TVec3<FRealSingle>>& Points, TSpatialHashType<FRealSingle>& SpatialHash, const FRealSingle MinSpatialLodSize) const;
+template void FTriangleMesh::BuildSpatialHash<FRealDouble>(const TConstArrayView<TVec3<FRealDouble>>& Points, TSpatialHashType<FRealDouble>& SpatialHash, const FRealDouble MinSpatialLodSize) const;
 
 template<typename T>
 bool FTriangleMesh::PointProximityQuery(const TSpatialHashType<T>& SpatialHash, const TConstArrayView<TVec3<T>>& Points, const int32 PointIndex, const TVec3<T>& PointPosition, const T PointThickness, const T ThisThickness,
diff --git a/Engine/Source/Runtime/Experimental/Chaos/Public/Chaos/HierarchicalSpatialHash.h b/Engine/Source/Runtime/Experimental/Chaos/Public/Chaos/HierarchicalSpatialHash.h
--- a/Engine/Source/Runtime/Experimental/Chaos/Public/Chaos/HierarchicalSpatialHash.h
+++ b/Engine/Source/Runtime/Experimental/Chaos/Public/Chaos/HierarchicalSpatialHash.h
@@ -22,20 +22,26 @@ struct FHierarchicalSpatialHashCellIdx
 	}
 	
 	template<typename T>
-	static int32 LodForAABB(const TAABB<T, 3>& Box)
+	static int32 LodForCellSize(T CellSize)
 	{
-		const T LongestSideLen = Box.Extents().Max();
-		// want Ceil( log2( LongestSideLen) )
-		if (LongestSideLen <= (T).5)
+		// want Ceil( log2(CellSize) )
+		if (CellSize <= (T).5)
 		{
-			return -(int32)FMath::FloorLog2((uint32)FMath::FloorToInt32((T)1 / LongestSideLen));
+			return -(int32)FMath::FloorLog2((uint32)FMath::FloorToInt32((T)1 / CellSize));
 		}
 		else
 		{
-			return (int32)FMath::CeilLogTwo((uint32)FMath::CeilToInt32(LongestSideLen));
+			return (int32)FMath::CeilLogTwo((uint32)FMath::CeilToInt32(CellSize));
 		}
 	}
 
+	template<typename T>
+	static int32 LodForAABB(const TAABB<T, 3>& Box)
+	{
+		const T LongestSideLen = Box.Extents().Max();
+		return LodForCellSize(LongestSideLen);
+	}
+
 	template<typename T>
 	static T CellSizeForLod(const int32 Lod)
 	{
@@ -209,7 +215,7 @@ public:
 	}
 
 	template<typename ParticlesType>
-	void Initialize(const ParticlesType& Particles)
+	void Initialize(const ParticlesType& Particles, const T MinLodSize = (T)0)
 	{
 		TRACE_CPUPROFILER_EVENT_SCOPE(HierarchicalSpatialHash_Init);
 		Reset();
@@ -217,6 +223,8 @@ public:
 		Bounds.SetNumUninitialized(Particles.Num());
 		HashMap.PreallocateElementsForConcurrentAdd(Particles.Num() * 8, Particles.Num() * 2);
 
+		const int32 MinAllowableLod = MinLodSize > (T)0 ? FHierarchicalSpatialHashCellIdx::LodForCellSize(MinLodSize) : std::numeric_limits<int32>::min();
+
 		// For most common UsedLods, store whether or not the Lod is used in an array without a lock.
 		// For uncommon Lods, use a critical section to store directly into the UsedLods map.
 		constexpr int32 MaxPreAllocatedUsedLodValue = 10; // i.e., a single element 2^10 = 1024 across
@@ -226,12 +234,12 @@ public:
 		FCriticalSection UsedLodsCriticalSection;
 
 		PhysicsParallelFor(Particles.Num(),
-			[this, &Particles, MaxPreAllocatedUsedLodValue, MinPreAllocatedUsedLodValue, &UsedLodsArray, &UsedLodsCriticalSection](int32 ParticleIdx)
+			[this, &Particles, MinAllowableLod, MaxPreAllocatedUsedLodValue, MinPreAllocatedUsedLodValue, &UsedLodsArray, &UsedLodsCriticalSection](int32 ParticleIdx)
 			{
 				const auto& Particle = Particles[ParticleIdx];
 				const TAABB<T, 3> ParticleBounds = Particle.BoundingBox();
 
-				const int32 Lod = FHierarchicalSpatialHashCellIdx::LodForAABB(ParticleBounds);
+				const int32 Lod = FMath::Max(MinAllowableLod, FHierarchicalSpatialHashCellIdx::LodForAABB(ParticleBounds));
 				const T CellSize = FHierarchicalSpatialHashCellIdx::CellSizeForLod<T>(Lod);
 
 				const PayloadType Payload = Particle.template GetPayload<PayloadType>(ParticleIdx);
diff --git a/Engine/Source/Runtime/Experimental/Chaos/Public/Chaos/PBDTriangleMeshCollisions.h b/Engine/Source/Runtime/Experimental/Chaos/Public/Chaos/PBDTriangleMeshCollisions.h
--- a/Engine/Source/Runtime/Experimental/Chaos/Public/Chaos/PBDTriangleMeshCollisions.h
+++ b/Engine/Source/Runtime/Experimental/Chaos/Public/Chaos/PBDTriangleMeshCollisions.h
@@ -153,7 +153,7 @@ public:
 		}
 	}
 
-	void Init(const FSolverParticles& Particles);
+	void Init(const FSolverParticles& Particles, const FSolverReal MinProximityQueryRadius = (FSolverReal)0.);
 
 	void SetGlobalIntersectionAnalysis(bool bInGlobalIntersectionAnalysis) { bGlobalIntersectionAnalysis = bInGlobalIntersectionAnalysis; }
 	void SetContourMinimization(bool bInContourMinimization) { bContourMinimization = bInContourMinimization; }
diff --git a/Engine/Source/Runtime/Experimental/Chaos/Public/Chaos/TriangleMesh.h b/Engine/Source/Runtime/Experimental/Chaos/Public/Chaos/TriangleMesh.h
--- a/Engine/Source/Runtime/Experimental/Chaos/Public/Chaos/TriangleMesh.h
+++ b/Engine/Source/Runtime/Experimental/Chaos/Public/Chaos/TriangleMesh.h
@@ -306,8 +306,9 @@ namespace Chaos
 		template<typename T>
 		using TSpatialHashType = THierarchicalSpatialHash<int32, T>;
 
+		// Hierarchy will only go down to lods as small as MinSpatialLodSize.
 		template<typename T>
-		void BuildSpatialHash(const TConstArrayView<TVec3<T>>& Points, TSpatialHashType<T>& SpatialHash) const;
+		void BuildSpatialHash(const TConstArrayView<TVec3<T>>& Points, TSpatialHashType<T>& SpatialHash, const T MinSpatialLodSize = (T)0.) const;
 
 		template<typename T>
 		bool PointProximityQuery(const TSpatialHashType<T>& SpatialHash, const TConstArrayView<TVec3<T>>& Points, const int32 PointIndex, const TVec3<T>& PointPosition, const T PointThickness, const T ThisThickness,
