Iris: Fix wrong comment and added const to some bit array views to better tell they are read-only

#rnx

diff --git a/Engine/Source/Runtime/Experimental/Iris/Core/Private/Iris/ReplicationSystem/Prioritization/ReplicationPrioritization.cpp b/Engine/Source/Runtime/Experimental/Iris/Core/Private/Iris/ReplicationSystem/Prioritization/ReplicationPrioritization.cpp
--- a/Engine/Source/Runtime/Experimental/Iris/Core/Private/Iris/ReplicationSystem/Prioritization/ReplicationPrioritization.cpp
+++ b/Engine/Source/Runtime/Experimental/Iris/Core/Private/Iris/ReplicationSystem/Prioritization/ReplicationPrioritization.cpp
@@ -119,7 +119,7 @@ public:
 		}
 	}
 
-	EBatchProcessStatus PrepareBatch(FPerConnectionInfo& ConnInfo, FNetBitArrayView Objects, const uint8* PrioritizerIndices, const float* InDefaultPriorities)
+	EBatchProcessStatus PrepareBatch(FPerConnectionInfo& ConnInfo, const FNetBitArrayView Objects, const uint8* PrioritizerIndices, const float* InDefaultPriorities)
 	{
 		IRIS_PROFILER_SCOPE(FReplicationPrioritization_PrioritizeForConnection_PrepareBatch);
 
@@ -666,7 +666,7 @@ void FReplicationPrioritization::UpdatePrioritiesForNewAndDeletedObjects()
  * 1. Sort indices by prioritizer first, index second. Try to keep static priority objects last.
  * 2. Loop through index list until new prioritizer is found and pass the info to the prioritizer for processing.
  */
-void FReplicationPrioritization::PrioritizeForConnection(uint32 ConnId, FPrioritizerBatchHelper& BatchHelper, FNetBitArrayView Objects)
+void FReplicationPrioritization::PrioritizeForConnection(uint32 ConnId, FPrioritizerBatchHelper& BatchHelper, const FNetBitArrayView Objects)
 {
 	IRIS_PROFILER_SCOPE(FReplicationPrioritization_PrioritizeForConnection);
 
diff --git a/Engine/Source/Runtime/Experimental/Iris/Core/Private/Iris/ReplicationSystem/ReplicationBridge.cpp b/Engine/Source/Runtime/Experimental/Iris/Core/Private/Iris/ReplicationSystem/ReplicationBridge.cpp
--- a/Engine/Source/Runtime/Experimental/Iris/Core/Private/Iris/ReplicationSystem/ReplicationBridge.cpp
+++ b/Engine/Source/Runtime/Experimental/Iris/Core/Private/Iris/ReplicationSystem/ReplicationBridge.cpp
@@ -464,7 +464,7 @@ void UReplicationBridge::InternalAddSubObject(FNetRefHandle OwnerHandle, FNetRef
 
 	if (NetRefHandleManager->AddSubObject(OwnerHandle, SubObjectHandle, InsertRelativeToSubObjectHandle, AddSubObjectFlags))
 	{
-		// If the subobject is not new we need to mark it as dirty to pick it up for replication with its new parent
+		// If the subobject is new we need to mark it dirty to pick it up for replication with its new parent
 		MarkNetObjectStateDirty(ReplicationSystem->GetId(), NetRefHandleManager->GetInternalIndex(SubObjectHandle));
 
 		// We set the priority of subobjects to be static as they will be prioritized with owner
