Geometry collection : add extra checks for out of bounds for mass attributes trying to get more data related to the crash happening around here

diff --git a/Engine/Source/Runtime/Experimental/Chaos/Private/PhysicsProxy/GeometryCollectionPhysicsProxy.cpp b/Engine/Source/Runtime/Experimental/Chaos/Private/PhysicsProxy/GeometryCollectionPhysicsProxy.cpp
--- a/Engine/Source/Runtime/Experimental/Chaos/Private/PhysicsProxy/GeometryCollectionPhysicsProxy.cpp
+++ b/Engine/Source/Runtime/Experimental/Chaos/Private/PhysicsProxy/GeometryCollectionPhysicsProxy.cpp
@@ -1690,10 +1690,11 @@ FGeometryCollectionPhysicsProxy::BuildClusters_Internal(
 	}
 
 	check(Parameters.RestCollection);
-	const TManagedArray<Chaos::FReal>& Mass =
-		Parameters.RestCollection->GetAttribute<Chaos::FReal>("Mass", FTransformCollection::TransformGroup);
-	const TManagedArray<FVector3f>& InertiaTensor = 
-		Parameters.RestCollection->GetAttribute<FVector3f>("InertiaTensor", FTransformCollection::TransformGroup);
+	const TManagedArray<Chaos::FReal>& Mass = Parameters.RestCollection->GetAttribute<Chaos::FReal>("Mass", FTransformCollection::TransformGroup);
+	const TManagedArray<FVector3f>& InertiaTensor = Parameters.RestCollection->GetAttribute<FVector3f>("InertiaTensor", FTransformCollection::TransformGroup);
+
+	Mass.RangeCheck(CollectionClusterIndex);
+	InertiaTensor.RangeCheck(CollectionClusterIndex);
 
 	const Chaos::FReal ScaledMass = AdjustMassForScale(Mass[CollectionClusterIndex]);
 	const Chaos::FVec3f ScaledInertia = AdjustInertiaForScale((Chaos::FVec3f)InertiaTensor[CollectionClusterIndex]);
