Apply/AddRadialImpulse functionality added to allow application via Interfaces or Blueprints.

diff --git a/Engine/Source/Runtime/Engine/Classes/PhysicsEngine/PhysicsObjectBlueprintLibrary.h b/Engine/Source/Runtime/Engine/Classes/PhysicsEngine/PhysicsObjectBlueprintLibrary.h
--- a/Engine/Source/Runtime/Engine/Classes/PhysicsEngine/PhysicsObjectBlueprintLibrary.h
+++ b/Engine/Source/Runtime/Engine/Classes/PhysicsEngine/PhysicsObjectBlueprintLibrary.h
@@ -21,4 +21,7 @@ public:
 
 	UFUNCTION(BlueprintPure, Category = "Physics Object")
 	static FTransform GetPhysicsObjectWorldTransform(UPrimitiveComponent* Component, FName BoneName);
+
+	UFUNCTION(BlueprintCallable, Category = "Physics Object")
+	static void ApplyRadialImpulse(UPrimitiveComponent* Component, FVector Origin, float Radius, float Strength, enum ERadialImpulseFalloff Falloff, bool bApplyStrain);
 };
\ No newline at end of file
diff --git a/Engine/Source/Runtime/Engine/Private/PhysicsEngine/PhysicsObjectBlueprintLibrary.cpp b/Engine/Source/Runtime/Engine/Private/PhysicsEngine/PhysicsObjectBlueprintLibrary.cpp
--- a/Engine/Source/Runtime/Engine/Private/PhysicsEngine/PhysicsObjectBlueprintLibrary.cpp
+++ b/Engine/Source/Runtime/Engine/Private/PhysicsEngine/PhysicsObjectBlueprintLibrary.cpp
@@ -1,8 +1,15 @@
 // Copyright Epic Games, Inc. All Rights Reserved.
 #include "PhysicsEngine/PhysicsObjectBlueprintLibrary.h"
 #include "PhysicsEngine/PhysicsObjectExternalInterface.h"
+#include "Chaos/PhysicsObjectCollisionInterface.h"
+#include "Chaos/PhysicsObjectInterface.h"
+#include "Chaos/PhysicsObjectInternalInterface.h"
+
+#include "PhysicsProxy/ClusterUnionPhysicsProxy.h"
+#include "PhysicsProxy/SingleParticlePhysicsProxy.h"
 
 #include "Components/PrimitiveComponent.h"
+#include "PhysicsEngine/ClusterUnionComponent.h"
 
 FClosestPhysicsObjectResult UPhysicsObjectBlueprintLibrary::GetClosestPhysicsObjectFromWorldLocation(UPrimitiveComponent* Component, const FVector& WorldLocation)
 {
@@ -46,3 +53,19 @@ FTransform UPhysicsObjectBlueprintLibrary::GetPhysicsObjectWorldTransform(UPrimi
 	
 	return FTransform::Identity;
 }
+
+void UPhysicsObjectBlueprintLibrary::ApplyRadialImpulse(UPrimitiveComponent* Component, FVector Origin, float Radius, float Strength, enum ERadialImpulseFalloff Falloff, bool bApplyStrain)
+{
+	if (Component == nullptr)
+	{
+		return;
+	}
+
+	if (Chaos::FPhysicsObject* PhysicsObject = Component->GetPhysicsObjectByName(NAME_None))
+	{
+		TArray<Chaos::FPhysicsObjectHandle>  PhysicsObjects;
+		PhysicsObjects.Add(PhysicsObject);
+		FLockedWritePhysicsObjectExternalInterface Interface = FPhysicsObjectExternalInterface::LockWrite(PhysicsObjects);
+		return Interface->AddRadialImpulse(PhysicsObjects, Origin, Radius, Strength, Falloff, bApplyStrain, true);
+	}
+}
diff --git a/Engine/Source/Runtime/Experimental/Chaos/Private/Chaos/PhysicsObjectInterface.cpp b/Engine/Source/Runtime/Experimental/Chaos/Private/Chaos/PhysicsObjectInterface.cpp
--- a/Engine/Source/Runtime/Experimental/Chaos/Private/Chaos/PhysicsObjectInterface.cpp
+++ b/Engine/Source/Runtime/Experimental/Chaos/Private/Chaos/PhysicsObjectInterface.cpp
@@ -8,6 +8,18 @@
 #include "Chaos/PhysicsObjectInternal.h"
 #include "PhysicsProxy/SingleParticlePhysicsProxy.h"
 #include "Math/Transform.h"
+#include "Chaos/ChaosEngineInterface.h"
+
+namespace PhysicsObjectInterfaceCVars
+{
+	static float StrainModifier = 200;
+	static FAutoConsoleVariableRef CVarStrainModifier(
+		TEXT("Chaos.Debug.StrainModifier"),
+		StrainModifier,
+		TEXT("Modify the strain by this factor"),
+		ECVF_Default
+	);
+}
 
 namespace
 {
@@ -66,6 +78,39 @@ namespace
 			}
 		}
 	}
+
+	template<typename T>
+	float AddRadialImpulseHelper(T ParticleHandle, FVector Origin, float Radius, float Strength, enum ERadialImpulseFalloff Falloff, float VelocityRatio = 1.0f, bool bInvalidate = true)
+	{
+		using namespace Chaos;
+
+		if (ParticleHandle == nullptr)
+		{
+			return 0.0f;
+		}
+
+		const FVec3 CurrentImpulseVelocity = ParticleHandle->LinearImpulseVelocity();
+
+		float FalloffStrength = Strength;
+		if (Falloff == ERadialImpulseFalloff::RIF_Linear)
+		{
+			//Radius should always be greater than distance - if we get here and it's not, something has gone
+			//wrong with detection
+			double Distance = FVec3::Distance(ParticleHandle->X(), Origin);
+			FalloffStrength = static_cast<float>(FalloffStrength * (1.0f - FMath::Min(1.0f, Distance / Radius)));
+		}
+
+		FVec3 Impulse = ParticleHandle->X() - Origin;
+		Impulse = Impulse.GetSafeNormal();
+		Impulse = Impulse * FalloffStrength;
+		FVec3 Velocity = Impulse * ParticleHandle->InvM() * VelocityRatio;
+
+		FVec3 FinalVeloc = CurrentImpulseVelocity + Velocity;
+
+		ParticleHandle->SetLinearImpulseVelocity(CurrentImpulseVelocity + Velocity, bInvalidate);
+
+		return FalloffStrength;
+	}
 }
 
 FName FClosestPhysicsObjectResult::HitName() const
@@ -712,6 +757,85 @@ namespace Chaos
 		}
 	}
 
+	template<EThreadContext Id>
+	void FWritePhysicsObjectInterface<Id>::AddRadialImpulse(TArrayView<const FPhysicsObjectHandle> InObjects, FVector Origin, float Radius, float Strength, enum ERadialImpulseFalloff Falloff, bool bApplyStrain, bool bInvalidate)
+	{
+		if (Chaos::FPBDRigidsSolver* RigidSolver = Chaos::FPhysicsObjectInterface::GetSolver(InObjects))
+		{
+			//put onto physics thread
+			RigidSolver->EnqueueCommandImmediate([InObjects = TArray<FPhysicsObjectHandle>{ InObjects },
+				RigidSolver,
+				Origin,
+				Radius,
+				Strength,
+				Falloff,
+				bApplyStrain,
+				bInvalidate]()
+				{
+					using namespace Chaos;
+					for (FPhysicsObjectHandle Object : InObjects)
+					{
+						if (Object == nullptr)
+						{
+							continue;
+						}
+
+						//can apply to clusters or rigid particles
+						TThreadParticle<EThreadContext::Internal>* Particle = Object->GetParticle<EThreadContext::Internal>();
+						if (FPBDRigidClusteredParticleHandle* Cluster = reinterpret_cast<FPBDRigidClusteredParticleHandle*>(Particle))
+						{
+							if (bApplyStrain)
+							{
+								FRigidClustering& Clustering = RigidSolver->GetEvolution()->GetRigidClustering();
+								TArray<FPBDRigidParticleHandle*>* ChildrenHandles = Clustering.GetChildrenMap().Find(Cluster);
+
+								if (ChildrenHandles == nullptr)
+								{
+									continue;
+								}
+
+								for (FPBDRigidParticleHandle* ChildHandle : *ChildrenHandles)
+								{
+									float VelocityRatio = 1.0f / static_cast<float>(ChildrenHandles->Num());
+
+									float FalloffStrength = AddRadialImpulseHelper(ChildHandle, Origin, Radius, Strength, Falloff, VelocityRatio, bInvalidate);
+
+									//to do: remove cvar when material system is in place and densities are updated
+									FalloffStrength = PhysicsObjectInterfaceCVars::StrainModifier * FalloffStrength;
+
+									Clustering.SetExternalStrain(ChildHandle->CastToClustered(), FalloffStrength);
+								}
+							}
+							else
+							{
+								FPBDRigidParticleHandle* ParticleHandle = Cluster->CastToRigidParticle();
+
+								if (ParticleHandle != nullptr && ParticleHandle->IsSleeping())
+								{
+									RigidSolver->GetEvolution()->SetParticleObjectState(ParticleHandle, Chaos::EObjectStateType::Dynamic);
+									RigidSolver->GetEvolution()->GetParticles().MarkTransientDirtyParticle(ParticleHandle);
+								}
+
+								AddRadialImpulseHelper(ParticleHandle, Origin, Radius, Strength, Falloff, bInvalidate);
+							}
+						}
+						else
+						{
+							FPBDRigidParticleHandle* ParticleHandle = reinterpret_cast<FPBDRigidParticleHandle*>(Particle);
+
+							if (ParticleHandle != nullptr && ParticleHandle->IsSleeping())
+							{
+								RigidSolver->GetEvolution()->SetParticleObjectState(ParticleHandle, Chaos::EObjectStateType::Dynamic);
+								RigidSolver->GetEvolution()->GetParticles().MarkTransientDirtyParticle(ParticleHandle);
+							}
+
+							AddRadialImpulseHelper(ParticleHandle, Origin, Radius, Strength, Falloff, bInvalidate);
+						}
+					}
+				});
+		}
+	}
+
 	template<EThreadContext Id>
 	void FWritePhysicsObjectInterface<Id>::UpdateShapeCollisionFlags(TArrayView<const FPhysicsObjectHandle> InObjects, bool bSimCollision, bool bQueryCollision)
 	{
diff --git a/Engine/Source/Runtime/Experimental/Chaos/Public/Chaos/PhysicsObjectCollisionInterface.h b/Engine/Source/Runtime/Experimental/Chaos/Public/Chaos/PhysicsObjectCollisionInterface.h
--- a/Engine/Source/Runtime/Experimental/Chaos/Public/Chaos/PhysicsObjectCollisionInterface.h
+++ b/Engine/Source/Runtime/Experimental/Chaos/Public/Chaos/PhysicsObjectCollisionInterface.h
@@ -42,7 +42,7 @@ namespace Chaos
 				FVector::Zero(),
 				[this, ObjectA, &OutOverlaps](const FShapeOverlapData& A, const FShapeOverlapData& B, const FMTDInfo&)
 				{
-					ChaosInterface::FOverlapHit Overlap;
+					TOverlapHit Overlap;
 					Overlap.Shape = A.Shape;
 					Overlap.Actor = Interface.GetParticle(ObjectA);
 					OutOverlaps.Add(Overlap);
@@ -169,7 +169,7 @@ namespace Chaos
 						{
 							bHasOverlap = true;
 
-							ChaosInterface::FOverlapHit Overlap;
+							TOverlapHit Overlap;
 							Overlap.Shape = Shape;
 							Overlap.Actor = Interface.GetParticle(IterObject);
 							OutOverlaps.Add(Overlap);
diff --git a/Engine/Source/Runtime/Experimental/Chaos/Public/Chaos/PhysicsObjectInterface.h b/Engine/Source/Runtime/Experimental/Chaos/Public/Chaos/PhysicsObjectInterface.h
--- a/Engine/Source/Runtime/Experimental/Chaos/Public/Chaos/PhysicsObjectInterface.h
+++ b/Engine/Source/Runtime/Experimental/Chaos/Public/Chaos/PhysicsObjectInterface.h
@@ -37,6 +37,8 @@ struct CHAOS_API FClosestPhysicsObjectResult
 	FName HitName() const;
 };
 
+enum ERadialImpulseFalloff : int;
+
 namespace Chaos
 {
 	class FPBDRigidsSolver;
@@ -117,6 +119,7 @@ namespace Chaos
 		void WakeUp(TArrayView<const FPhysicsObjectHandle> InObjects);
 		void AddForce(TArrayView<const FPhysicsObjectHandle> InObjects, const FVector& Force, bool bInvalidate);
 		void AddTorque(TArrayView<const FPhysicsObjectHandle> InObjects, const FVector& Torque, bool bInvalidate);
+		void AddRadialImpulse(TArrayView<const FPhysicsObjectHandle> InObjects, FVector Origin, float Radius, float Strength, enum ERadialImpulseFalloff Falloff, bool bApplyStrain, bool bInvalidate);
 		
 		void UpdateShapeCollisionFlags(TArrayView<const FPhysicsObjectHandle> InObjects, bool bSimCollision, bool bQueryCollision);
 		void UpdateShapeFilterData(TArrayView<const FPhysicsObjectHandle> InObjects, const FCollisionFilterData& QueryData, const FCollisionFilterData& SimData);
