ChaosFlesh : Add Gauss Seildel neohookean method

diff --git a/Engine/Plugins/Experimental/ChaosFlesh/Source/ChaosFleshEngine/Private/ChaosFlesh/ChaosDeformableSolverComponent.cpp b/Engine/Plugins/Experimental/ChaosFlesh/Source/ChaosFleshEngine/Private/ChaosFlesh/ChaosDeformableSolverComponent.cpp
--- a/Engine/Plugins/Experimental/ChaosFlesh/Source/ChaosFleshEngine/Private/ChaosFlesh/ChaosDeformableSolverComponent.cpp
+++ b/Engine/Plugins/Experimental/ChaosFlesh/Source/ChaosFleshEngine/Private/ChaosFlesh/ChaosDeformableSolverComponent.cpp
@@ -213,9 +213,10 @@ void UDeformableSolverComponent::Reset()
 			, SolverForces.bEnableGravity
 			, SolverConstraints.CorotatedConstraints.bEnableCorotatedConstraint
 			, SolverConstraints.bEnablePositionTargets
-			, GaussSeidelConstraints.bUseGaussSeidelConstraints
-			, GaussSeidelConstraints.bUseSOR
-			, GaussSeidelConstraints.OmegaSOR
+			, SolverConstraints.GaussSeidelConstraints.bUseGaussSeidelConstraints
+			, SolverConstraints.GaussSeidelConstraints.bUseSOR
+			, SolverConstraints.GaussSeidelConstraints.OmegaSOR
+			, SolverConstraints.GaussSeidelConstraints.bUseGSNeohookean
 		}));
 
 		for (TObjectPtr<UDeformablePhysicsComponent>& DeformableComponent : ConnectedObjects.DeformableComponents)
diff --git a/Engine/Plugins/Experimental/ChaosFlesh/Source/ChaosFleshEngine/Public/ChaosFlesh/ChaosDeformableSolverComponent.h b/Engine/Plugins/Experimental/ChaosFlesh/Source/ChaosFleshEngine/Public/ChaosFlesh/ChaosDeformableSolverComponent.h
--- a/Engine/Plugins/Experimental/ChaosFlesh/Source/ChaosFleshEngine/Public/ChaosFlesh/ChaosDeformableSolverComponent.h
+++ b/Engine/Plugins/Experimental/ChaosFlesh/Source/ChaosFleshEngine/Public/ChaosFlesh/ChaosDeformableSolverComponent.h
@@ -99,12 +99,27 @@ struct FSolverGaussSeidelConstraintsGroup
 {
 	GENERATED_USTRUCT_BODY()
 
+	/**
+	* Enable the Gauss Seidel solver instead of the existing XPBD.
+	*/
 	UPROPERTY(EditAnywhere, Category = "GaussSeidelConstraints")
 		bool bUseGaussSeidelConstraints = false;
 
+	/**
+	* Enable another model that runs simulation faster.
+	*/
+	UPROPERTY(EditAnywhere, Category = "GaussSeidelConstraints")
+		bool bUseGSNeohookean = false;
+
+	/**
+	* Enable acceleration technique for Gauss Seidel solver to make simulation look better within a limited budget.
+	*/
 	UPROPERTY(EditAnywhere, Category = "GaussSeidelConstraints")
 		bool bUseSOR = true;
 
+	/**
+	* Acceleration related parameter. Tune it down if simulation becomes unstable. 
+	*/
 	UPROPERTY(EditAnywhere, Category = "GaussSeidelConstraints")
 		float OmegaSOR = 1.6f;
 };
@@ -153,6 +168,12 @@ struct FSolverConstraintsGroup
 
 	UPROPERTY(EditAnywhere, Category = "Constraints")
 	FSolverCorotatedConstraintsGroup CorotatedConstraints;
+
+	/**
+	* These are options for another solver. 
+	*/
+	UPROPERTY(EditAnywhere, BlueprintReadOnly, Category = "Constraints")
+	FSolverGaussSeidelConstraintsGroup GaussSeidelConstraints;
 };
 
 
@@ -223,8 +244,7 @@ public:
 	UPROPERTY(EditAnywhere, BlueprintReadOnly, Category = "Physics")
 	FSolverConstraintsGroup SolverConstraints;
 
-	UPROPERTY(EditAnywhere, BlueprintReadOnly, Category = "Physics")
-	FSolverGaussSeidelConstraintsGroup GaussSeidelConstraints;
+	
 
 	UPROPERTY(EditAnywhere, BlueprintReadOnly, Category = "Physics")
 	FSolverForcesGroup SolverForces;
diff --git a/Engine/Source/Runtime/Experimental/Chaos/Private/Chaos/Deformable/ChaosDeformableSolver.cpp b/Engine/Source/Runtime/Experimental/Chaos/Private/Chaos/Deformable/ChaosDeformableSolver.cpp
--- a/Engine/Source/Runtime/Experimental/Chaos/Private/Chaos/Deformable/ChaosDeformableSolver.cpp
+++ b/Engine/Source/Runtime/Experimental/Chaos/Private/Chaos/Deformable/ChaosDeformableSolver.cpp
@@ -1151,21 +1151,42 @@ namespace Chaos::Softs
 	{
 		PERF_SCOPE(STAT_ChaosDeformableSolver_InitializeGaussSeidelConstraintVariables);
 
-		GSCorotatedConstraints.Reset(new Chaos::Softs::FGaussSeidelCorotatedConstraints<FSolverReal, FSolverParticles>(
-			Evolution->Particles(), *AllElements, *AllTetEMeshArray, *AllTetNuMeshArray, MoveTemp(*AllTetAlphaJArray), MoveTemp(*AllIncidentElements), MoveTemp(*AllIncidentElementsLocal), 0, Evolution->Particles().Size(), Property.bDoQuasistatics, Property.bUseSOR, Property.OmegaSOR, GDeformableXPBDCorotatedParams));
-		Evolution->ResetConstraintRules();
-		int32 InitIndex1 = Evolution->AddConstraintInitRange(1, true);
-		Evolution->ConstraintInits()[InitIndex1] =
-			[this](FSolverParticles& InParticles, const FSolverReal Dt)
+		if (Property.bUseGSNeohookean)
 		{
-			this->GSCorotatedConstraints->Init(Dt, InParticles);
-		};
-		int32 ConstraintIndex1 = Evolution->AddConstraintRuleRange(1, true);
-		Evolution->ConstraintRules()[ConstraintIndex1] =
-			[this](FSolverParticles& InParticles, const FSolverReal Dt)
+			GSNeohookeanConstraints.Reset(new Chaos::Softs::FGaussSeidelNeohookeanConstraints<FSolverReal, FSolverParticles>(
+				Evolution->Particles(), *AllElements, *AllTetEMeshArray, *AllTetNuMeshArray, MoveTemp(*AllTetAlphaJArray), MoveTemp(*AllIncidentElements), MoveTemp(*AllIncidentElementsLocal), 0, Evolution->Particles().Size(), Property.bDoQuasistatics, Property.bUseSOR, Property.OmegaSOR, GDeformableXPBDCorotatedParams));
+			Evolution->ResetConstraintRules();
+			int32 InitIndex1 = Evolution->AddConstraintInitRange(1, true);
+			Evolution->ConstraintInits()[InitIndex1] =
+				[this](FSolverParticles& InParticles, const FSolverReal Dt)
+			{
+				this->GSNeohookeanConstraints->Init(Dt, InParticles);
+			};
+			int32 ConstraintIndex1 = Evolution->AddConstraintRuleRange(1, true);
+			Evolution->ConstraintRules()[ConstraintIndex1] =
+				[this](FSolverParticles& InParticles, const FSolverReal Dt)
+			{
+				this->GSNeohookeanConstraints->Apply(InParticles, Dt);
+			};
+		} 
+		else
 		{
-			this->GSCorotatedConstraints->ApplyInParallel(InParticles, Dt);
-		};
+			GSCorotatedConstraints.Reset(new Chaos::Softs::FGaussSeidelCorotatedConstraints<FSolverReal, FSolverParticles>(
+				Evolution->Particles(), *AllElements, *AllTetEMeshArray, *AllTetNuMeshArray, MoveTemp(*AllTetAlphaJArray), MoveTemp(*AllIncidentElements), MoveTemp(*AllIncidentElementsLocal), 0, Evolution->Particles().Size(), Property.bDoQuasistatics, Property.bUseSOR, Property.OmegaSOR, GDeformableXPBDCorotatedParams));
+			Evolution->ResetConstraintRules();
+			int32 InitIndex1 = Evolution->AddConstraintInitRange(1, true);
+			Evolution->ConstraintInits()[InitIndex1] =
+				[this](FSolverParticles& InParticles, const FSolverReal Dt)
+			{
+				this->GSCorotatedConstraints->Init(Dt, InParticles);
+			};
+			int32 ConstraintIndex1 = Evolution->AddConstraintRuleRange(1, true);
+			Evolution->ConstraintRules()[ConstraintIndex1] =
+				[this](FSolverParticles& InParticles, const FSolverReal Dt)
+			{
+				this->GSCorotatedConstraints->Apply(InParticles, Dt);
+			};
+		}
 	}
 
 
diff --git a/Engine/Source/Runtime/Experimental/Chaos/Private/Chaos/GraphColoring.cpp b/Engine/Source/Runtime/Experimental/Chaos/Private/Chaos/GraphColoring.cpp
--- a/Engine/Source/Runtime/Experimental/Chaos/Private/Chaos/GraphColoring.cpp
+++ b/Engine/Source/Runtime/Experimental/Chaos/Private/Chaos/GraphColoring.cpp
@@ -1439,7 +1439,7 @@ TArray<TArray<int32>> Chaos::ComputeNodalColoring(const TArray<TVec4<int32>>& Gr
 		}
 	}
 
-	int32 SizeColors = ParticleColors.Max();
+	int32 SizeColors = FMath::Max<int32>(ParticleColors);
 
 	ParticlesPerColor.SetNum(0);
 	ParticlesPerColor.Init(TArray<int32>(), SizeColors + 1);
diff --git a/Engine/Source/Runtime/Experimental/Chaos/Private/Chaos/NeohookeanModel.cpp b/Engine/Source/Runtime/Experimental/Chaos/Private/Chaos/NeohookeanModel.cpp
new file mode 100644
--- /dev/null
+++ b/Engine/Source/Runtime/Experimental/Chaos/Private/Chaos/NeohookeanModel.cpp
@@ -0,0 +1,45 @@
+// Copyright Epic Games, Inc. All Rights Reserved.
+#include "Chaos/NeohookeanModel.h"
+#include "Chaos/ImplicitQRSVD.h"
+
+
+
+namespace Chaos::Softs
+{
+
+template <typename T>
+T PsiNeohookeanMM(const Chaos::PMatrix<T, 3, 3>& F, const T mu, const T lambda)
+{
+	T J = F.Determinant();
+	Chaos::PMatrix<T, 3, 3> FTF = F* F.GetTransposed();
+	T lambda_hat = mu + lambda;
+	T gamma = T(1) + mu / lambda_hat;
+	return T(0.5) * mu * (FTF.GetAt(0,0) + FTF.GetAt(1,1) + FTF.GetAt(2,2) - 3) + T(0.5) * lambda_hat * (J - gamma) * (J - gamma);
+}
+
+
+template <typename T>
+void PNeohookeanMM(const Chaos::PMatrix<T, 3, 3>& Fe, const T mu, const T lambda, Chaos::PMatrix<T, 3, 3>& P)
+{
+	T J = Fe.Determinant();
+	T lambda_hat = mu + lambda;
+	Chaos::PMatrix<T, 3, 3> JFinvT((T)0.);
+	JFinvT.SetAt(0, 0, Fe.GetAt(1, 1) * Fe.GetAt(2, 2) - Fe.GetAt(2, 1) * Fe.GetAt(1, 2));
+	JFinvT.SetAt(0, 1, Fe.GetAt(2, 0) * Fe.GetAt(1, 2) - Fe.GetAt(1, 0) * Fe.GetAt(2, 2));
+	JFinvT.SetAt(0, 2, Fe.GetAt(1, 0) * Fe.GetAt(2, 1) - Fe.GetAt(2, 0) * Fe.GetAt(1, 1));
+	JFinvT.SetAt(1, 0, Fe.GetAt(2, 1) * Fe.GetAt(0, 2) - Fe.GetAt(0, 1) * Fe.GetAt(2, 2));
+	JFinvT.SetAt(1, 1, Fe.GetAt(0, 0) * Fe.GetAt(2, 2) - Fe.GetAt(2, 0) * Fe.GetAt(0, 2));
+	JFinvT.SetAt(1, 2, Fe.GetAt(2, 0) * Fe.GetAt(0, 1) - Fe.GetAt(0, 0) * Fe.GetAt(2, 1));
+	JFinvT.SetAt(2, 0, Fe.GetAt(0, 1) * Fe.GetAt(1, 2) - Fe.GetAt(1, 1) * Fe.GetAt(0, 2));
+	JFinvT.SetAt(2, 1, Fe.GetAt(1, 0) * Fe.GetAt(0, 2) - Fe.GetAt(0, 0) * Fe.GetAt(1, 2));
+	JFinvT.SetAt(2, 2, Fe.GetAt(0, 0) * Fe.GetAt(1, 1) - Fe.GetAt(1, 0) * Fe.GetAt(0, 1));
+	T gamma = T(1) + mu / lambda_hat;
+	P = T(mu) * Fe+T(lambda_hat) * (J - gamma) * JFinvT;
+}
+
+}
+
+template CHAOS_API void Chaos::Softs::PNeohookeanMM<Chaos::FReal>(const Chaos::PMatrix<Chaos::FReal, 3, 3>& Fe, const Chaos::FReal mu, const Chaos::FReal lambda, Chaos::PMatrix<Chaos::FReal, 3, 3>& P);
+template CHAOS_API void Chaos::Softs::PNeohookeanMM<Chaos::FRealSingle>(const Chaos::PMatrix<Chaos::FRealSingle, 3, 3>& Fe, const Chaos::FRealSingle mu, const Chaos::FRealSingle lambda, Chaos::PMatrix<Chaos::FRealSingle, 3, 3>& P);
+template CHAOS_API Chaos::FReal Chaos::Softs::PsiNeohookeanMM<Chaos::FReal>(const Chaos::PMatrix<Chaos::FReal, 3, 3>& F, const Chaos::FReal mu, const Chaos::FReal lambda);
+template CHAOS_API Chaos::FRealSingle Chaos::Softs::PsiNeohookeanMM<Chaos::FRealSingle>(const Chaos::PMatrix<Chaos::FRealSingle, 3, 3>& F, const Chaos::FRealSingle mu, const Chaos::FRealSingle lambda);
diff --git a/Engine/Source/Runtime/Experimental/Chaos/Public/Chaos/Deformable/ChaosDeformableSolver.h b/Engine/Source/Runtime/Experimental/Chaos/Public/Chaos/Deformable/ChaosDeformableSolver.h
--- a/Engine/Source/Runtime/Experimental/Chaos/Public/Chaos/Deformable/ChaosDeformableSolver.h
+++ b/Engine/Source/Runtime/Experimental/Chaos/Public/Chaos/Deformable/ChaosDeformableSolver.h
@@ -12,6 +12,7 @@
 #include "Chaos/TriangleMesh.h"
 #include "Chaos/XPBDCorotatedConstraints.h"
 #include "Chaos/Deformable/GaussSeidelCorotatedConstraints.h"
+#include "Chaos/Deformable/GaussSeidelNeohookeanConstraints.h"
 #include "Chaos/XPBDWeakConstraints.h"
 #include "Chaos/BlendedXPBDCorotatedConstraints.h"
 #include "Chaos/XPBDGridBasedCorotatedConstraints.h"
@@ -176,6 +177,7 @@ namespace Chaos::Softs
 		TUniquePtr<Softs::FPBDEvolution> Evolution;
 		TArray<TUniquePtr<Softs::FXPBDCorotatedConstraints<Softs::FSolverReal, Softs::FSolverParticles>>> CorotatedConstraints;
 		TUniquePtr<Softs::FGaussSeidelCorotatedConstraints<Softs::FSolverReal, Softs::FSolverParticles>> GSCorotatedConstraints;
+		TUniquePtr<Softs::FGaussSeidelNeohookeanConstraints<Softs::FSolverReal, Softs::FSolverParticles>> GSNeohookeanConstraints;
 		TArray<TUniquePtr<Softs::FXPBDWeakConstraints<Softs::FSolverReal, Softs::FSolverParticles>>> WeakConstraints;
 		TArray<TUniquePtr<Softs::FBlendedXPBDCorotatedConstraints<Softs::FSolverReal, Softs::FSolverParticles>>> BlendedCorotatedConstraints;
 		TUniquePtr<Softs::FXPBDGridBasedCorotatedConstraints<Softs::FSolverReal, Softs::FSolverParticles>> GridBasedCorotatedConstraint;
diff --git a/Engine/Source/Runtime/Experimental/Chaos/Public/Chaos/Deformable/ChaosDeformableSolverTypes.h b/Engine/Source/Runtime/Experimental/Chaos/Public/Chaos/Deformable/ChaosDeformableSolverTypes.h
--- a/Engine/Source/Runtime/Experimental/Chaos/Public/Chaos/Deformable/ChaosDeformableSolverTypes.h
+++ b/Engine/Source/Runtime/Experimental/Chaos/Public/Chaos/Deformable/ChaosDeformableSolverTypes.h
@@ -36,7 +36,8 @@ namespace Chaos::Softs
 			bool InbEnablePositionTargets = true, 
 			bool InbUseGaussSeidelConstraints = false, 
 			bool InbUseSOR = true,
-			FSolverReal InOmegaSOR = (FSolverReal)1.6)
+			FSolverReal InOmegaSOR = (FSolverReal)1.6, 
+			bool InbUseGSNeohookean = false)
 			: NumSolverSubSteps(InNumSolverSubSteps)
 			, NumSolverIterations(InNumSolverIterations)
 			, FixTimeStep(InFixTimeStep)
@@ -58,6 +59,7 @@ namespace Chaos::Softs
 			, bUseGaussSeidelConstraints(InbUseGaussSeidelConstraints)
 			, bUseSOR(InbUseSOR)
 			, OmegaSOR(InOmegaSOR)
+			, bUseGSNeohookean(InbUseGSNeohookean)
 		{}
 
 		int32 NumSolverSubSteps = 5;
@@ -81,6 +83,7 @@ namespace Chaos::Softs
 		bool bUseGaussSeidelConstraints = false;
 		bool bUseSOR = true;
 		FSolverReal OmegaSOR = (FSolverReal)1.6;
+		bool bUseGSNeohookean = false;
 	};
 
 
diff --git a/Engine/Source/Runtime/Experimental/Chaos/Public/Chaos/Deformable/GaussSeidelCorotatedConstraints.h b/Engine/Source/Runtime/Experimental/Chaos/Public/Chaos/Deformable/GaussSeidelCorotatedConstraints.h
--- a/Engine/Source/Runtime/Experimental/Chaos/Public/Chaos/Deformable/GaussSeidelCorotatedConstraints.h
+++ b/Engine/Source/Runtime/Experimental/Chaos/Public/Chaos/Deformable/GaussSeidelCorotatedConstraints.h
@@ -8,6 +8,7 @@
 #include "Chaos/GraphColoring.h"
 #include "Chaos/NewtonCorotatedCache.h"
 #include "Chaos/Framework/Parallel.h"
+//#include "Chaos/Deformable/GaussSeidelWeakConstraints.h"
 
 
 namespace Chaos::Softs
@@ -23,9 +24,9 @@ namespace Chaos::Softs
 		using Base::DmInverse;
 		using Base::MuElementArray;
 		using Base::LambdaElementArray;
+		using Base::CorotatedParams;
 
 	public:
-		//this one only accepts tetmesh input and mesh
 		FGaussSeidelCorotatedConstraints(
 			const ParticleType& InParticles,
 			const TArray<TVector<int32, 4>>& InMesh,
@@ -73,6 +74,7 @@ namespace Chaos::Softs
 				X_k.Init(TVector<T, 3>(T(0.)), ParticleEndIndex - ParticleStartIndex);
 			}
 			InitColor(InParticles);
+			InitializeLambdas();
 		}
 
 		virtual ~FGaussSeidelCorotatedConstraints() {}
@@ -80,17 +82,29 @@ namespace Chaos::Softs
 
 		void Apply(ParticleType& Particles, const T Dt) const 
 		{
-			//SCOPE_CYCLE_COUNTER(STAT_ChaosGaussSeidelCorotated);
-			TRACE_CPUPROFILER_EVENT_SCOPE(TEXT("STAT_ChaosGaussSeidelCorotatedApply"));
+			TRACE_CPUPROFILER_EVENT_SCOPE(STAT_ChaosGaussSeidelApply);
 			for (int32 i = 0; i < ParticlesPerColor.Num(); i++)
 			{
-				PhysicsParallelFor(ParticlesPerColor[i].Num(), [&](const int32 j)
+				int32 NumBatch = ParticlesPerColor[i].Num() / CorotatedParams.XPBDCorotatedBatchSize;
+				if (ParticlesPerColor[i].Num() % CorotatedParams.XPBDCorotatedBatchSize != 0)
+				{
+					NumBatch += 1;
+				}
+
+				PhysicsParallelFor(NumBatch, [&](const int32 BatchIndex)
 					{
-						int32 ParticleIndex = ParticlesPerColor[i][j];
-						Chaos::TVector<T, 3> Dx = ComputeDeltax(ParticleIndex, Particle2Incident[ParticleIndex - ParticleStartIndex], Particles,
-							Dt);
-						Particles.P(ParticleIndex) += Dx;
-					}, ParticlesPerColor[i].Num() < Base::CorotatedParams.XPBDCorotatedBatchThreshold);
+						for (int32 BatchSubIndex = 0; BatchSubIndex < CorotatedParams.XPBDCorotatedBatchSize; BatchSubIndex++) {
+							int32 TaskIndex = CorotatedParams.XPBDCorotatedBatchSize * BatchIndex + BatchSubIndex;
+							if (TaskIndex < ParticlesPerColor[i].Num())
+							{
+								int32 ParticleIndex = ParticlesPerColor[i][TaskIndex];
+								Chaos::TVector<T, 3> Dx = ComputeDeltax(ParticleIndex, Particle2Incident[ParticleIndex - ParticleStartIndex], Particles,
+									Dt);
+								Particles.P(ParticleIndex) += Dx;
+							}
+						}
+					}, NumBatch < CorotatedParams.XPBDCorotatedBatchThreshold);
+
 			}
 			ApplySOR(Particles, Dt);
 			CurrentIt += 1;
@@ -122,7 +136,6 @@ namespace Chaos::Softs
 						{
 							Particles.P(ParticleIndex) = OmegaSOR * (Particles.P(ParticleIndex) - X_k_1[i]) + X_k_1[i];
 						}
-						//X_k_1[i] = Particles.X(ParticleIndex);
 						X_k_1[i] = X_k[i];
 						X_k[i] = Particles.P(ParticleIndex);
 					}, ParticleEndIndex - ParticleStartIndex < 1000);
@@ -146,11 +159,9 @@ namespace Chaos::Softs
 				Chaos::PMatrix<T, 3, 3> Fe = Base::F(e, Particles);
 				Chaos::PMatrix<T, 3, 3> P((T)0.);
 
-				//corotated model may switch to others:
-				Chaos::Softs::PCorotated(Fe, MuElementArray[e], LambdaElementArray[e], P);
+				ComputeStress(Fe, MuElementArray[e], LambdaElementArray[e], P);
 
 				Chaos::PMatrix<T, 3, 3> force_term = -Measure[e] * DmInvT * P;
-				// last term: contribution from element e:
 				Chaos::TVector<T, 3> dx((T)0.);
 				if (local_index > 0)
 				{
@@ -200,70 +211,8 @@ namespace Chaos::Softs
 
 				T Coeff = Dt * Dt * Measure[ElementIndex];
 
-				Chaos::PMatrix<T, 3, 3> JFinvT((T)0.);
-				JFinvT.SetAt(0, 0, Fe.GetAt(1, 1) * Fe.GetAt(2, 2) - Fe.GetAt(2, 1) * Fe.GetAt(1, 2));
-				JFinvT.SetAt(0, 1, Fe.GetAt(2, 0) * Fe.GetAt(1, 2) - Fe.GetAt(1, 0) * Fe.GetAt(2, 2));
-				JFinvT.SetAt(0, 2, Fe.GetAt(1, 0) * Fe.GetAt(2, 1) - Fe.GetAt(2, 0) * Fe.GetAt(1, 1));
-				JFinvT.SetAt(1, 0, Fe.GetAt(2, 1) * Fe.GetAt(0, 2) - Fe.GetAt(0, 1) * Fe.GetAt(2, 2));
-				JFinvT.SetAt(1, 1, Fe.GetAt(0, 0) * Fe.GetAt(2, 2) - Fe.GetAt(2, 0) * Fe.GetAt(0, 2));
-				JFinvT.SetAt(1, 2, Fe.GetAt(2, 0) * Fe.GetAt(0, 1) - Fe.GetAt(0, 0) * Fe.GetAt(2, 1));
-				JFinvT.SetAt(2, 0, Fe.GetAt(0, 1) * Fe.GetAt(1, 2) - Fe.GetAt(1, 1) * Fe.GetAt(0, 2));
-				JFinvT.SetAt(2, 1, Fe.GetAt(1, 0) * Fe.GetAt(0, 2) - Fe.GetAt(0, 0) * Fe.GetAt(1, 2));
-				JFinvT.SetAt(2, 2, Fe.GetAt(0, 0) * Fe.GetAt(1, 1) - Fe.GetAt(1, 0) * Fe.GetAt(0, 1));
-
-				Chaos::PMatrix<T, 3, 3> JFinv = JFinvT.GetTransposed();
-
-				if (local_index == 0) {
-					// first mu term:
-					T DmInvsum = T(0);
-					for (int32 nu = 0; nu < 3; nu++) {
-						T localDmsum = T(0);
-						for (int32 k = 0; k < 3; k++) {
-							localDmsum += DmInv.GetAt(k, nu);
-						}
-						DmInvsum += localDmsum * localDmsum;
-					}
-					for (int32 alpha = 0; alpha < 3; alpha++) {
-						final_hessian.SetAt(alpha, alpha, final_hessian.GetAt(alpha, alpha) + Coeff * T(2) * MuElementArray[ElementIndex] * DmInvsum);
-					}
-
-					// second lambda term:
-					Chaos::PMatrix<T, 3, 3> DmInvJFinv = JFinv * DmInv;
-					Chaos::TVector<T, 3> l((T)0.);
-					for (int32 alpha = 0; alpha < 3; alpha++) {
-						for (int32 k = 0; k < 3; k++) {
-							l[alpha] += DmInvJFinv.GetAt(k, alpha);
-						}
-					}
-					for (int32 alpha = 0; alpha < 3; alpha++)
-					{
-						final_hessian.SetRow(alpha, final_hessian.GetRow(alpha) + Coeff * LambdaElementArray[ElementIndex] * l[alpha] * l);
-					}
+				ComputeHessianHelper(Fe, DmInv, MuElementArray[ElementIndex], LambdaElementArray[ElementIndex], local_index, Coeff, final_hessian);
 
-				}
-				else {
-					// first mu term:
-					T DmInvsum = T(0);
-					for (int32 nu = 0; nu < 3; nu++)
-					{
-						DmInvsum += DmInv.GetAt(local_index - 1, nu) * DmInv.GetAt(local_index - 1, nu);
-					}
-					for (int32 alpha = 0; alpha < 3; alpha++)
-					{
-						final_hessian.SetAt(alpha, alpha, final_hessian.GetAt(alpha, alpha) + Coeff * T(2) * MuElementArray[ElementIndex] * DmInvsum);
-					}
-
-					// lambda term:
-					Chaos::PMatrix<T, 3, 3> DmInvJFinv = JFinv * DmInv;
-					Chaos::TVector<T, 3> l((T)0.);
-					for (int32 alpha = 0; alpha < 3; alpha++) {
-						l[alpha] = DmInvJFinv.GetAt(local_index - 1, alpha);
-					}
-					for (int32 alpha = 0; alpha < 3; alpha++)
-					{
-						final_hessian.SetRow(alpha, final_hessian.GetRow(alpha) + Coeff * LambdaElementArray[ElementIndex] * l[alpha] * l);
-					}
-				}
 			}
 			return final_hessian;
 		}
@@ -274,6 +223,77 @@ namespace Chaos::Softs
 			ParticlesPerColor = ComputeNodalColoring(MeshConstraints, Particles, ParticleStartIndex, ParticleEndIndex, IncidentElements, IncidentElementsLocal);
 		}
 
+		 virtual void InitializeLambdas()
+		 {
+			 ComputeStress = [](const Chaos::PMatrix<T, 3, 3>& Fe, const T mu, const T lambda, Chaos::PMatrix<T, 3, 3>& P)
+			 {
+				 PCorotated(Fe, mu, lambda, P);
+			 };
+			 ComputeHessianHelper = [](const Chaos::PMatrix<T, 3, 3>& Fe, const Chaos::PMatrix<T, 3, 3>& DmInv, const T mu, const T lambda, const int32 local_index, const T Coeff, Chaos::PMatrix<T, 3, 3>& final_hessian) 
+			 {
+				 Chaos::PMatrix<T, 3, 3> JFinvT((T)0.);
+				 JFinvT.SetAt(0, 0, Fe.GetAt(1, 1) * Fe.GetAt(2, 2) - Fe.GetAt(2, 1) * Fe.GetAt(1, 2));
+				 JFinvT.SetAt(0, 1, Fe.GetAt(2, 0) * Fe.GetAt(1, 2) - Fe.GetAt(1, 0) * Fe.GetAt(2, 2));
+				 JFinvT.SetAt(0, 2, Fe.GetAt(1, 0) * Fe.GetAt(2, 1) - Fe.GetAt(2, 0) * Fe.GetAt(1, 1));
+				 JFinvT.SetAt(1, 0, Fe.GetAt(2, 1) * Fe.GetAt(0, 2) - Fe.GetAt(0, 1) * Fe.GetAt(2, 2));
+				 JFinvT.SetAt(1, 1, Fe.GetAt(0, 0) * Fe.GetAt(2, 2) - Fe.GetAt(2, 0) * Fe.GetAt(0, 2));
+				 JFinvT.SetAt(1, 2, Fe.GetAt(2, 0) * Fe.GetAt(0, 1) - Fe.GetAt(0, 0) * Fe.GetAt(2, 1));
+				 JFinvT.SetAt(2, 0, Fe.GetAt(0, 1) * Fe.GetAt(1, 2) - Fe.GetAt(1, 1) * Fe.GetAt(0, 2));
+				 JFinvT.SetAt(2, 1, Fe.GetAt(1, 0) * Fe.GetAt(0, 2) - Fe.GetAt(0, 0) * Fe.GetAt(1, 2));
+				 JFinvT.SetAt(2, 2, Fe.GetAt(0, 0) * Fe.GetAt(1, 1) - Fe.GetAt(1, 0) * Fe.GetAt(0, 1));
+
+				 Chaos::PMatrix<T, 3, 3> JFinv = JFinvT.GetTransposed();
+
+				 if (local_index == 0) {
+					 T DmInvsum = T(0);
+					 for (int32 nu = 0; nu < 3; nu++) {
+						 T localDmsum = T(0);
+						 for (int32 k = 0; k < 3; k++) {
+							 localDmsum += DmInv.GetAt(k, nu);
+						 }
+						 DmInvsum += localDmsum * localDmsum;
+					 }
+					 for (int32 alpha = 0; alpha < 3; alpha++) {
+						 final_hessian.SetAt(alpha, alpha, final_hessian.GetAt(alpha, alpha) + Coeff * T(2) * mu * DmInvsum);
+					 }
+
+					 Chaos::PMatrix<T, 3, 3> DmInvJFinv = JFinv * DmInv;
+					 Chaos::TVector<T, 3> l((T)0.);
+					 for (int32 alpha = 0; alpha < 3; alpha++) {
+						 for (int32 k = 0; k < 3; k++) {
+							 l[alpha] += DmInvJFinv.GetAt(k, alpha);
+						 }
+					 }
+					 for (int32 alpha = 0; alpha < 3; alpha++)
+					 {
+						 final_hessian.SetRow(alpha, final_hessian.GetRow(alpha) + Coeff * lambda * l[alpha] * l);
+					 }
+
+				 }
+				 else {
+					 T DmInvsum = T(0);
+					 for (int32 nu = 0; nu < 3; nu++)
+					 {
+						 DmInvsum += DmInv.GetAt(local_index - 1, nu) * DmInv.GetAt(local_index - 1, nu);
+					 }
+					 for (int32 alpha = 0; alpha < 3; alpha++)
+					 {
+						 final_hessian.SetAt(alpha, alpha, final_hessian.GetAt(alpha, alpha) + Coeff * T(2) * mu * DmInvsum);
+					 }
+
+					 Chaos::PMatrix<T, 3, 3> DmInvJFinv = JFinv * DmInv;
+					 Chaos::TVector<T, 3> l((T)0.);
+					 for (int32 alpha = 0; alpha < 3; alpha++) {
+						 l[alpha] = DmInvJFinv.GetAt(local_index - 1, alpha);
+					 }
+					 for (int32 alpha = 0; alpha < 3; alpha++)
+					 {
+						 final_hessian.SetRow(alpha, final_hessian.GetRow(alpha) + Coeff * lambda * l[alpha] * l);
+					 }
+				 }
+			 };
+		 }
+
 	private:
 
 		TVector<T, 3> ComputeDeltax(const int32 p, const int32 IncidentIndex, const ParticleType& Particles,
@@ -297,6 +317,8 @@ namespace Chaos::Softs
 
 		
 
+	protected:	
+
 		TArray<int32> Particle2Incident;
 		TArray<TArray<int32>> IncidentElements;
 		TArray<TArray<int32>> IncidentElementsLocal;
@@ -311,6 +333,8 @@ namespace Chaos::Softs
 		bool bDoSOR = true;
 		T OmegaSOR = T(1.6);
 		mutable int32 CurrentIt = 0;
+		TFunction<void(const Chaos::PMatrix<T, 3, 3>&, const T, const T, Chaos::PMatrix<T, 3, 3>& )> ComputeStress;
+		TFunction<void(const Chaos::PMatrix<T, 3, 3>&, const Chaos::PMatrix<T, 3, 3>&, const T, const T, const int32, const T, Chaos::PMatrix<T, 3, 3>&)> ComputeHessianHelper;
 	};
 
 
diff --git a/Engine/Source/Runtime/Experimental/Chaos/Public/Chaos/Deformable/GaussSeidelNeohookeanConstraints.h b/Engine/Source/Runtime/Experimental/Chaos/Public/Chaos/Deformable/GaussSeidelNeohookeanConstraints.h
new file mode 100644
--- /dev/null
+++ b/Engine/Source/Runtime/Experimental/Chaos/Public/Chaos/Deformable/GaussSeidelNeohookeanConstraints.h
@@ -0,0 +1,139 @@
+// Copyright Epic Games, Inc. All Rights Reserved.
+#pragma once
+
+#include "Chaos/PBDSpringConstraintsBase.h"
+#include "Chaos/Deformable/GaussSeidelCorotatedConstraints.h"
+#include "ChaosStats.h"
+#include "Chaos/ImplicitQRSVD.h"
+#include "Chaos/GraphColoring.h"
+#include "Chaos/NewtonCorotatedCache.h"
+#include "Chaos/NeohookeanModel.h"
+#include "Chaos/Framework/Parallel.h"
+
+
+namespace Chaos::Softs
+{
+	template <typename T, typename ParticleType>
+	class FGaussSeidelNeohookeanConstraints : public FGaussSeidelCorotatedConstraints <T, ParticleType>
+	{
+		typedef FGaussSeidelCorotatedConstraints<T, ParticleType> GSBase;
+		using GSBase::ParticlesPerColor;
+		using GSBase::IncidentElements;
+		using GSBase::IncidentElementsLocal;
+		using GSBase::ParticleStartIndex;
+		using GSBase::LocalNewtonTol;
+		using GSBase::Particle2Incident;
+		using GSBase::bDoQuasistatics;
+		using GSBase::ComputeStress;
+		using GSBase::ComputeHessianHelper;
+		typedef FXPBDCorotatedConstraints<T, ParticleType> Base;
+		using Base::MeshConstraints;
+		using Base::LambdaArray;
+		using Base::Measure;
+		using Base::DmInverse;
+		using Base::MuElementArray;
+		using Base::LambdaElementArray;
+		using Base::CorotatedParams;
+
+	public:
+		FGaussSeidelNeohookeanConstraints(
+			const ParticleType& InParticles,
+			const TArray<TVector<int32, 4>>& InMesh,
+			const TArray<T>& EMeshArray,
+			const TArray<T>& NuMeshArray,
+			TArray<T>&& AlphaJMeshArray,
+			TArray<TArray<int32>>&& IncidentElementsIn,
+			TArray<TArray<int32>>&& IncidentElementsLocalIn,
+			const int32 ParticleStartIndexIn,
+			const int32 ParticleEndIndexIn,
+			const bool bDoQuasistaticsIn = false,
+			const bool bDoSORIn = true,
+			const T InOmegaSOR = (T)1.6,
+			const FDeformableXPBDCorotatedParams& InParams = FDeformableXPBDCorotatedParams(),
+			const T& NuMesh = (T).3,
+			const bool bRecordMetricIn = false
+		)
+			: GSBase(InParticles, InMesh, EMeshArray, NuMeshArray, MoveTemp(AlphaJMeshArray), MoveTemp(IncidentElementsIn), MoveTemp(IncidentElementsLocalIn), ParticleStartIndexIn, ParticleEndIndexIn, bDoQuasistaticsIn, bDoSORIn, InOmegaSOR, InParams, NuMesh, bRecordMetricIn)
+		{
+			InitializeLambdas();
+		}
+
+		virtual ~FGaussSeidelNeohookeanConstraints() {}
+
+	protected:
+
+		void InitializeLambdas() override
+		{
+			ComputeStress = [](const Chaos::PMatrix<T, 3, 3>& Fe, const T mu, const T lambda, Chaos::PMatrix<T, 3, 3>& P)
+			{
+				PNeohookeanMM(Fe, mu, lambda, P);
+			};
+			ComputeHessianHelper = [](const Chaos::PMatrix<T, 3, 3>& Fe, const Chaos::PMatrix<T, 3, 3>& DmInv, const T mu, const T lambda, const int32 local_index, const T Coeff, Chaos::PMatrix<T, 3, 3>& final_hessian)
+			{
+				T LambdaHat = mu + lambda;
+				Chaos::PMatrix<T, 3, 3> JFinvT((T)0.);
+				JFinvT.SetAt(0, 0, Fe.GetAt(1, 1) * Fe.GetAt(2, 2) - Fe.GetAt(2, 1) * Fe.GetAt(1, 2));
+				JFinvT.SetAt(0, 1, Fe.GetAt(2, 0) * Fe.GetAt(1, 2) - Fe.GetAt(1, 0) * Fe.GetAt(2, 2));
+				JFinvT.SetAt(0, 2, Fe.GetAt(1, 0) * Fe.GetAt(2, 1) - Fe.GetAt(2, 0) * Fe.GetAt(1, 1));
+				JFinvT.SetAt(1, 0, Fe.GetAt(2, 1) * Fe.GetAt(0, 2) - Fe.GetAt(0, 1) * Fe.GetAt(2, 2));
+				JFinvT.SetAt(1, 1, Fe.GetAt(0, 0) * Fe.GetAt(2, 2) - Fe.GetAt(2, 0) * Fe.GetAt(0, 2));
+				JFinvT.SetAt(1, 2, Fe.GetAt(2, 0) * Fe.GetAt(0, 1) - Fe.GetAt(0, 0) * Fe.GetAt(2, 1));
+				JFinvT.SetAt(2, 0, Fe.GetAt(0, 1) * Fe.GetAt(1, 2) - Fe.GetAt(1, 1) * Fe.GetAt(0, 2));
+				JFinvT.SetAt(2, 1, Fe.GetAt(1, 0) * Fe.GetAt(0, 2) - Fe.GetAt(0, 0) * Fe.GetAt(1, 2));
+				JFinvT.SetAt(2, 2, Fe.GetAt(0, 0) * Fe.GetAt(1, 1) - Fe.GetAt(1, 0) * Fe.GetAt(0, 1));
+
+				Chaos::PMatrix<T, 3, 3> JFinv = JFinvT.GetTransposed();
+
+				if (local_index == 0) {
+					T DmInvsum = T(0);
+					for (int32 nu = 0; nu < 3; nu++) {
+						T localDmsum = T(0);
+						for (int32 k = 0; k < 3; k++) {
+							localDmsum += DmInv.GetAt(k, nu);
+						}
+						DmInvsum += localDmsum * localDmsum;
+					}
+					for (int32 alpha = 0; alpha < 3; alpha++) {
+						final_hessian.SetAt(alpha, alpha, final_hessian.GetAt(alpha, alpha) + Coeff * mu * DmInvsum);
+					}
+
+					Chaos::PMatrix<T, 3, 3> DmInvJFinv = JFinv * DmInv;
+					Chaos::TVector<T, 3> l((T)0.);
+					for (int32 alpha = 0; alpha < 3; alpha++) {
+						for (int32 k = 0; k < 3; k++) {
+							l[alpha] += DmInvJFinv.GetAt(k, alpha);
+						}
+					}
+					for (int32 alpha = 0; alpha < 3; alpha++)
+					{
+						final_hessian.SetRow(alpha, final_hessian.GetRow(alpha) + Coeff * LambdaHat * l[alpha] * l);
+					}
+
+				}
+				else {
+					T DmInvsum = T(0);
+					for (int32 nu = 0; nu < 3; nu++)
+					{
+						DmInvsum += DmInv.GetAt(local_index - 1, nu) * DmInv.GetAt(local_index - 1, nu);
+					}
+					for (int32 alpha = 0; alpha < 3; alpha++)
+					{
+						final_hessian.SetAt(alpha, alpha, final_hessian.GetAt(alpha, alpha) + Coeff * mu * DmInvsum);
+					}
+
+					Chaos::PMatrix<T, 3, 3> DmInvJFinv = JFinv * DmInv;
+					Chaos::TVector<T, 3> l((T)0.);
+					for (int32 alpha = 0; alpha < 3; alpha++) {
+						l[alpha] = DmInvJFinv.GetAt(local_index - 1, alpha);
+					}
+					for (int32 alpha = 0; alpha < 3; alpha++)
+					{
+						final_hessian.SetRow(alpha, final_hessian.GetRow(alpha) + Coeff * LambdaHat * l[alpha] * l);
+					}
+				}
+			};
+		}
+
+	};
+
+}// End namespace Chaos::Softs
diff --git a/Engine/Source/Runtime/Experimental/Chaos/Public/Chaos/NeohookeanModel.h b/Engine/Source/Runtime/Experimental/Chaos/Public/Chaos/NeohookeanModel.h
new file mode 100644
--- /dev/null
+++ b/Engine/Source/Runtime/Experimental/Chaos/Public/Chaos/NeohookeanModel.h
@@ -0,0 +1,14 @@
+// Copyright Epic Games, Inc. All Rights Reserved.
+#pragma once 
+#include "Chaos/Matrix.h"
+
+namespace Chaos::Softs
+{
+template <typename T>
+T PsiNeohookeanMM(const Chaos::PMatrix<T, 3, 3>& F, const T mu, const T lambda);
+
+template <typename T>
+void PNeohookeanMM(const Chaos::PMatrix<T, 3, 3>& Fe, const T mu, const T lambda, Chaos::PMatrix<T, 3, 3>& P);
+
+}
+
