Iris - Fixed issue with subobjects fitered out and then allowed to be replicated again. Before they wouldn't replicate unless it or another subobject had dirty state.

diff --git a/Engine/Plugins/Runtime/ReplicationSystemTestPlugin/Source/Private/Tests/ReplicationSystem/Filtering/TestGroups.cpp b/Engine/Plugins/Runtime/ReplicationSystemTestPlugin/Source/Private/Tests/ReplicationSystem/Filtering/TestGroups.cpp
--- a/Engine/Plugins/Runtime/ReplicationSystemTestPlugin/Source/Private/Tests/ReplicationSystem/Filtering/TestGroups.cpp
+++ b/Engine/Plugins/Runtime/ReplicationSystemTestPlugin/Source/Private/Tests/ReplicationSystem/Filtering/TestGroups.cpp
@@ -748,9 +748,7 @@ UE_NET_TEST_FIXTURE(FTestGroupsFixture, GroupFilterAPISubObjectIsFilteredOutWith
 	// Verify that both object and subobject are enabled
 	{
 		VerifyObjectFilterStatus(ServerHandle, Client->ConnectionIdOnServer, true);
-		
-		//$IRIS TODO: this fails since the subobject is not found on the client.  Needs to be looked into.
-		//VerifyObjectFilterStatus(ServerSubObjectHandle, Client->ConnectionIdOnServer, true);
+		VerifyObjectFilterStatus(ServerSubObjectHandle, Client->ConnectionIdOnServer, true);
 	}
 }
 
diff --git a/Engine/Source/Runtime/Experimental/Iris/Core/Private/Iris/ReplicationSystem/ReplicationWriter.cpp b/Engine/Source/Runtime/Experimental/Iris/Core/Private/Iris/ReplicationSystem/ReplicationWriter.cpp
--- a/Engine/Source/Runtime/Experimental/Iris/Core/Private/Iris/ReplicationSystem/ReplicationWriter.cpp
+++ b/Engine/Source/Runtime/Experimental/Iris/Core/Private/Iris/ReplicationSystem/ReplicationWriter.cpp
@@ -467,6 +467,20 @@ void FReplicationWriter::StartReplication(uint32 InternalIndex)
 	UE_LOG_REPLICATIONWRITER_CONN(TEXT("ReplicationWriter.StartReplication for ( InternalIndex: %u ) %s"), InternalIndex, ToCStr(Data.RefHandle.ToString()));
 
 	ObjectsWithDirtyChanges.SetBit(InternalIndex);
+
+	// Subobject needs to mark its owner as dirty as the subobject could have been filtered out and now allowed to replicate again.
+	if (Info.IsSubObject)
+	{
+		FNetRefHandleManager::FReplicatedObjectData& ObjectData = NetRefHandleManager->GetReplicatedObjectDataNoCheck(InternalIndex);
+		if (ensureAlways(ObjectData.SubObjectRootIndex != FNetRefHandleManager::InvalidInternalIndex))
+		{
+			// If owner is not pending destroy we mark the state of the SubObject to SubObjectPendingDestroy and mark owner as having dirty subobjects which will 
+			// destroy the subobject using the replicated state path of the owner.
+			FReplicationInfo& OwnerInfo = ReplicatedObjects[ObjectData.SubObjectRootIndex];
+			ObjectsWithDirtyChanges.SetBit(ObjectData.SubObjectRootIndex);
+			OwnerInfo.HasDirtySubObjects = 1U;
+		}
+	}
 }
 
 void FReplicationWriter::StopReplication(uint32 InternalIndex)
