Iris - Don't ensure when tearing off subobjects and root object is no longer scopable.

#rnx

diff --git a/Engine/Source/Runtime/Experimental/Iris/Core/Private/Iris/ReplicationSystem/ReplicationOperationsInternal.cpp b/Engine/Source/Runtime/Experimental/Iris/Core/Private/Iris/ReplicationSystem/ReplicationOperationsInternal.cpp
--- a/Engine/Source/Runtime/Experimental/Iris/Core/Private/Iris/ReplicationSystem/ReplicationOperationsInternal.cpp
+++ b/Engine/Source/Runtime/Experimental/Iris/Core/Private/Iris/ReplicationSystem/ReplicationOperationsInternal.cpp
@@ -110,8 +110,8 @@ uint32 FReplicationInstanceOperationsInternal::CopyObjectStateData(FNetBitStream
 			if (const uint32 SubObjectOwnerIndex = Object.SubObjectRootIndex)
 			{
 				const bool bIsOwnerScopable = NetRefHandleManager.IsScopableIndex(SubObjectOwnerIndex);
-				// Dependent objects should not ensure if the owner isn't scopable.
-				ensureMsgf(bIsOwnerScopable || Object.IsDependentObject(), TEXT("SubObject ( InternaIndex: %u ) with NetRefHandle (Id=%u) is trying to dirty parent ( InternalIndex: %u ) not in scope."), InternalIndex, Object.RefHandle.GetId(), SubObjectOwnerIndex);
+				// Dependent objects should not ensure if the owner isn't scopable. Subobjects pending tear off is ok too.
+				ensureAlwaysMsgf(bIsOwnerScopable || Object.IsDependentObject() || Object.bTearOff, TEXT("SubObject ( InternaIndex: %u ) with NetRefHandle (Id=%u) is trying to dirty parent ( InternalIndex: %u ) not in scope."), InternalIndex, Object.RefHandle.GetId(), SubObjectOwnerIndex);
 				if (bIsOwnerScopable)
 				{
 					// Do we want to control this separately for subobjects? Or should they respect the setting on the owner?
