Add IsValidIndex method to managed arrays

diff --git a/Engine/Source/Runtime/Experimental/Chaos/Public/GeometryCollection/ManagedArray.h b/Engine/Source/Runtime/Experimental/Chaos/Public/GeometryCollection/ManagedArray.h
--- a/Engine/Source/Runtime/Experimental/Chaos/Public/GeometryCollection/ManagedArray.h
+++ b/Engine/Source/Runtime/Experimental/Chaos/Public/GeometryCollection/ManagedArray.h
@@ -525,6 +525,16 @@ public:
 		return Num;
 	}
 
+	/**
+	* return true if index is in array range.
+	*
+	* @param Index Index to check.
+	*/
+	FORCEINLINE bool IsValidIndex(int32 Index) const
+	{
+		return Array.IsValidIndex(Index);
+	}
+
 	/**
 	* Checks if index is in array range.
 	*
