Made FHeightField::FData::Serialize() avoid using TArray<TUniqueObj<FBoxFloat3>>, as FBoxFloat3 is immovable and shouldn't be used by movable allocators.

diff --git a/Engine/Source/Runtime/Experimental/Chaos/Public/Chaos/HeightField.h b/Engine/Source/Runtime/Experimental/Chaos/Public/Chaos/HeightField.h
--- a/Engine/Source/Runtime/Experimental/Chaos/Public/Chaos/HeightField.h
+++ b/Engine/Source/Runtime/Experimental/Chaos/Public/Chaos/HeightField.h
@@ -9,6 +9,7 @@
 #include "ChaosArchive.h"
 #include "Serialization/MemoryWriter.h"
 #include "Math/NumericLimits.h"
+#include "Templates/UniqueObj.h"
 #include "Templates/UnrealTypeTraits.h"
 #include "UniformGrid.h"
 #include "Utilities.h"
@@ -641,7 +642,7 @@ namespace Chaos
 					{
 						// todo(chaos) this may not matter if the Vector types are handling serialization properly 
 						// legacy, need to keep the inner box type as float ( not FReal ) 
-						TArray<FBoxFloat3> CellBounds;
+						TArray<TUniqueObj<FBoxFloat3>> CellBounds;
 						Ar << CellBounds;
 					}
 					else if(Ar.CustomVer(FExternalPhysicsCustomObjectVersion::GUID) < FExternalPhysicsCustomObjectVersion::HeightfieldUsesHeightsDirectly)
