Iris
* Enable culling of non-relevant objects from the poll+copy phases.  Controlled via Net.Iris.CullNonRelevant
* Now accumulating dirty flags for dirty objects not cleaned after getting culled.
* DirtyObjectTracker now tracking objects set dirty in the current tick and objects not cleaned in previous frames.
* Now tracking list of CleanedObjects during the frame and only resetting dirty status of the cleaned objects.

#rnx

diff --git a/Engine/Plugins/Runtime/ReplicationSystemTestPlugin/Source/Private/Tests/ReplicationSystem/TestDirtyNetObjectTracker.cpp b/Engine/Plugins/Runtime/ReplicationSystemTestPlugin/Source/Private/Tests/ReplicationSystem/TestDirtyNetObjectTracker.cpp
--- a/Engine/Plugins/Runtime/ReplicationSystemTestPlugin/Source/Private/Tests/ReplicationSystem/TestDirtyNetObjectTracker.cpp
+++ b/Engine/Plugins/Runtime/ReplicationSystemTestPlugin/Source/Private/Tests/ReplicationSystem/TestDirtyNetObjectTracker.cpp
@@ -78,11 +78,54 @@ UE_NET_TEST_FIXTURE(FDirtyNetObjectTrackerTestFixture, CanClearDirtyObjects)
 	MarkNetObjectStateDirty(ReplicationSystemId, NetObjectIndexRangeStart);
 	MarkNetObjectStateDirty(ReplicationSystemId, NetObjectIndexRangeEnd);
 
-	DirtyNetObjectTracker->ClearDirtyNetObjects();
+	FNetBitArray CleanedObjects;
+	CleanedObjects.Init(NetObjectIndexRangeEnd+1);
 
-	FDirtyObjectsAccessor DirtyObjectsAccessor(*DirtyNetObjectTracker);
-	const FNetBitArrayView DirtyObjects = DirtyObjectsAccessor.GetDirtyNetObjects();
-	UE_NET_ASSERT_FALSE(DirtyObjects.IsAnyBitSet());
+	CleanedObjects.SetBit(NetObjectIndexRangeStart);
+	CleanedObjects.SetBit(NetObjectIndexRangeEnd);
+
+	DirtyNetObjectTracker->UpdateAccumulatedDirtyList();
+	DirtyNetObjectTracker->ClearDirtyNetObjects(MakeNetBitArrayView(CleanedObjects));
+
+	const FNetBitArrayView AccumulatedDirtyObjects = DirtyNetObjectTracker->GetAccumulatedDirtyNetObjects();
+	UE_NET_ASSERT_FALSE(AccumulatedDirtyObjects.IsAnyBitSet());
+}
+
+UE_NET_TEST_FIXTURE(FDirtyNetObjectTrackerTestFixture, DelayedDirtyBitTracking)
+{
+
+	const uint32 FirstObjectIndex = NetObjectIndexRangeStart;
+	const uint32 SecondObjectIndex = NetObjectIndexRangeStart + 1;
+	MarkNetObjectStateDirty(ReplicationSystemId, FirstObjectIndex);
+	MarkNetObjectStateDirty(ReplicationSystemId, SecondObjectIndex);
+
+	// Clean first object
+	{
+		FNetBitArray CleanedObjects;
+		CleanedObjects.Init(NetObjectIndexRangeEnd + 1);
+		CleanedObjects.SetBit(FirstObjectIndex);
+
+		DirtyNetObjectTracker->UpdateAccumulatedDirtyList();
+		DirtyNetObjectTracker->ClearDirtyNetObjects(MakeNetBitArrayView(CleanedObjects));
+	}
+
+	const FNetBitArrayView AccumulatedDirtyObjects = DirtyNetObjectTracker->GetAccumulatedDirtyNetObjects();
+
+	UE_NET_ASSERT_FALSE(AccumulatedDirtyObjects.GetBit(FirstObjectIndex));
+	UE_NET_ASSERT_TRUE(AccumulatedDirtyObjects.GetBit(SecondObjectIndex));
+
+	// Clean second object
+	{
+		FNetBitArray CleanedObjects;
+		CleanedObjects.Init(NetObjectIndexRangeEnd + 1);
+		CleanedObjects.SetBit(SecondObjectIndex);
+
+		DirtyNetObjectTracker->UpdateAccumulatedDirtyList();
+		DirtyNetObjectTracker->ClearDirtyNetObjects(MakeNetBitArrayView(CleanedObjects));
+	}
+
+	UE_NET_ASSERT_FALSE(AccumulatedDirtyObjects.GetBit(FirstObjectIndex));
+	UE_NET_ASSERT_FALSE(AccumulatedDirtyObjects.GetBit(SecondObjectIndex));
 }
 
 UE_NET_TEST(DirtyNetObjectTracker, MarkingObjectAsDirtyInNonExistingSystemDoesNotCrash)
diff --git a/Engine/Source/Runtime/Engine/Private/NetDriver.cpp b/Engine/Source/Runtime/Engine/Private/NetDriver.cpp
--- a/Engine/Source/Runtime/Engine/Private/NetDriver.cpp
+++ b/Engine/Source/Runtime/Engine/Private/NetDriver.cpp
@@ -3640,7 +3640,7 @@ void UNetDriver::NotifyActorIsTraveling(AActor* TravelingActor)
 
 void UNetDriver::ForceNetUpdate(AActor* Actor)
 {
-	// Let Replication Driver handle it if its exists
+	// Let Replication Driver handle it if it exists
 	if (ReplicationDriver)
 	{
 		ReplicationDriver->ForceNetUpdate(Actor);
diff --git a/Engine/Source/Runtime/Experimental/Iris/Core/Private/Iris/ReplicationSystem/DirtyNetObjectTracker.cpp b/Engine/Source/Runtime/Experimental/Iris/Core/Private/Iris/ReplicationSystem/DirtyNetObjectTracker.cpp
--- a/Engine/Source/Runtime/Experimental/Iris/Core/Private/Iris/ReplicationSystem/DirtyNetObjectTracker.cpp
+++ b/Engine/Source/Runtime/Experimental/Iris/Core/Private/Iris/ReplicationSystem/DirtyNetObjectTracker.cpp
@@ -28,13 +28,7 @@ namespace UE::Net::Private
 {
 
 FDirtyNetObjectTracker::FDirtyNetObjectTracker()
-: NetRefHandleManager(nullptr)
-, DirtyNetObjectContainer(nullptr)
-, ReplicationSystemId(InvalidReplicationSystemId)
-, DirtyNetObjectWordCount(0)
-, NetObjectIdRangeStart(0U)
-, NetObjectIdRangeEnd(0U)
-, NetObjectIdCount(0)
+: ReplicationSystemId(InvalidReplicationSystemId)
 {
 }
 
@@ -65,6 +59,8 @@ void FDirtyNetObjectTracker::Init(const FDirtyNetObjectTrackerInitParams& Params
 	DirtyNetObjectContainer = new StorageType[DirtyNetObjectWordCount];
 	FMemory::Memzero(DirtyNetObjectContainer, DirtyNetObjectWordCount * sizeof(StorageType));
 
+	AccumulatedDirtyNetObjects.Init(NetObjectIdCount);
+
 	AllowExternalAccess();
 
 	UE_LOG_DIRTYOBJECTTRACKER(TEXT("FDirtyNetObjectTracker::Init %u Id, Start:%u, End: %u"), ReplicationSystemId, NetObjectIdRangeStart, NetObjectIdRangeEnd);
@@ -91,6 +87,7 @@ void FDirtyNetObjectTracker::UpdateDirtyNetObjects()
 	LockExternalAccess();
 
 	bHasPolledGlobalDirtyTracker = true;
+
 	const TSet<FNetHandle>& GlobalDirtyNetObjects = FGlobalDirtyNetObjectTracker::GetDirtyNetObjects(GlobalDirtyTrackerPollHandle);
 	for (FNetHandle NetHandle : GlobalDirtyNetObjects)
 	{
@@ -103,9 +100,32 @@ void FDirtyNetObjectTracker::UpdateDirtyNetObjects()
 		}
 	}
 
+	const uint32* GlobalScopeListData = NetRefHandleManager->GetScopableInternalIndicesView().GetData();
+	uint32* AccumulatedDirtyNetObjectsData = AccumulatedDirtyNetObjects.GetData();
+
+	const uint32 NumWords = AccumulatedDirtyNetObjects.GetNumWords();
+	for (uint32 WordIndex = 0; WordIndex != NumWords; ++WordIndex)
+	{
+		// Due to objects having been marked as dirty and later removed we must make sure that all dirty objects are still in scope.
+		uint32 DirtyObjectWord = DirtyNetObjectContainer[WordIndex] & GlobalScopeListData[WordIndex];
+		DirtyNetObjectContainer[WordIndex] = DirtyObjectWord;
+
+		// Add new dirty objects to the accumulated list
+		AccumulatedDirtyNetObjectsData[WordIndex] = AccumulatedDirtyNetObjectsData[WordIndex] | DirtyObjectWord;
+	}
+
 	AllowExternalAccess();
 }
 
+void FDirtyNetObjectTracker::UpdateAccumulatedDirtyList()
+{
+	IRIS_PROFILER_SCOPE(FDirtyNetObjectTracker_UpdateDirtyNetObjects)
+
+	FNetBitArrayView DirtyObjectsThisFrame(DirtyNetObjectContainer, NetObjectIdCount);
+
+	MakeNetBitArrayView(AccumulatedDirtyNetObjects).Combine(DirtyObjectsThisFrame, FNetBitArrayView::OrOp);
+}
+
 void FDirtyNetObjectTracker::MarkNetObjectDirty(uint32 NetObjectIndex)
 {
 #if UE_NET_THREAD_SAFETY_CHECK
@@ -138,7 +158,7 @@ void FDirtyNetObjectTracker::AllowExternalAccess()
 #endif
 }
 
-FNetBitArrayView FDirtyNetObjectTracker::GetDirtyNetObjects() const
+FNetBitArrayView FDirtyNetObjectTracker::GetDirtyNetObjectsThisFrame()
 {
 #if UE_NET_THREAD_SAFETY_CHECK
 	checkf(!bIsExternalAccessAllowed, TEXT("Cannot access the DirtyNetObjects bitarray unless its locked for multithread access."));
@@ -146,7 +166,7 @@ FNetBitArrayView FDirtyNetObjectTracker::GetDirtyNetObjects() const
 	return FNetBitArrayView(DirtyNetObjectContainer, NetObjectIdCount);
 }
 
-void FDirtyNetObjectTracker::ClearDirtyNetObjects()
+void FDirtyNetObjectTracker::ClearDirtyNetObjects(const FNetBitArrayView& CleanNetObjects)
 {
 	LockExternalAccess();
 
@@ -156,8 +176,12 @@ void FDirtyNetObjectTracker::ClearDirtyNetObjects()
 		FGlobalDirtyNetObjectTracker::ResetDirtyNetObjects(GlobalDirtyTrackerPollHandle);
 	}
 
+	// Keep dirty flags for objects not cleaned this frame.
+	MakeNetBitArrayView(AccumulatedDirtyNetObjects).Combine(CleanNetObjects, FNetBitArrayView::AndNotOp);
+
+	// Clear the current frame dirty objects
 	FMemory::Memzero(DirtyNetObjectContainer, DirtyNetObjectWordCount*sizeof(StorageType));
-	
+
 	AllowExternalAccess();
 
 	std::atomic_thread_fence(std::memory_order_seq_cst);
diff --git a/Engine/Source/Runtime/Experimental/Iris/Core/Private/Iris/ReplicationSystem/DirtyNetObjectTracker.h b/Engine/Source/Runtime/Experimental/Iris/Core/Private/Iris/ReplicationSystem/DirtyNetObjectTracker.h
--- a/Engine/Source/Runtime/Experimental/Iris/Core/Private/Iris/ReplicationSystem/DirtyNetObjectTracker.h
+++ b/Engine/Source/Runtime/Experimental/Iris/Core/Private/Iris/ReplicationSystem/DirtyNetObjectTracker.h
@@ -35,17 +35,23 @@ public:
 
 	void Init(const FDirtyNetObjectTrackerInitParams& Params);
 
-	/* Update dirty objects with the set of globally marked dirty objects. **/
+	/** Update dirty objects with the set of globally marked dirty objects. */
 	void UpdateDirtyNetObjects();
 
+	/** Add all the current frame dirty objects set into the accumulated list */
+	void UpdateAccumulatedDirtyList();
+
 	/** Set safety permissions so no one can write in the bit array via the public methods */
 	void LockExternalAccess();
 
 	/** Release safety permissions and allow to write in the bit array via the public methods */
 	void AllowExternalAccess();
 
-	/** Reset the global and local dirty objects lists */
-	void ClearDirtyNetObjects();
+	/** Reset the global and local dirty objects lists for those objects that are now clean */
+	void ClearDirtyNetObjects(const FNetBitArrayView& CleanNetObjects);
+
+	/** Returns the list of objects that are dirty this frame or were dirty in previous frames but not cleaned up at that time. */
+	const FNetBitArrayView GetAccumulatedDirtyNetObjects() const { return MakeNetBitArrayView(AccumulatedDirtyNetObjects); }
 
 private:
 	friend IRISCORE_API void MarkNetObjectStateDirty(uint32 ReplicationSystemId, uint32 NetObjectIndex);
@@ -59,16 +65,26 @@ private:
 	void MarkNetObjectDirty(uint32 NetObjectIndex);
 
 	/** Can only be accessed via FDirtyObjectsAccessor */
-	FNetBitArrayView GetDirtyNetObjects() const;
+	FNetBitArrayView GetDirtyNetObjectsThisFrame();
+
+private:
+
+	// Dirty objects that persist across frames.
+	FNetBitArray AccumulatedDirtyNetObjects;
 
-	const FNetRefHandleManager* NetRefHandleManager;
-	StorageType* DirtyNetObjectContainer;
+	// List of objects set to be dirty this frame. Is always reset at the end of the net tick flush
+	StorageType* DirtyNetObjectContainer = nullptr;
+
+	const FNetRefHandleManager* NetRefHandleManager = nullptr;
+	
 	FGlobalDirtyNetObjectTracker::FPollHandle GlobalDirtyTrackerPollHandle;
+
 	uint32 ReplicationSystemId;
-	uint32 DirtyNetObjectWordCount;
-	uint32 NetObjectIdRangeStart;
-	uint32 NetObjectIdRangeEnd;
-	uint32 NetObjectIdCount;
+
+	uint32 DirtyNetObjectWordCount = 0;
+	uint32 NetObjectIdRangeStart = 0;
+	uint32 NetObjectIdRangeEnd = 0;
+	uint32 NetObjectIdCount = 0;
 	
 	bool bHasPolledGlobalDirtyTracker = false;
 
@@ -77,6 +93,9 @@ private:
 #endif
 };
 
+/**
+ * Gives access to the list of dirty objects while detecting non-thread safe access to it.
+ */
 class FDirtyObjectsAccessor
 {
 public:
@@ -91,8 +110,8 @@ public:
 		DirtyNetObjectTracker.AllowExternalAccess();
 	}
 
-	FNetBitArrayView GetDirtyNetObjects()				{ return DirtyNetObjectTracker.GetDirtyNetObjects(); }
-	const FNetBitArrayView GetDirtyNetObjects() const	{ return DirtyNetObjectTracker.GetDirtyNetObjects(); }
+	FNetBitArrayView GetDirtyNetObjects()				{ return DirtyNetObjectTracker.GetDirtyNetObjectsThisFrame(); }
+	const FNetBitArrayView GetDirtyNetObjects() const	{ return DirtyNetObjectTracker.GetDirtyNetObjectsThisFrame(); }
 
 private:
 	FDirtyNetObjectTracker& DirtyNetObjectTracker;
diff --git a/Engine/Source/Runtime/Experimental/Iris/Core/Private/Iris/ReplicationSystem/Filtering/ReplicationFiltering.cpp b/Engine/Source/Runtime/Experimental/Iris/Core/Private/Iris/ReplicationSystem/Filtering/ReplicationFiltering.cpp
--- a/Engine/Source/Runtime/Experimental/Iris/Core/Private/Iris/ReplicationSystem/Filtering/ReplicationFiltering.cpp
+++ b/Engine/Source/Runtime/Experimental/Iris/Core/Private/Iris/ReplicationSystem/Filtering/ReplicationFiltering.cpp
@@ -22,6 +22,10 @@ namespace UE::Net::Private
 
 DEFINE_LOG_CATEGORY_STATIC(LogIrisFiltering, Log, All);
 
+
+bool bCVarRepFilterCullNonRelevant = true;
+static FAutoConsoleVariableRef CVarRepFilterCullNonRelevant(TEXT("Net.Iris.CullNonRelevant"), bCVarRepFilterCullNonRelevant, TEXT("When enabled will cull replicated actors that are not relevant to any client."), ECVF_Default);
+
 FName GetStaticFilterName(FNetObjectFilterHandle Filter)
 {
 	switch (Filter)
@@ -297,32 +301,33 @@ void FReplicationFiltering::FilterPrePoll()
 
 	UpdateSubObjectFilters();
 
-	FDirtyObjectsAccessor DirtyObjectsAccessor(ReplicationSystem->GetReplicationSystemInternal()->GetDirtyNetObjectTracker());
-	const FNetBitArrayView DirtyObjects = DirtyObjectsAccessor.GetDirtyNetObjects();
-
 	if (HasRawFilters())
 	{
-		UpdateDynamicFilters(ENetFilterType::PrePoll_Raw, DirtyObjects);
+		UpdateDynamicFilters(ENetFilterType::PrePoll_Raw);
 	}
 	else
 	{
 		// Dynamic filters are responsible for updating ObjectsInScope. 
 		// Do it here if no filters were executed.
-		auto UpdateConnectionScope = [this](uint32 ConnectionId)
+		ValidConnections.ForAllSetBits([this](uint32 ConnectionId)
 		{
 			FPerConnectionInfo& ConnectionInfo = ConnectionInfos[ConnectionId];
 			ConnectionInfo.ObjectsInScope.Copy(ConnectionInfo.ObjectsInScopeBeforeDynamicFiltering);
-		};
-
-		ValidConnections.ForAllSetBits(UpdateConnectionScope);
+		});
 	}
 
-	//$IRIS TODO: Filter culling is still WIP instead update the relevant list so it's equal to the global scope
-	//FilterNonRelevantObjects(DirtyObjects);
-	NetRefHandleManager->GetRelevantObjectsInternalIndices().Copy(MakeNetBitArrayView(NetRefHandleManager->GetScopableInternalIndices()));
+	if (bCVarRepFilterCullNonRelevant)
+	{
+		FilterNonRelevantObjects();
+	}
+	else
+	{
+		// Make every object in the global scope part of the relevant list
+		NetRefHandleManager->GetRelevantObjectsInternalIndices().Copy(MakeNetBitArrayView(NetRefHandleManager->GetScopableInternalIndices()));
+	}
 }
 
-void FReplicationFiltering::FilterNonRelevantObjects(const FNetBitArrayView& DirtyObjects)
+void FReplicationFiltering::FilterNonRelevantObjects()
 {
 	IRIS_PROFILER_SCOPE(FReplicationFiltering_FilterNonRelevantObjects);
 
@@ -340,10 +345,11 @@ void FReplicationFiltering::FilterNonRelevantObjects(const FNetBitArrayView& Dir
 	const uint32 MaxWords = GlobalRelevantObjects.GetNumWords();
 	for (uint32 WordIndex = 0; WordIndex < MaxWords; ++WordIndex)
 	{
+		// Build the list of always relevant objects (objects that have no filters)
 		GlobalRelevantData[WordIndex] = GlobalScopeData[WordIndex] & ~(WithOwnerData[WordIndex] | ConnectionFiltersData[WordIndex] | DynamicFilteredData[WordIndex] | GroupFilteredData[WordIndex]);
 	}
 
-	// Build the list of currently relevant objects. e.g. always relevant objects + filterable objects relevant to atleast one connection.	
+	// Build the list of currently relevant objects. e.g. always relevant objects + filterable objects relevant to at least one connection.	
 	auto MergeConnectionScopes = [this, &GlobalRelevantObjects](uint32 ConnectionId)
 	{
 		const FNetBitArrayView ConnectionFilteredScope = MakeNetBitArrayView(ConnectionInfos[ConnectionId].ObjectsInScope);
@@ -353,12 +359,12 @@ void FReplicationFiltering::FilterNonRelevantObjects(const FNetBitArrayView& Dir
 	ValidConnections.ForAllSetBits(MergeConnectionScopes);
 }
 
-void FReplicationFiltering::UpdateDynamicFilters(ENetFilterType FilterPass, const FNetBitArrayView& DirtyObjects)
+void FReplicationFiltering::UpdateDynamicFilters(ENetFilterType FilterPass)
 {
 	/**
 	* Dynamic filters allows users to filter out objects based on arbitrary criteria.
 	*/
-	NotifyFiltersOfDirtyObjects(FilterPass, DirtyObjects);
+	NotifyFiltersOfDirtyObjects(FilterPass);
 
 	PreUpdateDynamicFiltering(FilterPass);
 	UpdateDynamicFiltering(FilterPass);
@@ -369,10 +375,7 @@ void FReplicationFiltering::FilterPostPoll()
 {
 	if (HasFragmentFilters())
 	{
-		FDirtyObjectsAccessor DirtyObjectsAccessor(ReplicationSystem->GetReplicationSystemInternal()->GetDirtyNetObjectTracker());
-		const FNetBitArrayView DirtyObjects = DirtyObjectsAccessor.GetDirtyNetObjects();
-
-		UpdateDynamicFilters(ENetFilterType::PostPoll_FragmentBased, DirtyObjects);
+		UpdateDynamicFilters(ENetFilterType::PostPoll_FragmentBased);
 	}
 }
 
@@ -1334,20 +1337,23 @@ void FReplicationFiltering::PostUpdateDynamicFiltering(ENetFilterType FilterType
 	}
 }
 
-void FReplicationFiltering::NotifyFiltersOfDirtyObjects(ENetFilterType FilterType, const FNetBitArrayView& DirtyObjects)
+void FReplicationFiltering::NotifyFiltersOfDirtyObjects(ENetFilterType FilterType)
 {
+	FDirtyObjectsAccessor DirtyObjectsAccessor(ReplicationSystem->GetReplicationSystemInternal()->GetDirtyNetObjectTracker());
+	const FNetBitArrayView DirtyObjectsThisFrame = DirtyObjectsAccessor.GetDirtyNetObjects();
+
 	FUpdateDirtyObjectsBatchHelper BatchHelper(NetRefHandleManager, DynamicFilterInfos, FilterType);
 
 	constexpr SIZE_T MaxBatchObjectCount = FUpdateDirtyObjectsBatchHelper::Constants::MaxObjectCountPerBatch;
 	uint32 ObjectIndices[MaxBatchObjectCount];
 
 	const uint32 BitCount = ~0U;
-	for (uint32 ObjectCount, StartIndex = 0; (ObjectCount = DirtyObjects.GetSetBitIndices(StartIndex, BitCount, ObjectIndices, MaxBatchObjectCount)) > 0; )
+	for (uint32 ObjectCount, StartIndex = 0; (ObjectCount = DirtyObjectsThisFrame.GetSetBitIndices(StartIndex, BitCount, ObjectIndices, MaxBatchObjectCount)) > 0; )
 	{
 		BatchNotifyFiltersOfDirtyObjects(BatchHelper, ObjectIndices, ObjectCount);
 
 		StartIndex = ObjectIndices[ObjectCount - 1] + 1U;
-		if ((StartIndex == DirtyObjects.GetNumBits()) | (ObjectCount < MaxBatchObjectCount))
+		if ((StartIndex == DirtyObjectsThisFrame.GetNumBits()) | (ObjectCount < MaxBatchObjectCount))
 		{
 			break;
 		}
@@ -1830,6 +1836,8 @@ void FReplicationFiltering::InternalSetGroupFilterStatus(FNetObjectGroupHandle G
 		FPerObjectInfo* ConnectionState = GetPerObjectInfo(GroupInfos[GroupHandle].ConnectionStateIndex);
 		if (GetConnectionFilterStatus(*ConnectionState, ConnectionId) == ENetFilterStatus::Allow)
 		{
+			IRIS_PROFILER_SCOPE(FReplicationFiltering_InternalSetGroupFilterStatus_Disallow);
+
 			// Mark filter as active for the connection
 			SetConnectionFilterStatus(*ConnectionState, ConnectionId, ENetFilterStatus::Disallow);
 
@@ -1868,6 +1876,8 @@ void FReplicationFiltering::InternalSetGroupFilterStatus(FNetObjectGroupHandle G
 		FPerObjectInfo* ConnectionState = GetPerObjectInfo(GroupInfos[GroupHandle].ConnectionStateIndex);
 		if (GetConnectionFilterStatus(*ConnectionState, ConnectionId) == ENetFilterStatus::Disallow)
 		{
+			IRIS_PROFILER_SCOPE(FReplicationFiltering_InternalSetGroupFilterStatus_Allow);
+
 			// Mark filter as no longer being active for the connection
 			SetConnectionFilterStatus(*ConnectionState, ConnectionId, ENetFilterStatus::Allow);
 
diff --git a/Engine/Source/Runtime/Experimental/Iris/Core/Private/Iris/ReplicationSystem/Filtering/ReplicationFiltering.h b/Engine/Source/Runtime/Experimental/Iris/Core/Private/Iris/ReplicationSystem/Filtering/ReplicationFiltering.h
--- a/Engine/Source/Runtime/Experimental/Iris/Core/Private/Iris/ReplicationSystem/Filtering/ReplicationFiltering.h
+++ b/Engine/Source/Runtime/Experimental/Iris/Core/Private/Iris/ReplicationSystem/Filtering/ReplicationFiltering.h
@@ -175,13 +175,13 @@ private:
 	void UpdateGroupFiltering();
 	void UpdateSubObjectFilters();
 
-	void UpdateDynamicFilters(ENetFilterType FilterPass, const FNetBitArrayView& DirtyObjects);
+	void UpdateDynamicFilters(ENetFilterType FilterPass);
 	void PreUpdateDynamicFiltering(ENetFilterType FilterType);
 	void UpdateDynamicFiltering(ENetFilterType FilterType);
 	void PostUpdateDynamicFiltering(ENetFilterType FilterType);
 
 	/** Build the list of always relevant objects + objects that are currently relevant to at least one connection. */
-	void FilterNonRelevantObjects(const FNetBitArrayView& DirtyObjects);
+	void FilterNonRelevantObjects();
 
 	bool HasDynamicFilters() const;
 	bool HasRawFilters() const;
@@ -212,7 +212,7 @@ private:
 
 	void RemoveFromDynamicFilter(uint32 ObjectIndex, uint32 FilterIndex);
 
-	void NotifyFiltersOfDirtyObjects(ENetFilterType FilterType, const FNetBitArrayView& DirtyObjects);
+	void NotifyFiltersOfDirtyObjects(ENetFilterType FilterType);
 	void BatchNotifyFiltersOfDirtyObjects(FUpdateDirtyObjectsBatchHelper& BatchHelper, const uint32* ObjectIndices, uint32 ObjectCount);
 
 	void InvalidateBaselinesForObject(uint32 ObjectIndex, uint32 NewOwningConnectionId, uint32 PrevOwningConnectionId);
diff --git a/Engine/Source/Runtime/Experimental/Iris/Core/Private/Iris/ReplicationSystem/NetRefHandleManager.cpp b/Engine/Source/Runtime/Experimental/Iris/Core/Private/Iris/ReplicationSystem/NetRefHandleManager.cpp
--- a/Engine/Source/Runtime/Experimental/Iris/Core/Private/Iris/ReplicationSystem/NetRefHandleManager.cpp
+++ b/Engine/Source/Runtime/Experimental/Iris/Core/Private/Iris/ReplicationSystem/NetRefHandleManager.cpp
@@ -24,6 +24,7 @@ FNetRefHandleManager::FNetRefHandleManager(FReplicationProtocolManager& InReplic
 , ScopableInternalIndices(MaxActiveObjectCount)
 , PrevFrameScopableInternalIndices(MaxActiveObjectCount)
 , RelevantObjectsInternalIndices(MaxActiveObjectCount)
+, CleanedObjectsInternalIndices(MaxActiveObjectCount)
 , AssignedInternalIndices(MaxActiveObjectCount)
 , SubObjectInternalIndices(MaxActiveObjectCount)
 , DependentObjectInternalIndices(MaxActiveObjectCount)
diff --git a/Engine/Source/Runtime/Experimental/Iris/Core/Private/Iris/ReplicationSystem/NetRefHandleManager.h b/Engine/Source/Runtime/Experimental/Iris/Core/Private/Iris/ReplicationSystem/NetRefHandleManager.h
--- a/Engine/Source/Runtime/Experimental/Iris/Core/Private/Iris/ReplicationSystem/NetRefHandleManager.h
+++ b/Engine/Source/Runtime/Experimental/Iris/Core/Private/Iris/ReplicationSystem/NetRefHandleManager.h
@@ -182,8 +182,9 @@ public:
 	// Get internal index from NetHandle
 	inline FInternalNetRefIndex GetInternalIndexFromNetHandle(FNetHandle Handle) const;
 
-	// Get bitarray for all currently scopable internal indices
+	/** Get bitarray for all currently scopable internal indices */
 	const FNetBitArray& GetScopableInternalIndices() const { return ScopableInternalIndices; }
+	const FNetBitArrayView GetScopableInternalIndicesView() const { return MakeNetBitArrayView(ScopableInternalIndices); }
 
 	// Get bitarray for all internal indices that was scopable last update
 	const FNetBitArray& GetPrevFrameScopableInternalIndices() const { return PrevFrameScopableInternalIndices; }
@@ -192,6 +193,9 @@ public:
 	/** List of objects that are always relevant or currently relevant to at least one connection. */
 	FNetBitArrayView GetRelevantObjectsInternalIndices() const { return MakeNetBitArrayView(RelevantObjectsInternalIndices); }
 
+	/** List of objects on whom we copied their dirty state this frame */
+	FNetBitArrayView GetCleanedObjectsInternalIndices() const { return MakeNetBitArrayView(CleanedObjectsInternalIndices); }
+
 	// Get bitarray for all internal indices that currently are assigned
 	const FNetBitArray& GetAssignedInternalIndices() const { return AssignedInternalIndices; }
 
@@ -242,6 +246,9 @@ public:
 
 	const FRefHandleMap& GetReplicatedHandles() const { return RefHandleToInternalIndex; }
 
+	// Get the replicated object represented by a given internal index.
+	UObject* GetReplicatedObjectInstance(FInternalNetRefIndex ObjectIndex) const { return ReplicatedInstances[ObjectIndex]; }
+
 	const TArray<UObject*>& GetReplicatedInstances() const { return ReplicatedInstances; }
 
 	void AddReferencedObjects(FReferenceCollector& Collector);
@@ -304,9 +311,12 @@ private:
 	// Which internal indices were used last net frame. This can be used to find out which ones are new and deleted this frame. 
 	FNetBitArray PrevFrameScopableInternalIndices;
 
-	/** The contains the ScopableInternalIndices list minus filtered objects that are not relevant to any connection this frame. */
+	/** This contains the ScopableInternalIndices list minus filtered objects that are not relevant to any connection this frame. */
 	FNetBitArray RelevantObjectsInternalIndices;
 
+	/** List of objects that were dirty but whom we copied their dirty data this frame */
+	FNetBitArray CleanedObjectsInternalIndices;
+
 	// Bitset containing all internal indices that are assigned
 	FNetBitArray AssignedInternalIndices;
 
diff --git a/Engine/Source/Runtime/Experimental/Iris/Core/Private/Iris/ReplicationSystem/ObjectPollFrequencyLimiter.cpp b/Engine/Source/Runtime/Experimental/Iris/Core/Private/Iris/ReplicationSystem/ObjectPollFrequencyLimiter.cpp
--- a/Engine/Source/Runtime/Experimental/Iris/Core/Private/Iris/ReplicationSystem/ObjectPollFrequencyLimiter.cpp
+++ b/Engine/Source/Runtime/Experimental/Iris/Core/Private/Iris/ReplicationSystem/ObjectPollFrequencyLimiter.cpp
@@ -38,7 +38,7 @@ void FObjectPollFrequencyLimiter::Deinit()
 	FrameCounters.Empty();
 }
 
-void FObjectPollFrequencyLimiter::Update(const FNetBitArrayView& ScopableObjects, const FNetBitArrayView& DirtyObjects, FNetBitArrayView& OutObjectsToPoll)
+void FObjectPollFrequencyLimiter::Update(const FNetBitArrayView& RelevantObjects, const FNetBitArrayView& DirtyObjects, FNetBitArrayView& OutObjectsToPoll)
 {
 	IRIS_PROFILER_SCOPE(ObjectPollFrequencyLimiter_Update);
 
@@ -55,7 +55,7 @@ void FObjectPollFrequencyLimiter::Update(const FNetBitArrayView& ScopableObjects
 	uint8* CountersData = FrameCounters.GetData();
 	const uint8* FramesBetweenUpdatesData = FramesBetweenUpdates.GetData();
 
-	const WordType* ScopableObjectsData = ScopableObjects.GetData();
+	const WordType* RelevantObjectsData = RelevantObjects.GetData();
 	const WordType* DirtyObjectsData = DirtyObjects.GetData();
 	WordType* ObjectsToPollData = OutObjectsToPoll.GetData();
 
@@ -71,10 +71,10 @@ void FObjectPollFrequencyLimiter::Update(const FNetBitArrayView& ScopableObjects
 
 		for (uint32 ObjectIndex = 0, ObjectEndIndex = MaxInternalHandle;
 			ObjectIndex <= ObjectEndIndex; 
-			ObjectIndex += WordBitCount, CountersData += WordBitCount, FramesBetweenUpdatesData += WordBitCount, ++ScopableObjectsData, ++DirtyObjectsData, ++ObjectsToPollData)
+			ObjectIndex += WordBitCount, CountersData += WordBitCount, FramesBetweenUpdatesData += WordBitCount, ++RelevantObjectsData, ++DirtyObjectsData, ++ObjectsToPollData)
 		{
 			// Skip ranges with no scopable objects.
-			const WordType ObjectsInScopeWord = *ScopableObjectsData;
+			const WordType ObjectsInScopeWord = *RelevantObjectsData;
 			if (!ObjectsInScopeWord)
 			{
 				continue;
@@ -111,10 +111,10 @@ void FObjectPollFrequencyLimiter::Update(const FNetBitArrayView& ScopableObjects
 	{
 		for (uint32 ObjectIndex = 0, ObjectEndIndex = MaxInternalHandle;
 			ObjectIndex <= ObjectEndIndex;
-			ObjectIndex += WordBitCount, CountersData += WordBitCount, FramesBetweenUpdatesData += WordBitCount, ++ScopableObjectsData, ++DirtyObjectsData, ++ObjectsToPollData)
+			ObjectIndex += WordBitCount, CountersData += WordBitCount, FramesBetweenUpdatesData += WordBitCount, ++RelevantObjectsData, ++DirtyObjectsData, ++ObjectsToPollData)
 		{
 			// Skip ranges with no scopable objects.
-			WordType ObjectsInScopeWord = *ScopableObjectsData;
+			WordType ObjectsInScopeWord = *RelevantObjectsData;
 			if (!ObjectsInScopeWord)
 			{
 				continue;
diff --git a/Engine/Source/Runtime/Experimental/Iris/Core/Private/Iris/ReplicationSystem/ObjectPollFrequencyLimiter.h b/Engine/Source/Runtime/Experimental/Iris/Core/Private/Iris/ReplicationSystem/ObjectPollFrequencyLimiter.h
--- a/Engine/Source/Runtime/Experimental/Iris/Core/Private/Iris/ReplicationSystem/ObjectPollFrequencyLimiter.h
+++ b/Engine/Source/Runtime/Experimental/Iris/Core/Private/Iris/ReplicationSystem/ObjectPollFrequencyLimiter.h
@@ -29,7 +29,11 @@ public:
 
 	void SetPollWithObject(FInternalNetRefIndex ObjectToPollWithInternalIndex, FInternalNetRefIndex InternalIndex);
 
-	void Update(const FNetBitArrayView& ScopableObjects, const FNetBitArrayView& DirtyObjects, FNetBitArrayView& OutObjectsToPoll);
+	/** 
+	* Produces the list of objects that should be polled this frame.
+	* This list is composed of relevant objects that are dirty or that hit their poll period this frame.
+	*/
+	void Update(const FNetBitArrayView& RelevantObjects, const FNetBitArrayView& DirtyObjects, FNetBitArrayView& OutObjectsToPoll);
 
 private:
 	uint32 GetPollFramePeriodForFrequency(float PollFrequency) const;
diff --git a/Engine/Source/Runtime/Experimental/Iris/Core/Private/Iris/ReplicationSystem/ObjectReplicationBridge.cpp b/Engine/Source/Runtime/Experimental/Iris/Core/Private/Iris/ReplicationSystem/ObjectReplicationBridge.cpp
--- a/Engine/Source/Runtime/Experimental/Iris/Core/Private/Iris/ReplicationSystem/ObjectReplicationBridge.cpp
+++ b/Engine/Source/Runtime/Experimental/Iris/Core/Private/Iris/ReplicationSystem/ObjectReplicationBridge.cpp
@@ -9,6 +9,7 @@
 #include "Iris/Core/IrisMemoryTracker.h"
 #include "Iris/Core/IrisProfiler.h"
 
+#include "Net/Core/NetBitArrayPrinter.h"
 #include "Net/Core/NetHandle/NetHandleManager.h"
 #include "Net/Core/PropertyConditions/PropertyConditionsDelegates.h"
 #include "Net/Core/Trace/NetTrace.h"
@@ -708,7 +709,7 @@ void UObjectReplicationBridge::PreUpdateAndPollImpl(FNetRefHandle Handle)
 
 	FDirtyObjectsAccessor DirtyObjectsAccessor(ReplicationSystemInternal->GetDirtyNetObjectTracker());
 
-	FNetBitArrayView DirtyObjects = DirtyObjectsAccessor.GetDirtyNetObjects();
+	FNetBitArrayView DirtyObjectsThisFrame = DirtyObjectsAccessor.GetDirtyNetObjects();
 
 	struct FPreUpdateAndPollStats
 	{
@@ -719,7 +720,7 @@ void UObjectReplicationBridge::PreUpdateAndPollImpl(FNetRefHandle Handle)
 
 	FPreUpdateAndPollStats Stats;
 
-	auto UpdateAndPollFunction = [&ReplicatedInstances, this, &LocalNetRefHandleManager, &DirtyObjects, &Stats](uint32 InternalObjectIndex)
+	auto UpdateAndPollFunction = [&ReplicatedInstances, this, &LocalNetRefHandleManager, &DirtyObjectsThisFrame, &Stats](uint32 InternalObjectIndex)
 	{
 		const FNetRefHandleManager::FReplicatedObjectData& ObjectData = LocalNetRefHandleManager.GetReplicatedObjectDataNoCheck(InternalObjectIndex);
 		if (ObjectData.InstanceProtocol)
@@ -749,7 +750,7 @@ void UObjectReplicationBridge::PreUpdateAndPollImpl(FNetRefHandle Handle)
 				const bool bWasMarkedDirty = FReplicationInstanceOperations::PollAndRefreshCachedPropertyData(ObjectData.InstanceProtocol, PollOptions);
 				if (bWasMarkedDirty)
 				{
-					DirtyObjects.SetBit(InternalObjectIndex);
+					DirtyObjectsThisFrame.SetBit(InternalObjectIndex);
 				}
 				++Stats.PolledObjectCount;
 			}
@@ -758,7 +759,7 @@ void UObjectReplicationBridge::PreUpdateAndPollImpl(FNetRefHandle Handle)
 
 	const FNetBitArray& PrevScopableObjects = LocalNetRefHandleManager.GetPrevFrameScopableInternalIndices();
 
-	auto PushModelUpdateAndPollFunction = [&ReplicatedInstances, this, &LocalNetRefHandleManager, &DirtyObjects, &PrevScopableObjects, &Stats](uint32 InternalObjectIndex)
+	auto PushModelUpdateAndPollFunction = [&ReplicatedInstances, this, &LocalNetRefHandleManager, &DirtyObjectsThisFrame, &PrevScopableObjects, &Stats](uint32 InternalObjectIndex)
 	{
 		const FNetRefHandleManager::FReplicatedObjectData& ObjectData = LocalNetRefHandleManager.GetReplicatedObjectDataNoCheck(InternalObjectIndex);
 		if (const FReplicationInstanceProtocol* InstanceProtocol = ObjectData.InstanceProtocol)
@@ -767,7 +768,7 @@ void UObjectReplicationBridge::PreUpdateAndPollImpl(FNetRefHandle Handle)
 
 			const EReplicationInstanceProtocolTraits InstanceTraits = InstanceProtocol->InstanceTraits;
 			const bool bNeedsPoll = EnumHasAnyFlags(InstanceTraits, EReplicationInstanceProtocolTraits::NeedsPoll);
-			bool bIsDirtyObject = DirtyObjects.GetBit(InternalObjectIndex);
+			bool bIsDirtyObject = DirtyObjectsThisFrame.GetBit(InternalObjectIndex);
 
 			// Call per-instance PreUpdate function
 			if (PreUpdateInstanceFunction && EnumHasAnyFlags(InstanceTraits, EReplicationInstanceProtocolTraits::NeedsPreSendUpdate))
@@ -783,7 +784,7 @@ void UObjectReplicationBridge::PreUpdateAndPollImpl(FNetRefHandle Handle)
 					bIsDirtyObject = FGlobalDirtyNetObjectTracker::IsNetObjectStateDirty(ObjectData.NetHandle);
 					if (bIsDirtyObject)
 					{
-						DirtyObjects.SetBit(bIsDirtyObject);
+						DirtyObjectsThisFrame.SetBit(bIsDirtyObject);
 					}
 				}
 			}
@@ -846,7 +847,7 @@ void UObjectReplicationBridge::PreUpdateAndPollImpl(FNetRefHandle Handle)
 
 			if (bWasMarkedDirty)
 			{
-				DirtyObjects.SetBit(InternalObjectIndex);
+				DirtyObjectsThisFrame.SetBit(InternalObjectIndex);
 			}
 		}
 	};
@@ -868,19 +869,22 @@ void UObjectReplicationBridge::PreUpdateAndPollImpl(FNetRefHandle Handle)
 	// Update every relevant objects from here
 
 	const FNetBitArrayView RelevantObjects = LocalNetRefHandleManager.GetRelevantObjectsInternalIndices();
-	const FNetBitArray& WantToBeDormantObjects = LocalNetRefHandleManager.GetWantToBeDormantInternalIndices();
+	const FNetBitArrayView WantToBeDormantObjects = MakeNetBitArrayView(LocalNetRefHandleManager.GetWantToBeDormantInternalIndices());
 
 	// Filter the set of objects considered for pre-update and polling
-	// We always want to consider objects marked as dirty and their subobjects
 	FNetBitArray ObjectsConsideredForPolling(LocalNetRefHandleManager.GetMaxActiveObjectCount());
-	FNetBitArrayView FrequencyBasedObjectsToPollView = MakeNetBitArrayView(ObjectsConsideredForPolling);
+	FNetBitArrayView ObjectsConsideredForPollingView = MakeNetBitArrayView(ObjectsConsideredForPolling);
+
+	// We always want to consider objects marked as dirty and their subobjects
+	const FNetBitArrayView AccumulatedDirtyObjects = ReplicationSystemInternal->GetDirtyNetObjectTracker().GetAccumulatedDirtyNetObjects();
+
 	if (bUseFrequencyBasedPolling)
 	{
-		PollFrequencyLimiter->Update(RelevantObjects, DirtyObjects, FrequencyBasedObjectsToPollView);
+		PollFrequencyLimiter->Update(RelevantObjects, AccumulatedDirtyObjects, ObjectsConsideredForPollingView);
 	}
 	else
 	{
-		FrequencyBasedObjectsToPollView.Copy(RelevantObjects);
+		ObjectsConsideredForPollingView.Copy(RelevantObjects);
 	}
 
 	// Mask off objects pending dormancy as we do not want to poll/pre-update them unless they are marked for flush or are dirty
@@ -889,7 +893,7 @@ void UObjectReplicationBridge::PreUpdateAndPollImpl(FNetRefHandle Handle)
 		IRIS_PROFILER_SCOPE(PreUpdateAndPollImpl_Dormancy);
 
 		// Mask off objects pending dormancy that are not dirty
-		FrequencyBasedObjectsToPollView.CombineMultiple(FNetBitArrayView::AndNotOp, MakeNetBitArrayView(WantToBeDormantObjects), FNetBitArrayView::AndNotOp, DirtyObjects);
+		ObjectsConsideredForPollingView.CombineMultiple(FNetBitArrayView::AndNotOp, WantToBeDormantObjects, FNetBitArrayView::AndNotOp, AccumulatedDirtyObjects);
 
 		// Add objects that have requested to flush dormancy
 		for (FNetRefHandle HandlePendingFlush : MakeArrayView(DormantHandlesPendingFlush))
@@ -936,12 +940,18 @@ void UObjectReplicationBridge::PreUpdateAndPollImpl(FNetRefHandle Handle)
 			}
 		};
 
+		// Make a list of objects that were set to be dirty this frame but are also relevant
+		FNetBitArray DirtyAndRelevantObjects(RelevantObjects.GetNumBits());
+		FNetBitArrayView DirtyAndRelevantObjectsView = MakeNetBitArrayView(DirtyAndRelevantObjects);
+
+		DirtyAndRelevantObjectsView.Set(RelevantObjects, FNetBitArray::AndOp, DirtyObjectsThisFrame);
+
 		// Update subobjects' owner first and owners' subobjects second. It's the only way to properly mark all groups of objects in two passes.
 		const FNetBitArrayView SubObjects = MakeNetBitArrayView(LocalNetRefHandleManager.GetSubObjectInternalIndices());
-		FNetBitArrayView::ForAllSetBits(DirtyObjects, SubObjects, FNetBitArray::AndOp, PropagateSubObjectDirtinessToOwner);
-		FNetBitArrayView::ForAllSetBits(DirtyObjects, SubObjects, FNetBitArray::AndNotOp, PropagateOwnerDirtinessToSubObjectsAndDependentObjects);
+		FNetBitArrayView::ForAllSetBits(DirtyAndRelevantObjectsView, SubObjects, FNetBitArray::AndOp, PropagateSubObjectDirtinessToOwner);
+		FNetBitArrayView::ForAllSetBits(DirtyAndRelevantObjectsView, SubObjects, FNetBitArray::AndNotOp, PropagateOwnerDirtinessToSubObjectsAndDependentObjects);
 			
-		// Patch dependent objects, not subtle as this point
+		// If an object with dependents is about to be polled, force his dependents to poll at the same time.
 		{
 			IRIS_PROFILER_SCOPE(PreUpdateAndPollImpl_PatchDependentObjects);
 
@@ -949,7 +959,10 @@ void UObjectReplicationBridge::PreUpdateAndPollImpl(FNetRefHandle Handle)
 			FNetBitArray::ForAllSetBits(TempObjectsConsideredForPolling, NetRefHandleManager->GetObjectsWithDependentObjectsInternalIndices(), FNetBitArray::AndOp,
 				[&LocalNetRefHandleManager, &ObjectsConsideredForPolling](FInternalNetRefIndex ObjectIndex) 
 				{
-					LocalNetRefHandleManager.ForAllDependentObjectsRecursive(ObjectIndex, [&ObjectsConsideredForPolling, ObjectIndex](FInternalNetRefIndex DependentObjectIndex) { ObjectsConsideredForPolling.SetBit(DependentObjectIndex); });
+					LocalNetRefHandleManager.ForAllDependentObjectsRecursive(ObjectIndex, [&ObjectsConsideredForPolling, ObjectIndex](FInternalNetRefIndex DependentObjectIndex)
+					{ 
+						ObjectsConsideredForPolling.SetBit(DependentObjectIndex); 
+					});
 				}
 			);
 		}
@@ -1007,8 +1020,8 @@ void UObjectReplicationBridge::UpdateInstancesWorldLocation()
 	};
 
 	FDirtyObjectsAccessor DirtyObjectsAccessor(ReplicationSystemInternal->GetDirtyNetObjectTracker());
-	const FNetBitArrayView DirtyObjects = DirtyObjectsAccessor.GetDirtyNetObjects();
-	DirtyObjects.ForAllSetBits(UpdateInstanceWorldLocation);
+	const FNetBitArrayView DirtyObjectsThisFrame = DirtyObjectsAccessor.GetDirtyNetObjects();
+	DirtyObjectsThisFrame.ForAllSetBits(UpdateInstanceWorldLocation);
 
 }
 
diff --git a/Engine/Source/Runtime/Experimental/Iris/Core/Private/Iris/ReplicationSystem/ReplicationSystem.cpp b/Engine/Source/Runtime/Experimental/Iris/Core/Private/Iris/ReplicationSystem/ReplicationSystem.cpp
--- a/Engine/Source/Runtime/Experimental/Iris/Core/Private/Iris/ReplicationSystem/ReplicationSystem.cpp
+++ b/Engine/Source/Runtime/Experimental/Iris/Core/Private/Iris/ReplicationSystem/ReplicationSystem.cpp
@@ -220,15 +220,12 @@ public:
 
 	void UpdateDirtyObjectList()
 	{
-		FDirtyNetObjectTracker& DirtyNetObjectTracker = ReplicationSystemInternal.GetDirtyNetObjectTracker();
-		DirtyNetObjectTracker.UpdateDirtyNetObjects();
-
-		FDirtyObjectsAccessor DirtyObjectsAccessor(ReplicationSystemInternal.GetDirtyNetObjectTracker());
-		FNetBitArrayView DirtyObjects = DirtyObjectsAccessor.GetDirtyNetObjects();
+		ReplicationSystemInternal.GetDirtyNetObjectTracker().UpdateDirtyNetObjects();
+	}
 
-		// Due to objects having been marked as dirty and later removed we must make sure that all dirty objects are still in scope.
-		const FNetBitArrayView ScopableInternalIndices = MakeNetBitArrayView(ReplicationSystemInternal.GetNetRefHandleManager().GetScopableInternalIndices());
-		DirtyObjects.Combine(ScopableInternalIndices, FNetBitArrayView::AndOp);
+	void UpdateDirtyListPostPoll()
+	{
+		ReplicationSystemInternal.GetDirtyNetObjectTracker().UpdateAccumulatedDirtyList();
 	}
 
 	void UpdateWorldLocations()
@@ -329,22 +326,21 @@ public:
 
 		SerializationContext.SetInternalContext(&InternalContext);
 
+		// Build the list of dirty objects whom we will copy their data and "clean"
+		FNetBitArrayView CleanedObjects = NetRefHandleManager.GetCleanedObjectsInternalIndices();
+		const FNetBitArrayView AccumulatedDirtyNetObjects = ReplicationSystemInternal.GetDirtyNetObjectTracker().GetAccumulatedDirtyNetObjects();
+		const FNetBitArrayView RelevantObjects = NetRefHandleManager.GetRelevantObjectsInternalIndices();
+
+		CleanedObjects.Set(AccumulatedDirtyNetObjects, FNetBitArrayView::AndOp, RelevantObjects);
+
+
 		auto CopyFunction = [&ChangeMaskWriter, &Cache, &NetRefHandleManager, &CopiedObjectCount, &SerializationContext](uint32 DirtyIndex)
 		{
 			CopiedObjectCount += FReplicationInstanceOperationsInternal::CopyObjectStateData(ChangeMaskWriter, Cache, NetRefHandleManager, SerializationContext, DirtyIndex);
 		};
 
-
-		// Only iterate over dirty and relevant objects
-		FDirtyObjectsAccessor DirtyObjectsAccessor(ReplicationSystemInternal.GetDirtyNetObjectTracker());
-		const FNetBitArrayView DirtyObjects = DirtyObjectsAccessor.GetDirtyNetObjects();
-
-		// Copy all ReplicatedObjects with dirty state data
-		DirtyObjects.ForAllSetBits(CopyFunction);
-
-		//$IRIS TODO Only copy relevant objects but keep dirty flag for those not copied.
-		//FNetBitArrayView RelevantObjects = NetRefHandleManager.GetRelevantObjectsInternalIndices();
-		//FNetBitArrayView::ForAllSetBits(DirtyObjects, RelevantObjects, FNetBitArrayView::AndOp, CopyFunction);
+		// Only copy both dirty and relevant objects
+		CleanedObjects.ForAllSetBits(CopyFunction);
 
 		const uint32 ReplicationSystemId = ReplicationSystem->GetId();
 		UE_NET_TRACE_FRAME_STATSCOUNTER(ReplicationSystemId, ReplicationSystem.CopiedObjectCount, CopiedObjectCount, ENetTraceVerbosity::Trace);
@@ -355,19 +351,16 @@ public:
 		IRIS_PROFILER_SCOPE(FReplicationSystem_ResetObjectStateDirtiness);
 
 		FNetRefHandleManager& NetRefHandleManager = ReplicationSystemInternal.GetNetRefHandleManager();
-		FDirtyNetObjectTracker& DirtyNetObjectTracker = ReplicationSystemInternal.GetDirtyNetObjectTracker();
+		const FNetBitArrayView CleanedObjects = NetRefHandleManager.GetCleanedObjectsInternalIndices();
 
-		auto ResetDirtinessFunction = [&NetRefHandleManager](uint32 DirtyIndex)
+		// Reset object dirtyness
+		CleanedObjects.ForAllSetBits([&NetRefHandleManager](uint32 DirtyIndex)
 		{
 			FReplicationInstanceOperationsInternal::ResetObjectStateDirtiness(NetRefHandleManager, DirtyIndex);
-		};
+		});
 
-		FDirtyObjectsAccessor DirtyObjectsAccessor(DirtyNetObjectTracker);
-		const FNetBitArrayView DirtyObjects = DirtyObjectsAccessor.GetDirtyNetObjects();
-
-		DirtyObjects.ForAllSetBits(ResetDirtinessFunction);
-
-		DirtyNetObjectTracker.ClearDirtyNetObjects();
+		// Reset cleaned objects in the tracker
+		ReplicationSystemInternal.GetDirtyNetObjectTracker().ClearDirtyNetObjects(CleanedObjects);
 	}
 
 	void ProcessNetObjectAttachmentSendQueue(FNetBlobManager::EProcessMode ProcessMode)
@@ -572,7 +565,7 @@ void UReplicationSystem::PreSendUpdate(float DeltaSeconds)
 	{
 		UE_NET_TRACE_FRAME_STATSCOUNTER(GetId(), ReplicationSystem.ReplicatedObjectCount, InternalSys.GetNetRefHandleManager().GetActiveObjectCount(), ENetTraceVerbosity::Verbose);
 
-		// Refresh dirty object list
+		// Refresh the dirty objects we were told about.
 		Impl->UpdateDirtyObjectList();
 
 		// Update world locations. We need this to happen before both filtering and prioritization.
@@ -582,7 +575,10 @@ void UReplicationSystem::PreSendUpdate(float DeltaSeconds)
 		Impl->UpdateFilterPrePoll();
 
 		// Invoke any operations we need to do before copying state data
-		InternalSys.GetReplicationBridge()->CallPreSendUpdate(DeltaSeconds);		
+		InternalSys.GetReplicationBridge()->CallPreSendUpdate(DeltaSeconds);
+
+		// Finalize the dirty list with objects set dirty during the poll phase
+		Impl->UpdateDirtyListPostPoll();
 
 		// Update conditionals
 		Impl->UpdateConditionals();
