Fix missing include

#rnx

diff --git a/Engine/Source/Runtime/Experimental/Iris/Core/Private/Iris/ReplicationSystem/Filtering/ReplicationFiltering.cpp b/Engine/Source/Runtime/Experimental/Iris/Core/Private/Iris/ReplicationSystem/Filtering/ReplicationFiltering.cpp
--- a/Engine/Source/Runtime/Experimental/Iris/Core/Private/Iris/ReplicationSystem/Filtering/ReplicationFiltering.cpp
+++ b/Engine/Source/Runtime/Experimental/Iris/Core/Private/Iris/ReplicationSystem/Filtering/ReplicationFiltering.cpp
@@ -2,6 +2,8 @@
 
 #include "ReplicationFiltering.h"
 #include "HAL/PlatformMath.h"
+#include "HAL/IConsoleManager.h"
+
 #include "Iris/IrisConfigInternal.h"
 #include "Iris/IrisConstants.h"
 #include "Iris/Core/IrisLog.h"
