Add option for cluster-convex-from-leaf dataflow node to ignore proximity when merging hulls
 + clean up settings parameters for cluster-convex functions to use a struct, so adding new options is easier

diff --git a/Engine/Plugins/Experimental/GeometryCollectionPlugin/Source/GeometryCollectionNodes/Private/Dataflow/GeometryCollectionUtilityNodes.cpp b/Engine/Plugins/Experimental/GeometryCollectionPlugin/Source/GeometryCollectionNodes/Private/Dataflow/GeometryCollectionUtilityNodes.cpp
--- a/Engine/Plugins/Experimental/GeometryCollectionPlugin/Source/GeometryCollectionNodes/Private/Dataflow/GeometryCollectionUtilityNodes.cpp
+++ b/Engine/Plugins/Experimental/GeometryCollectionPlugin/Source/GeometryCollectionNodes/Private/Dataflow/GeometryCollectionUtilityNodes.cpp
@@ -172,21 +172,21 @@ void FGenerateClusterConvexHullsFromLeafHullsDataflowNode::Evaluate(Dataflow::FC
 			if (IsConnected(&OptionalSelectionFilter))
 			{
 				const FDataflowTransformSelection& InOptionalSelectionFilter = GetValue<FDataflowTransformSelection>(Context, &OptionalSelectionFilter);
+				FGeometryCollectionConvexUtility::FClusterConvexHullSettings Settings(InConvexCount, InErrorToleranceInCm, bPreferExternalCollisionShapes);
+				Settings.AllowMergesMethod = AllowMerges;
 				FGeometryCollectionConvexUtility::GenerateClusterConvexHullsFromLeafHulls(
 					*GeomCollection,
-					InConvexCount,
-					InErrorToleranceInCm,
-					InOptionalSelectionFilter.AsArray(),
-					bPreferExternalCollisionShapes
+					Settings,
+					InOptionalSelectionFilter.AsArray()
 				);
 			}
 			else
 			{
+				FGeometryCollectionConvexUtility::FClusterConvexHullSettings Settings(InConvexCount, InErrorToleranceInCm, bPreferExternalCollisionShapes);
+				Settings.AllowMergesMethod = AllowMerges;
 				FGeometryCollectionConvexUtility::GenerateClusterConvexHullsFromLeafHulls(
 					*GeomCollection,
-					InConvexCount,
-					InErrorToleranceInCm,
-					bPreferExternalCollisionShapes
+					Settings
 				);
 			}
 
@@ -219,21 +219,19 @@ void FGenerateClusterConvexHullsFromChildrenHullsDataflowNode::Evaluate(Dataflow
 			if (IsConnected(&OptionalSelectionFilter))
 			{
 				const FDataflowTransformSelection& InOptionalSelectionFilter = GetValue<FDataflowTransformSelection>(Context, &OptionalSelectionFilter);
+				FGeometryCollectionConvexUtility::FClusterConvexHullSettings Settings(InConvexCount, InErrorToleranceInCm, bPreferExternalCollisionShapes);
 				FGeometryCollectionConvexUtility::GenerateClusterConvexHullsFromChildrenHulls(
 					*GeomCollection,
-					InConvexCount,
-					InErrorToleranceInCm,
-					InOptionalSelectionFilter.AsArray(),
-					bPreferExternalCollisionShapes
+					Settings,
+					InOptionalSelectionFilter.AsArray()
 				);
 			}
 			else
 			{
+				FGeometryCollectionConvexUtility::FClusterConvexHullSettings Settings(InConvexCount, InErrorToleranceInCm, bPreferExternalCollisionShapes);
 				FGeometryCollectionConvexUtility::GenerateClusterConvexHullsFromChildrenHulls(
 					*GeomCollection,
-					InConvexCount,
-					InErrorToleranceInCm,
-					bPreferExternalCollisionShapes
+					Settings
 				);
 			}
 
diff --git a/Engine/Plugins/Experimental/GeometryCollectionPlugin/Source/GeometryCollectionNodes/Public/Dataflow/GeometryCollectionUtilityNodes.h b/Engine/Plugins/Experimental/GeometryCollectionPlugin/Source/GeometryCollectionNodes/Public/Dataflow/GeometryCollectionUtilityNodes.h
--- a/Engine/Plugins/Experimental/GeometryCollectionPlugin/Source/GeometryCollectionNodes/Public/Dataflow/GeometryCollectionUtilityNodes.h
+++ b/Engine/Plugins/Experimental/GeometryCollectionPlugin/Source/GeometryCollectionNodes/Public/Dataflow/GeometryCollectionUtilityNodes.h
@@ -162,6 +162,10 @@ public:
 	UPROPERTY(EditAnywhere, Category = "Convex")
 	bool bPreferExternalCollisionShapes = true;
 
+	/** Method to determine which convex hull pairs can potentially be merged */
+	UPROPERTY(EditAnywhere, Category = "Convex")
+	EAllowConvexMergeMethod AllowMerges = EAllowConvexMergeMethod::ByProximity;
+
 	/** Optional transform selection to compute cluster hulls on -- if not provided, all cluster hulls will be computed. */
 	UPROPERTY(meta = (DataflowInput, DataflowIntrinsic))
 	FDataflowTransformSelection OptionalSelectionFilter;
diff --git a/Engine/Source/Runtime/Experimental/Chaos/Private/GeometryCollection/GeometryCollectionConvexUtility.cpp b/Engine/Source/Runtime/Experimental/Chaos/Private/GeometryCollection/GeometryCollectionConvexUtility.cpp
--- a/Engine/Source/Runtime/Experimental/Chaos/Private/GeometryCollection/GeometryCollectionConvexUtility.cpp
+++ b/Engine/Source/Runtime/Experimental/Chaos/Private/GeometryCollection/GeometryCollectionConvexUtility.cpp
@@ -1645,24 +1645,24 @@ void FGeometryCollectionConvexUtility::GenerateLeafConvexHulls(FGeometryCollecti
 	RemoveEmptyConvexHulls(Collection);
 }
 
-void FGeometryCollectionConvexUtility::GenerateClusterConvexHullsFromChildrenHulls(FGeometryCollection& Collection, int32 ConvexCount, double ErrorToleranceInCm, const TArrayView<const int32> TransformSubset, bool bUseExternalCollisionIfAvailable)
+void FGeometryCollectionConvexUtility::GenerateClusterConvexHullsFromChildrenHulls(FGeometryCollection& Collection, const FClusterConvexHullSettings& Settings, const TArrayView<const int32> TransformSubset)
 {
-	GenerateClusterConvexHullsFromLeafOrChildrenHullsInternal(Collection, ConvexCount, ErrorToleranceInCm, true, true/*bUseDirectChildren*/, TransformSubset, bUseExternalCollisionIfAvailable);
+	GenerateClusterConvexHullsFromLeafOrChildrenHullsInternal(Collection, Settings, true, true/*bUseDirectChildren*/, TransformSubset);
 }
-void FGeometryCollectionConvexUtility::GenerateClusterConvexHullsFromChildrenHulls(FGeometryCollection& Collection, int32 ConvexCount, double ErrorToleranceInCm, bool bUseExternalCollisionIfAvailable)
+void FGeometryCollectionConvexUtility::GenerateClusterConvexHullsFromChildrenHulls(FGeometryCollection& Collection, const FClusterConvexHullSettings& Settings)
 {
-	GenerateClusterConvexHullsFromLeafOrChildrenHullsInternal(Collection, ConvexCount, ErrorToleranceInCm, false, true/*bUseDirectChildren*/, TArrayView<const int32>(), bUseExternalCollisionIfAvailable);
+	GenerateClusterConvexHullsFromLeafOrChildrenHullsInternal(Collection, Settings, false, true/*bUseDirectChildren*/, TArrayView<const int32>());
 }
-void FGeometryCollectionConvexUtility::GenerateClusterConvexHullsFromLeafHulls(FGeometryCollection& Collection, int32 ConvexCount, double ErrorToleranceInCm, const TArrayView<const int32> TransformSubset, bool bUseExternalCollisionIfAvailable)
+void FGeometryCollectionConvexUtility::GenerateClusterConvexHullsFromLeafHulls(FGeometryCollection& Collection, const FClusterConvexHullSettings& Settings, const TArrayView<const int32> TransformSubset)
 {
-	GenerateClusterConvexHullsFromLeafOrChildrenHullsInternal(Collection, ConvexCount, ErrorToleranceInCm, true, false/*bUseDirectChildren*/, TransformSubset, bUseExternalCollisionIfAvailable);
+	GenerateClusterConvexHullsFromLeafOrChildrenHullsInternal(Collection, Settings, true, false/*bUseDirectChildren*/, TransformSubset);
 }
-void FGeometryCollectionConvexUtility::GenerateClusterConvexHullsFromLeafHulls(FGeometryCollection& Collection, int32 ConvexCount, double ErrorToleranceInCm, bool bUseExternalCollisionIfAvailable)
+void FGeometryCollectionConvexUtility::GenerateClusterConvexHullsFromLeafHulls(FGeometryCollection& Collection, const FClusterConvexHullSettings& Settings)
 {
-	GenerateClusterConvexHullsFromLeafOrChildrenHullsInternal(Collection, ConvexCount, ErrorToleranceInCm, false, false/*bUseDirectChildren*/, TArrayView<const int32>(), bUseExternalCollisionIfAvailable);
+	GenerateClusterConvexHullsFromLeafOrChildrenHullsInternal(Collection, Settings, false, false/*bUseDirectChildren*/, TArrayView<const int32>());
 }
 
-void FGeometryCollectionConvexUtility::GenerateClusterConvexHullsFromLeafOrChildrenHullsInternal(FGeometryCollection& Collection, int32 ConvexCount, double ErrorToleranceInCm, bool bOnlySubset, bool bUseDirectChildren, const TArrayView<const int32> TransformSubset, bool bUseExternalCollisionIfAvailable)
+void FGeometryCollectionConvexUtility::GenerateClusterConvexHullsFromLeafOrChildrenHullsInternal(FGeometryCollection& Collection, const FClusterConvexHullSettings& Settings, bool bOnlySubset, bool bUseDirectChildren, const TArrayView<const int32> TransformSubset)
 {
 	static FName ConvexGroupName("Convex");
 	static FName ConvexHullAttributeName("ConvexHull");
@@ -1754,7 +1754,7 @@ void FGeometryCollectionConvexUtility::GenerateClusterConvexHullsFromLeafOrChild
 					FTransform ChildToParentTransform = InnerTransform.GetRelativeTransform(ParentTransform);
 
 					const FSharedImplicit ExternalCollision = (ExternalCollisionAttribute.IsValid()) ? ExternalCollisionAttribute.Get()[SourceTransformIndex] : FSharedImplicit();
-					if (ExternalCollision && bUseExternalCollisionIfAvailable)
+					if (ExternalCollision && Settings.bUseExternalCollisionIfAvailable)
 					{
 						const int32 ExternalHullsStart = ExternalHulls.Num();
 						ConvertImplicitToConvexArray(*ExternalCollision, FTransform::Identity, ExternalHulls);
@@ -1781,12 +1781,12 @@ void FGeometryCollectionConvexUtility::GenerateClusterConvexHullsFromLeafOrChild
 						}
 					}
 				}
-				// get hull proximity out of geometry proximity.  (Note: Only works if merging leaf convexes ... need a more general source of proximity!)
-				if (bUseDirectChildren)
+				// Set HullProximity for convex merges -- the connections between convex hulls that can be merged
+				if (bUseDirectChildren || Settings.AllowMergesMethod == EAllowConvexMergeMethod::Any)
 				{
-					// when using direct childre, proximity data may not be available ( if children are not leaves ) 
-					// so we use declare everything connected to everything for now 
-					// @todo this is suboptimal in some cases and we could precompute proximity and use it instead 
+					// With MergeMethod of Any, any pair of convex hulls can be merged
+					// Note: Currently bUseDirectChildren only supports this method
+					// This is likely to be slower and may consider additional merges that the proximity method cannot
 					for (int32 IndexA = 0; IndexA < SourceTransformIndices.Num(); IndexA++)
 					{
 						const int32 TransformIndexA = SourceTransformIndices[IndexA];
@@ -1809,7 +1809,7 @@ void FGeometryCollectionConvexUtility::GenerateClusterConvexHullsFromLeafOrChild
 						}
 					}
 				}
-				else
+				else // Settings.MergeMethod == EAllowConvexMergeMethod::ByProximity
 				{
 					for (int32 SourceTransformIndex : SourceTransformIndices)
 					{
@@ -1844,7 +1844,7 @@ void FGeometryCollectionConvexUtility::GenerateClusterConvexHullsFromLeafOrChild
 				};
 				UE::Geometry::FConvexDecomposition3 ConvexDecomposition;
 				ConvexDecomposition.InitializeFromHulls(Hulls.Num(), GetHullVolume, GetHullNumVertices, GetHullVertex, HullProximity);
-				ConvexDecomposition.MergeBest(ConvexCount, ErrorToleranceInCm, 0, true);
+				ConvexDecomposition.MergeBest(Settings.ConvexCount, Settings.ErrorToleranceInCm, 0, true);
 				
 				// reset existing hulls for this transform index
 				for (int32 ConvexIndex : TransformToConvexIndices[TransformIndex])
diff --git a/Engine/Source/Runtime/Experimental/Chaos/Public/GeometryCollection/GeometryCollectionConvexUtility.h b/Engine/Source/Runtime/Experimental/Chaos/Public/GeometryCollection/GeometryCollectionConvexUtility.h
--- a/Engine/Source/Runtime/Experimental/Chaos/Public/GeometryCollection/GeometryCollectionConvexUtility.h
+++ b/Engine/Source/Runtime/Experimental/Chaos/Public/GeometryCollection/GeometryCollectionConvexUtility.h
@@ -38,6 +38,15 @@ enum class EGenerateConvexMethod : uint8
 	IntersectExternalWithComputed
 };
 
+UENUM()
+enum class EAllowConvexMergeMethod : uint8
+{
+	// Only allow merging convex hulls that are in proximity
+	ByProximity,
+	// Allow any pair of convex hulls to merge
+	Any
+};
+
 namespace UE::GeometryCollectionConvexUtility
 {
 	// Used to pass computed convex hulls along with metadata
@@ -99,10 +108,23 @@ public:
 	static FGeometryCollectionConvexData CreateNonOverlappingConvexHullData(FGeometryCollection* GeometryCollection, double FractionAllowRemove = .3, double SimplificationDistanceThreshold = 0.0, double CanExceedFraction = .5, 
 		EConvexOverlapRemoval OverlapRemovalMethod = EConvexOverlapRemoval::All, double OverlapRemovalShrinkPercent = 0.0, UE::GeometryCollectionConvexUtility::FConvexHulls* ComputedLeafHullsToModify = nullptr);
 
-	static void GenerateClusterConvexHullsFromChildrenHulls(FGeometryCollection& Collection, int32 ConvexCount, double ErrorToleranceInCm, const TArrayView<const int32> TransformSubset, bool bUseExternalCollisionIfAvailable = true);
-	static void GenerateClusterConvexHullsFromChildrenHulls(FGeometryCollection& Collection, int32 ConvexCount, double ErrorToleranceInCm, bool bUseExternalCollisionIfAvailable = true);
-	static void GenerateClusterConvexHullsFromLeafHulls(FGeometryCollection& Collection, int32 ConvexCount, double ErrorToleranceInCm, const TArrayView<const int32> OptionalTransformSubset, bool bUseExternalCollisionIfAvailable = true);
-	static void GenerateClusterConvexHullsFromLeafHulls(FGeometryCollection& Collection, int32 ConvexCount, double ErrorToleranceInCm, bool bUseExternalCollisionIfAvailable = true);
+	struct FClusterConvexHullSettings
+	{
+		FClusterConvexHullSettings() = default;
+		FClusterConvexHullSettings(int32 ConvexCount, double ErrorToleranceInCm, bool bUseExternalCollisionIfAvailable = true)
+			: ConvexCount(ConvexCount), ErrorToleranceInCm(ErrorToleranceInCm), bUseExternalCollisionIfAvailable(bUseExternalCollisionIfAvailable)
+		{}
+
+		int32 ConvexCount = 4;
+		double ErrorToleranceInCm = 0.0;
+		bool bUseExternalCollisionIfAvailable = true;
+		EAllowConvexMergeMethod AllowMergesMethod = EAllowConvexMergeMethod::ByProximity;
+	};
+
+	static void GenerateClusterConvexHullsFromChildrenHulls(FGeometryCollection& Collection, const FClusterConvexHullSettings& Settings, const TArrayView<const int32> TransformSubset);
+	static void GenerateClusterConvexHullsFromChildrenHulls(FGeometryCollection& Collection, const FClusterConvexHullSettings& Settings);
+	static void GenerateClusterConvexHullsFromLeafHulls(FGeometryCollection& Collection, const FClusterConvexHullSettings& Settings, const TArrayView<const int32> OptionalTransformSubset);
+	static void GenerateClusterConvexHullsFromLeafHulls(FGeometryCollection& Collection, const FClusterConvexHullSettings& Settings);
 
 	// Additional settings for filtering when the EGenerateConvexMethod::IntersectExternalWithComputed is applied
 	struct FIntersectionFilters
@@ -182,6 +204,6 @@ private:
 	static void CreateConvexHullAttributesIfNeeded(FManagedArrayCollection& GeometryCollection);
 
 	// Implementation for GenerateClusterConvexHullsFromLeafHulls, supporting the full-collection and subset cases
-	static void GenerateClusterConvexHullsFromLeafOrChildrenHullsInternal(FGeometryCollection& Collection, int32 ConvexCount, double ErrorToleranceInCm, bool bOnlySubset, bool bUseDirectChildren, const TArrayView<const int32> OptionalTransformSubset, bool bUseExternalCollisionIfAvailable = true);
+	static void GenerateClusterConvexHullsFromLeafOrChildrenHullsInternal(FGeometryCollection& Collection, const FClusterConvexHullSettings& Settings, bool bOnlySubset, bool bUseDirectChildren, const TArrayView<const int32> OptionalTransformSubset);
 };
 
