Chaos - big island manager refactor

diff --git a/Engine/Source/Runtime/Experimental/Chaos/Private/Chaos/Island/IslandManager.cpp b/Engine/Source/Runtime/Experimental/Chaos/Private/Chaos/Island/IslandManager.cpp
--- a/Engine/Source/Runtime/Experimental/Chaos/Private/Chaos/Island/IslandManager.cpp
+++ b/Engine/Source/Runtime/Experimental/Chaos/Private/Chaos/Island/IslandManager.cpp
@@ -1,1326 +1,2319 @@
 // Copyright Epic Games, Inc. All Rights Reserved.
 #include "Chaos/Island/IslandManager.h"
-#include "Chaos/Island/SolverIsland.h"
-#include "Chaos/Island/IslandGraph.h"
 #include "Chaos/ConstraintHandle.h"
 #include "Chaos/ParticleHandle.h"
+#include "Chaos/ParticleIterator.h"
 #include "Chaos/PBDConstraintContainer.h"
-#include "Chaos/PBDRigidParticles.h"
 #include "Chaos/PBDRigidsSOAs.h"
+#include "Chaos/PerParticleGravity.h"
 #include "Chaos/PhysicsMaterialUtilities.h"
-#include "ChaosLog.h"
 
-#include "Framework/Threading.h"
-#include "Misc/App.h"
+#include "ChaosStats.h"
 
-//PRAGMA_DISABLE_OPTIMIZATION
 
-namespace Chaos
-{
-/** Cvar to enable/disable the island sleeping */
-bool bChaosSolverSleepEnabled = true;
-FAutoConsoleVariableRef CVarChaosSolverSleepEnabled(TEXT("p.Chaos.Solver.SleepEnabled"), bChaosSolverSleepEnabled, TEXT(""));
-
-/** Cvar to override the sleep counter threshold if necessary */
-int32 ChaosSolverCollisionDefaultSleepCounterThresholdCVar = 20;
-FAutoConsoleVariableRef CVarChaosSolverCollisionDefaultSleepCounterThreshold(TEXT("p.ChaosSolverCollisionDefaultSleepCounterThreshold"), ChaosSolverCollisionDefaultSleepCounterThresholdCVar, TEXT("Default counter threshold for sleeping.[def:20]"));
+// Extra check for debug mode
+// NOTE: Should be disabled for checkin (except debug builds)
+#ifndef CHAOS_CONSTRAINTGRAPH_CHECK_ENABLED
+#define CHAOS_CONSTRAINTGRAPH_CHECK_ENABLED (UE_BUILD_DEBUG)
+#endif
 
-/** Cvar to override the sleep linear threshold if necessary */
-FRealSingle ChaosSolverCollisionDefaultLinearSleepThresholdCVar = 0.001f; // .001 unit mass cm
-FAutoConsoleVariableRef CVarChaosSolverCollisionDefaultLinearSleepThreshold(TEXT("p.ChaosSolverCollisionDefaultLinearSleepThreshold"), ChaosSolverCollisionDefaultLinearSleepThresholdCVar, TEXT("Default linear threshold for sleeping.[def:0.001]"));
+#if CHAOS_CONSTRAINTGRAPH_CHECK_ENABLED
+#define CHAOS_CONSTRAINTGRAPH_CHECK(X) check(X)
+#else
+#define CHAOS_CONSTRAINTGRAPH_CHECK(X)
+#endif
 
-/** Cvar to override the sleep angular threshold if necessary */
-FRealSingle ChaosSolverCollisionDefaultAngularSleepThresholdCVar = 0.0087f;  //~1/2 unit mass degree
-FAutoConsoleVariableRef CVarChaosSolverCollisionDefaultAngularSleepThreshold(TEXT("p.ChaosSolverCollisionDefaultAngularSleepThreshold"), ChaosSolverCollisionDefaultAngularSleepThresholdCVar, TEXT("Default angular threshold for sleeping.[def:0.0087]"));
+#define CHAOS_CONSTRAINTGRAPH_VALIDATE_TEST(X) \
+	bIsValid = bIsValid && ensure(X); \
+	if (!bIsValid) { return false; }
 
-bool bChaosSolverValidateGraph = (CHAOS_CONSTRAINTHANDLE_DEBUG_ENABLED != 0);
-FAutoConsoleVariableRef CVarChaosSolverValidateGraph(TEXT("p.Chaos.Solver.ValidateGraph"), bChaosSolverValidateGraph, TEXT(""));
 
-bool bChaosSolverPersistentGraph = true;
-FAutoConsoleVariableRef CVarChaosSolverPersistentGraph(TEXT("p.Chaos.Solver.PersistentGraph"), bChaosSolverPersistentGraph, TEXT(""));
+DECLARE_CYCLE_STAT(TEXT("IslandManager::UpdateParticles"), STAT_IslandManager_UpdateParticles, STATGROUP_ChaosIslands);
+DECLARE_CYCLE_STAT(TEXT("IslandManager::Merge"), STAT_IslandManager_MergeIslands, STATGROUP_ChaosIslands);
+DECLARE_CYCLE_STAT(TEXT("IslandManager::Split"), STAT_IslandManager_SplitIslands, STATGROUP_ChaosIslands);
+DECLARE_CYCLE_STAT(TEXT("IslandManager::Levels"), STAT_IslandManager_AssignLevels, STATGROUP_ChaosIslands);
+DECLARE_CYCLE_STAT(TEXT("IslandManager::Finalize"), STAT_IslandManager_Finalize, STATGROUP_ChaosIslands);
+DECLARE_CYCLE_STAT(TEXT("IslandManager::Validate"), STAT_IslandManager_Validate, STATGROUP_ChaosIslands);
 
-namespace Private
+namespace Chaos::CVars
 {
-	
-/** Check if a particle is dynamic or sleeping */
-FORCEINLINE bool IsDynamicParticle(const FGeometryParticleHandle* ParticleHandle)
-{
-	return (ParticleHandle->ObjectState() == EObjectStateType::Dynamic) ||
-		   (ParticleHandle->ObjectState() == EObjectStateType::Sleeping);
+	extern int32 ChaosSolverCollisionPositionShockPropagationIterations;
+	extern int32 ChaosSolverCollisionVelocityShockPropagationIterations;
+	extern bool bChaosSolverPersistentGraph;
+
+	bool bChaosConstraintGraphValidate = (CHAOS_CONSTRAINTGRAPH_CHECK_ENABLED != 0);
+	FAutoConsoleVariableRef CVarChaosConstraintGraphValidate(TEXT("p.Chaos.ConstraintGraph.Validate"), bChaosConstraintGraphValidate, TEXT("Enable per-tick ConstraintGraph validation checks/assertions"));
+
+	/** Cvar to enable/disable the island sleeping */
+	bool bChaosSolverSleepEnabled = true;
+	FAutoConsoleVariableRef CVarChaosSolverSleepEnabled(TEXT("p.Chaos.Solver.SleepEnabled"), bChaosSolverSleepEnabled, TEXT(""));
+
+	/** Cvar to override the sleep counter threshold if necessary */
+	int32 ChaosSolverCollisionDefaultSleepCounterThreshold = 20;
+	FAutoConsoleVariableRef CVarChaosSolverCollisionDefaultSleepCounterThreshold(TEXT("p.ChaosSolverCollisionDefaultSleepCounterThreshold"), ChaosSolverCollisionDefaultSleepCounterThreshold, TEXT("Default counter threshold for sleeping.[def:20]"));
+
+	/** Cvar to override the sleep linear threshold if necessary */
+	FRealSingle ChaosSolverCollisionDefaultLinearSleepThreshold = 0.001f; // .001 unit mass cm
+	FAutoConsoleVariableRef CVarChaosSolverCollisionDefaultLinearSleepThreshold(TEXT("p.ChaosSolverCollisionDefaultLinearSleepThreshold"), ChaosSolverCollisionDefaultLinearSleepThreshold, TEXT("Default linear threshold for sleeping.[def:0.001]"));
+
+	/** Cvar to override the sleep angular threshold if necessary */
+	FRealSingle ChaosSolverCollisionDefaultAngularSleepThreshold = 0.0087f;  //~1/2 unit mass degree
+	FAutoConsoleVariableRef CVarChaosSolverCollisionDefaultAngularSleepThreshold(TEXT("p.ChaosSolverCollisionDefaultAngularSleepThreshold"), ChaosSolverCollisionDefaultAngularSleepThreshold, TEXT("Default angular threshold for sleeping.[def:0.0087]"));
 }
 
-/** Check if a particle is not moving */
-bool IsStationaryParticle(const FGeometryParticleHandle* ParticleHandle)
+
+namespace Chaos::Private
 {
-	if (ParticleHandle->ObjectState() == EObjectStateType::Kinematic)
+	//////////////////////////////////////////////////////////////////////////////////////////////////
+	//////////////////////////////////////////////////////////////////////////////////////////////////
+	// 
+	// Utility Functions
+	// 
+	//////////////////////////////////////////////////////////////////////////////////////////////////
+	//////////////////////////////////////////////////////////////////////////////////////////////////
+
+
+	bool IsParticleDynamic(const FGeometryParticleHandle* Particle)
 	{
-		// For kinematic particles check whether their current mode and target will cause the particle to move
-		const FKinematicGeometryParticleHandle* KinematicParticle = ParticleHandle->CastToKinematicParticle();
-		if (KinematicParticle->KinematicTarget().GetMode() == EKinematicTargetMode::Position)
-		{
-			return (KinematicParticle->X() - KinematicParticle->KinematicTarget().GetTargetPosition()).IsZero()
-					&& (KinematicParticle->R() * KinematicParticle->KinematicTarget().GetTargetRotation().Inverse()).IsIdentity();
-		}
-		// For all other modes (Velocity, None and Reset) check that the velocity is non-zero
-		else
-		{
-			return KinematicParticle->V().IsZero() && KinematicParticle->W().IsZero();
-		}
+		return (Particle != nullptr) && ((Particle->ObjectState() == EObjectStateType::Dynamic) || (Particle->ObjectState() == EObjectStateType::Sleeping));
 	}
-	else
+
+	bool IsParticleSleeping(const FGeometryParticleHandle* Particle)
 	{
-		return (ParticleHandle->ObjectState() == EObjectStateType::Static) ||
-			   (ParticleHandle->ObjectState() == EObjectStateType::Sleeping);
+		return (Particle != nullptr) && (Particle->ObjectState() == EObjectStateType::Sleeping);
 	}
-}
-
-/** Check if an island is sleeping or not given a linear/angular velocities and sleeping thresholds */
-inline bool IsIslandSleeping(const FReal MaxLinearSpeed2, const FReal MaxAngularSpeed2,
-								  const FReal LinearSleepingThreshold, const FReal AngularSleepingThreshold,
-								  const int32 CounterThreshold, int32& SleepCounter)
-{
-	const FReal MaxLinearSpeed = FMath::Sqrt(MaxLinearSpeed2);
-	const FReal MaxAngularSpeed = FMath::Sqrt(MaxAngularSpeed2);
 
-	if (MaxLinearSpeed < LinearSleepingThreshold && MaxAngularSpeed < AngularSleepingThreshold)
+	bool WasParticleAsleep(const FGeometryParticleHandle* Particle)
 	{
-		if (SleepCounter >= CounterThreshold)
+		if (const FPBDRigidParticleHandle* Rigid = Particle->CastToRigidParticle())
 		{
-			return true;
+			return Rigid->WasSleeping();
 		}
-		else
+		return false;
+	}
+
+	bool IsParticleMoving(const FGeometryParticleHandle* Particle)
+	{
+		bool bIsStationary = true;
+		if (Particle != nullptr)
 		{
-			SleepCounter++;
+			if (Particle->ObjectState() == EObjectStateType::Kinematic)
+			{
+				// For kinematic particles check whether their current mode and target will cause the particle to move
+				// For all other modes (Velocity, None and Reset) check that the velocity is non-zero
+				const FKinematicGeometryParticleHandle* Kinematic = Particle->CastToKinematicParticle();
+				const FKinematicTarget& KinematicTarget = Kinematic->KinematicTarget();
+				if (KinematicTarget.GetMode() == EKinematicTargetMode::Position)
+				{
+					bIsStationary = (Kinematic->X() - KinematicTarget.GetTargetPosition()).IsZero() && (Kinematic->R() * KinematicTarget.GetTargetRotation().Inverse()).IsIdentity();
+				}
+				else
+				{
+					bIsStationary = Kinematic->V().IsZero() && Kinematic->W().IsZero();
+				}
+			}
+			else
+			{
+				bIsStationary = (Particle->ObjectState() == EObjectStateType::Static) || (Particle->ObjectState() == EObjectStateType::Sleeping);
+			}
 		}
+		return !bIsStationary;
 	}
-	else
+
+	bool IsParticleNeedsResim(const FGeometryParticleHandle* Particle)
 	{
-		SleepCounter = 0;
+		return (Particle != nullptr) && (Particle->SyncState() != ESyncState::InSync);
 	}
-	return false;
-}
-	
-/** Compute sleeping thresholds given a solver island  */
-inline bool ComputeSleepingThresholds(
-	FPBDIsland* Island,
-	const TArrayCollectionArray<TSerializablePtr<FChaosPhysicsMaterial>>& PhysicsMaterials,
-	const TArrayCollectionArray<TUniquePtr<FChaosPhysicsMaterial>>& PerParticlePhysicsMaterials,
-	const THandleArray<FChaosPhysicsMaterial>& SimMaterials,
-	FReal& LinearSleepingThreshold, 
-	FReal& AngularSleepingThreshold, 
-	FReal& MaxLinearSpeed2, 
-	FReal& MaxAngularSpeed2, 
-	int32& SleepCounterThreshold)
-{
-	LinearSleepingThreshold = FLT_MAX;  
-	AngularSleepingThreshold = FLT_MAX; 
-	MaxLinearSpeed2 = 0.f; 
-	MaxAngularSpeed2 = 0.f;
-	SleepCounterThreshold = 0;
 
-	bool bHaveSleepThreshold = false;
-	for (FPBDIslandParticle& IslandParticle : Island->GetParticles())
+	bool GetParticleSleepThresholds(
+		const FGeometryParticleHandle* Particle,
+		const TArrayCollectionArray<TSerializablePtr<FChaosPhysicsMaterial>>* PhysicsMaterials,
+		const TArrayCollectionArray<TUniquePtr<FChaosPhysicsMaterial>>* PerParticlePhysicsMaterials,
+		const THandleArray<FChaosPhysicsMaterial>* SimMaterials,
+		FRealSingle& OutSleepLinearThreshold,
+		FRealSingle& OutSleepAngularThreshold,
+		int& OutSleepCounterThreshold)
 	{
-		FGeometryParticleHandle* ParticleHandle = IslandParticle.GetParticle();
-		if (ParticleHandle)
+		OutSleepLinearThreshold = 0;
+		OutSleepAngularThreshold = 0;
+		OutSleepCounterThreshold = TNumericLimits<int32>::Max();
+
+		const FPBDRigidParticleHandle* Rigid = Particle->CastToRigidParticle();
+		if ((Rigid != nullptr) && (Rigid->SleepType() != ESleepType::NeverSleep))
 		{
-			if (FPBDRigidParticleHandle* PBDRigid = ParticleHandle->CastToRigidParticle())
+			const FChaosPhysicsMaterial* PhysicsMaterial = Private::GetFirstPhysicsMaterial(Rigid, PhysicsMaterials, PerParticlePhysicsMaterials, SimMaterials);
+			if (PhysicsMaterial != nullptr)
 			{
-				// Should we change this condition to be if (!IsStationaryParticle(ParticleHandle))
-				// to be in sync with what is done for the graph island sleeping flag?
-				if	(IsDynamicParticle(ParticleHandle) && !PBDRigid->Sleeping())
-				{
-					// If any body in the island is not allowed to sleep, the whole island cannot sleep
-					// @todo(chaos): if this is a common thing, we should set a flag on the island when it has a particle
-					// with this property enabled and skip the sleep check altogether
-					if (PBDRigid->SleepType() == ESleepType::NeverSleep)
-					{
-						return false;
-					}
-
-					bHaveSleepThreshold = true;
+				OutSleepLinearThreshold = FRealSingle(PhysicsMaterial->SleepingLinearThreshold);
+				OutSleepAngularThreshold = FRealSingle(PhysicsMaterial->SleepingAngularThreshold);
+				OutSleepCounterThreshold = PhysicsMaterial->SleepCounterThreshold;
+			}
+			else
+			{
+				OutSleepLinearThreshold = CVars::ChaosSolverCollisionDefaultLinearSleepThreshold;
+				OutSleepAngularThreshold = CVars::ChaosSolverCollisionDefaultAngularSleepThreshold;
+				OutSleepCounterThreshold = CVars::ChaosSolverCollisionDefaultSleepCounterThreshold;
+			}
 
-					MaxLinearSpeed2 = FMath::Max(PBDRigid->VSmooth().SizeSquared(), MaxLinearSpeed2);
-					MaxAngularSpeed2 = FMath::Max(PBDRigid->WSmooth().SizeSquared(), MaxAngularSpeed2);
+			return true;
+		}
 
-					const FChaosPhysicsMaterial* PhysicsMaterial = Private::GetFirstPhysicsMaterial(PBDRigid, PhysicsMaterials, PerParticlePhysicsMaterials, &SimMaterials);
+		return false;
+	}
 
-					const FReal LocalSleepingLinearThreshold = PhysicsMaterial ? PhysicsMaterial->SleepingLinearThreshold :
-										ChaosSolverCollisionDefaultLinearSleepThresholdCVar;
-					const FReal LocalAngularSleepingThreshold = PhysicsMaterial ? PhysicsMaterial->SleepingAngularThreshold :
-										ChaosSolverCollisionDefaultAngularSleepThresholdCVar;
-					const int32 LocalSleepCounterThreshold = PhysicsMaterial ? PhysicsMaterial->SleepCounterThreshold :
-										ChaosSolverCollisionDefaultSleepCounterThresholdCVar;
+	bool GetParticleDisableThresholds(
+		const FGeometryParticleHandle* Particle,
+		const TArrayCollectionArray<TSerializablePtr<FChaosPhysicsMaterial>>* PhysicsMaterials,
+		const TArrayCollectionArray<TUniquePtr<FChaosPhysicsMaterial>>* PerParticlePhysicsMaterials,
+		const THandleArray<FChaosPhysicsMaterial>* SimMaterials,
+		FRealSingle& OutDisableLinearThreshold,
+		FRealSingle& OutDisableAngularThreshold)
+	{
+		OutDisableLinearThreshold = 0;
+		OutDisableAngularThreshold = 0;
 
-					LinearSleepingThreshold = FMath::Min(LinearSleepingThreshold, LocalSleepingLinearThreshold);
-					AngularSleepingThreshold = FMath::Min(AngularSleepingThreshold, LocalAngularSleepingThreshold);
-					SleepCounterThreshold = FMath::Max(SleepCounterThreshold, LocalSleepCounterThreshold);
-				}
+		const FPBDRigidParticleHandle* Rigid = Particle->CastToRigidParticle();
+		if ((Rigid != nullptr) && (Rigid->SleepType() != ESleepType::NeverSleep))
+		{
+			const FChaosPhysicsMaterial* PhysicsMaterial = Private::GetFirstPhysicsMaterial(Rigid, PhysicsMaterials, PerParticlePhysicsMaterials, SimMaterials);
+			if (PhysicsMaterial != nullptr)
+			{
+				OutDisableLinearThreshold = FRealSingle(PhysicsMaterial->DisabledLinearThreshold);
+				OutDisableAngularThreshold = FRealSingle(PhysicsMaterial->DisabledAngularThreshold);
+				return true;
 			}
 		}
+
+		return false;
 	}
-	return bHaveSleepThreshold;
-}
 
-FPBDIslandManager::FPBDIslandManager() 
-	: Islands()
-	, IslandGraph(MakeUnique<GraphType>())
-	, MaxParticleIndex(INDEX_NONE)
-	, bIslandsPopulated(false)
-	, bEndTickCalled(true)
-{
-	IslandGraph->SetOwner(this);
-}
 
-FPBDIslandManager::~FPBDIslandManager()
-{
-}
+	//////////////////////////////////////////////////////////////////////////////////////////////////
+	//////////////////////////////////////////////////////////////////////////////////////////////////
+	// 
+	// FPBDIslandParticle
+	// 
+	//////////////////////////////////////////////////////////////////////////////////////////////////
+	//////////////////////////////////////////////////////////////////////////////////////////////////
 
-void FPBDIslandManager::SetIsDeterministic(const bool bInIsDeterministic)
-{
-	IslandGraph->SetIsDeterministic(bInIsDeterministic);
-}
 
-void FPBDIslandManager::Reset()
-{
-	RemoveConstraints();
-	RemoveParticles();
+	FPBDIslandParticle::FPBDIslandParticle()
+	{
+	}
 
-	Islands.Reset();
-}
+	FPBDIslandParticle::FPBDIslandParticle(FGeometryParticleHandle* InParticle)
+	{
+		Reuse(InParticle);
+	}
 
-void FPBDIslandManager::AddConstraintContainer(FPBDConstraintContainer& ConstraintContainer)
-{
-	const int32 ContainerId = ConstraintContainer.GetContainerId();
+	FPBDIslandParticle::~FPBDIslandParticle()
+	{
+		Trash();
+	}
 
-	if (ConstraintContainers.Num() <= ContainerId)
+	int32 FPBDIslandParticle::GetIslandId() const
 	{
-		ConstraintContainers.SetNum(ContainerId + 1);
+		if (Island != nullptr)
+		{
+			return Island->GetIslandId();
+		}
+		return INDEX_NONE;
 	}
 
-	ConstraintContainers[ContainerId] = &ConstraintContainer;
-}
+	void FPBDIslandParticle::Reuse(FGeometryParticleHandle* InParticle)
+	{
+		Particle = InParticle;
+
+		if (Particle != nullptr)
+		{
+			Particle->SetConstraintGraphNode(this);
+		}
+	}
+
+	void FPBDIslandParticle::Trash()
+	{
+		check((Particle == nullptr) || (Particle->GetConstraintGraphNode() == this));
+		check(ArrayIndex == INDEX_NONE);
+		check(Edges.IsEmpty());
+		check(Island == nullptr);
+		check(IslandArrayIndex == INDEX_NONE);
+		check(ArrayIndex == INDEX_NONE);
+
+		if (Particle != nullptr)
+		{
+			Particle->SetConstraintGraphNode(nullptr);
+			Particle = nullptr;
+		}
+
+		Level = 0;
+		Flags.Reset();
+		VisitEpoch = INDEX_NONE;
+		ResimFrame = INDEX_NONE;
+	}
 
-void FPBDIslandManager::InitializeGraph(const TParticleView<FPBDRigidParticles>& PBDRigids)
-{
-	// If we hit this, there's a missing call to EndTick since the previous tick
-	ensure(bEndTickCalled);
-	bEndTickCalled = false;
 
-	// @todo(chaos): make this handle when the user sleeps all particles in an island. Currently this
-	// will put the island to sleep as expected, but only after removing all the constraints from the island.
-	// The result is that the particles in the island will have no collisions on the first frame after waking.
-	// See unit test: DISABLED_TestConstraintGraph_ParticleSleep_Manual
+	//////////////////////////////////////////////////////////////////////////////////////////////////
+	//////////////////////////////////////////////////////////////////////////////////////////////////
+	// 
+	// FPBDIslandConstraint
+	// 
+	//////////////////////////////////////////////////////////////////////////////////////////////////
+	//////////////////////////////////////////////////////////////////////////////////////////////////
 
-	MaxParticleIndex = 0;
-	ReserveParticles(PBDRigids.Num());
 
-	// Reset the island merge tracking data - adding particles may cause islands to merge (e.g., if a particle has switched from kinematic to dynamic
-	// @todo(chaos): put this into a method on the graph. Or maybe this isn't necessary since Merge/Split should always leave a valid empty state...
-	for (auto& GraphIsland : IslandGraph->GraphIslands)
+	FPBDIslandConstraint::FPBDIslandConstraint()
 	{
-		GraphIsland.ChildrenIslands.Reset();
-		GraphIsland.ParentIsland = INDEX_NONE;
 	}
 
-	// We update the valid/steady state of the nodes in case any state changed. We don't need to do this for newly added particles, 
-	// so we do it before AddParticles and AddConstraints
-	// @todo(chaos): can we do this when the properties change instead?
-	for(int32 NodeIndex = 0, NumNodes = IslandGraph->GraphNodes.GetMaxIndex(); NodeIndex < NumNodes; ++NodeIndex)
+	FPBDIslandConstraint::FPBDIslandConstraint(const int32 InContainerId, FConstraintHandle* InConstraint)
 	{
-		if (IslandGraph->GraphNodes.IsValidIndex(NodeIndex))
-		{
-			FGeometryParticleHandle* ParticleHandle = IslandGraph->GraphNodes[NodeIndex].NodeItem;
+		Reuse(InContainerId, InConstraint);
+	}
 
-			IslandGraph->UpdateNode(NodeIndex, IsDynamicParticle(ParticleHandle), IsStationaryParticle(ParticleHandle));
+	FPBDIslandConstraint::~FPBDIslandConstraint()
+	{
+		Trash();
+	}
 
-			// If the particle is static or kinematic and has no edges, remove it
-			if (!IslandGraph->GraphNodes[NodeIndex].bValidNode)
-			{
-				if (IslandGraph->GraphNodes[NodeIndex].NodeEdges.Num() == 0)
-				{
-					RemoveParticle(ParticleHandle);
-				}
-			}
+	int32 FPBDIslandConstraint::GetIslandId() const
+	{
+		if (Island != nullptr)
+		{
+			return Island->GetIslandId();
 		}
+		return INDEX_NONE;
 	}
 
-	if (!bChaosSolverPersistentGraph)
+	void FPBDIslandConstraint::Reuse(const int32 InContainerId, FConstraintHandle* InConstraint)
 	{
-		// We are adding all the particles from the solver in case some were just created/activated
-		for (auto& RigidParticle : PBDRigids)
+		ContainerIndex = InContainerId;
+		Constraint = InConstraint;
+
+		if (Constraint != nullptr)
 		{
-			AddParticle(RigidParticle.Handle());
+			Constraint->SetConstraintGraphEdge(this);
 		}
-
-		// NOTE: This only removes constraints from awake islands and it leaves the sleeping ones. This is important for 
-		// semi-persistent collisions because we do not run collision detection on sleeping particles, and we need the graph
-		// to keep it's sleeping edges so that we know what collisions to wake when an island is awakened. The constraint
-		// management system also uses the sleeping state as a lock to prevent constraint destruction.
-		IslandGraph->RemoveAllAwakeEdges();
 	}
-}
 
-void FPBDIslandManager::EndTick()
-{
-	// Islands are transient for use during the solver phase - clear them at the end of the frame
-	for (auto& Island : Islands)
+	void FPBDIslandConstraint::Trash()
 	{
-		Island->ClearParticles();
-		Island->ClearConstraints();
-	}
+		check((Constraint == nullptr) || (Constraint->GetConstraintGraphEdge() == this));
+		check(Nodes[0] == nullptr);
+		check(Nodes[1] == nullptr);
+		check(NodeArrayIndices[0] == INDEX_NONE);
+		check(NodeArrayIndices[1] == INDEX_NONE);
+		check(Island == nullptr);
+		check(IslandArrayIndex == INDEX_NONE);
+		check(ArrayIndex == INDEX_NONE);
+
+		if (Constraint != nullptr)
+		{
+			Constraint->SetConstraintGraphEdge(nullptr);
+			Constraint = nullptr;
+		}
 
-	// Let the graph perform and end-of-frame cleanup
-	IslandGraph->EndTick();
+		VisitEpoch = INDEX_NONE;
+		ContainerIndex = INDEX_NONE;
+		Level = INDEX_NONE;
+		Flags.Reset();
+	}
 
-	// We can no longer use the Island's particle and constraint lists
-	bIslandsPopulated = false;
 
-	bEndTickCalled = true;
-}
+	//////////////////////////////////////////////////////////////////////////////////////////////////
+	//////////////////////////////////////////////////////////////////////////////////////////////////
+	// 
+	// FPBDIsland
+	// 
+	//////////////////////////////////////////////////////////////////////////////////////////////////
+	//////////////////////////////////////////////////////////////////////////////////////////////////
 
-void FPBDIslandManager::RemoveParticles()
-{
-	RemoveConstraints();
 
-	IslandGraph->ResetNodes();
+	FPBDIsland::FPBDIsland()
+	{
+		Reuse();
+	}
 
-	for (auto& Island : Islands)
+	FPBDIsland::~FPBDIsland()
 	{
-		Island->ClearParticles();
+		Trash();
 	}
-}
 
-void FPBDIslandManager::RemoveConstraints()
-{
-	IslandGraph->ResetEdges();
+	void FPBDIsland::Reuse()
+	{
+		Flags.bIsSleepAllowed = true;
+		Flags.bIsSleeping = true;
+	}
 
-	for(auto& Island : Islands)
+	void FPBDIsland::Trash()
 	{
-		Island->ClearConstraints();
-		Island->ClearParticles();
+		check(Nodes.IsEmpty());
+		check(NumEdges == 0);
+		check(ArrayIndex == INDEX_NONE);
+
+		MergeSet = nullptr;
+		MergeSetIslandIndex = INDEX_NONE;
+		SleepCounter = 0;
+		DisableCounter = 0;
+		ResimFrame = INDEX_NONE;
+		Flags.Reset();
 	}
-}
 
-int32 FPBDIslandManager::ReserveParticles(const int32 NumParticles)
-{
-	const int32 MaxIndex = IslandGraph->MaxNumNodes();
-	IslandGraph->ReserveNodes(NumParticles);
-	
-	Islands.Reserve(NumParticles);
-	IslandIndexing.Reserve(NumParticles);
-	
-	return FMath::Max(0,NumParticles - MaxIndex);
-}
+	void FPBDIsland::UpdateSyncState()
+	{
+		Flags.bNeedsResim = false;
+		ResimFrame = INDEX_NONE;
 
-void FPBDIslandManager::ReserveConstraints(const int32 NumConstraints)
-{
-	IslandGraph->ReserveEdges(NumConstraints);
-}
+		// If any of our particles want a resim, the whole island is resimmed
+		for (TArray<FPBDIslandConstraint*>& Edges : ContainerEdges)
+		{
+			for (FPBDIslandConstraint* Edge : Edges)
+			{
+				for (int32 NodeIndex = 0; NodeIndex < 2; ++NodeIndex)
+				{
+					if (FPBDIslandParticle* EdgeNode = Edge->Nodes[NodeIndex])
+					{
+						Flags.bNeedsResim = ((Flags.bNeedsResim | EdgeNode->Flags.bNeedsResim) != 0);
+						ResimFrame = (ResimFrame == INDEX_NONE) ? EdgeNode->ResimFrame : FMath::Min(ResimFrame, EdgeNode->ResimFrame);
+					}
+				}
+			}
+		}
+	}
 
-// Callback from the IslandGraph to allow us to store the node index
-void FPBDIslandManager::GraphNodeAdded(FGeometryParticleHandle* ParticleHandle, const int32 NodeIndex)
-{
-	check(ParticleHandle->ConstraintGraphIndex() == INDEX_NONE);
+	//////////////////////////////////////////////////////////////////////////////////////////////////
+	//////////////////////////////////////////////////////////////////////////////////////////////////
+	// 
+	// FPBDIslandMergeSet
+	// 
+	//////////////////////////////////////////////////////////////////////////////////////////////////
+	//////////////////////////////////////////////////////////////////////////////////////////////////
 
-	ParticleHandle->SetConstraintGraphIndex(NodeIndex);
-}
 
-void FPBDIslandManager::GraphNodeRemoved(FGeometryParticleHandle* ParticleHandle)
-{
-	check(ParticleHandle->ConstraintGraphIndex() != INDEX_NONE);
+	FPBDIslandMergeSet::FPBDIslandMergeSet()
+	{
+		Reuse();
+	}
 
-	ParticleHandle->SetConstraintGraphIndex(INDEX_NONE);
-}
+	FPBDIslandMergeSet::~FPBDIslandMergeSet()
+	{
+		Trash();
+	}
 
-// Callback from the IslandGraph to allow us to store the edge index
-void FPBDIslandManager::GraphEdgeAdded(const FConstraintHandleHolder& ConstraintHandle, const int32 EdgeIndex)
-{
-	check(ConstraintHandle->GetConstraintGraphIndex() == INDEX_NONE);
+	void FPBDIslandMergeSet::Reuse()
+	{
+	}
 
-	ConstraintHandle->SetConstraintGraphIndex(EdgeIndex);
-}
+	void FPBDIslandMergeSet::Trash()
+	{
+		check(ArrayIndex == INDEX_NONE);
 
-void FPBDIslandManager::GraphEdgeRemoved(const FConstraintHandleHolder& ConstraintHandle)
-{
-	check(ConstraintHandle->GetConstraintGraphIndex() != INDEX_NONE);
+		Islands.Reset();
+		NumEdges = 0;
+	}
 
-	ConstraintHandle->SetConstraintGraphIndex(INDEX_NONE);
-}
 
-int32 FPBDIslandManager::AddParticle(FGeometryParticleHandle* ParticleHandle)
-{
-#if CHAOS_CONSTRAINTHANDLE_DEBUG_ENABLED
-	ensure(!FConstGenericParticleHandle(ParticleHandle)->Disabled());
-#endif
+	//////////////////////////////////////////////////////////////////////////////////////////////////
+	//////////////////////////////////////////////////////////////////////////////////////////////////
+	// 
+	// FPBDIslandManager
+	// 
+	//////////////////////////////////////////////////////////////////////////////////////////////////
+	//////////////////////////////////////////////////////////////////////////////////////////////////
+
+
+	FPBDIslandManager::FPBDIslandManager(FPBDRigidsSOAs& InParticles)
+		: Particles(InParticles)
+		, PhysicsMaterials(nullptr)
+		, PerParticlePhysicsMaterials(nullptr)
+		, SimMaterials(nullptr)
+		, Nodes(1000)
+		, Edges(1000)
+		, Islands(1000)
+		, MergeSets(1000)
+	{
+	}
 
-	return TryAddParticleIfDynamic(ParticleHandle,  /*IslandIndex*/INDEX_NONE);
-}
+	FPBDIslandManager::~FPBDIslandManager()
+	{
+	}
 
-int32 FPBDIslandManager::TryAddParticle(FGeometryParticleHandle* ParticleHandle, const int32 IslandIndex)
-{
-	if (ParticleHandle)
+	int32 FPBDIslandManager::GetNextVisitEpoch()
 	{
-		if (!ParticleHandle->IsInConstraintGraph())
+		// Don't wrap to negative because INDEX_NONE is used as "uninitialized epoch"
+		if (NextVisitEpoch == TNumericLimits<int32>::Max())
 		{
-			MaxParticleIndex = FMath::Max(MaxParticleIndex, ParticleHandle->UniqueIdx().Idx);
-
-			const bool bIsDynamic = IsDynamicParticle(ParticleHandle);
-			const bool bIsStationary = IsStationaryParticle(ParticleHandle);
-			IslandGraph->AddNode(ParticleHandle, bIsDynamic, IslandIndex, bIsStationary);
+			NextVisitEpoch = 0;
 		}
+		return NextVisitEpoch++;
+	}
 
-		return ParticleHandle->ConstraintGraphIndex();
+	void FPBDIslandManager::SetIsDeterministic(const bool bInIsDeterministic)
+	{
+		bIsDeterministic = bInIsDeterministic;
+		ApplyDeterminism();
 	}
-	return INDEX_NONE;
-}
 
-int32 FPBDIslandManager::TryAddParticleIfDynamic(FGeometryParticleHandle* ParticleHandle, const int32 IslandIndex)
-{
-	if (ParticleHandle)
+	void FPBDIslandManager::AddConstraintContainer(const FPBDConstraintContainer& Container)
 	{
-		if (IsDynamicParticle(ParticleHandle))
+		if (ConstraintContainers.Num() < Container.GetContainerId() + 1)
 		{
-			return TryAddParticle(ParticleHandle, IslandIndex);
+			ConstraintContainers.SetNumZeroed(Container.GetContainerId() + 1, false);
 		}
-
-		return ParticleHandle->ConstraintGraphIndex();
+		ConstraintContainers[Container.GetContainerId()] = &Container;
 	}
-	return INDEX_NONE;
-}
 
-int32 FPBDIslandManager::TryAddParticleIfNotDynamic(FGeometryParticleHandle* ParticleHandle, const int32 IslandIndex)
-{
-	if (ParticleHandle)
+	void FPBDIslandManager::RemoveConstraintContainer(const FPBDConstraintContainer& Container)
 	{
-		if (!IsDynamicParticle(ParticleHandle))
+		for (int32 EdgeIndex = Edges.Num() - 1; EdgeIndex >= 0; --EdgeIndex)
 		{
-			return TryAddParticle(ParticleHandle, IslandIndex);
+			if (Edges[EdgeIndex]->ContainerIndex == Container.GetContainerId())
+			{
+				RemoveConstraint(Edges[EdgeIndex]->Constraint);
+			}
 		}
 
-		return ParticleHandle->ConstraintGraphIndex();
+		ConstraintContainers[Container.GetContainerId()] = nullptr;
 	}
-	return INDEX_NONE;
-}
 
-void FPBDIslandManager::AddConstraint(const uint32 ContainerId, FConstraintHandle* ConstraintHandle, const TVec2<FGeometryParticleHandle*>& ConstrainedParticles)
-{
-	if (ConstraintHandle)
+	void FPBDIslandManager::SetMaterialContainers(const TArrayCollectionArray<TSerializablePtr<FChaosPhysicsMaterial>>* InPhysicsMaterials, const TArrayCollectionArray<TUniquePtr<FChaosPhysicsMaterial>>* InPerParticlePhysicsMaterials, const THandleArray<FChaosPhysicsMaterial>* InSimMaterials)
+	{
+		PhysicsMaterials = InPhysicsMaterials;
+		PerParticlePhysicsMaterials = InPerParticlePhysicsMaterials;
+		SimMaterials = InSimMaterials;
+	}
+
+	void FPBDIslandManager::SetGravityForces(const FPerParticleGravity* InGravity)
+	{
+		Gravity = InGravity;
+	}
+
+	void FPBDIslandManager::SetDisableCounterThreshold(const int32 InDisableCounterThreshold)
+	{
+		DisableCounterThreshold = InDisableCounterThreshold;
+	}
+
+	int32 FPBDIslandManager::GetNumConstraintContainers() const
 	{
-		if (ConstraintHandle->GetConstraintGraphIndex() != INDEX_NONE)
+		return ConstraintContainers.Num();
+	}
+
+	void FPBDIslandManager::Reset()
+	{
+		for (int32 EdgeIndex = Edges.Num() - 1; EdgeIndex >= 0; --EdgeIndex)
 		{
-			// Already in the graph
-			return;
+			RemoveConstraint(Edges[EdgeIndex]->Constraint);
 		}
+		check(Edges.IsEmpty());
 
-		// Are the particles dynamic (including asleep)?
-		const bool bDynamicParticle0 = ConstrainedParticles[0] && IsDynamicParticle(ConstrainedParticles[0]);
-		const bool bDynamicParticle1 = ConstrainedParticles[1] && IsDynamicParticle(ConstrainedParticles[1]);
-		
-		// We only add the constraint if one particle is non-kinematic
-		// @todo(chaos): we can relax the above restriction now if necessary, but so far no need
-		if (bDynamicParticle0 || bDynamicParticle1)
+		for (int32 NodeIndex = Nodes.Num() - 1; NodeIndex >= 0; --NodeIndex)
 		{
-			// If either of the particles is kinematic, it will not have been added by AddParticle, so we must add it here
-			const int32 NodeIndex0 = TryAddParticleIfNotDynamic(ConstrainedParticles[0], /*IslandIndex*/INDEX_NONE);
-			const int32 NodeIndex1 = TryAddParticleIfNotDynamic(ConstrainedParticles[1], /*IslandIndex*/INDEX_NONE);
+			RemoveParticle(Nodes[NodeIndex]->Particle);
+		}
+		check(Nodes.IsEmpty());
+
+		MergeSets.Reset();
+		Islands.Reset();
+		Edges.Reset();
+		Nodes.Reset();
+	}
 
-			// All required particles should be in the graph at this point
-			check((NodeIndex0 != INDEX_NONE) || (ConstrainedParticles[0] == nullptr));
-			check((NodeIndex1 != INDEX_NONE) || (ConstrainedParticles[1] == nullptr));
+	int32 FPBDIslandManager::ReserveParticles(const int32 InNumParticles)
+	{
+		const int32 NumNodes = Nodes.Num();
+		Nodes.Reserve(InNumParticles);
+		return FMath::Max(0, InNumParticles - NumNodes);
+	}
 
-			// Add the constraint to the graph
-			const int32 EdgeIndex = IslandGraph->AddEdge(ConstraintHandle, ContainerId, NodeIndex0, NodeIndex1);
+	int32 FPBDIslandManager::GetParticleLevel(FGeometryParticleHandle* Particle) const
+	{
+		if (const FPBDIslandParticle* Node = GetGraphNode(Particle))
+		{
+			return Node->Level;
+		}
+		return INDEX_NONE;
+	}
 
-			// If we were added to a sleeping island, make sure the constraint is flagged as sleeping.
-			// Adding constraints between 2 sleeping particles does not wake them because we need to handle
-			// streaming which may amortize particle and constraint creation over multiple ticks.
-			const FGraphEdge& GraphEdge = IslandGraph->GraphEdges[EdgeIndex];
-			if (GraphEdge.IslandIndex != INDEX_NONE)
+	void FPBDIslandManager::UpdateParticleMaterial(FGeometryParticleHandle* Particle)
+	{
+		if (FPBDIslandParticle* Node = GetGraphNode(Particle))
+		{
+			if (Node->Flags.bIsDynamic)
 			{
-				const bool bIsSleeping = IslandGraph->GraphIslands[GraphEdge.IslandIndex].bIsSleeping;
-				if (bIsSleeping)
-				{
-					ConstraintHandle->SetIsSleeping(bIsSleeping);
-				}
+				UpdateGraphNodeSleepSettings(Node);
 			}
 		}
 	}
-}
 
-void FPBDIslandManager::RemoveParticle(FGeometryParticleHandle* ParticleHandle)
-{
-	if (ParticleHandle)
+	const FPBDIsland* FPBDIslandManager::GetParticleIsland(const FGeometryParticleHandle* Particle) const
 	{
-		IslandGraph->RemoveNode(ParticleHandle);
+		if (const FPBDIslandParticle* Node = GetGraphNode(Particle))
+		{
+			return Node->Island;
+		}
+		return nullptr;
+	}
 
-#if CHAOS_CONSTRAINTHANDLE_DEBUG_DETAILED_ENABLED
-		DebugCheckParticleNotInGraph(ParticleHandle);
-#endif
+	const FPBDIsland* FPBDIslandManager::GetConstraintIsland(const FConstraintHandle* Constraint) const
+	{
+		if (const FPBDIslandConstraint* Edge = GetGraphEdge(Constraint))
+		{
+			return Edge->Island;
+		}
+		return nullptr;
 	}
-}
 
-void FPBDIslandManager::RemoveConstraint(const uint32 ContainerId, FConstraintHandle* ConstraintHandle)
-{
-	if (ConstraintHandle)
+	int32 FPBDIslandManager::GetIslandIndex(const FPBDIsland* Island) const
 	{
-		const int32 EdgeIndex = ConstraintHandle->GetConstraintGraphIndex();
-		if (IslandGraph->GraphEdges.IsValidIndex(EdgeIndex))
+		if (Island != nullptr)
 		{
-			IslandGraph->RemoveEdge(EdgeIndex);
+			return Island->GetArrayIndex();
 		}
+		return INDEX_NONE;
+	}
 
-#if CHAOS_CONSTRAINTHANDLE_DEBUG_DETAILED_ENABLED
-		DebugCheckConstraintNotInGraph(ConstraintHandle);
-#endif
+	int32 FPBDIslandManager::GetIslandArrayIndex(const FPBDIslandConstraint* Edge) const
+	{
+		if (Edge != nullptr)
+		{
+			return Edge->IslandArrayIndex;
+		}
+		return INDEX_NONE;
 	}
-}
 
-void FPBDIslandManager::RemoveConstraints(const uint32 ContainerId)
-{
-	for (int32 EdgeIndex = 0; EdgeIndex < IslandGraph->GraphEdges.Num(); ++EdgeIndex)
+	TArray<const FPBDIsland*> FPBDIslandManager::FindParticleIslands(const FGeometryParticleHandle* Particle) const
 	{
-		if (IslandGraph->GraphEdges.IsValidIndex(EdgeIndex))
+		TArray<const FPBDIsland*> ParticleIslands;
+
+		if (const FPBDIslandParticle* Node = GetGraphNode(Particle))
 		{
-			FGraphEdge& GraphEdge = IslandGraph->GraphEdges[EdgeIndex];
-			if (GraphEdge.ItemContainer == ContainerId)
+			if (Node->Flags.bIsDynamic)
+			{
+				// Dynamic particles are in one island
+				if (Node->Island != nullptr)
+				{
+					ParticleIslands.Add(Node->Island);
+				}
+			}
+			else
 			{
-				IslandGraph->RemoveEdge(EdgeIndex);
+				// Kinematic particles are in many islands and we need to check the edges
+				for (FPBDIslandConstraint* Edge : Node->Edges)
+				{
+					if (Edge->Island != nullptr)
+					{
+						ParticleIslands.AddUnique(Edge->Island);
+					}
+				}
 			}
 		}
+
+		return ParticleIslands;
 	}
-}
 
-void FPBDIslandManager::RemoveParticleConstraints(FGeometryParticleHandle* ParticleHandle, const uint32 ContainerId)
-{
-	if (ParticleHandle)
+	TArray<const FGeometryParticleHandle*> FPBDIslandManager::FindParticlesInIslands(const TArray<const FPBDIsland*> InIslands) const
 	{
-		// Find the graph node for this particle
-		if (const int32* PNodeIndex = IslandGraph->ItemNodes.Find(ParticleHandle))
+		TArray<const FGeometryParticleHandle*> IslandParticles;
+
+		for (const FPBDIsland* Island : InIslands)
 		{
-			const int32 NodeIndex = *PNodeIndex;
-			if (IslandGraph->GraphNodes.IsValidIndex(NodeIndex))
+			for (const TArray<FPBDIslandConstraint*>& ContainerEdges : Island->ContainerEdges)
 			{
-				const FGraphNode& GraphNode = IslandGraph->GraphNodes[NodeIndex];
-
-				// Loop over the edges attached to this node and remove the constraints in the specified container
-				// NOTE: Can't use ranged-for iterator when removing items, but we can just loop and remove since it's a sparse array
-				for (int32 NodeEdgeIndex = 0, NumNodeEdges = GraphNode.NodeEdges.GetMaxIndex(); NodeEdgeIndex < NumNodeEdges; ++NodeEdgeIndex)
+				for (const FPBDIslandConstraint* Edge : ContainerEdges)
 				{
-					if (GraphNode.NodeEdges.IsValidIndex(NodeEdgeIndex))
+					if (Edge->Nodes[0] != nullptr)
 					{
-						const int32 EdgeIndex = GraphNode.NodeEdges[NodeEdgeIndex];
-						if (IslandGraph->GraphEdges.IsValidIndex(EdgeIndex))
-						{
-							const FGraphEdge& GraphEdge = IslandGraph->GraphEdges[EdgeIndex];
-
-							// Remove the constraint if it is the right type (same ContainerId)
-							if (GraphEdge.ItemContainer == ContainerId)
-							{
-#if CHAOS_CONSTRAINTHANDLE_DEBUG_ENABLED
-								// All constraints should know their graph index
-								ensure(GraphEdge.EdgeItem->GetConstraintGraphIndex() != INDEX_NONE);
-#endif
-
-								RemoveConstraint(ContainerId, GraphEdge.EdgeItem.Get());
-							}
-						}
+						IslandParticles.AddUnique(Edge->Nodes[0]->GetParticle());
+					}
+					if (Edge->Nodes[1] != nullptr)
+					{
+						IslandParticles.AddUnique(Edge->Nodes[1]->GetParticle());
 					}
 				}
-
-				// Note: we may have left the particle in the island without any constraints on it. Is that ok?
-				// It should be fine as long as the node knows that it is still in it/them, which it does.
 			}
 		}
 
-#if CHAOS_CONSTRAINTHANDLE_DEBUG_ENABLED
-		DebugParticleConstraintsNotInGraph(ParticleHandle, ContainerId);
-#endif
+		return IslandParticles;
 	}
-}
 
-int32 FPBDIslandManager::GetParticleLevel(FGeometryParticleHandle* ParticleHandle) const
-{
-	return IslandGraph->GetNodeItemLevel(ParticleHandle);
-}
+	TArray<const FConstraintHandle*> FPBDIslandManager::FindConstraintsInIslands(const TArray<const FPBDIsland*> InIslands, int32 ContainerId) const
+	{
+		TArray<const FConstraintHandle*> IslandConstraints;
 
-void FPBDIslandManager::ResetIslands(const TParticleView<FPBDRigidParticles>& PBDRigids)
-{
-	// @todo 
-}
+		for (const FPBDIsland* Island : InIslands)
+		{
+			for (const FPBDIslandConstraint* Edge : Island->ContainerEdges[ContainerId])
+			{
+				IslandConstraints.Add(Edge->GetConstraint());
+			}
+		}
 
-inline bool SolverIslandSortPredicate(const TUniquePtr<FPBDIsland>& SolverIslandL, const TUniquePtr<FPBDIsland>& SolverIslandR)
-{
-	return SolverIslandL->GetNumConstraints() < SolverIslandR->GetNumConstraints();
-}
+		return IslandConstraints;
+	}
 
-void FPBDIslandManager::GraphIslandAdded(const int32 IslandIndex)
-{
-	FGraphIsland& GraphIsland = IslandGraph->GraphIslands[IslandIndex];
-	check(GraphIsland.IslandItem == nullptr);
+	void FPBDIslandManager::AddParticle(FGeometryParticleHandle* Particle)
+	{
+	}
 
-	// An graph island was created, so we should assign an island manager island to it
-	Islands.Reserve(IslandGraph->MaxNumIslands());
-	if (!Islands.IsValidIndex(IslandIndex))
+	void FPBDIslandManager::RemoveParticle(FGeometryParticleHandle* Particle)
 	{
-		Islands.EmplaceAt(IslandIndex, MakeUnique<FPBDIsland>(ConstraintContainers.Num()));
+		if (FPBDIslandParticle* Node = GetGraphNode(Particle))
+		{
+			// All constraints for this particle should have been removed first
+			check(Node->Edges.Num() == 0);
+			check(Node->Island == nullptr);
+
+			DestroyGraphNode(Node);
+		}
 	}
 
-	// NOTE: The index may be reused and we may already have an island so we need to reset the island
-	Islands[IslandIndex]->Reuse();
-	GraphIsland.IslandItem = Islands[IslandIndex].Get();
+	void FPBDIslandManager::UpdateParticles()
+	{
+		SCOPE_CYCLE_COUNTER(STAT_IslandManager_UpdateParticles);
 
-}
+		const int32 VisitEpoch = GetNextVisitEpoch();
 
-void FPBDIslandManager::InitIslands()
-{
-	Islands.Reserve(IslandGraph->MaxNumIslands());
-	IslandIndexing.SetNum(IslandGraph->MaxNumIslands(),false);
-	int32 LocalIsland = 0;
+		// Process state changed from registered particles.
+		// It would be nice to iterate over transient particles here but the view we need
+		// (Particles.GetActiveDynamicMovingKinematicParticlesView()) actually iterates over
+		// arrays of handles so we lose the benefits. We would also have to visit nodes that
+		// are not represented in that view
+		for (int32 NodeIndex = 0; NodeIndex < Nodes.Num(); ++NodeIndex)
+		{
+			if (NodeIndex < Nodes.Num() - 1)
+			{
+				// Prefetch the next handle
+				FPlatformMisc::Prefetch(Nodes[NodeIndex + 1]->Particle);
+			}
 
-	// Sync of the solver islands first and reserve the required space
-	for (int32 IslandIndex = 0, NumIslands = IslandGraph->MaxNumIslands(); IslandIndex < NumIslands; ++IslandIndex)
+			FPBDIslandParticle* Node = Nodes[NodeIndex];
+			UpdateGraphNode(Node);
+		}
+
+		// The above loop may leave orphaned nodes in the array when dyanmics are made
+		// kinematic and we have removed any kinematic-kinematic constraints. We don't 
+		// handle node removal above because updating one node to kinematic can result 
+		// in removal of many nodes which is awkward while iterating
+		for (int32 NodeIndex = Nodes.Num() - 1; NodeIndex >= 0; --NodeIndex)
+		{
+			FPBDIslandParticle* Node = Nodes[NodeIndex];
+			if (Node->Edges.IsEmpty())
+			{
+				DestroyGraphNode(Node);
+			}
+		}
+	}
+
+	void FPBDIslandManager::AddConstraint(const int32 ContainerId, FConstraintHandle* Constraint, const TVec2<FGeometryParticleHandle*>& ConstrainedParticles)
 	{
-		if (IslandGraph->GraphIslands.IsValidIndex(IslandIndex))
+		if (GetGraphEdge(Constraint) == nullptr)
 		{
-			FGraphIsland& GraphIsland = IslandGraph->GraphIslands[IslandIndex];
+			// Create the edge representing the constraint. Also creates the nodes if necessary
+			if (FPBDIslandConstraint* Edge = CreateGraphEdge(ContainerId, Constraint, ConstrainedParticles))
+			{
+				// Put the edge/nodes into an island and flag islands for merging if necessary
+				// NOTE: We do not wake the island here. It will be awake already if one of the
+				// particles is awake. It both particles are asleep we do not want to wake them
+				// here - it would break support for bStartAwake==false for connected bodies.
+				AssignEdgeIsland(Edge);
+			}
+		}
+	}
 
-			// NOTE: We do not check !IsSleeping() here because we need to sync the sleep state of islands
-			// that were just put to sleep this tick. See PopulateIslands.
+	void FPBDIslandManager::RemoveConstraint(FConstraintHandle* Constraint)
+	{
+		if (FPBDIslandConstraint* Edge = GetGraphEdge(Constraint))
+		{
+			FPBDIslandParticle* Node0 = Edge->Nodes[0];
+			FPBDIslandParticle* Node1 = Edge->Nodes[1];
+			FPBDIsland* Island = Edge->Island;
 
-			// If the island item is not set on the graph we check if the solver island
-			// at the right index is already there. If not we create a new solver island.
-			check(Islands.IsValidIndex(IslandIndex));
-			FPBDIsland* Island = Islands[IslandIndex].Get();
-			
-			// We then transfer the persistent flag and the graph dense index to the solver island
-			Island->SetIsPersistent(GraphIsland.bIsPersistent);
-			Island->SetIsSleeping(GraphIsland.bIsSleeping);
-			Island->SetIsSleepingChanged(GraphIsland.bIsSleeping != GraphIsland.bWasSleeping);
-			Island->SetIslandIndex(LocalIsland);
+			// Remove edge from the island, and possibly destroy the island (if empty)
+			RemoveEdgeFromIsland(Edge);
 
-			// We update the IslandIndexing to retrieve the graph sparse and persistent index from the dense one.
-			IslandIndexing[LocalIsland] = IslandIndex;
-			LocalIsland++;
+			// Destroy the edge
+			DestroyGraphEdge(Edge);
 
-			Island->ReserveParticles(GraphIsland.NumNodes);
-			Island->SetResimFrame(INDEX_NONE);
+			// We also destroy the nodes if they have no edges left
+			if ((Node0 != nullptr) && Node0->Edges.IsEmpty())
+			{
+				RemoveNodeFromIsland(Node0);
+				DestroyGraphNode(Node0);
+			}
+			if ((Node1 != nullptr) && Node1->Edges.IsEmpty())
+			{
+				RemoveNodeFromIsland(Node1);
+				DestroyGraphNode(Node1);
+			}
 
-			// Reset of the sleep counter if the island is :
-			// - Non persistent since we are starting incrementing the counter once the island is persistent and if values below the threshold
-			// - Sleeping since as soon as it wakes up we could start incrementing the counter as well
-			// - Just woken because we may not have been asleep for a whole frame (this probably makes the previous check unnecessary)
-			if (!Island->IsPersistent() || Island->IsSleeping() || Island->IsSleepingChanged())
-			{ 
-				Island->SetSleepCounter(0);
+			// And destroy the island if it is now empty
+			if ((Island != nullptr) && (Island->NumEdges == 0))
+			{
+				DestroyIsland(Island);
 			}
+		}
+	}
 
-			// We reset the persistent flag to be true on the island graph for next frame. It may get set to false again
-			// if any particles are added or removed
-			// @todo(chaos): It does not feel right to reset this here. Should probably be in EndTick? (but See AdvanceClustering which
-			// is currently before that, and generates new particles that get added to the graph).
-			GraphIsland.bIsPersistent = true;
+	void FPBDIslandManager::RemoveParticleConstraints(FGeometryParticleHandle* Particle, int32 ContainerId)
+	{
+		if (FPBDIslandParticle* Node = GetGraphNode(Particle))
+		{
+			// Reverse loop because we remove as we iterate
+			for (int32 EdgeIndex = Node->Edges.Num() - 1; EdgeIndex >= 0; --EdgeIndex)
+			{
+				FPBDIslandConstraint* Edge = Node->Edges[EdgeIndex];
+				if (Edge->ContainerIndex == ContainerId)
+				{
+					// NOTE: Destroys the edge and possibly this and other nodes (if this was their last constraint)
+					RemoveConstraint(Edge->Constraint);
+				}
+			}
 		}
-		else if (Islands.IsValidIndex(IslandIndex))
+	}
+
+	void FPBDIslandManager::RemoveContainerConstraints(int32 ContainerId)
+	{
+		// Reverse loop because we remove as we iterate
+		for (int32 EdgeIndex = Edges.Num() - 1; EdgeIndex >= 0; --EdgeIndex)
 		{
-			Islands.RemoveAt(IslandIndex);
+			FPBDIslandConstraint* Edge = Edges[EdgeIndex];
+			if (Edge->ContainerIndex == ContainerId)
+			{
+				RemoveConstraint(Edge->Constraint);
+			}
 		}
 	}
 
-	IslandIndexing.SetNum(LocalIsland, false);
-}
+	void FPBDIslandManager::RemoveAllConstraints()
+	{
+		// Reverse loop because we remove as we iterate
+		for (int32 EdgeIndex = Edges.Num() - 1; EdgeIndex >= 0; --EdgeIndex)
+		{
+			FPBDIslandConstraint* Edge = Edges[EdgeIndex];
+			RemoveConstraint(Edge->Constraint);
+		}
 
-void FPBDIslandManager::SyncIslands(FPBDRigidsSOAs& Particles, const bool IsResimming)
-{
-	// Update of the sync and sleep state of all particles in all awake islands
-	for(auto& Island : Islands)
+		// Removing all constraints will remove all islands
+		check(Islands.IsEmpty());
+	}
+
+	void FPBDIslandManager::SetParticleIslandIsSleeping(FGeometryParticleHandle* Particle, const bool bInIsSleeping)
 	{
-		if(!Island->IsSleeping() || Island->IsSleepingChanged())
+		if (FPBDIslandParticle* Node = GetGraphNode(Particle))
 		{
-			Island->UpdateSyncState(Particles);
-			Island->PropagateSleepState(Particles);
+			if ((Node->Island != nullptr) && !Node->Island->Flags.bIsSleeping)
+			{
+				Node->Island->Flags.bIsSleeping = true;
+				Node->Island->SleepCounter = 0;
+				PropagateIslandSleep(Node->Island);
+			}
 		}
 	}
 
-#if CHAOS_CONSTRAINTHANDLE_DEBUG_ENABLED
-	DebugCheckIslands();
-#endif
+	void FPBDIslandManager::UpdateIslands()
+	{
+		ProcessIslands();
+	}
 
-	// We can now safely use the GetIslandParticles and GetIslandConstraints accessors
-	bIslandsPopulated = true;
-}
+	void FPBDIslandManager::UpdateSleep()
+	{
+		ProcessSleep();
+	}
 
-void FPBDIslandManager::GraphNodeLevelAssigned(FGraphNode& GraphNode)
-{
-	// Add the particle to its island. This will be called in the order that levels are assigned
-	// which allows us to skip the sorting-by-level step. NOTE: We only add dynamic particles to 
-	// islands at this stage (kinematics are added to the solver when we gather data).
-	if (GraphNode.bValidNode)
+	void FPBDIslandManager::UpdateDisable(TFunctionRef<void(FPBDRigidParticleHandle*)> ParticleDisableFunctor)
 	{
-		// We have a dynamic particle - add it to its island
-		check(GraphNode.IslandIndex != INDEX_NONE);
-		FPBDIsland* Island = IslandGraph->GraphIslands[GraphNode.IslandIndex].IslandItem;
+		ProcessDisable(ParticleDisableFunctor);
+	}
 
-#if CHAOS_CONSTRAINTHANDLE_DEBUG_ENABLED
-		check(!Island->DebugContainsParticle(GraphNode.NodeItem));
-#endif
+	void FPBDIslandManager::EndTick()
+	{
+		ApplyDeterminism();
+	}
 
-		Island->AddParticle(GraphNode.NodeItem);
+	int32 FPBDIslandManager::GetParticleLevel(const FPBDIslandParticle* Node) const
+	{
+		if (Node != nullptr)
+		{
+			return Node->Level;
+		}
+		return INDEX_NONE;
+	}
 
-		if(GraphNode.ResimFrame != INDEX_NONE)
+	int32 FPBDIslandManager::GetParticleColor(const FPBDIslandParticle* Node) const
+	{
+		return INDEX_NONE;
+	}
+
+	int32 FPBDIslandManager::GetConstraintLevel(const FPBDIslandConstraint* Edge) const
+	{
+		if (Edge != nullptr)
 		{
-			const int32 ResimFrame = (Island->GetResimFrame() == INDEX_NONE) ? 
-					GraphNode.ResimFrame : FMath::Min(Island->GetResimFrame(), GraphNode.ResimFrame);
-			Island->SetResimFrame(ResimFrame);
+			return Edge->Level;
 		}
+		return INDEX_NONE;
 	}
-}
 
-void FPBDIslandManager::GraphEdgeLevelAssigned(FGraphEdge& GraphEdge)
-{
-	check(GraphEdge.IslandIndex != INDEX_NONE);
-	FPBDIsland* Island = IslandGraph->GraphIslands[GraphEdge.IslandIndex].IslandItem;
+	int32 FPBDIslandManager::GetConstraintColor(const FPBDIslandConstraint* Edge) const
+	{
+		return INDEX_NONE;
+	}
 
-#if CHAOS_CONSTRAINTHANDLE_DEBUG_ENABLED
-	ensure(GraphEdge.bValidEdge);
-	check(!Island->DebugContainsConstraint(GraphEdge.EdgeItem));
-#endif
+	int32 FPBDIslandManager::GetParticleResimFrame(const FGeometryParticleHandle* Particle) const
+	{
+		if (const FPBDIslandParticle* Node = GetGraphNode(Particle))
+		{
+			return Node->ResimFrame;
+		}
+		return INDEX_NONE;
+	}
 
-	// Persistent ordering between ticks is now built-in because the graph is persistent so the constraint order will 
-	// not change from tick to tick even with parallel collision detection (though the order will be different between 
-	// machines unles deterministic mode is set on the evolution). As a result this sort key is not required but for 
-	// now we still calculate it because we can disable persistent islands on a cvar (bChaosSolverPersistentGraph).
-	// 
-	// For non-persistent graph mode:
-	// The sub sort key is for consistent ordering of constraints when both Level and Color are the same.
-	// Without this the ordering within a level/color would depend on the order they were added to the graph
-	// which depends on collision detection order, which is multi-threaded and not deterministic (by default).
-	// The sort key is based on Particle's NodeIndex because particles remain in the graph with the same index 
-	// as long as they are active. We need to pack the key into 32 bits so we truncate the particle indices. 
-	// This should matter only very rarely. If we hit a conflict we may see a slight non-deterministic behaviour 
-	// difference because the solve  order may vary from frame to frame, but it will still solve ok.
-	const uint32 SubSortKey = (((uint32(GraphEdge.FirstNode) & 0xFFFF) << 16) | (uint32(GraphEdge.SecondNode) & 0xFFFF));
-
-	Island->AddConstraint(GraphEdge.EdgeItem, GraphEdge.LevelIndex, GraphEdge.ColorIndex, SubSortKey);
-}
+	void FPBDIslandManager::SetParticleResimFrame(FGeometryParticleHandle* Particle, const int32 ResimFrame)
+	{
+		if (FPBDIslandParticle* Node = GetGraphNode(Particle))
+		{
+			Node->ResimFrame = ResimFrame;
+		}
+	}
 
-void FPBDIslandManager::UpdateIslands(FPBDRigidsSOAs& Particles, const bool IsResimming)
-{
-	// Assign all particles and constraints to islands, merging and splitting islands as required
-	IslandGraph->UpdateGraph();
+	void FPBDIslandManager::ResetParticleResimFrame(const int32 ResetFrame)
+	{
+		for (FPBDIslandParticle* Node : Nodes)
+		{
+			Node->ResimFrame = ResetFrame;
+		}
+	}
 
-#if CHAOS_CONSTRAINTHANDLE_DEBUG_ENABLED
-	DebugCheckParticlesInGraph(Particles.GetNonDisabledDynamicView());
-#endif
+	FPBDIslandParticle* FPBDIslandManager::GetGraphNode(const FGeometryParticleHandle* Particle) const
+	{
+		if (Particle != nullptr)
+		{
+			return Particle->GetConstraintGraphNode();
+		}
+		return nullptr;
+	}
+
+	FPBDIslandParticle* FPBDIslandManager::GetGraphNode(const FTransientGeometryParticleHandle& Particle) const
+	{
+		return Particle.GetConstraintGraphNode();
+	}
+
+	FPBDIslandParticle* FPBDIslandManager::CreateGraphNode(FGeometryParticleHandle* Particle)
+	{
+		if (Particle != nullptr)
+		{
+			check(GetGraphNode(Particle) == nullptr);
 
-	// Initializes our set of islands to match the state of the graph islands, ready to receive the particles and constraints.
-	InitIslands();
+			FPBDIslandParticle* Node = Nodes.Alloc(Particle);
 
-	// Reeset all particle and constraint levels, colors etc
-	IslandGraph->InitSorting();
+			// Initial state
+			Node->Flags.bIsDynamic = IsParticleDynamic(Particle);
+			Node->Flags.bIsSleeping = IsParticleSleeping(Particle);
+			Node->Flags.bIsMoving = IsParticleMoving(Particle);
+			Node->Flags.bNeedsResim = IsParticleNeedsResim(Particle);
+			UpdateGraphNodeSleepSettings(Node);
 
-	// Assign a level to every constraint and particle.
-	// NOTE: This will call GraphNodeLevelAssigned/GraphEdgeLevelAssigned for every
-	// particle/constraint in level order. We use this to add constraints and particles
-	// to their islands in level-sorted order, so we don't need to sort again below
-	IslandGraph->ComputeLevels();
+			return Node;
+		}
+		return nullptr;
+	}
 
-	// Assign colors
-	// @todo(chaos): re-enable graph coloring, although it's probably not worth it!
-	//IslandGraph->ComputeColors();
+	FPBDIslandParticle* FPBDIslandManager::GetOrCreateGraphNode(FGeometryParticleHandle* Particle)
+	{
+		FPBDIslandParticle* Node = GetGraphNode(Particle);
+		if (Node == nullptr)
+		{
+			Node = CreateGraphNode(Particle);
+		}
+		return Node;
+	}
 
-	// Sync the state of all particles and constraints (primarily sleep state)
-	SyncIslands(Particles, IsResimming);
+	void FPBDIslandManager::DestroyGraphNode(FPBDIslandParticle* Node)
+	{
+		if (Node != nullptr)
+		{
+			// We should have already been removed from any island
+			check(Node->Island == nullptr);
+
+			Nodes.Free(Node);
+		}
+	}
 
-	// Sort each island's constraints based on level, color, etc
-	// This is not required when we have a persistent graph because we already added constraints to islands in 
-	// Level-order, and the sub-level order does not change from tick to tick when persistence is enabled.
-	if (!bChaosSolverPersistentGraph)
+	FPBDIslandConstraint* FPBDIslandManager::GetGraphEdge(const FConstraintHandle* Constraint)
 	{
-		for (TUniquePtr<FPBDIsland>& Island : Islands)
+		check(Constraint != nullptr);
+
+		return Constraint->GetConstraintGraphEdge();
+	}
+
+	const FPBDIslandConstraint* FPBDIslandManager::GetGraphEdge(const FConstraintHandle* Constraint) const
+	{
+		return const_cast<FPBDIslandManager*>(this)->GetGraphEdge(Constraint);
+	}
+
+	FPBDIslandConstraint* FPBDIslandManager::CreateGraphEdge(const int32 ContainerId, FConstraintHandle* Constraint, const TVec2<FGeometryParticleHandle*>& ConstrainedParticles)
+	{
+		check(Constraint != nullptr);
+		check(GetGraphEdge(Constraint) == nullptr);
+
+		// Make sure our particles are represented in the graph
+		if (IsParticleDynamic(ConstrainedParticles[0]) || IsParticleDynamic(ConstrainedParticles[1]))
+		{
+			FPBDIslandParticle* Node0 = GetOrCreateGraphNode(ConstrainedParticles[0]);
+			FPBDIslandParticle* Node1 = GetOrCreateGraphNode(ConstrainedParticles[1]);
+
+			// Create the edge and bind it to the nodes
+			FPBDIslandConstraint* Edge = Edges.Alloc(ContainerId, Constraint);
+
+			BindEdgeToNodes(Edge, Node0, Node1);
+
+			return Edge;
+		}
+
+		return nullptr;
+	}
+
+	void FPBDIslandManager::DestroyGraphEdge(FPBDIslandConstraint* Edge)
+	{
+		if (Edge != nullptr)
+		{
+			// We should have already been removed from any island
+			check(Edge->Island == nullptr);
+
+			// Disconnect from the nodes
+			UnbindEdgeFromNodes(Edge);
+
+			Edges.Free(Edge);
+		}
+	}
+
+	void FPBDIslandManager::BindEdgeToNodes(FPBDIslandConstraint* Edge, FPBDIslandParticle* Node0, FPBDIslandParticle* Node1)
+	{
+		check(Edge != nullptr);
+		check(Edge->Nodes[0] == nullptr);
+		check(Edge->Nodes[1] == nullptr);
+		check(Edge->Island == nullptr);
+
+		if (Node0 != nullptr)
+		{
+			Edge->Nodes[0] = Node0;
+			Edge->NodeArrayIndices[0] = Node0->Edges.Add(Edge);
+		}
+		if (Node1 != nullptr)
 		{
-			if (!Island->IsSleeping())
+			Edge->Nodes[1] = Node1;
+			Edge->NodeArrayIndices[1] = Node1->Edges.Add(Edge);
+		}
+
+		// To retain order between ticks for constraints at the same level we assign a permanent sort key to
+		// each constraint. This also allows us to be deterministic as long as constraints are added in the same order.
+		Edge->LevelSortKey = NextLevelSortKey++;
+	}
+
+	void FPBDIslandManager::UnbindEdgeFromNodes(FPBDIslandConstraint* Edge)
+	{
+		for (int32 NodeIndex = 0; NodeIndex < 2; ++NodeIndex)
+		{
+			if (FPBDIslandParticle* Node = Edge->Nodes[NodeIndex])
 			{
-				Island->SortConstraints();
+				const int32 ArrayIndex = Edge->NodeArrayIndices[NodeIndex];
+				check(Node->Edges[ArrayIndex] == Edge);
+
+				// Remove the edge from the node
+				Node->Edges.RemoveAtSwap(ArrayIndex, 1, false);
+				Edge->Nodes[NodeIndex] = nullptr;
+				Edge->NodeArrayIndices[NodeIndex] = INDEX_NONE;
+
+				// Update the array index of the edge we swapped in
+				if (ArrayIndex < Node->Edges.Num())
+				{
+					// We don't know which of the edges nodes we are...
+					if (Node->Edges[ArrayIndex]->Nodes[0] == Node)
+					{
+						Node->Edges[ArrayIndex]->NodeArrayIndices[0] = ArrayIndex;
+					}
+					else if (Node->Edges[ArrayIndex]->Nodes[1] == Node)
+					{
+						Node->Edges[ArrayIndex]->NodeArrayIndices[1] = ArrayIndex;
+					}
+				}
 			}
 		}
 	}
 
-	ValidateIslands();
-}
+	void FPBDIslandManager::UpdateGraphNodeSleepSettings(FPBDIslandParticle* Node)
+	{
+		check(Node != nullptr);
 
-void FPBDIslandManager::ValidateIslands() const
-{
-	if (!bChaosSolverValidateGraph)
+		if (Node->Flags.bIsDynamic)
+		{
+			FRealSingle SleepLinearThreshold, SleepAngularThreshold;
+			FRealSingle DisableLinearThreshold, DisableAngularThreshold;
+			int32 SleepCounterThreshold;
+			GetParticleSleepThresholds(Node->Particle, PhysicsMaterials, PerParticlePhysicsMaterials, SimMaterials, SleepLinearThreshold, SleepAngularThreshold, SleepCounterThreshold);
+			GetParticleDisableThresholds(Node->Particle, PhysicsMaterials, PerParticlePhysicsMaterials, SimMaterials, DisableLinearThreshold, DisableAngularThreshold);
+
+			Node->SleepLinearThresholdSq = FMath::Square(SleepLinearThreshold);
+			Node->SleepAngularThresholdSq = FMath::Square(SleepAngularThreshold);
+			Node->SleepCounterThreshold = SleepCounterThreshold;
+			Node->DisableLinearThresholdSq = FMath::Square(DisableLinearThreshold);
+			Node->DisableAngularThresholdSq = FMath::Square(DisableAngularThreshold);
+		}
+	}
+
+	void FPBDIslandManager::UpdateGraphNode(FPBDIslandParticle* Node)
 	{
-		return;
+		check(Node != nullptr);
+
+		const bool bWasDynamic = Node->Flags.bIsDynamic;
+		const bool bWasMoving = Node->Flags.bIsMoving;
+		const bool bWasSleeping = Node->Flags.bIsSleeping;
+		const bool bIsDynamic = IsParticleDynamic(Node->Particle);
+		const bool bIsMoving = IsParticleMoving(Node->Particle);
+		const bool bIsSleeping = IsParticleSleeping(Node->Particle);
+		const bool bNeedsResim = IsParticleNeedsResim(Node->Particle);
+
+		Node->Flags.bIsDynamic = bIsDynamic;
+		Node->Flags.bIsMoving = bIsMoving;
+		Node->Flags.bIsSleeping = bIsSleeping;
+		Node->Flags.bNeedsResim = bNeedsResim;
+
+		// Did we change between dynamic and kinematic?
+		if (bIsDynamic != bWasDynamic)
+		{
+			// If we have turned dynamic, we need to merge all islands we are in
+			// and update our sleep settings (we don't initialize sleep settings for kinematics)
+			if (bIsDynamic)
+			{
+				UpdateGraphNodeSleepSettings(Node);
+
+				for (FPBDIslandConstraint* Edge : Node->Edges)
+				{
+					// Merge or assign node islands
+					MergeNodeIslands(Edge->Nodes[0], Edge->Nodes[1]);
+				}
+			}
+
+			// If we have turned kinematic, we need to be removed from our island
+			// because kinematics are not directly tracked in islands
+			if (!bIsDynamic)
+			{
+				RemoveNodeFromIsland(Node);
+			}
+
+			// Update the edge because the particle type has changed
+			// NOTE: reverse iteration since we may remove edges as we go
+			for (int32 EdgeIndex = Node->Edges.Num() - 1; EdgeIndex >= 0; --EdgeIndex)
+			{
+				UpdateGraphEdge(Node->Edges[EdgeIndex]);
+			}
+		}
+
+		// If the particle changed, we need to check the island's sleep state
+		// This is to cover the case where any particle is woken, or all particles are manually slept
+		if (bIsDynamic && (!bWasDynamic || (bIsSleeping != bWasSleeping) || bIsMoving))
+		{
+			const bool bIsSleepAllowed = true;
+			EnqueueIslandCheckSleep(Node->Island, bIsSleepAllowed);
+		}
+
+		// If we have a moving kinematic its island(s) cannot sleep
+		if (!bIsDynamic && bIsMoving)
+		{
+			const bool bIsSleepAllowed = false;
+			for (FPBDIslandConstraint* Edge : Node->Edges)
+			{
+				EnqueueIslandCheckSleep(Edge->Island, bIsSleepAllowed);
+			}
+		}
 	}
 
-	// Check that nodes joined by an edge are in the same island as the edge
-	for (const FGraphEdge& Edge : IslandGraph->GraphEdges)
+	void FPBDIslandManager::UpdateGraphEdge(FPBDIslandConstraint* Edge)
 	{
-		const FGraphNode* Node0 = (Edge.FirstNode != INDEX_NONE) ? &IslandGraph->GraphNodes[Edge.FirstNode] : nullptr;
-		const FGraphNode* Node1 = (Edge.SecondNode != INDEX_NONE) ? &IslandGraph->GraphNodes[Edge.SecondNode] : nullptr;
+		const bool bIsDynamic0 = (Edge->Nodes[0] != nullptr) && Edge->Nodes[0]->Flags.bIsDynamic;
+		const bool bIsDynamic1 = (Edge->Nodes[1] != nullptr) && Edge->Nodes[1]->Flags.bIsDynamic;
+		if (!bIsDynamic0 && !bIsDynamic1)
+		{
+			// If we get here, both particles in the constraint are now kinematic and
+			// should have already been removed from the island (see FPBDIslandManager::UpdateGraphNode)
+			check((Edge->Nodes[0] == nullptr) || (Edge->Nodes[0]->Island == nullptr));
+			check((Edge->Nodes[1] == nullptr) || (Edge->Nodes[1]->Island == nullptr));
+
+			// Remove the edge from its island
+			RemoveEdgeFromIsland(Edge);
 
-		ensure((Node0 == nullptr) || !Node0->bValidNode || (Node0->IslandIndex == Edge.IslandIndex));
-		ensure((Node1 == nullptr) || !Node1->bValidNode || (Node1->IslandIndex == Edge.IslandIndex));
+			// Destroy the edge
+			DestroyGraphEdge(Edge);
+		}
+		else
+		{
+			// We should be in an island already if at least one of the nodes is dynamic
+			check(Edge->Island != nullptr);
+		}
 	}
-}
 
-bool FPBDIslandManager::SleepInactive(const int32 IslandIndex,
-	const TArrayCollectionArray<TSerializablePtr<FChaosPhysicsMaterial>>& PhysicsMaterials,
-	const TArrayCollectionArray<TUniquePtr<FChaosPhysicsMaterial>>& PerParticlePhysicsMaterials,
-	const THandleArray<FChaosPhysicsMaterial>& SimMaterials)
-{
-	// Only the persistent islands could start sleeping
-	const int32 GraphIndex = GetGraphIndex(IslandIndex);
-	if (!bChaosSolverSleepEnabled || !Islands.IsValidIndex(GraphIndex) ||
-		(Islands.IsValidIndex(GraphIndex) && !Islands[GraphIndex]->IsPersistent()) )
+	FPBDIsland* FPBDIslandManager::CreateIsland()
 	{
-		return false;
+		FPBDIsland* Island = Islands.Alloc();
+		Island->ContainerEdges.SetNum(GetNumConstraintContainers());
+		return Island;
 	}
 
-	FReal LinearSleepingThreshold = FLT_MAX,  AngularSleepingThreshold = FLT_MAX, MaxLinearSpeed2 = 0.f, MaxAngularSpeed2 = 0.f;
-	int32 SleepCounterThreshold = 0, NumDynamicParticles = 0;
+	void FPBDIslandManager::DestroyIsland(FPBDIsland* Island)
+	{
+		if (Island != nullptr)
+		{
+			check(Island->Nodes.IsEmpty());
+			check(Island->NumEdges == 0);
 
-	// Compute of the linear/angular velocities + thresholds to make islands sleeping 
-	if (ComputeSleepingThresholds(Islands[GraphIndex].Get(), PhysicsMaterials, PerParticlePhysicsMaterials,
-		SimMaterials, LinearSleepingThreshold, AngularSleepingThreshold,
-		MaxLinearSpeed2, MaxAngularSpeed2, SleepCounterThreshold))
+			Islands.Free(Island);
+		}
+	}
+
+	void FPBDIslandManager::AssignNodeIsland(FPBDIslandParticle* Node)
 	{
-		int32 SleepCounter = Islands[GraphIndex]->GetSleepCounter();
-		const bool bSleepingIsland = IsIslandSleeping(MaxLinearSpeed2, MaxAngularSpeed2, LinearSleepingThreshold,
-			AngularSleepingThreshold, SleepCounterThreshold, SleepCounter);
-		
-		Islands[GraphIndex]->SetSleepCounter(SleepCounter);
+		check(Node != nullptr);
+		check(Node->Island == nullptr);
+
+		FPBDIsland* Island = CreateIsland();
 
-		return bSleepingIsland;
+		AddNodeToIsland(Node, Island);
 	}
-	return false;
-}
 
-void FPBDIslandManager::SleepIsland(FPBDRigidsSOAs& Particles, const int32 IslandIndex)
-{
-	// @todo(chaos): this function will not work when called manually at the start of the frame based on user request.
-	// Currently this is not required so not important. See disabled unit test: TestConstraintGraph_SleepIsland
-	ensure(bIslandsPopulated);
+	void FPBDIslandManager::AssignEdgeIsland(FPBDIslandConstraint* Edge)
+	{
+		check(Edge != nullptr);
+		check(Edge->Island == nullptr);
 
-	SetIslandSleeping(Particles, IslandIndex, true);
-}
+		FPBDIsland* Island = MergeNodeIslands(Edge->Nodes[0], Edge->Nodes[1]);
 
-void FPBDIslandManager::SetIslandSleeping(FPBDRigidsSOAs& Particles, const int32 IslandIndex, const bool bIsSleeping)
-{
-	// NOTE: This Sleep function is called from the graph update to set the state of the 
-	// island and all particles in it.
-	const int32 GraphIndex = GetGraphIndex(IslandIndex);
-	if (Islands.IsValidIndex(GraphIndex))
+		AddEdgeToIsland(Edge, Island);
+	}
+
+	void FPBDIslandManager::AddNodeToIsland(FPBDIslandParticle* Node, FPBDIsland* Island)
+	{
+		check(Island != nullptr);
+
+		if (Node != nullptr)
+		{
+			// Only dynamic particles are kept in the island's node array.
+			// Kinematics may be in many islands (we need to check the node's edges to visit them).
+			if (Node->Flags.bIsDynamic)
+			{
+				Node->Island = Island;
+				Node->IslandArrayIndex = Island->Nodes.Add(Node);
+			}
+
+			Island->Flags.bItemsAdded = true;
+
+			const bool bIsSleepAllowed = true;
+			EnqueueIslandCheckSleep(Island, bIsSleepAllowed);
+		}
+	}
+
+	void FPBDIslandManager::RemoveNodeFromIsland(FPBDIslandParticle* Node)
+	{
+		if ((Node != nullptr) && (Node->Island != nullptr))
+		{
+			FPBDIsland* Island = Node->Island;
+			const int32 ArrayIndex = Node->IslandArrayIndex;
+			check(Island->Nodes[ArrayIndex] == Node);
+
+			Island->Nodes.RemoveAtSwap(ArrayIndex, 1, false);
+			if (ArrayIndex < Island->Nodes.Num())
+			{
+				Island->Nodes[ArrayIndex]->IslandArrayIndex = ArrayIndex;
+			}
+
+			Island->Flags.bItemsRemoved = true;
+
+			Node->Island = nullptr;
+			Node->IslandArrayIndex = INDEX_NONE;
+		}
+	}
+
+	void FPBDIslandManager::AddEdgeToIsland(FPBDIslandConstraint* Edge, FPBDIsland* Island)
+	{
+		check(Edge != nullptr);
+		check(Island != nullptr);
+
+		const int32 ContainerIndex = Edge->ContainerIndex;
+		Edge->IslandArrayIndex = Island->ContainerEdges[ContainerIndex].Add(Edge);
+		Edge->Island = Island;
+
+		Island->Flags.bItemsAdded = true;
+		Island->NumEdges = Island->NumEdges + 1;
+
+		const bool bIsSleepAllowed = true;
+		EnqueueIslandCheckSleep(Island, bIsSleepAllowed);
+	}
+
+	void FPBDIslandManager::RemoveEdgeFromIsland(FPBDIslandConstraint* Edge)
+	{
+		check(Edge != nullptr);
+
+		if (Edge->Island != nullptr)
+		{
+			FPBDIsland* Island = Edge->Island;
+
+			const int32 ContainerIndex = Edge->ContainerIndex;
+			const int32 EdgeIndex = Edge->IslandArrayIndex;
+			check(Island->ContainerEdges[ContainerIndex][EdgeIndex] == Edge);
+
+			Island->ContainerEdges[ContainerIndex].RemoveAtSwap(EdgeIndex, 1, false);
+			if (EdgeIndex < Island->ContainerEdges[ContainerIndex].Num())
+			{
+				Island->ContainerEdges[ContainerIndex][EdgeIndex]->IslandArrayIndex = EdgeIndex;
+			}
+
+			Island->Flags.bItemsRemoved = true;
+			Island->NumEdges = Island->NumEdges - 1;
+
+			Edge->Island = nullptr;
+			Edge->IslandArrayIndex = INDEX_NONE;
+		}
+	}
+
+	void FPBDIslandManager::EnqueueIslandCheckSleep(FPBDIsland* Island, const bool bIsSleepAllowed)
 	{
-		const bool bWasSleeping = Islands[GraphIndex]->IsSleeping();
-		if (bWasSleeping != bIsSleeping)
+		if (Island != nullptr)
+		{
+			Island->Flags.bCheckSleep = true;
+			Island->Flags.bIsSleepAllowed = (Island->Flags.bIsSleepAllowed != 0) && bIsSleepAllowed;
+		}
+	}
+
+	FPBDIsland* FPBDIslandManager::MergeNodeIslands(FPBDIslandParticle* Node0, FPBDIslandParticle* Node1)
+	{
+		// NOTE: We only regsiter the islands for merging here. The actual merge takes place later in ProcessMerges.
+
+		FPBDIsland* Island0 = (Node0 != nullptr) ? Node0->Island : nullptr;
+		FPBDIsland* Island1 = (Node1 != nullptr) ? Node1->Island : nullptr;
+
+		// Don't try to merge island into self
+		if ((Island0 != nullptr) && (Island0 == Island1))
+		{
+			return Island0;
+		}
+
+		FPBDIsland* Island = nullptr;
+		if ((Island0 != nullptr) && (Island1 != nullptr))
+		{
+			Island = MergeIslands(Island0, Island1);
+		}
+		else
 		{
-			IslandGraph->GraphIslands[GraphIndex].bIsSleeping = bIsSleeping;
-			Islands[GraphIndex]->SetIsSleeping(bIsSleeping);
-			Islands[GraphIndex]->SetIsSleepingChanged(true);
-			Islands[GraphIndex]->PropagateSleepState(Particles);
+			if (Island0 != nullptr) // && (Island1 == nullptr)
+			{
+				Island = Island0;
+				AddNodeToIsland(Node1, Island);
+			}
+			else if (Island1 != nullptr) // && (Island0 == nullptr)
+			{
+				Island = Island1;
+				AddNodeToIsland(Node0, Island);
+			}
+			else // ((Island0 == nullptr) && (Island1 == nullptr))
+			{
+				// NOTE: Islands are created in the sleeping state and will be set to
+				// awake if either node is moving when we add them
+				Island = CreateIsland();
+				AddNodeToIsland(Node0, Island);
+				AddNodeToIsland(Node1, Island);
+			}
+		}
+
+		return Island;
+	}
+
+	FPBDIsland* FPBDIslandManager::MergeIslands(FPBDIsland* Island0, FPBDIsland* Island1)
+	{
+		// NOTE: We only regsiters the islands for merging here. The actual merge takes place later in ProcessMerges.
+		check(Island0 != nullptr);
+		check(Island1 != nullptr);
+
+		if (Island0 == Island1)
+		{
+			return Island0;
+		}
+
+		// Add the islands to a merge set, which keeps track of all the islands that have been merged into each other
+		// If both islands are already assigned to a merge set, merge the merge sets (if different)
+		// Otherwise add the island to the existing or created merge set
+		FPBDIslandMergeSet* MergeSet0 = Island0->MergeSet;
+		FPBDIslandMergeSet* MergeSet1 = Island1->MergeSet;
+		if ((MergeSet0 == nullptr) || (MergeSet0 != MergeSet1))
+		{
+			if ((MergeSet0 != nullptr) && (MergeSet1 != nullptr))
+			{
+				CombineMergeSets(MergeSet0, MergeSet1);
+			}
+			else if (MergeSet0 != nullptr) // && (MergeSet1 == nullptr)
+			{
+				AddIslandToMergeSet(Island1, MergeSet0);
+			}
+			else if (MergeSet1 != nullptr) // && (MergeSet0 == nullptr)
+			{
+				AddIslandToMergeSet(Island0, MergeSet1);
+			}
+			else // (MergeSet0 == nullptr) && (MergSet1 == nullptr)
+			{
+				CreateMergeSet(Island0, Island1);
+			}
+		}
+
+		return (Island0->NumEdges >= Island1->NumEdges) ? Island0 : Island1;
+	}
+
+	FPBDIslandMergeSet* FPBDIslandManager::CreateMergeSet(FPBDIsland* Island0, FPBDIsland* Island1)
+	{
+		FPBDIslandMergeSet* MergeSet = MergeSets.Alloc();
+
+		AddIslandToMergeSet(Island0, MergeSet);
+		AddIslandToMergeSet(Island1, MergeSet);
+
+		return MergeSet;
+	}
+
+	void FPBDIslandManager::DestroyMergeSet(FPBDIslandMergeSet* MergeSet)
+	{
+		if (MergeSet != nullptr)
+		{
+			MergeSets.Free(MergeSet);
+		}
+	}
+
+	void FPBDIslandManager::AddIslandToMergeSet(FPBDIsland* Island, FPBDIslandMergeSet* MergeSet)
+	{
+		check(Island != nullptr);
+		check(Island->MergeSet == nullptr);
+
+		// Bind the island to the merge set
+		Island->MergeSet = MergeSet;
+		Island->MergeSetIslandIndex = MergeSet->Islands.Add(Island);
+		MergeSet->NumEdges += Island->NumEdges;
+	}
+
+	void FPBDIslandManager::RemoveIslandFromMergeSet(FPBDIsland* Island)
+	{
+		check(Island != nullptr);
+
+		if (Island->MergeSet != nullptr)
+		{
+			FPBDIslandMergeSet* MergeSet = Island->MergeSet;
+
+			// Remove from the list of islands to merge
+			const int32 IslandIndex = Island->MergeSetIslandIndex;
+			MergeSet->Islands.RemoveAtSwap(IslandIndex, 1, false);
+			if (IslandIndex < MergeSet->Islands.Num())
+			{
+				MergeSet->Islands[IslandIndex]->MergeSetIslandIndex = IslandIndex;
+			}
+
+			// Update edge count
+			MergeSet->NumEdges -= Island->NumEdges;
+
+			Island->MergeSet = nullptr;
+			Island->MergeSetIslandIndex = INDEX_NONE;
+		}
+	}
+
+	FPBDIslandMergeSet* FPBDIslandManager::CombineMergeSets(FPBDIslandMergeSet* MergeSetParent, FPBDIslandMergeSet* MergeSetChild)
+	{
+		check(MergeSetParent != nullptr);
+		check(MergeSetChild != nullptr);
+
+		if (MergeSetParent == MergeSetChild)
+		{
+			return MergeSetParent;
+		}
+
+		// Select the largest set as the merge parent
+		// (we should merge into the set with the largest island, but this should be good enough)
+		if (MergeSetParent->NumEdges < MergeSetChild->NumEdges)
+		{
+			Swap(MergeSetParent, MergeSetChild);
+		}
+
+		// Tell all islands about their new merge set
+		int32 ChildIslandIndex = MergeSetParent->Islands.Num();
+		for (FPBDIsland* Island : MergeSetChild->Islands)
+		{
+			Island->MergeSet = MergeSetParent;
+			Island->MergeSetIslandIndex = ChildIslandIndex++;
+		}
+
+		// Add the islands to their new merge set
+		MergeSetParent->Islands.Append(MoveTemp(MergeSetChild->Islands));
+
+		// Update edge count
+		MergeSetParent->NumEdges += MergeSetChild->NumEdges;
+
+		// Destroy the emptied merge set
+		DestroyMergeSet(MergeSetChild);
+
+		return MergeSetParent;
+	}
+
+	FPBDIsland* FPBDIslandManager::GetMergeSetParentIsland(FPBDIslandMergeSet* MergeSet, int32& OutNumNodes, const TArrayView<int32>& OutNumContainerEdges)
+	{
+		check(MergeSet != nullptr);
+
+		// Initialize outputs
+		FPBDIsland* LargestIsland = nullptr;
+		OutNumNodes = 0;
+		for (int32 ContainerIndex = 0; ContainerIndex < OutNumContainerEdges.Num(); ++ContainerIndex)
+		{
+			OutNumContainerEdges[ContainerIndex] = 0;
+		}
+
+		// Find the largest island and count the nodes and edges from all islands
+		for (FPBDIsland* Island : MergeSet->Islands)
+		{
+			if ((LargestIsland == nullptr) || (Island->NumEdges > LargestIsland->NumEdges))
+			{
+				LargestIsland = Island;
+			}
+
+			OutNumNodes += Island->Nodes.Num();
+
+			for (int32 ContainerIndex = 0; ContainerIndex < Island->ContainerEdges.Num(); ++ContainerIndex)
+			{
+				OutNumContainerEdges[ContainerIndex] += Island->ContainerEdges[ContainerIndex].Num();
+			}
+		}
+
+		return LargestIsland;
+	}
+
+	void FPBDIslandManager::ProcessIslands()
+	{
+		ProcessMerges();
+		ProcessSplits();
+		ProcessWakes();
+		AssignLevels();
+		FinalizeIslands();
+	}
+
+	void FPBDIslandManager::ProcessMerges()
+	{
+		SCOPE_CYCLE_COUNTER(STAT_IslandManager_MergeIslands);
+
+		int32 NumNodes;
+		TArray<int32> NumContainerEdges;
+		NumContainerEdges.SetNumUninitialized(GetNumConstraintContainers());
+
+		// Could go wide here...
+		for (FPBDIslandMergeSet* MergeSet : MergeSets)
+		{
+			// Find the island that we will merge into (the biggest one) and count the nodes and edges
+			if (FPBDIsland* ParentIsland = GetMergeSetParentIsland(MergeSet, NumNodes, MakeArrayView(NumContainerEdges)))
+			{
+				// Reserve space for all the nodes and edges
+				ParentIsland->Nodes.Reserve(NumNodes);
+				for (int32 ContainerIndex = 0; ContainerIndex < NumContainerEdges.Num(); ++ContainerIndex)
+				{
+					ParentIsland->ContainerEdges[ContainerIndex].Reserve(NumContainerEdges[ContainerIndex]);
+				}
+
+				// Merge each child island into the parent island
+				for (FPBDIsland* ChildIsland : MergeSet->Islands)
+				{
+					ProcessIslandMerge(ParentIsland, ChildIsland);
+				}
+
+				// Reset the parent island's merge set ready for next tick (the children will have been destroyed)
+				ParentIsland->MergeSet = nullptr;
+				ParentIsland->MergeSetIslandIndex = INDEX_NONE;
+			}
+		}
+
+		MergeSets.Reset();
+	}
+
+	void FPBDIslandManager::ProcessIslandMerge(FPBDIsland* ParentIsland, FPBDIsland* ChildIsland)
+	{
+		if (ParentIsland != ChildIsland)
+		{
+			// Combine island state
+			ParentIsland->SleepCounter = FMath::Min(ParentIsland->SleepCounter, ChildIsland->SleepCounter);
+			ParentIsland->Flags.bItemsAdded = true;
+			ParentIsland->Flags.bItemsRemoved = !!ParentIsland->Flags.bItemsRemoved || !!ChildIsland->Flags.bItemsRemoved;
+			ParentIsland->Flags.bIsSleepAllowed = !!ParentIsland->Flags.bIsSleepAllowed && !!ChildIsland->Flags.bIsSleepAllowed;
+			ParentIsland->Flags.bIsSleeping = !!ParentIsland->Flags.bIsSleeping && !!ChildIsland->Flags.bIsSleeping;
+			ParentIsland->Flags.bWasSleeping = !!ParentIsland->Flags.bWasSleeping || !!ChildIsland->Flags.bWasSleeping;
+			ParentIsland->Flags.bCheckSleep = !!ParentIsland->Flags.bCheckSleep || !!ChildIsland->Flags.bCheckSleep;
+			ParentIsland->Flags.bNeedsResim = false;	// Calculated in FinalizeIslands
+
+			// Tell all the nodes which island they belong to and what index they will have (after we move the elements)
+			int32 NextIslandArrayIndex = ParentIsland->Nodes.Num();
+			for (FPBDIslandParticle* Node : ChildIsland->Nodes)
+			{
+				if (Node->Flags.bIsDynamic)
+				{
+					Node->Island = ParentIsland;
+					Node->IslandArrayIndex = NextIslandArrayIndex++;
+				}
+			}
+
+			// Tell all the edges which island they belong to and what index they will have (after we move the elements)
+			for (int32 ContainerIndex = 0; ContainerIndex < ChildIsland->ContainerEdges.Num(); ++ContainerIndex)
+			{
+				int32 NextContainerEdgeIndex = ParentIsland->ContainerEdges[ContainerIndex].Num();
+				for (FPBDIslandConstraint* Edge : ChildIsland->ContainerEdges[ContainerIndex])
+				{
+					Edge->Island = ParentIsland;
+					Edge->IslandArrayIndex = NextContainerEdgeIndex++;
+				}
+			}
+
+			// Move the nodes and edges to the parent
+			ParentIsland->Nodes.Append(MoveTemp(ChildIsland->Nodes));
+			ChildIsland->Nodes.Reset();
+
+			for (int32 ContainerIndex = 0; ContainerIndex < ChildIsland->ContainerEdges.Num(); ++ContainerIndex)
+			{
+				ParentIsland->NumEdges += ChildIsland->ContainerEdges[ContainerIndex].Num();
+				ParentIsland->ContainerEdges[ContainerIndex].Append(MoveTemp(ChildIsland->ContainerEdges[ContainerIndex]));
+
+				ChildIsland->ContainerEdges[ContainerIndex].Reset();
+			}
+			ChildIsland->NumEdges = 0;
+
+			// Destroy the child island
+			DestroyIsland(ChildIsland);
+		}
+	}
+
+	void FPBDIslandManager::ProcessSplits()
+	{
+		SCOPE_CYCLE_COUNTER(STAT_IslandManager_SplitIslands);
+
+		// Could go wide here...(though a bit complicated since we create and destroy islands)
+		for (int32 IslandIndex = Islands.Num() - 1; IslandIndex >= 0; --IslandIndex)
+		{
+			FPBDIsland* Island = Islands[IslandIndex];
+			if (Island->Flags.bItemsRemoved)
+			{
+				ProcessIslandSplits(Island);
+			}
+		}
+	}
+
+	void FPBDIslandManager::ProcessIslandSplits(FPBDIsland* Island)
+	{
+		check(Island != nullptr);
+		check(Island->Flags.bItemsRemoved);
+
+		// Reset split tracking indicesRemo
+		for (FPBDIslandParticle* Node : Island->Nodes)
+		{
+			Node->Island = nullptr;
+			Node->IslandArrayIndex = INDEX_NONE;
+		}
+		for (int32 ContainerIndex = 0; ContainerIndex < Island->ContainerEdges.Num(); ++ContainerIndex)
+		{
+			for (FPBDIslandConstraint* Edge : Island->ContainerEdges[ContainerIndex])
+			{
+				Edge->Island = nullptr;
+				Edge->IslandArrayIndex = INDEX_NONE;
+			}
 		}
+		const int32 VisitEpoch = GetNextVisitEpoch();
 
-		// Reset the sleep counter with every wake call, even if already awake so that we don't immediately go back to sleep
-		if (!bIsSleeping)
+		// Extract the nodes and reset the current island. We rebuild below.
+		TArray<FPBDIslandParticle*> IslandNodes;
+		Swap(IslandNodes, Island->Nodes);
+		for (int32 ContainerIndex = 0; ContainerIndex < Island->ContainerEdges.Num(); ++ContainerIndex)
 		{
-			Islands[GraphIndex]->SetSleepCounter(0);
+			Island->ContainerEdges[ContainerIndex].Reset();
 		}
-	}
-}
+		Island->NumEdges = 0;
 
-TArray<int32> FPBDIslandManager::FindParticleIslands(const FGeometryParticleHandle* ParticleHandle) const
-{
-	TArray<int32> IslandIndices;
+		// The next island to add nodes/edges to
+		FPBDIsland* CurrentIsland = Island;
 
-	// Initially build the list of sparse island indices based on the connected edges (for kinematics)
-	if (ParticleHandle)
-	{
-		// Find the graph node for this particle
-		if (const int32* PNodeIndex = IslandGraph->ItemNodes.Find(ParticleHandle))
+		// Visit all nodes and put connected nodes into an island
+		TArray<FPBDIslandParticle*> NodeQueue;
+		NodeQueue.Reserve(Island->Nodes.Num());
+
+		// We need to remember this value bacause it gets overwriten when
+		// rebuilding the islands below (see AddNodeToIsland)
+		const bool bCheckSleep = Island->Flags.bCheckSleep;
+
+		for (FPBDIslandParticle* RootNode : IslandNodes)
 		{
-			const int32 NodeIndex = *PNodeIndex;
-			if (IslandGraph->GraphNodes.IsValidIndex(NodeIndex))
+			// If we haven't seen this node yet and it is dynamic, start forming a new island
+			if ((RootNode->VisitEpoch != VisitEpoch) && RootNode->Flags.bIsDynamic)
 			{
-				const FGraphNode& GraphNode = IslandGraph->GraphNodes[NodeIndex];
-
-				if (GraphNode.bValidNode)
+				// Create an island if we need one
+				if (CurrentIsland == nullptr)
 				{
-					// A dynamic particle is only in one island
-					if ((GraphNode.IslandIndex != INDEX_NONE) && IslandGraph->GraphIslands.IsValidIndex(GraphNode.IslandIndex))
-					{
-						IslandIndices.Add(GraphNode.IslandIndex);
-					}
+					CurrentIsland = CreateIsland();
 				}
-				else
+
+				// Add the root node to the island and visit all connected nodes
+				AddNodeToIsland(RootNode, CurrentIsland);
+				NodeQueue.Add(RootNode);
+				RootNode->VisitEpoch = VisitEpoch;
+
+				// Populate the island with all connected nodes and edges
+				while (!NodeQueue.IsEmpty())
 				{
-					// A non-dynamic particle can be in many islands, depending on the connections to dynamics
-					for (int32 EdgeIndex : GraphNode.NodeEdges)
+					FPBDIslandParticle* NextNode = NodeQueue.Pop(false);
+
+					// Visit all the edges connected to the current node
+					for (FPBDIslandConstraint* Edge : NextNode->Edges)
 					{
-						if (IslandGraph->GraphEdges.IsValidIndex(EdgeIndex))
+						if (Edge->VisitEpoch != VisitEpoch)
 						{
-							const FGraphEdge& GraphEdge = IslandGraph->GraphEdges[EdgeIndex];
-							if ((GraphEdge.IslandIndex != INDEX_NONE) && IslandGraph->GraphIslands.IsValidIndex(GraphEdge.IslandIndex))
+							// Add the edge to the current island
+							AddEdgeToIsland(Edge, CurrentIsland);
+							Edge->VisitEpoch = VisitEpoch;
+
+							// Queue the other connected node for processing
+							FPBDIslandParticle* EdgeOtherNode = (NextNode == Edge->Nodes[0]) ? Edge->Nodes[1] : Edge->Nodes[0];
+							if ((EdgeOtherNode != nullptr) && EdgeOtherNode->Flags.bIsDynamic && (EdgeOtherNode->VisitEpoch != VisitEpoch))
 							{
-								IslandIndices.Add(GraphEdge.IslandIndex);
+								AddNodeToIsland(EdgeOtherNode, CurrentIsland);
+								NodeQueue.Add(EdgeOtherNode);
+								EdgeOtherNode->VisitEpoch = VisitEpoch;
 							}
 						}
 					}
 				}
+
+				// Set the island state
+				CurrentIsland->SleepCounter = Island->SleepCounter;
+				CurrentIsland->Flags.bItemsAdded = true;
+				CurrentIsland->Flags.bItemsRemoved = true;
+				CurrentIsland->Flags.bIsSleepAllowed = Island->Flags.bIsSleepAllowed;
+				CurrentIsland->Flags.bIsSleeping = Island->Flags.bIsSleeping;
+				CurrentIsland->Flags.bWasSleeping = Island->Flags.bWasSleeping;
+				CurrentIsland->Flags.bCheckSleep = bCheckSleep;
+				CurrentIsland->Flags.bNeedsResim = Island->Flags.bNeedsResim;
+
+				// We are done with this island
+				CurrentIsland = nullptr;
 			}
 		}
+	}
+
+	void FPBDIslandManager::AssignLevels()
+	{
+		SCOPE_CYCLE_COUNTER(STAT_IslandManager_AssignLevels);
+
+		// We only need to assign levels if shock propagation is enabled (but we always sort)
+		const bool bAssignLevels = ((CVars::ChaosSolverCollisionPositionShockPropagationIterations > 0) || (CVars::ChaosSolverCollisionVelocityShockPropagationIterations > 0));
 
-		// Map the islands into non-sparse indices
-		for (int32 IslandIndicesIndex = 0; IslandIndicesIndex < IslandIndices.Num(); ++IslandIndicesIndex)
+		// Could go wide here...
+		for (FPBDIsland* Island : Islands)
 		{
-			const int32 GraphIslandIndex = IslandIndices[IslandIndicesIndex];
-			const int32 SolverIslandIndex = IslandGraph->GraphIslands[GraphIslandIndex].IslandItem->GetIslandIndex();
-			IslandIndices[IslandIndicesIndex] = SolverIslandIndex;
+			// Levels will only change when awake and when something has been added or removed
+			if (!Island->Flags.bIsSleeping && (!!Island->Flags.bItemsAdded || !!Island->Flags.bItemsRemoved))
+			{
+				if (bAssignLevels)
+				{
+					AssignIslandLevels(Island);
+				}
+				SortIslandEdges(Island);
+			}
 		}
 	}
 
-	return IslandIndices;
-}
-
-int32 FPBDIslandManager::GetConstraintIsland(const FConstraintHandle* ConstraintHandle) const
-{
-	const int32 EdgeIndex = ConstraintHandle->GetConstraintGraphIndex();
-	if (IslandGraph->GraphEdges.IsValidIndex(EdgeIndex))
+	void FPBDIslandManager::AssignIslandLevels(FPBDIsland* Island)
 	{
-		int32 GraphIslandIndex = IslandGraph->GraphEdges[EdgeIndex].IslandIndex;
-		if (GraphIslandIndex != INDEX_NONE)
+		// Reset levels: 0 means unassigned; 1 means directly attached to a kinematic
+		// (default 0 is so that the SortKey work even when we don't assign levels)
+		for (FPBDIslandParticle* Node : Island->Nodes)
 		{
-			const int32 SolverIslandIndex = IslandGraph->GraphIslands[GraphIslandIndex].IslandItem->GetIslandIndex();
-			return SolverIslandIndex;
+			Node->Level = 0;
+		}
+		for (int32 ContainerIndex = 0; ContainerIndex < Island->ContainerEdges.Num(); ++ContainerIndex)
+		{
+			for (FPBDIslandConstraint* Edge : Island->ContainerEdges[ContainerIndex])
+			{
+				Edge->Level = 0;
+			}
 		}
-	}
-	return INDEX_NONE;
-}
-
-bool FPBDIslandManager::IslandNeedsResim(const int32 IslandIndex) const
-{
-	return Islands[GetGraphIndex(IslandIndex)]->NeedsResim();
-}
 
-int32 FPBDIslandManager::GetParticleResimFrame(const FGeometryParticleHandle* ParticleHandle) const
-{
-	if (ParticleHandle && IslandGraph->GraphNodes.IsValidIndex(ParticleHandle->ConstraintGraphIndex()))
-	{
-		return IslandGraph->GraphNodes[ParticleHandle->ConstraintGraphIndex()].ResimFrame;
-	}
-	return INDEX_NONE;
-}
+		if (Island->NumEdges < 2)
+		{
+			return;
+		}
 
-void FPBDIslandManager::SetParticleResimFrame(FGeometryParticleHandle* ParticleHandle, const int32 ResimFrame)
-{
-	if (ParticleHandle && IslandGraph->GraphNodes.IsValidIndex(ParticleHandle->ConstraintGraphIndex()))
-	{
-		IslandGraph->GraphNodes[ParticleHandle->ConstraintGraphIndex()].ResimFrame = ResimFrame;
-	}
-}
+		TArray<FPBDIslandParticle*> NodeQueue;
+		NodeQueue.Reserve(Island->Nodes.Num());
 
-void FPBDIslandManager::ResetParticleResimFrame(const int32 ResetFrame)
-{
-	for(auto& GraphNode : IslandGraph->GraphNodes)
-	{
-		GraphNode.ResimFrame = ResetFrame;
-	}
-}
+		// Initialize edge levels and populate the queue with all the level-0 nodes
+		for (int32 ContainerIndex = 0; ContainerIndex < Island->ContainerEdges.Num(); ++ContainerIndex)
+		{
+			for (FPBDIslandConstraint* Edge : Island->ContainerEdges[ContainerIndex])
+			{
+				check((Edge->Nodes[0] != nullptr) || (Edge->Nodes[1] != nullptr));
+				const bool bIsDynamic0 = ((Edge->Nodes[0] != nullptr) && Edge->Nodes[0]->Flags.bIsDynamic);
+				const bool bIsDynamic1 = ((Edge->Nodes[1] != nullptr) && Edge->Nodes[1]->Flags.bIsDynamic);
 
-///////////////////////////////////////////////////////////////////////////////////////////////////
-///////////////////////////////////////////////////////////////////////////////////////////////////
-///////////////////////////////////////////////////////////////////////////////////////////////////
-//
-//
-// BEGIN DEBUG STUFF
-//
-//
-///////////////////////////////////////////////////////////////////////////////////////////////////
-///////////////////////////////////////////////////////////////////////////////////////////////////
-///////////////////////////////////////////////////////////////////////////////////////////////////
-
-#if CHAOS_CONSTRAINTHANDLE_DEBUG_ENABLED
-bool FPBDIslandManager::DebugCheckParticlesInGraph(const TParticleView<FPBDRigidParticles>& Particles) const
-{
-	// Make sure that all particles that should be in the graph actually are.
-	// We used to loop over all particles every trick and re-add them, but now we only add/remove when necessary.
-	bool bSuccess = true;
+				if (!bIsDynamic0 && (Edge->Nodes[1]->Level == 0))
+				{
+					Edge->Level = 1;
+					Edge->Nodes[1]->Level = 1;
+					NodeQueue.Add(Edge->Nodes[1]);
+				}
+				if (!bIsDynamic1 && (Edge->Nodes[0]->Level == 0))
+				{
+					Edge->Level = 1;
+					Edge->Nodes[0]->Level = 1;
+					NodeQueue.Add(Edge->Nodes[0]);
+				}
+			}
+		}
 
-	// All particles in the graph should be enabled
-	for (const FGraphNode& Node : IslandGraph->GraphNodes)
-	{
-		FPBDRigidParticleHandle* Rigid = Node.NodeItem->CastToRigidParticle();
-		if (Rigid != nullptr)
+		// Breadth-first visit all nodes, assign a level to the connected node, and enqeue it
+		for (int32 NodeIndex = 0; NodeIndex < NodeQueue.Num(); ++NodeIndex)
 		{
-			bSuccess = bSuccess && !Rigid->Disabled();
+			FPBDIslandParticle* Node = NodeQueue[NodeIndex];
+			for (FPBDIslandConstraint* Edge : Node->Edges)
+			{
+				if (Edge->Level == 0)
+				{
+					Edge->Level = Node->Level + 1;
+
+					FPBDIslandParticle* OtherNode = (Node == Edge->Nodes[0]) ? Edge->Nodes[1] : Edge->Nodes[0];
+					if ((OtherNode != nullptr) && OtherNode->Flags.bIsDynamic && (OtherNode->Level == 0))
+					{
+						OtherNode->Level = Node->Level + 1;
+						NodeQueue.Add(OtherNode);
+					}
+				}
+			}
 		}
-		ensure(bSuccess);
 	}
 
-	// All dynamic particles should be in the graph
-	for (const auto& ParticleHandle : Particles)
+	void FPBDIslandManager::SortIslandEdges(FPBDIsland* Island)
 	{
-		FConstGenericParticleHandle Particle = ParticleHandle.Handle();
-		if (!Particle->IsDynamic())
+		// Sort the constraints based on level, and then add order (within each level)
+		for (int32 ContainerIndex = 0; ContainerIndex < Island->ContainerEdges.Num(); ++ContainerIndex)
 		{
-			continue;
+			Island->ContainerEdges[ContainerIndex].Sort(
+				[this](const FPBDIslandConstraint& L, const FPBDIslandConstraint& R)
+				{
+					return L.GetSortKey() < R.GetSortKey();
+				});
+
+			int32 IslandArrayIndex = 0;
+			for (FPBDIslandConstraint* Edge : Island->ContainerEdges[ContainerIndex])
+			{
+				Edge->IslandArrayIndex = IslandArrayIndex++;
+			}
 		}
+	}
 
-		int32* PNodeIndex = IslandGraph->ItemNodes.Find(Particle->Handle());
-		bSuccess = bSuccess && (PNodeIndex != nullptr);
-		if (ensure(bSuccess))
+	void FPBDIslandManager::ProcessWakes()
+	{
+		for (FPBDIsland* Island : Islands)
 		{
-			const int32 NodeIndex = *PNodeIndex;
-			bSuccess = bSuccess && IslandGraph->GraphNodes.IsValidIndex(NodeIndex);
-			if (ensure(bSuccess))
+			if (Island->Flags.bCheckSleep)
 			{
-				const FGraphNode& GraphNode = IslandGraph->GraphNodes[NodeIndex];
-				bSuccess = bSuccess && (GraphNode.NodeItem == Particle->Handle());
-				ensure(bSuccess);
+				// See if we should sleep or wake based on particle state and whether we are allowed to sleep
+				// (islands with a moving kinematic are not allowed to sleep)
+				bool bIsSleeping = Island->Flags.bIsSleepAllowed;
+				if (bIsSleeping)
+				{
+					for (const FPBDIslandParticle* Node : Island->Nodes)
+					{
+						bIsSleeping = bIsSleeping && Node->Flags.bIsSleeping;
+						if (!bIsSleeping)
+						{
+							break;
+						}
+					}
+				}
+
+				if (bIsSleeping != Island->Flags.bIsSleeping)
+				{
+					Island->Flags.bIsSleeping = bIsSleeping;
+					Island->SleepCounter = 0;
+				}
 			}
 		}
 	}
 
-	for (const FGraphEdge& GraphEdge : IslandGraph->GraphEdges)
+	void FPBDIslandManager::FinalizeIslands()
 	{
-		const int32 FirstNodeIndex = GraphEdge.FirstNode;
-		const int32 SecondNodeIndex = GraphEdge.SecondNode;
+		SCOPE_CYCLE_COUNTER(STAT_IslandManager_Finalize);
 
-		// Every constraint should reference at least one particle
-		bSuccess = bSuccess && ((FirstNodeIndex != INDEX_NONE) || (SecondNodeIndex != INDEX_NONE));
-		ensure(bSuccess);
-
-		// Every particle referenced by a constraint should be in the graph (including kinematics and statics)
-		// The node must point to the correct particle
-		if (FirstNodeIndex != INDEX_NONE)
+		for (int32 IslandIndex = Islands.Num() - 1; IslandIndex >= 0; --IslandIndex)
 		{
-			bSuccess = bSuccess && IslandGraph->GraphNodes.IsValidIndex(FirstNodeIndex);
-			if (ensure(bSuccess))
+			FPBDIsland* Island = Islands[IslandIndex];
+
+			// Remove empty islands
+			if (Island->NumEdges == 0)
 			{
-				bSuccess = bSuccess && (IslandGraph->GraphNodes[FirstNodeIndex].NodeItem == GraphEdge.EdgeItem->GetConstrainedParticles()[0]);
-				ensure(bSuccess);
+				DestroyIsland(Island);
+				continue;
 			}
-		}
-		if (SecondNodeIndex != INDEX_NONE)
-		{
-			bSuccess = bSuccess && IslandGraph->GraphNodes.IsValidIndex(SecondNodeIndex);
-			if (ensure(bSuccess))
+
+			// If the island is awake or just changed state, make sure all the particles and constraints agree
+			// NOTE: also handles the case where we add to a sleeping island (by checking bCheckSleep)
+			// @todo(chaos): we should only need to do this when the state changes...?
+			if (!Island->Flags.bIsSleeping || (Island->Flags.bIsSleeping != Island->Flags.bWasSleeping) || !!Island->Flags.bCheckSleep)
 			{
-				bSuccess = bSuccess && (IslandGraph->GraphNodes[SecondNodeIndex].NodeItem == GraphEdge.EdgeItem->GetConstrainedParticles()[1]);
-				ensure(bSuccess);
+				PropagateIslandSleep(Island);
+				Island->UpdateSyncState();
 			}
-		}
-	}
 
-	return bSuccess;
-}
+			// Reset of the sleep counter if the island is:
+			// - Not allowed to sleep this tick
+			// - Sleeping, since as soon as it wakes up we could start incrementing the counter
+			// - Just woken because we may not have been asleep for a whole frame (this probably makes the previous check unnecessary)
+			if (!Island->Flags.bIsSleepAllowed || !!Island->Flags.bIsSleeping || (Island->Flags.bIsSleeping != Island->Flags.bWasSleeping))
+			{
+				Island->SleepCounter = 0;
+			}
 
-bool FPBDIslandManager::DebugCheckParticleNotInGraph(const FGeometryParticleHandle* ParticleHandle) const
-{
-	bool bSuccess = true;
-	for (const auto& Node : IslandGraph->GraphNodes)
-	{
-		if (Node.NodeItem == ParticleHandle)
-		{
-			bSuccess = false;
-			ensure(false);
+			Island->Flags.bWasSleeping = Island->Flags.bIsSleeping;
+			Island->Flags.bCheckSleep = false;
+			Island->Flags.bItemsAdded = false;
+			Island->Flags.bItemsRemoved = false;
+			Island->ResimFrame = INDEX_NONE;
 		}
+
+		Validate();
 	}
 
-	for (const auto& ItemNode : IslandGraph->ItemNodes)
+	void FPBDIslandManager::ProcessSleep()
 	{
-		if (ItemNode.Key == ParticleHandle)
+		// Isloated particles are not kept in any island and need to be handled separately.
+		ProcessParticlesSleep();
+
+		// @todo(chaos): can go wide except for PropagateSleepState
+		for (FPBDIsland* Island : Islands)
 		{
-			bSuccess = false;
-			ensure(false);
+			if (!Island->Flags.bIsSleeping && !!Island->Flags.bIsSleepAllowed && !Island->Flags.bIsUsingCache)
+			{
+				// Update the sleep state based on particle movement etc
+				ProcessIslandSleep(Island);
+
+				if (Island->Flags.bIsSleeping)
+				{
+					PropagateIslandSleep(Island);
+				}
+			}
+
+			// This gets set to false again next tick if there are moving kinematics in the island
+			Island->Flags.bIsSleepAllowed = true;
 		}
 	}
 
-	for (const auto& Edge : IslandGraph->GraphEdges)
+	void FPBDIslandManager::ProcessParticlesSleep()
 	{
-		if (Edge.EdgeItem.GetParticle0() == ParticleHandle)
+		// We only need to process particles that have zero gravity because under gravity
+		// the particles will only sleep when held in place by a constraint and that will
+		// be handled in ProcessIslandSleep.
+		// @todo(chaos): keep track of particles with zero gravity? There usually are very few
+		// so it's a shame to have to visit all isolated particles here
+		TArray<FGeometryParticleHandle*> SleptParticles;
+		TArray<FGeometryParticleHandle*> DisabledParticles;
+		for (FTransientPBDRigidParticleHandle& Rigid : Particles.GetActiveDynamicMovingKinematicParticlesView())
 		{
-			bSuccess = false;
-			ensure(false);
+			if (Rigid.IsDynamic() && !Rigid.IsInConstraintGraph())
+			{
+				// If we have gravity we will not stop unless constrained
+				// so we only care about particles with zero gravity
+				bool bHasGravity = Rigid.GravityEnabled();
+				if (bHasGravity && (Gravity != nullptr))
+				{
+					const int32 GravityGroup = Rigid.GravityGroupIndex();
+					if (GravityGroup != INDEX_NONE)
+					{
+						bHasGravity = !Gravity->GetAcceleration(GravityGroup).IsNearlyZero();
+					}
+				}
+
+				if (!bHasGravity)
+				{
+					FRealSingle SleepLinearThreshold, SleepAngularThreshold;
+					int32 SleepCounterThreshold;
+					GetParticleSleepThresholds(Rigid.Handle(), PhysicsMaterials, PerParticlePhysicsMaterials, SimMaterials, SleepLinearThreshold, SleepAngularThreshold, SleepCounterThreshold);
+
+					// Check for sleep
+					if ((SleepLinearThreshold > 0) || (SleepAngularThreshold > 0))
+					{
+						int32 SleepCounter = 0;
+						if (SleepCounterThreshold < TNumericLimits<int32>::Max())
+						{
+							// Isolated particles have a max sleep counter of 127 (to reduce counter space in the particle)
+							SleepCounterThreshold = FMath::Min(SleepCounterThreshold, TNumericLimits<int8>::Max());
+
+							// Did we exceed the velocity threshold?
+							if ((Rigid.VSmooth().SizeSquared() > FMath::Square(SleepLinearThreshold)) 
+								|| (Rigid.WSmooth().SizeSquared() > FMath::Square(SleepAngularThreshold)))
+							{
+								continue;
+							}
+
+							// If we get here we want to sleep
+							// Update the counter and sleep if we exceed it
+							static_assert(sizeof(decltype(Rigid.SleepCounter())) == 1, "Expected int8 for SleepCounter(). Update clamp below");
+							SleepCounter = FMath::Min(int32(Rigid.SleepCounter()) + 1, int32(TNumericLimits<int8>::Max()));
+							if (SleepCounter > SleepCounterThreshold)
+							{
+								SleptParticles.Add(Rigid.Handle());
+								SleepCounter = 0;
+							}
+						}
+						Rigid.SetSleepCounter(int8(SleepCounter));
+					}
+				}
+			}
 		}
-		if (Edge.EdgeItem.GetParticle1() == ParticleHandle)
+
+		if (!SleptParticles.IsEmpty())
 		{
-			bSuccess = false;
-			ensure(false);
+			Particles.DeactivateParticles(SleptParticles);
 		}
 	}
-	return bSuccess;
-}
 
-bool FPBDIslandManager::DebugParticleConstraintsNotInGraph(const FGeometryParticleHandle* ParticleHandle, const int32 ContainerId) const
-{
-	bool bSuccess = true;
-	for (const auto& Edge : IslandGraph->GraphEdges)
+	void FPBDIslandManager::ProcessIslandSleep(FPBDIsland* Island)
 	{
-		if (Edge.ItemContainer == ContainerId)
+		bool bWithinSleepThreshold = true;
+		int32 SleepCounterThreshold = 0;
+		for (FPBDIslandParticle* Node : Island->GetParticles())
 		{
-			if ((Edge.EdgeItem.GetParticle0() == ParticleHandle) || (Edge.EdgeItem.GetParticle1() == ParticleHandle))
+			// All zeroes means never sleep/disable
+			if ((Node->SleepLinearThresholdSq <= 0) && (Node->SleepAngularThresholdSq <= 0))
+			{
+				bWithinSleepThreshold = false;
+				break;
+			}
+
+			// Check the particle state against the thresholds
+			FConstGenericParticleHandle P = Node->GetParticle();
+
+			// Did we exceed the velocity threshold?
+			if ((P->VSmooth().SizeSquared() > Node->SleepLinearThresholdSq) 
+				|| (P->WSmooth().SizeSquared() > Node->SleepAngularThresholdSq))
 			{
-				bSuccess = false;
-				ensure(false);
+				bWithinSleepThreshold = false;
+				break;
 			}
+
+			// Take the longest sleep time
+			SleepCounterThreshold = FMath::Max(SleepCounterThreshold, Node->SleepCounterThreshold);
 		}
-	}
-	return bSuccess;
-}
 
-bool FPBDIslandManager::DebugCheckConstraintNotInGraph(const FConstraintHandle* ConstraintHandle) const
-{
-	bool bSuccess = true;
-	for (const auto& Edge : IslandGraph->GraphEdges)
-	{
-		if (Edge.EdgeItem.Get() == ConstraintHandle)
+		int32 SleepCounter = 0;
+		if (bWithinSleepThreshold)
 		{
-			bSuccess = false;
-			ensure(false);
+			// If we get here we may want to sleep
+			SleepCounter = Island->SleepCounter + 1;
+			if (SleepCounter > SleepCounterThreshold)
+			{
+				Island->Flags.bIsSleeping = true;
+				SleepCounter = 0;
+			}
 		}
+		Island->SleepCounter = SleepCounter;
 	}
 
-	for (const auto& ItemEdge : IslandGraph->ItemEdges)
+	void FPBDIslandManager::PropagateIslandSleep(FPBDIsland* Island)
 	{
-		if (ItemEdge.Key.Get() == ConstraintHandle)
+		bool bRebuildViews = false;
+		for (FPBDIslandParticle* IslandNode : Island->Nodes)
+		{
+			if (Island->Flags.bIsSleeping != IslandNode->Flags.bIsSleeping)
+			{
+				FGeometryParticleHandle* Particle = IslandNode->GetParticle();
+
+				// Put to sleep...
+				if (Island->Flags.bIsSleeping)
+				{
+					Particles.DeactivateParticle(Particle, true);
+				}
+
+				// Wake up...
+				if (!Island->Flags.bIsSleeping)
+				{
+					Particles.ActivateParticle(Particle, true);
+
+					// When we wake particles, we have skipped their integrate step which causes some issues:
+					//	- we have zero velocity (no gravity or external forces applied)
+					//	- the world transforms cached in the ShapesArray will be at the last post-integrate positions
+					//	  which doesn't match what the velocity is telling us
+					// This causes problems for the solver - essentially we have an "initial overlap" situation.
+					// @todo(chaos): We could just run (partial) integrate here for this particle, but we don't know about the Evolution - fix this
+					// Better solution: Leave Particle X and P in current state for sleeping particles (rather than setting X = P)
+					for (const TUniquePtr<FPerShapeData>& Shape : Particle->ShapesArray())
+					{
+						Shape->UpdateLeafWorldTransform(Particle);
+					}
+				}
+
+				IslandNode->Flags.bIsSleeping = Island->Flags.bIsSleeping;
+				bRebuildViews = true;
+			}
+		}
+
+		// Update views if we changed particle state
+		if (bRebuildViews)
 		{
-			bSuccess = false;
-			ensure(false);
+			Particles.RebuildViews();
 		}
-	}
 
-	for (int32 IslandIndex = 0; IslandIndex < NumIslands(); ++IslandIndex)
-	{
-		const FPBDIsland* Island = GetIsland(IslandIndex);
-		for (int32 ContainerIndex = 0; ContainerIndex < Island->GetNumConstraintContainers(); ++ContainerIndex)
+		// Set the constraint sleep state to match
+		for (TArray<FPBDIslandConstraint*>& IslandEdges : Island->ContainerEdges)
 		{
-			for (const FPBDIslandConstraint& Constraint : Island->GetConstraints(ContainerIndex))
+			for (FPBDIslandConstraint* IslandEdge : IslandEdges)
 			{
-				if (Constraint.GetConstraint() == ConstraintHandle)
+				if (Island->Flags.bIsSleeping != IslandEdge->Flags.bIsSleeping)
 				{
-					bSuccess = false;
-					ensure(false);
+					IslandEdge->GetConstraint()->SetIsSleeping(Island->Flags.bIsSleeping);
+					IslandEdge->Flags.bIsSleeping = Island->Flags.bIsSleeping;
 				}
 			}
 		}
 	}
-	return bSuccess;
-}
-
-bool FPBDIslandManager::DebugCheckIslandConstraints() const
-{
-	bool bSuccess = true;
 
-	for (const auto& GraphEdge : IslandGraph->GraphEdges)
+	void FPBDIslandManager::ProcessDisable(TFunctionRef<void(FPBDRigidParticleHandle*)> ParticleDisableFunctor)
 	{
-		const FConstraintHandleHolder& ConstraintHolder = GraphEdge.EdgeItem;
-
-		// NOTE: The check for both null here is for unit tests that do not use particles...not ideal...
-		if ((ConstraintHolder.GetParticle0() != nullptr) || (ConstraintHolder.GetParticle1() != nullptr))
+		// @todo(chaos): parallelize
+		if (DisableCounterThreshold == TNumericLimits<int32>::Max())
 		{
-			// Make sure that if we have an IslandIndex it is a valid one
-			const bool bIslandIndexOK = (GraphEdge.IslandIndex == INDEX_NONE) || IslandGraph->GraphIslands.IsValidIndex(GraphEdge.IslandIndex);
-			bSuccess = bSuccess && bIslandIndexOK;
-			ensure(bIslandIndexOK);
+			return;
+		}
 
-			// At least one particle on each constraint must be dynamic if we are in an awake island.
-			// If we do not have a dynamic, we must have no island or it must be asleep
-			const bool bInAwakeIsland = (GraphEdge.IslandIndex != INDEX_NONE) && IslandGraph->GraphIslands.IsValidIndex(GraphEdge.IslandIndex) && !IslandGraph->GraphIslands[GraphEdge.IslandIndex].bIsSleeping;
-			const bool bIsDynamic0 = (ConstraintHolder.GetParticle0() != nullptr) ? FConstGenericParticleHandle(ConstraintHolder.GetParticle0())->IsDynamic() : false;
-			const bool bIsDynamic1 = (ConstraintHolder.GetParticle1() != nullptr) ? FConstGenericParticleHandle(ConstraintHolder.GetParticle1())->IsDynamic() : false;
-			
-			if (bInAwakeIsland)
-			{
-				// If the island is awake, both particle should be too
-				bSuccess = bSuccess && (bIsDynamic0 || bIsDynamic1);
-				ensure(bIsDynamic0 || bIsDynamic1);
+		TArray<FPBDRigidParticleHandle*> DisableParticles;
 
-				// If the island is awake, constraint should be too
-				if (ConstraintHolder->SupportsSleeping())
-				{
-					const bool bIsAsleep = ConstraintHolder->IsSleeping();
-					ensure(!bIsAsleep);
-				}
-			}
-			else
+		// Check to see if islolated particles (not in an island) need to be disabled
+		// We prefer to iterate over islands because nodes in the graph cache their thresholds
+		ProcessParticlesDisable(DisableParticles);
+
+		// Check islands to see if we want to disable the particles in them
+		for (FPBDIsland* Island : Islands)
+		{
+			if (!Island->Flags.bIsSleeping && !Island->IsUsingCache())
 			{
-				// If the island is asleep, make sure the constraint is too
-				if (ConstraintHolder->SupportsSleeping())
-				{
-					const bool bIsAsleep = ConstraintHolder->IsSleeping();
-					ensure(bIsAsleep);
-				}
+				// Update the sleep state based on particle movement etc
+				ProcessIslandDisable(Island, DisableParticles);
 			}
+		}
 
-			if (!bIsDynamic0 && !bIsDynamic1)
+		// Disable the particles we identified
+		for (FPBDRigidParticleHandle* Particle : DisableParticles)
+		{
+			if (Particle != nullptr)
 			{
-				bSuccess = bSuccess && !bInAwakeIsland;
-				ensure(!bInAwakeIsland);
+				ParticleDisableFunctor(Particle);
 			}
 		}
 	}
 
-	return bSuccess;
-}
-
-void FPBDIslandManager::DebugCheckParticleIslands(const FGeometryParticleHandle* ParticleHandle) const
-{
-	// Check that every particle is included in the islands they think and no others
-	const int32* PNodeIndex = IslandGraph->ItemNodes.Find(ParticleHandle);
-	if (PNodeIndex != nullptr)
+	void FPBDIslandManager::ProcessParticlesDisable(TArray<FPBDRigidParticleHandle*>& OutDisableParticles)
 	{
-		const int32 NodeIndex = *PNodeIndex;
-		const auto& Node = IslandGraph->GraphNodes[NodeIndex];
-		ensure(Node.NodeItem == ParticleHandle);
-
-		// Find the islands containing this node (either directly, or referenced by an constraint)
-		TArray<int32> NodeSolverIslands;
-		for (int32 IslandIndex = 0; IslandIndex < NumIslands(); ++IslandIndex)
+		// Check all isolated particle (not in any island)
+		for (FTransientPBDRigidParticleHandle& Particle : Particles.GetActiveDynamicMovingKinematicParticlesView())
 		{
-			const FPBDIsland* Island = GetIsland(IslandIndex);
-			for (const FPBDIslandParticle& IslandParticle : Island->GetParticles())
+			if (Particle.IsDynamic() && !Particle.IsInConstraintGraph())
 			{
-				const FGeometryParticleHandle* IslandParticleHandle = IslandParticle.GetParticle();
+				FRealSingle DisableLinearThreshold, DisableAngularThreshold;
+				GetParticleDisableThresholds(Particle.Handle(), PhysicsMaterials, PerParticlePhysicsMaterials, SimMaterials, DisableLinearThreshold, DisableAngularThreshold);
 
-				if (IslandParticleHandle == ParticleHandle)
-				{
-					NodeSolverIslands.AddUnique(IslandIndex);
-				}
-			}
-			for (int32 ContainerIndex = 0; ContainerIndex < Island->GetNumConstraintContainers(); ++ContainerIndex)
-			{
-				for (const FPBDIslandConstraint& IslandConstraint : Island->GetConstraints(ContainerIndex))
+				// Check for disable
+				if ((DisableLinearThreshold > 0) || (DisableAngularThreshold > 0))
 				{
-					if ((IslandConstraint.GetConstraint()->GetConstrainedParticles()[0] == ParticleHandle) || (IslandConstraint.GetConstraint()->GetConstrainedParticles()[1] == ParticleHandle))
+					int32 DisableCounter = 0;
+
+					const FReal VSq = Particle.VSmooth().SizeSquared();
+					const FReal WSq = Particle.WSmooth().SizeSquared();
+					if ((VSq < FMath::Square(DisableLinearThreshold)) && (WSq < FMath::Square(DisableAngularThreshold)))
 					{
-						NodeSolverIslands.AddUnique(IslandIndex);
+						DisableCounter = Particle.DisableCounter() + 1;
+						if (DisableCounter > DisableCounterThreshold)
+						{
+							OutDisableParticles.Add(Particle.Handle());
+						}
 					}
+
+					Particle.SetDisableCounter(int8(DisableCounter));
 				}
 			}
 		}
-
-		// NOTE: Particles are always assigneed to an island, but if we do not have any constraints, the node may not be in the Island
-		if (FConstGenericParticleHandle(ParticleHandle)->IsDynamic())
-		{
-			// Dynamic particles should be in 1 island
-			ensure(Node.IslandIndex != INDEX_NONE);
-		}
-		else
-		{
-			// There's not much we can check for non-dynamic particles. They should not have an IslandIndex because that's not used
-			// for kinematics, but they may or may not be in an island. E.g., if the kinematic is connected to dynamic it will
-			// be in one or more islands, but if it has no connections or is connected only to another kinematic it will not be in any islands.
-			ensure(Node.IslandIndex == INDEX_NONE);
-		}
 	}
-}
 
-void FPBDIslandManager::DebugCheckConstraintIslands(const FConstraintHandle* ConstraintHandle, const int32 EdgeIslandIndex) const
-{
-	// Find all the islands containing this edge
-	TArray<int32> EdgeSolverIslands;
-	for (int32 SolverIslandIndex = 0; SolverIslandIndex < NumIslands(); ++SolverIslandIndex)
+	void FPBDIslandManager::ProcessIslandDisable(FPBDIsland* Island, TArray<FPBDRigidParticleHandle*>& OutDisableParticles)
 	{
-		const FPBDIsland* Island = GetIsland(SolverIslandIndex);
-		for (int32 ContainerIndex = 0; ContainerIndex < Island->GetNumConstraintContainers(); ++ContainerIndex)
+		for (FPBDIslandParticle* Node : Island->GetParticles())
 		{
-			for (const FPBDIslandConstraint& IslandConstraint : Island->GetConstraints(ContainerIndex))
+			// All zeroes means never disable
+			if ((Node->DisableLinearThresholdSq <= 0) && (Node->DisableAngularThresholdSq <= 0))
+			{
+				continue;
+			}
+
+			if (FPBDRigidParticleHandle* Particle = Node->GetParticle()->CastToRigidParticle())
 			{
-				// Make sure all the constraints are in the correct list according to type
-				ensure(ContainerIndex == IslandConstraint.GetConstraint()->GetContainerId());
+				int32 DisableCounter = 0;
 
-				if (IslandConstraint.GetConstraint() == ConstraintHandle)
+				// Did we exceed the velocity thresholds?
+				const FReal VSq = Particle->V().SizeSquared();
+				const FReal WSq = Particle->W().SizeSquared();
+				if ((VSq < Node->DisableLinearThresholdSq) && (WSq < Node->DisableAngularThresholdSq))
 				{
-					EdgeSolverIslands.AddUnique(SolverIslandIndex);
+					// We are within the velocity thresholds, so see if we should disable
+					DisableCounter = Particle->DisableCounter() + 1;
+					if (DisableCounter > DisableCounterThreshold)
+					{
+						OutDisableParticles.Add(Particle);
+					}
 				}
+
+				Particle->SetDisableCounter(int8(DisableCounter));
 			}
 		}
 	}
 
-	if (EdgeIslandIndex != INDEX_NONE)
+	void FPBDIslandManager::ApplyDeterminism()
+	{
+		// Nothing to do as long as we sorted constraints
+	}
+
+	bool FPBDIslandManager::Validate() const
 	{
-		if (IslandGraph->GraphIslands[EdgeIslandIndex].bIsSleeping)
+		SCOPE_CYCLE_COUNTER(STAT_IslandManager_Validate);
+
+		if (!CVars::bChaosConstraintGraphValidate)
+		{
+			return true;
+		}
+
+		bool bIsValid = true;
+
+		// Make sure we have processed the merges etc
+		CHAOS_CONSTRAINTGRAPH_VALIDATE_TEST(MergeSets.IsEmpty());
+
+		// Check various Node properties
+		for (FPBDIslandParticle* Node : Nodes)
 		{
-			// Make sure the constraint and its island agree on sleepiness
-			if (ConstraintHandle->SupportsSleeping())
+			// All dynamic nodes should be in an island
+			CHAOS_CONSTRAINTGRAPH_VALIDATE_TEST(!Node->Flags.bIsDynamic || (Node->Island != nullptr));
+
+			// Kinematic nodes must not be in an island
+			CHAOS_CONSTRAINTGRAPH_VALIDATE_TEST(Node->Flags.bIsDynamic || (Node->Island == nullptr));
+
+			// Nodes without edges should have been removed
+			CHAOS_CONSTRAINTGRAPH_VALIDATE_TEST(Node->Edges.Num() > 0);
+
+			if (Node->Island != nullptr)
 			{
-				ensure(ConstraintHandle->IsSleeping());
+				// Make sure the island exists
+				CHAOS_CONSTRAINTGRAPH_VALIDATE_TEST(Islands[Node->Island->ArrayIndex] == Node->Island);
+
+				// Make sure the island knows about the node
+				CHAOS_CONSTRAINTGRAPH_VALIDATE_TEST(Node->Island->Nodes.Contains(Node));
+				CHAOS_CONSTRAINTGRAPH_VALIDATE_TEST(Node->Island->Nodes[Node->IslandArrayIndex] == Node);
+
+				// Make sure the particle sleep state is correct
+				if (Node->Island->Flags.bIsSleeping)
+				{
+					CHAOS_CONSTRAINTGRAPH_VALIDATE_TEST(Node->Particle->Sleeping())
+				}
 			}
 
-			// If we are sleeping, we usually don't appear in the solver islands (because we don't need to solve sleepers)
-			// but we do create solver islands for 1 tick after the island goes to sleep so that we can sync the sleep state
-			// (see PopulateIslands and SyncSleepState). 
-			if (EdgeSolverIslands.Num() > 0)
+			// Make sure the edge points to the node and that we're in the same island
+			for (FPBDIslandConstraint* Edge : Node->Edges)
 			{
-				ensure(EdgeSolverIslands.Num() == 1);
-				const int32 EdgeIsland = IslandIndexing[EdgeSolverIslands[0]];
-				ensure(EdgeIslandIndex == EdgeIsland);
+				CHAOS_CONSTRAINTGRAPH_VALIDATE_TEST((Edge->Nodes[0] == Node) || (Edge->Nodes[1] == Node));
+
+				if (Node->Flags.bIsDynamic)
+				{
+					CHAOS_CONSTRAINTGRAPH_VALIDATE_TEST(Node->Island == Edge->Island);
+				}
 			}
+
+			// Particle sleep state should match
+			CHAOS_CONSTRAINTGRAPH_VALIDATE_TEST(Node->Flags.bIsSleeping == IsParticleSleeping(Node->Particle));
 		}
-		else
+
+		// Check various Edge properties
+		for (FPBDIslandConstraint* Edge : Edges)
 		{
-			// Make sure the constraint and its island agree on sleepiness
-			if (ConstraintHandle->SupportsSleeping())
+			// Make sure nodes know about the edge
+			CHAOS_CONSTRAINTGRAPH_VALIDATE_TEST((Edge->Nodes[0] == nullptr) || Edge->Nodes[0]->Edges.Contains(Edge));
+			CHAOS_CONSTRAINTGRAPH_VALIDATE_TEST((Edge->Nodes[1] == nullptr) || Edge->Nodes[1]->Edges.Contains(Edge));
+
+			// All edges should be in an island (even inactive edges between two kinematics)
+			const bool bIsInIsland = (Edge->Island != nullptr);
+			CHAOS_CONSTRAINTGRAPH_VALIDATE_TEST(bIsInIsland);
+
+			if (bIsInIsland)
 			{
-				ensure(!ConstraintHandle->IsSleeping());
+				// Make sure the island exists
+				CHAOS_CONSTRAINTGRAPH_VALIDATE_TEST(Islands[Edge->Island->ArrayIndex] == Edge->Island);
+
+				// Make sure the island edge index is correct
+				CHAOS_CONSTRAINTGRAPH_VALIDATE_TEST(Edge->Island->ContainerEdges[Edge->ContainerIndex][Edge->IslandArrayIndex] == Edge);
+
+				// Make sure the island and edge sleep states match
+				if (Edge->Constraint->SupportsSleeping())
+				{
+					CHAOS_CONSTRAINTGRAPH_VALIDATE_TEST(Edge->Island->Flags.bIsSleeping == Edge->Flags.bIsSleeping);
+				}
 			}
 
-			// We are awake, make sure our island search produced the correct island (and only that one)
-			ensure(EdgeSolverIslands.Num() == 1);
-			if (EdgeSolverIslands.Num() > 0)
+			// Constraint sleep state should match edge
+			if (Edge->Constraint->SupportsSleeping())
 			{
-				const int32 EdgeIsland = IslandIndexing[EdgeSolverIslands[0]];
-				ensure(EdgeIslandIndex == EdgeIsland);
+				CHAOS_CONSTRAINTGRAPH_VALIDATE_TEST(Edge->Flags.bIsSleeping == Edge->Constraint->IsSleeping());
 			}
 		}
-	}
-	else
-	{
-		// We are not yet added to any islands
-		ensure(EdgeSolverIslands.Num() == 0);
-	}
-}
 
-void FPBDIslandManager::DebugCheckIslands() const
-{
-	// Check that the solver island state matches the graph island state
-	for (int32 IslandIndex = 0; IslandIndex < IslandGraph->GraphIslands.GetMaxIndex(); ++IslandIndex)
-	{
-		if (IslandGraph->GraphIslands.IsValidIndex(IslandIndex))
+		// Check various Island properties
+		for (FPBDIsland* Island : Islands)
 		{
-			if (IslandGraph->GraphIslands[IslandIndex].bIsSleeping != IslandGraph->GraphIslands[IslandIndex].IslandItem->IsSleeping())
+			// We are done with merges
+			CHAOS_CONSTRAINTGRAPH_VALIDATE_TEST(Island->MergeSet == nullptr);
+
+			// Island wake flag was reset
+			CHAOS_CONSTRAINTGRAPH_VALIDATE_TEST(Island->Flags.bCheckSleep == false);
+
+			// No dupes in the lists
+			for (FPBDIslandParticle* Node0 : Island->Nodes)
+			{
+				int32 Count = 0;
+				for (FPBDIslandParticle* Node1 : Island->Nodes)
+				{
+					if (Node0 == Node1)
+					{
+						++Count;
+					}
+				}
+				CHAOS_CONSTRAINTGRAPH_VALIDATE_TEST(Count == 1);
+			}
+			for (TArray<FPBDIslandConstraint*> IslandContainerEdges : Island->ContainerEdges)
 			{
-				ensure(false);
+				for (FPBDIslandConstraint* Edge0 : IslandContainerEdges)
+				{
+					int32 Count = 0;
+					for (FPBDIslandConstraint* Edge1 : IslandContainerEdges)
+					{
+						if (Edge0 == Edge1)
+						{
+							++Count;
+						}
+					}
+					CHAOS_CONSTRAINTGRAPH_VALIDATE_TEST(Count == 1);
+				}
 			}
 
-			const int32 SolverIslandIndex = IslandGraph->GraphIslands[IslandIndex].IslandItem->GetIslandIndex();
-			if (IslandIndexing[SolverIslandIndex] != IslandIndex)
+			// Edge Count is correct
+			int32 NumIslandEdges = 0;
+			for (TArray<FPBDIslandConstraint*> IslandContainerEdges : Island->ContainerEdges)
 			{
-				ensure(false);
+				NumIslandEdges += IslandContainerEdges.Num();
 			}
+			CHAOS_CONSTRAINTGRAPH_VALIDATE_TEST(NumIslandEdges == Island->NumEdges);
 		}
-	}
-
-	// Make sure islands only contain valid constraints
-	DebugCheckIslandConstraints();
-
-	// Check that every particle is included in the islands they think and no others
-	for (const auto& Node : IslandGraph->GraphNodes)
-	{
-		DebugCheckParticleIslands(Node.NodeItem);
-	}
 
-	// Check that every constraint is included in the island they think and no others
-	for (const auto& Edge : IslandGraph->GraphEdges)
-	{
-		DebugCheckConstraintIslands(Edge.EdgeItem, Edge.IslandIndex);
+		return true;
 	}
-}
-#endif
 
-}	// namespace Private
-}	// namespace Chaos
\ No newline at end of file
+} // namsepace Chaos::Private
