Chaos - big island manager refactor

diff --git a/Engine/Source/Runtime/Experimental/Chaos/Private/Chaos/Joint/PBDJointContainerSolver.cpp b/Engine/Source/Runtime/Experimental/Chaos/Private/Chaos/Joint/PBDJointContainerSolver.cpp
--- a/Engine/Source/Runtime/Experimental/Chaos/Private/Chaos/Joint/PBDJointContainerSolver.cpp
+++ b/Engine/Source/Runtime/Experimental/Chaos/Private/Chaos/Joint/PBDJointContainerSolver.cpp
@@ -1,6 +1,7 @@
 // Copyright Epic Games, Inc. All Rights Reserved.
 #include "Chaos/Joint/PBDJointContainerSolver.h"
 #include "Chaos/Joint/ChaosJointLog.h"
+#include "Chaos/Island/IslandManager.h"
 
 namespace Chaos
 {
@@ -130,12 +131,12 @@ namespace Chaos
 			}
 		}
 
-		void FPBDJointContainerSolver::AddConstraints(const TArrayView<Private::FPBDIslandConstraint>& IslandConstraints)
+		void FPBDJointContainerSolver::AddConstraints(const TArrayView<Private::FPBDIslandConstraint*>& IslandConstraints)
 		{
-			for (Private::FPBDIslandConstraint& IslandConstraint : IslandConstraints)
+			for (Private::FPBDIslandConstraint* IslandConstraint : IslandConstraints)
 			{
 				// We will only ever be given constraints from our container (asserts in non-shipping)
-				const int32 ContainerConstraintIndex = IslandConstraint.GetConstraint()->AsUnsafe<FPBDJointConstraintHandle>()->GetConstraintIndex();
+				const int32 ContainerConstraintIndex = IslandConstraint->GetConstraint()->AsUnsafe<FPBDJointConstraintHandle>()->GetConstraintIndex();
 
 				AddConstraint(ContainerConstraintIndex);
 			}
diff --git a/Engine/Source/Runtime/Experimental/Chaos/Private/Chaos/PBDCollisionConstraints.cpp b/Engine/Source/Runtime/Experimental/Chaos/Private/Chaos/PBDCollisionConstraints.cpp
--- a/Engine/Source/Runtime/Experimental/Chaos/Private/Chaos/PBDCollisionConstraints.cpp
+++ b/Engine/Source/Runtime/Experimental/Chaos/Private/Chaos/PBDCollisionConstraints.cpp
@@ -32,6 +32,8 @@ namespace Chaos
 	namespace CVars
 	{
 		extern bool bChaos_PBDCollisionSolver_UseJacobiPairSolver2;
+		
+		extern bool bChaosSolverPersistentGraph;
 	}
 
 	int32 CollisionParticlesBVHDepth = 4;
@@ -78,8 +80,6 @@ namespace Chaos
 
 	bool DebugDrawProbeDetection = false;
 	FAutoConsoleVariableRef CVarDebugDrawProbeDetection(TEXT("p.Chaos.Collision.DebugDrawProbeDetection"), DebugDrawProbeDetection, TEXT("Draw probe constraint detection."));
-	
-	extern bool bChaosSolverPersistentGraph;
 
 #if CHAOS_DEBUG_DRAW
 	namespace CVars
@@ -416,8 +416,8 @@ namespace Chaos
 		}
 
 		// This is a bit dodgy - we pass the FaceIndex to both material requests, knowing that at most one of the shapes will use it
-		const FChaosPhysicsMaterial* PhysicsMaterial0 = Private::GetPhysicsMaterial(Constraint.Particle[0], Constraint.GetShape0(), ShapeFaceIndex, MPhysicsMaterials, MPerParticlePhysicsMaterials, SimMaterials);
-		const FChaosPhysicsMaterial* PhysicsMaterial1 = Private::GetPhysicsMaterial(Constraint.Particle[1], Constraint.GetShape1(), ShapeFaceIndex, MPhysicsMaterials, MPerParticlePhysicsMaterials, SimMaterials);
+		const FChaosPhysicsMaterial* PhysicsMaterial0 = Private::GetPhysicsMaterial(Constraint.Particle[0], Constraint.GetShape0(), ShapeFaceIndex, &MPhysicsMaterials, &MPerParticlePhysicsMaterials, SimMaterials);
+		const FChaosPhysicsMaterial* PhysicsMaterial1 = Private::GetPhysicsMaterial(Constraint.Particle[1], Constraint.GetShape1(), ShapeFaceIndex, &MPhysicsMaterials, &MPerParticlePhysicsMaterials, SimMaterials);
 
 		FReal MaterialRestitution = 0;
 		FReal MaterialRestitutionThreshold = 0;
@@ -604,30 +604,28 @@ namespace Chaos
 		// Debugging/diagnosing: if we have collisions disabled, remove all collisions from the graph and don't add any more
 		if (!GetCollisionsEnabled())
 		{
-			IslandManager.RemoveConstraints(GetContainerId());
+			IslandManager.RemoveContainerConstraints(GetContainerId());
 			return;
 		}
 
-		// If we are running with a persistent graph, remove expired collisions	
-		if (bChaosSolverPersistentGraph)
-		{
-			// Find all expired constraints in the graph
-			TempCollisions.Reset();
-			IslandManager.VisitConstraintsInAwakeIslands(GetContainerId(),
-				[this](FConstraintHandle* ConstraintHandle)
+		// Remove expired collisions
+		// @chaos(todo): if graph persistent is disabled we remove all collisions, but in a non-optimal way...
+		TempCollisions.Reset();
+		const bool bRemoveAllAwakeCollisions = !CVars::bChaosSolverPersistentGraph;
+		IslandManager.VisitAwakeConstraints(GetContainerId(),
+			[this, bRemoveAllAwakeCollisions](const Private::FPBDIslandConstraint* IslandConstraint)
+			{
+				FPBDCollisionConstraintHandle* CollisionHandle = IslandConstraint->GetConstraint()->AsUnsafe<FPBDCollisionConstraintHandle>();
+				if (bRemoveAllAwakeCollisions || !CollisionHandle->IsEnabled() || CollisionHandle->IsProbe() || ConstraintAllocator.IsConstraintExpired(CollisionHandle->GetContact()))
 				{
-					FPBDCollisionConstraintHandle* CollisionHandle = ConstraintHandle->AsUnsafe<FPBDCollisionConstraintHandle>();
-					if (!CollisionHandle->IsEnabled() || CollisionHandle->IsProbe() || ConstraintAllocator.IsConstraintExpired(CollisionHandle->GetContact()))
-					{
-						TempCollisions.Add(CollisionHandle);
-					}
-				});
+					TempCollisions.Add(CollisionHandle);
+				}
+			});
 
-			// Remove expired constraints
-			for (FPBDCollisionConstraintHandle* CollisionHandle : TempCollisions)
-			{
-				IslandManager.RemoveConstraint(GetContainerId(), CollisionHandle);
-			}
+		// Remove expired constraints
+		for (FPBDCollisionConstraintHandle* CollisionHandle : TempCollisions)
+		{
+			IslandManager.RemoveConstraint(CollisionHandle);
 		}
 
 		// Collect all the new constraints that need to be added to the graph
diff --git a/Engine/Source/Runtime/Experimental/Chaos/Private/Chaos/PBDJointConstraints.cpp b/Engine/Source/Runtime/Experimental/Chaos/Private/Chaos/PBDJointConstraints.cpp
--- a/Engine/Source/Runtime/Experimental/Chaos/Private/Chaos/PBDJointConstraints.cpp
+++ b/Engine/Source/Runtime/Experimental/Chaos/Private/Chaos/PBDJointConstraints.cpp
@@ -843,7 +843,7 @@ namespace Chaos
 			}
 			else if (bIsInGraph && !bShouldBeInGraph)
 			{
-				IslandManager.RemoveConstraint(ContainerId, ConstraintHandle);
+				IslandManager.RemoveConstraint(ConstraintHandle);
 			}
 		}
 	}
diff --git a/Engine/Source/Runtime/Experimental/Chaos/Private/Chaos/PBDRigidsEvolution.cpp b/Engine/Source/Runtime/Experimental/Chaos/Private/Chaos/PBDRigidsEvolution.cpp
--- a/Engine/Source/Runtime/Experimental/Chaos/Private/Chaos/PBDRigidsEvolution.cpp
+++ b/Engine/Source/Runtime/Experimental/Chaos/Private/Chaos/PBDRigidsEvolution.cpp
@@ -200,7 +200,8 @@ namespace Chaos
 	};
 
 	FPBDRigidsEvolutionBase::FPBDRigidsEvolutionBase(FPBDRigidsSOAs& InParticles, THandleArray<FChaosPhysicsMaterial>& InSolverPhysicsMaterials, bool InIsSingleThreaded)
-	    : IslandGroupManager(IslandManager)
+	    : IslandManager(InParticles)
+		, IslandGroupManager(IslandManager)
 		, Particles(InParticles)
 		, SolverPhysicsMaterials(InSolverPhysicsMaterials)
 		, InternalAcceleration(nullptr)
@@ -213,7 +214,6 @@ namespace Chaos
 	{
 		Particles.GetParticleHandles().AddArray(&PhysicsMaterials);
 		Particles.GetParticleHandles().AddArray(&PerParticlePhysicsMaterials);
-		Particles.GetParticleHandles().AddArray(&ParticleDisableCount);
 		Particles.GetParticleHandles().AddArray(&Collided);
 
 		for (auto& Particle : InParticles.GetNonDisabledView())
@@ -228,7 +228,6 @@ namespace Chaos
 	{
 		Particles.GetParticleHandles().RemoveArray(&PhysicsMaterials);
 		Particles.GetParticleHandles().RemoveArray(&PerParticlePhysicsMaterials);
-		Particles.GetParticleHandles().RemoveArray(&ParticleDisableCount);
 		Particles.GetParticleHandles().RemoveArray(&Collided);
 		WaitOnAccelerationStructure();
 	}
@@ -1052,6 +1051,14 @@ namespace Chaos
 			{
 				Particles.MarkTransientDirtyParticle(Particle);
 			}
+
+			// Reset the sleep counter when the state changes
+			// NOTE: this sleep counter is only used for particles
+			// that are not in the graph (have no constraints on them)
+			if (bIsDynamic)
+			{
+				Particle->SetSleepCounter(0);
+			}
 		}
 
 		// If we are now dynamic and enabled, we need to be in the graph if not already there. We may not be
diff --git a/Engine/Source/Runtime/Experimental/Chaos/Private/Chaos/PBDRigidsEvolutionGBF.cpp b/Engine/Source/Runtime/Experimental/Chaos/Private/Chaos/PBDRigidsEvolutionGBF.cpp
--- a/Engine/Source/Runtime/Experimental/Chaos/Private/Chaos/PBDRigidsEvolutionGBF.cpp
+++ b/Engine/Source/Runtime/Experimental/Chaos/Private/Chaos/PBDRigidsEvolutionGBF.cpp
@@ -119,6 +119,9 @@ namespace Chaos
 		bool bChaosCollisionModiferBeforeCCD = false;
 		FAutoConsoleVariableRef  CVarChaosCollisionModiferBeforeCCD(TEXT("p.Chaos.Solver.CollisionModifiersBeforeCCD"), bChaosCollisionModiferBeforeCCD, TEXT("True: run the collision modifiers before CCD rewind is applied; False(default): run modifiers after CCD rewind. See comments in code."));
 
+		// Whether the constraint graph retains state between ticks. It will be very expensive with this disabled...
+		bool bChaosSolverPersistentGraph = true;
+		FAutoConsoleVariableRef CVarChaosSolverPersistentGraph(TEXT("p.Chaos.Solver.PersistentGraph"), bChaosSolverPersistentGraph, TEXT(""));
 
 		DECLARE_CYCLE_STAT(TEXT("FPBDRigidsEvolutionGBF::AdvanceOneTimeStep"), STAT_Evolution_AdvanceOneTimeStep, STATGROUP_Chaos);
 		DECLARE_CYCLE_STAT(TEXT("FPBDRigidsEvolutionGBF::UnclusterUnions"), STAT_Evolution_UnclusterUnions, STATGROUP_Chaos);
@@ -143,7 +146,6 @@ namespace Chaos
 		DECLARE_CYCLE_STAT(TEXT("FPBDRigidsEvolutionGBF::AddSleepingContacts"), STAT_Evolution_AddSleepingContacts, STATGROUP_Chaos);
 		DECLARE_CYCLE_STAT(TEXT("FPBDRigidsEvolutionGBF::PreApplyCallback"), STAT_Evolution_PreApplyCallback, STATGROUP_Chaos);
 		DECLARE_CYCLE_STAT(TEXT("FPBDRigidsEvolutionGBF::ParallelSolve"), STAT_Evolution_ParallelSolve, STATGROUP_Chaos);
-		DECLARE_CYCLE_STAT(TEXT("FPBDRigidsEvolutionGBF::BuildDisabledParticles"), STAT_Evolution_BuildDisabledParticles, STATGROUP_Chaos);
 		DECLARE_CYCLE_STAT(TEXT("FPBDRigidsEvolutionGBF::SaveParticlePostSolve"), STAT_Evolution_SavePostSolve, STATGROUP_Chaos);
 		DECLARE_CYCLE_STAT(TEXT("FPBDRigidsEvolutionGBF::DeactivateSleep"), STAT_Evolution_DeactivateSleep, STATGROUP_Chaos);
 		DECLARE_CYCLE_STAT(TEXT("FPBDRigidsEvolutionGBF::InertiaConditioning"), STAT_Evolution_InertiaConditioning, STATGROUP_Chaos);
@@ -264,11 +266,11 @@ void FPBDRigidsEvolutionGBF::ReloadParticlesCache()
 	FEvolutionResimCache* ResimCache = GetCurrentStepResimCache();
 	if ((ResimCache != nullptr) && ResimCache->IsResimming())
 	{
-		for (int32 IslandIndex = 0; IslandIndex < GetIslandManager().NumIslands(); ++IslandIndex)
+		for (int32 IslandIndex = 0; IslandIndex < GetIslandManager().GetNumIslands(); ++IslandIndex)
 		{
 			Private::FPBDIsland* Island = GetIslandManager().GetIsland(IslandIndex);
 			bool bIsUsingCache = false;
-			if (!Island->IsSleeping() && !GetIslandManager().IslandNeedsResim(IslandIndex))
+			if (!Island->IsSleeping() && !Island->NeedsResim())
 			{
 #if !(UE_BUILD_SHIPPING || UE_BUILD_TEST)
 				if (Chaos::FPhysicsSolverBase::IsNetworkPhysicsPredictionEnabled() && Chaos::FPhysicsSolverBase::CanDebugNetworkPhysicsPrediction())
@@ -276,9 +278,9 @@ void FPBDRigidsEvolutionGBF::ReloadParticlesCache()
 					UE_LOG(LogChaos, Log, TEXT("Reloading Island[%d] cache for %d particles"), IslandIndex, Island->GetParticles().Num());
 				}
 #endif
-				for (Private::FPBDIslandParticle& IslandParticle : Island->GetParticles())
+				for (Private::FPBDIslandParticle* IslandParticle : Island->GetParticles())
 				{
-					if (auto Rigid = IslandParticle.GetParticle()->CastToRigidParticle())
+					if (auto Rigid = IslandParticle->GetParticle()->CastToRigidParticle())
 					{
 						ResimCache->ReloadParticlePostSolve(*Rigid);
 					}
@@ -290,41 +292,6 @@ void FPBDRigidsEvolutionGBF::ReloadParticlesCache()
 	}
 }
 
-void FPBDRigidsEvolutionGBF::BuildDisabledParticles(const int32 Island, TArray<TArray<FPBDRigidParticleHandle*>>& DisabledParticles, TArray<bool>& SleepedIslands)
-{
-	for (Private::FPBDIslandParticle& IslandParticle : GetIslandManager().GetIsland(Island)->GetParticles())
-	{
-		// If a dynamic particle is moving slowly enough for long enough, disable it.
-		// @todo(mlentine): Find a good way of not doing this when we aren't using this functionality
-
-		// increment the disable count for the particle
-		auto PBDRigid = IslandParticle.GetParticle()->CastToRigidParticle();
-		if (PBDRigid && PBDRigid->ObjectState() == EObjectStateType::Dynamic)
-		{
-			if (PBDRigid->AuxilaryValue(PhysicsMaterials) && PBDRigid->V().SizeSquared() < PBDRigid->AuxilaryValue(PhysicsMaterials)->DisabledLinearThreshold &&
-				PBDRigid->W().SizeSquared() < PBDRigid->AuxilaryValue(PhysicsMaterials)->DisabledAngularThreshold)
-			{
-				++PBDRigid->AuxilaryValue(ParticleDisableCount);
-			}
-
-			// check if we're over the disable count threshold
-			if (PBDRigid->AuxilaryValue(ParticleDisableCount) > DisableThreshold)
-			{
-				PBDRigid->AuxilaryValue(ParticleDisableCount) = 0;
-				DisabledParticles[Island].Add(PBDRigid);
-			}
-
-			if (!(ensure(!FMath::IsNaN(PBDRigid->P()[0])) && ensure(!FMath::IsNaN(PBDRigid->P()[1])) && ensure(!FMath::IsNaN(PBDRigid->P()[2]))))
-			{
-				DisabledParticles[Island].Add(PBDRigid);
-			}
-		}
-	}
-
-	// Turn off if not moving
-	SleepedIslands[Island] = GetIslandManager().SleepInactive(Island, PhysicsMaterials, PerParticlePhysicsMaterials, SolverPhysicsMaterials);
-}
-
 void FPBDRigidsEvolutionGBF::UpdateCollisionSolverType()
 {
 	// If we have changed the collision solver type we must destroy any existing solvers so that they get recreated.
@@ -351,9 +318,6 @@ void FPBDRigidsEvolutionGBF::UpdateCollisionSolverType()
 	}
 }
 
-int32 DrawAwake = 0;
-FAutoConsoleVariableRef CVarDrawAwake(TEXT("p.chaos.DebugDrawAwake"),DrawAwake,TEXT("Draw particles that are awake"));
-
 void FPBDRigidsEvolutionGBF::AdvanceOneTimeStepImpl(const FReal Dt, const FSubStepInfo& SubStepInfo)
 {
 	SCOPE_CYCLE_COUNTER(STAT_Evolution_AdvanceOneTimeStep);
@@ -517,9 +481,9 @@ void FPBDRigidsEvolutionGBF::AdvanceOneTimeStepImpl(const FReal Dt, const FSubSt
 
 
 	TArray<bool> SleepedIslands;
-	SleepedIslands.SetNum(GetIslandManager().NumIslands());
+	SleepedIslands.SetNum(GetIslandManager().GetNumIslands());
 	TArray<TArray<FPBDRigidParticleHandle*>> DisabledParticles;
-	DisabledParticles.SetNum(GetIslandManager().NumIslands());
+	DisabledParticles.SetNum(GetIslandManager().GetNumIslands());
 	if(Dt > 0)
 	{
 		// Solve all the constraints
@@ -533,27 +497,9 @@ void FPBDRigidsEvolutionGBF::AdvanceOneTimeStepImpl(const FReal Dt, const FSubSt
 		// Post-solve CCD fixup to prevent the constraint solve from pushing CCD objects our of the world
 		{
 			SCOPE_CYCLE_COUNTER(STAT_Evolution_CCDCorrection);
-			CSV_SCOPED_TIMING_STAT(PhysicsVerbose, CCDCorrection);
+			CSV_SCOPED_TIMING_STAT(PhysicsVerbose, StepSolver_CCDCorrection);
 			CCDManager.ApplyCorrections(Dt);
 		}
-
-		// Determine which particles can be disabled and which islands are sleeping
-		{
-			SCOPE_CYCLE_COUNTER(STAT_Evolution_BuildDisabledParticles);
-
-			// @todo(chaos): improve this - batching will be poor when some island are much larger than others. 
-			// We can't just loop over groups because there are some islands with no constraints (and usually a single particle) 
-			// that do not get added to a group but still need otbe checked for sleeping
-			const int32 NumIslands = GetIslandManager().NumIslands();
-			PhysicsParallelFor(NumIslands, [&](const int32 IslandIndex)
-			{
-					Private::FPBDIsland* Island = GetIslandManager().GetIsland(IslandIndex);
-				if (!Island->IsSleeping() && !Island->IsUsingCache())
-				{
-					BuildDisabledParticles(IslandIndex, DisabledParticles, SleepedIslands);
-				}
-			});
-		}
 	}
 
 	{
@@ -581,18 +527,12 @@ void FPBDRigidsEvolutionGBF::AdvanceOneTimeStepImpl(const FReal Dt, const FSubSt
 	{
 		SCOPE_CYCLE_COUNTER(STAT_Evolution_DeactivateSleep);
 		CSV_SCOPED_TIMING_STAT(PhysicsVerbose, StepSolver_DeactivateSleep);
-		for (int32 Island = 0; Island < GetIslandManager().NumIslands(); ++Island)
-		{
-			if (SleepedIslands[Island])
-			{
-				GetIslandManager().SleepIsland(Particles, Island);
-			}
-			
-			for (const auto Particle : DisabledParticles[Island])
-			{
-				DisableParticle(Particle);
-			}
-		}
+
+		// Put any stationary islands to sleep (based on material settings)
+		GetIslandManager().UpdateSleep();
+		
+		// Disable stationary particles (based on material settings)
+		GetIslandManager().UpdateDisable([this](FPBDRigidParticleHandle* Rigid) { DisableParticle(Rigid); });
 	}
 
 	Clustering.AdvanceClustering(Dt, GetCollisionConstraints());
@@ -628,35 +568,6 @@ void FPBDRigidsEvolutionGBF::AdvanceOneTimeStepImpl(const FReal Dt, const FSubSt
 	{
 		SerializeToDisk(*this);
 	}
-
-#if CHAOS_DEBUG_DRAW
-	if(FDebugDrawQueue::IsDebugDrawingEnabled())
-	{
-		if(!!DrawAwake)
-		{
-			static const FColor IslandColors[] = {FColor::Green,FColor::Red,FColor::Yellow,
-				FColor::Blue,FColor::Orange,FColor::Black,FColor::Cyan,
-				FColor::Magenta,FColor::Purple,FColor::Turquoise};
-
-			static const int32 NumColors = sizeof(IslandColors) / sizeof(IslandColors[0]);
-			
-			for(const auto& Active : Particles.GetActiveParticlesView())
-			{
-				if(const auto* Geom = Active.Geometry().Get())
-				{
-					if(Geom->HasBoundingBox())
-					{
-						const int32 Island = Active.IslandIndex();
-						ensure(Island >= 0);
-						const int32 ColorIdx = Island % NumColors;
-						const FAABB3 LocalBounds = Geom->BoundingBox();
-						FDebugDrawQueue::GetInstance().DrawDebugBox(Active.X(),LocalBounds.Extents()*0.5f,Active.R(),IslandColors[ColorIdx],false,-1.f,0,0.f);
-					}
-				}
-			}
-		}
-	}
-#endif
 #endif
 }
 
@@ -751,6 +662,10 @@ FPBDRigidsEvolutionGBF::FPBDRigidsEvolutionGBF(
 
 	CollisionConstraints.SetCullDistance(DefaultCollisionCullDistance);
 
+	GetIslandManager().SetMaterialContainers(&PhysicsMaterials, &PerParticlePhysicsMaterials, &SolverPhysicsMaterials);
+	GetIslandManager().SetGravityForces(&GravityForces);
+	GetIslandManager().SetDisableCounterThreshold(DisableThreshold);
+
 	SetParticleUpdatePositionFunction([this](const TParticleView<FPBDRigidParticles>& ParticlesInput, const FReal Dt)
 	{
 		ParticlesInput.ParallelFor([&](auto& Particle, int32 Index)
diff --git a/Engine/Source/Runtime/Experimental/Chaos/Private/Chaos/PhysicsMaterialUtilities.cpp b/Engine/Source/Runtime/Experimental/Chaos/Private/Chaos/PhysicsMaterialUtilities.cpp
--- a/Engine/Source/Runtime/Experimental/Chaos/Private/Chaos/PhysicsMaterialUtilities.cpp
+++ b/Engine/Source/Runtime/Experimental/Chaos/Private/Chaos/PhysicsMaterialUtilities.cpp
@@ -11,24 +11,30 @@ namespace Chaos
 			const TGeometryParticleHandle<FReal, 3>* Particle,
 			const FShapeInstance* Shape,
 			const int32 ShapeFaceIndex,
-			const TArrayCollectionArray<TSerializablePtr<FChaosPhysicsMaterial>>& PhysicsMaterials,
-			const TArrayCollectionArray<TUniquePtr<FChaosPhysicsMaterial>>& PerParticlePhysicsMaterials,
+			const TArrayCollectionArray<TSerializablePtr<FChaosPhysicsMaterial>>* PhysicsMaterials,
+			const TArrayCollectionArray<TUniquePtr<FChaosPhysicsMaterial>>* PerParticlePhysicsMaterials,
 			const THandleArray<FChaosPhysicsMaterial>* const SimMaterials)
 		{
 			// Use the per-particle material if it exists
-			const FChaosPhysicsMaterial* UniquePhysicsMaterial = Particle->AuxilaryValue(PerParticlePhysicsMaterials).Get();
-			if (UniquePhysicsMaterial != nullptr)
+			if (PerParticlePhysicsMaterials != nullptr)
 			{
-				return UniquePhysicsMaterial;
+				const FChaosPhysicsMaterial* UniquePhysicsMaterial = Particle->AuxilaryValue(*PerParticlePhysicsMaterials).Get();
+				if (UniquePhysicsMaterial != nullptr)
+				{
+					return UniquePhysicsMaterial;
+				}
 			}
 
 			// Use the.... other?.... per particle material if it exists. 
 			// @todo(chaos): I assume this is to support writable unique materials, but there must be a better 
 			// way that adding another pointer to every particle. E.g., a ref count or shared setting on the material.
-			const FChaosPhysicsMaterial* PhysicsMaterial = Particle->AuxilaryValue(PhysicsMaterials).Get();
-			if (PhysicsMaterial != nullptr)
+			if (PhysicsMaterials != nullptr)
 			{
-				return PhysicsMaterial;
+				const FChaosPhysicsMaterial* PhysicsMaterial = Particle->AuxilaryValue(*PhysicsMaterials).Get();
+				if (PhysicsMaterial != nullptr)
+				{
+					return PhysicsMaterial;
+				}
 			}
 
 			// If no particle material, see if the shape has one
@@ -54,7 +60,7 @@ namespace Chaos
 			return nullptr;
 		}
 
-		const FChaosPhysicsMaterial* GetFirstPhysicsMaterial(const TGeometryParticleHandle<FReal, 3>* Particle, const TArrayCollectionArray<TSerializablePtr<FChaosPhysicsMaterial>>& PhysicsMaterials, const TArrayCollectionArray<TUniquePtr<FChaosPhysicsMaterial>>& PerParticlePhysicsMaterials, const THandleArray<FChaosPhysicsMaterial>* const SimMaterials)
+		const FChaosPhysicsMaterial* GetFirstPhysicsMaterial(const TGeometryParticleHandle<FReal, 3>* Particle, const TArrayCollectionArray<TSerializablePtr<FChaosPhysicsMaterial>>* PhysicsMaterials, const TArrayCollectionArray<TUniquePtr<FChaosPhysicsMaterial>>* PerParticlePhysicsMaterials, const THandleArray<FChaosPhysicsMaterial>* const SimMaterials)
 		{
 			const FShapeInstance* Shape = (Particle->ShapeInstances().IsEmpty()) ? nullptr : Particle->ShapeInstances()[0].Get();
 			return GetPhysicsMaterial(Particle, Shape, INDEX_NONE, PhysicsMaterials, PerParticlePhysicsMaterials, SimMaterials);
diff --git a/Engine/Source/Runtime/Experimental/Chaos/Private/PBDRigidsSolver.cpp b/Engine/Source/Runtime/Experimental/Chaos/Private/PBDRigidsSolver.cpp
--- a/Engine/Source/Runtime/Experimental/Chaos/Private/PBDRigidsSolver.cpp
+++ b/Engine/Source/Runtime/Experimental/Chaos/Private/PBDRigidsSolver.cpp
@@ -305,7 +305,7 @@ namespace Chaos
 
 		// @todo(chaos): move to physics project settings and set these to -1 when we are settled on values...
 		Chaos::FRealSingle ChaosSolverVelocityBoundsMultiplier = 1.0f;
-		Chaos::FRealSingle ChaosSolverMaxVelocityBoundsExpansion = 3.0f;
+		Chaos::FRealSingle ChaosSolverMaxVelocityBoundsExpansion = 3.0f;	// This should probably be a fraction of object size (see FParticlePairMidPhase::GenerateCollisions)
 		FAutoConsoleVariableRef CVarChaosSolverVelocityBoundsMultiplier(TEXT("p.Chaos.Solver.Collision.VelocityBoundsMultiplier"), ChaosSolverVelocityBoundsMultiplier, TEXT("Override velocity bounds multiplier (if >= 0)"));
 		FAutoConsoleVariableRef CVarChaosSolverMaxVelocityBoundsExpansion(TEXT("p.Chaos.Solver.Collision.MaxVelocityBoundsExpansion"), ChaosSolverMaxVelocityBoundsExpansion, TEXT("Override max velocity bounds expansion (if >= 0)"));
 
@@ -1723,8 +1723,8 @@ namespace Chaos
 									{
 										UE_LOG(LogChaos, Log, TEXT("Particle Dynamic At Position = %s | Velocity = %s | Quaternion = %s | Omega = %s | Resim Frame = %d | Sync State = %d | Needs Resim = %d"),   
 											*Handle.X().ToString(), *Handle.V().ToString(), *Handle.R().ToString(), *Handle.W().ToString(), 
-												GetEvolution()->GetIslandManager().GetIsland(Handle.IslandIndex())->GetResimFrame(), (uint8)Handle.SyncState(),
-											GetEvolution()->GetIslandManager().GetIsland(Handle.IslandIndex())->NeedsResim());
+											GetEvolution()->GetIslandManager().GetParticleIsland(Handle.Handle())->GetResimFrame(), (uint8)Handle.SyncState(),
+											GetEvolution()->GetIslandManager().GetParticleIsland(Handle.Handle())->NeedsResim());
 									}
 								}
 							}
@@ -1966,7 +1966,7 @@ TRACE_COUNTER_SET(ChaosTraceCounter_##Name, Value)
 		CHAOS_COUNTER_STAT(NumGeometryCollectionBodies, (int32)GetEvolution()->GetParticles().GetGeometryCollectionParticles().Size());
 
 		// Constraint counts
-		CHAOS_COUNTER_STAT(NumIslands, GetEvolution()->GetIslandManager().NumIslands());
+		CHAOS_COUNTER_STAT(NumIslands, GetEvolution()->GetIslandManager().GetNumIslands());
 		CHAOS_COUNTER_STAT(NumIslandGroups, GetEvolution()->GetIslandGroupManager().GetNumActiveGroups());
 		CHAOS_COUNTER_STAT(NumContacts, NumCollisionConstraints());
 		CHAOS_COUNTER_STAT(NumJoints, NumJointConstraints());
diff --git a/Engine/Source/Runtime/Experimental/Chaos/Public/Chaos/ChaosDebugDrawDeclares.h b/Engine/Source/Runtime/Experimental/Chaos/Public/Chaos/ChaosDebugDrawDeclares.h
--- a/Engine/Source/Runtime/Experimental/Chaos/Public/Chaos/ChaosDebugDrawDeclares.h
+++ b/Engine/Source/Runtime/Experimental/Chaos/Public/Chaos/ChaosDebugDrawDeclares.h
@@ -4,15 +4,11 @@
 #include "Chaos/Core.h"
 #include "Chaos/Declares.h"
 #include "Chaos/GeometryParticlesfwd.h"
+#include "Chaos/Island/IslandManagerFwd.h"
 #include "Misc/Build.h"
 
 namespace Chaos
 {
-	namespace Private
-	{
-		class FPBDIslandManager;
-	}
-
 	class FConstraintHandle;
 
 	class FPBDCollisionConstraints;
diff --git a/Engine/Source/Runtime/Experimental/Chaos/Public/Chaos/Collision/PBDCollisionContainerSolver.h b/Engine/Source/Runtime/Experimental/Chaos/Public/Chaos/Collision/PBDCollisionContainerSolver.h
--- a/Engine/Source/Runtime/Experimental/Chaos/Public/Chaos/Collision/PBDCollisionContainerSolver.h
+++ b/Engine/Source/Runtime/Experimental/Chaos/Public/Chaos/Collision/PBDCollisionContainerSolver.h
@@ -10,11 +10,6 @@
 
 namespace Chaos
 {
-	namespace Private
-	{
-		class FPBDIsland;
-		class FPBDIslandConstraint;
-	}
 	class FPBDCollisionConstraint;
 	class FPBDCollisionConstraints;
 	class FSolverBody;
@@ -42,7 +37,7 @@ namespace Chaos
 		// IslandGroup API
 		//
 		virtual void AddConstraints() override final;
-		virtual void AddConstraints(const TArrayView<Private::FPBDIslandConstraint>& ConstraintHandles) override final;
+		virtual void AddConstraints(const TArrayView<Private::FPBDIslandConstraint*>& ConstraintHandles) override final;
 		virtual void AddBodies(FSolverBodyContainer& SolverBodyContainer) override final;
 		virtual void GatherInput(const FReal Dt) override final;
 		virtual void GatherInput(const FReal Dt, const int32 BeginIndex, const int32 EndIndex) override final;
diff --git a/Engine/Source/Runtime/Experimental/Chaos/Public/Chaos/Collision/PBDCollisionContainerSolverJacobi.h b/Engine/Source/Runtime/Experimental/Chaos/Public/Chaos/Collision/PBDCollisionContainerSolverJacobi.h
--- a/Engine/Source/Runtime/Experimental/Chaos/Public/Chaos/Collision/PBDCollisionContainerSolverJacobi.h
+++ b/Engine/Source/Runtime/Experimental/Chaos/Public/Chaos/Collision/PBDCollisionContainerSolverJacobi.h
@@ -16,9 +16,6 @@ namespace Chaos
 
 	namespace Private
 	{
-		class FPBDIsland;
-		class FPBDIslandConstraint;
-
 		/**
 		 * The solver for a set of collision constraints. This collects all the data required to solve a set of collision
 		 * constraints into a contiguous, ordered buffer.
@@ -43,7 +40,7 @@ namespace Chaos
 			// IslandGroup API
 			//
 			virtual void AddConstraints() override final;
-			virtual void AddConstraints(const TArrayView<Private::FPBDIslandConstraint>& ConstraintHandles) override final;
+			virtual void AddConstraints(const TArrayView<Private::FPBDIslandConstraint*>& ConstraintHandles) override final;
 			virtual void AddBodies(FSolverBodyContainer& SolverBodyContainer) override final;
 			virtual void GatherInput(const FReal Dt) override final;
 			virtual void GatherInput(const FReal Dt, const int32 BeginIndex, const int32 EndIndex) override final;
diff --git a/Engine/Source/Runtime/Experimental/Chaos/Public/Chaos/Collision/PBDCollisionContainerSolverSimd.h b/Engine/Source/Runtime/Experimental/Chaos/Public/Chaos/Collision/PBDCollisionContainerSolverSimd.h
--- a/Engine/Source/Runtime/Experimental/Chaos/Public/Chaos/Collision/PBDCollisionContainerSolverSimd.h
+++ b/Engine/Source/Runtime/Experimental/Chaos/Public/Chaos/Collision/PBDCollisionContainerSolverSimd.h
@@ -16,9 +16,6 @@ namespace Chaos
 
 	namespace Private
 	{
-		class FPBDIsland;
-		class FPBDIslandConstraint;
-
 		template<int TNumLanes>
 		using TConstraintPtrSimd = TSimdValue<FPBDCollisionConstraint*, TNumLanes>;
 
@@ -44,7 +41,7 @@ namespace Chaos
 			// IslandGroup API
 			//
 			virtual void AddConstraints() override final;
-			virtual void AddConstraints(const TArrayView<Private::FPBDIslandConstraint>& ConstraintHandles) override final;
+			virtual void AddConstraints(const TArrayView<Private::FPBDIslandConstraint*>& ConstraintHandles) override final;
 			virtual void AddBodies(FSolverBodyContainer& SolverBodyContainer) override final;
 			virtual void GatherInput(const FReal Dt) override final;
 			virtual void GatherInput(const FReal Dt, const int32 BeginIndex, const int32 EndIndex) override final;
diff --git a/Engine/Source/Runtime/Experimental/Chaos/Public/Chaos/ConstraintHandle.h b/Engine/Source/Runtime/Experimental/Chaos/Public/Chaos/ConstraintHandle.h
--- a/Engine/Source/Runtime/Experimental/Chaos/Public/Chaos/ConstraintHandle.h
+++ b/Engine/Source/Runtime/Experimental/Chaos/Public/Chaos/ConstraintHandle.h
@@ -31,6 +31,11 @@ namespace Chaos
 	using FParticlePair = TVec2<FGeometryParticleHandle*>;
 	using FConstParticlePair = TVec2<const FGeometryParticleHandle*>;
 
+	namespace Private
+	{
+		class FPBDIslandConstraint;
+	}
+
 	/**
 	 * @brief A type id for constraint handles to support safe up/down casting (including intermediate classes in the hierrachy)
 	 *
@@ -95,18 +100,20 @@ namespace Chaos
 
 		FConstraintHandle() 
 			: ConstraintContainer(nullptr)
-			, GraphIndex(INDEX_NONE)
+			, GraphEdge(nullptr)
 		{
 		}
 
 		FConstraintHandle(FPBDConstraintContainer* InContainer)
 			: ConstraintContainer(InContainer)
-			, GraphIndex(INDEX_NONE)
+			, GraphEdge(nullptr)
 		{
 		}
 
 		virtual ~FConstraintHandle()
 		{
+			// Make sure we are not still in the graph
+			check(GraphEdge == nullptr);
 		}
 
 		virtual bool IsValid() const
@@ -126,17 +133,21 @@ namespace Chaos
 
 		bool IsInConstraintGraph() const
 		{
-			return (GraphIndex != INDEX_NONE);
+			return (GraphEdge != nullptr);
 		}
 
-		int32 GetConstraintGraphIndex() const
+		Private::FPBDIslandConstraint* GetConstraintGraphEdge() const
 		{
-			return GraphIndex;
+			return GraphEdge;
 		}
 
-		void SetConstraintGraphIndex(const int32 InIndex)
+		// NOTE: Should only be called by the IslandManager
+		void SetConstraintGraphEdge(Private::FPBDIslandConstraint* InEdge)
 		{
-			GraphIndex = InIndex;
+			// Check for double add
+			check((InEdge == nullptr) || (GraphEdge == nullptr));
+
+			GraphEdge = InEdge;
 		}
 
 		virtual TVec2<FGeometryParticleHandle*> GetConstrainedParticles() const = 0;
@@ -181,12 +192,16 @@ namespace Chaos
 			return STypeID;
 		}
 
+		// Deprecated API
+		UE_DEPRECATED(5.3, "Use GetConstraintGraphEdge") int32 GetConstraintGraphIndex() const { return INDEX_NONE; }
+		UE_DEPRECATED(5.3, "Not supported") void SetConstraintGraphIndex(const int32 InIndex) const {}
+
 	protected:
 		friend class FPBDConstraintContainer;
 
 		FPBDConstraintContainer* ConstraintContainer;
 		
-		int32 GraphIndex;
+		Private::FPBDIslandConstraint* GraphEdge;
 	};
 
 
diff --git a/Engine/Source/Runtime/Experimental/Chaos/Public/Chaos/Evolution/ABTestingConstraintContainerSolver.h b/Engine/Source/Runtime/Experimental/Chaos/Public/Chaos/Evolution/ABTestingConstraintContainerSolver.h
--- a/Engine/Source/Runtime/Experimental/Chaos/Public/Chaos/Evolution/ABTestingConstraintContainerSolver.h
+++ b/Engine/Source/Runtime/Experimental/Chaos/Public/Chaos/Evolution/ABTestingConstraintContainerSolver.h
@@ -12,10 +12,6 @@
 
 namespace Chaos
 {
-	namespace Private
-	{
-		class FPBDIslandConstraint;
-	}
 	class FConstraintHandleHolder;
 	class FSolverBodyContainer;
 
@@ -89,7 +85,7 @@ namespace Chaos
 				SolverB->AddConstraints();
 			}
 
-			virtual void AddConstraints(const TArrayView<Private::FPBDIslandConstraint>& Constraints) override final
+			virtual void AddConstraints(const TArrayView<Private::FPBDIslandConstraint*>& Constraints) override final
 			{
 				SolverA->AddConstraints(Constraints);
 				SolverB->AddConstraints(Constraints);
diff --git a/Engine/Source/Runtime/Experimental/Chaos/Public/Chaos/Evolution/IndexedConstraintContainer.h b/Engine/Source/Runtime/Experimental/Chaos/Public/Chaos/Evolution/IndexedConstraintContainer.h
--- a/Engine/Source/Runtime/Experimental/Chaos/Public/Chaos/Evolution/IndexedConstraintContainer.h
+++ b/Engine/Source/Runtime/Experimental/Chaos/Public/Chaos/Evolution/IndexedConstraintContainer.h
@@ -3,7 +3,7 @@
 
 #include "Chaos/ConstraintHandle.h"
 #include "Chaos/Evolution/SolverConstraintContainer.h"
-#include "Chaos/Island/SolverIsland.h"
+#include "Chaos/Island/IslandManager.h"
 #include "Chaos/PBDConstraintContainer.h"
 
 namespace Chaos
@@ -153,12 +153,12 @@ namespace Chaos
 			ensure(false);
 		}
 
-		virtual void AddConstraints(const TArrayView<Private::FPBDIslandConstraint>& Constraints) override final
+		virtual void AddConstraints(const TArrayView<Private::FPBDIslandConstraint*>& IslandConstraints) override final
 		{
-			for (const Private::FPBDIslandConstraint& Constraint : Constraints)
+			for (const Private::FPBDIslandConstraint* IslandConstraint : IslandConstraints)
 			{
 				// We will only ever be given constraints from our container (asserts in non-shipping)
-				const FIndexedConstraintHandle* IndexedConstraintHandle = Constraint.GetConstraint()->AsUnsafe<FIndexedConstraintHandle>();
+				const FIndexedConstraintHandle* IndexedConstraintHandle = IslandConstraint->GetConstraint()->AsUnsafe<FIndexedConstraintHandle>();
 
 				int32 ConstraintIndex = IndexedConstraintHandle->GetConstraintIndex();
 
diff --git a/Engine/Source/Runtime/Experimental/Chaos/Public/Chaos/Evolution/SolverConstraintContainer.h b/Engine/Source/Runtime/Experimental/Chaos/Public/Chaos/Evolution/SolverConstraintContainer.h
--- a/Engine/Source/Runtime/Experimental/Chaos/Public/Chaos/Evolution/SolverConstraintContainer.h
+++ b/Engine/Source/Runtime/Experimental/Chaos/Public/Chaos/Evolution/SolverConstraintContainer.h
@@ -2,13 +2,10 @@
 #pragma once
 
 #include "Chaos/ConstraintHandle.h"
+#include "Chaos/Island/IslandManagerFwd.h"
 
 namespace Chaos
 {
-	namespace Private
-	{
-		class FPBDIslandConstraint;
-	}
 	class FConstraintHandleHolder;
 	class FSolverBodyContainer;
 
@@ -66,7 +63,7 @@ namespace Chaos
 		 * NOTE: this should not do any actual data gathering - it should just add to the list of constraints in this group. All data
 		 * gathering is handled in GatherInput.
 		*/
-		virtual void AddConstraints(const TArrayView<Private::FPBDIslandConstraint>& Constraints) = 0;
+		virtual void AddConstraints(const TArrayView<Private::FPBDIslandConstraint*>& Constraints) = 0;
 
 		/**
 		 * Add all the required bodies to the body container (required for the constraints added with AddConstraints)
@@ -134,7 +131,7 @@ namespace Chaos
 			// We solve all constraints in the container in the order it prefers so nothing to do here
 		}
 
-		virtual void AddConstraints(const TArrayView<Private::FPBDIslandConstraint>& Constraints) override final
+		virtual void AddConstraints(const TArrayView<Private::FPBDIslandConstraint*>& Constraints) override final
 		{
 			// This solver container is for use with the a non-graph evolution. It will not call this function.
 			ensure(false);
diff --git a/Engine/Source/Runtime/Experimental/Chaos/Public/Chaos/GeometryParticles.h b/Engine/Source/Runtime/Experimental/Chaos/Public/Chaos/GeometryParticles.h
--- a/Engine/Source/Runtime/Experimental/Chaos/Public/Chaos/GeometryParticles.h
+++ b/Engine/Source/Runtime/Experimental/Chaos/Public/Chaos/GeometryParticles.h
@@ -37,6 +37,10 @@ namespace Chaos
 		CHAOS_API extern bool bCCDAxisThresholdUsesProbeShapes;
 	}
 
+	namespace Private
+	{
+		class FPBDIslandParticle;
+	}
 
 	/**
 	* Union between shape and shapes array pointers, used for passing around shapes with
@@ -177,7 +181,7 @@ namespace Chaos
 			TArrayCollection::AddArray(&MWeakParticleHandle);
 			TArrayCollection::AddArray(&MParticleConstraints);
 			TArrayCollection::AddArray(&MParticleCollisions);
-			TArrayCollection::AddArray(&MGraphIndex);
+			TArrayCollection::AddArray(&MGraphNode);
 			TArrayCollection::AddArray(&MResimType);
 			TArrayCollection::AddArray(&MEnabledDuringResim);
 			TArrayCollection::AddArray(&MLightWeightDisabled);
@@ -217,7 +221,7 @@ namespace Chaos
 			, MWeakParticleHandle(MoveTemp(Other.MWeakParticleHandle))
 			, MParticleConstraints(MoveTemp(Other.MParticleConstraints))
 			, MParticleCollisions(MoveTemp(Other.MParticleCollisions))
-			, MGraphIndex(MoveTemp(Other.MGraphIndex))
+			, MGraphNode(MoveTemp(Other.MGraphNode))
 			, MResimType(MoveTemp(Other.MResimType))
 			, MEnabledDuringResim(MoveTemp(Other.MEnabledDuringResim))
 			, MLightWeightDisabled(MoveTemp(Other.MLightWeightDisabled))
@@ -249,7 +253,7 @@ namespace Chaos
 			TArrayCollection::AddArray(&MWeakParticleHandle);
 			TArrayCollection::AddArray(&MParticleConstraints);
 			TArrayCollection::AddArray(&MParticleCollisions);
-			TArrayCollection::AddArray(&MGraphIndex);
+			TArrayCollection::AddArray(&MGraphNode);
 			TArrayCollection::AddArray(&MResimType);
 			TArrayCollection::AddArray(&MEnabledDuringResim);
 			TArrayCollection::AddArray(&MLightWeightDisabled);
@@ -291,7 +295,7 @@ namespace Chaos
 			TArrayCollection::AddArray(&MWeakParticleHandle);
 			TArrayCollection::AddArray(&MParticleConstraints);
 			TArrayCollection::AddArray(&MParticleCollisions);
-			TArrayCollection::AddArray(&MGraphIndex);
+			TArrayCollection::AddArray(&MGraphNode);
 			TArrayCollection::AddArray(&MResimType);
 			TArrayCollection::AddArray(&MEnabledDuringResim);
 			TArrayCollection::AddArray(&MLightWeightDisabled);
@@ -559,8 +563,8 @@ namespace Chaos
 			return MParticleCollisions[Index];
 		}
 
-		FORCEINLINE const int32 ConstraintGraphIndex(const int32 Index) const { return MGraphIndex[Index]; }
-		FORCEINLINE int32& ConstraintGraphIndex(const int32 Index) { return MGraphIndex[Index]; }
+		FORCEINLINE Private::FPBDIslandParticle* ConstraintGraphNode(const int32 Index) const { return MGraphNode[Index]; }
+		FORCEINLINE Private::FPBDIslandParticle*& ConstraintGraphNode(const int32 Index) { return MGraphNode[Index]; }
 
 		FORCEINLINE EResimType ResimType(const int32 Index) const { return MResimType[Index]; }
 		FORCEINLINE EResimType& ResimType(const int32 Index) { return MResimType[Index]; }
@@ -571,6 +575,11 @@ namespace Chaos
 		FORCEINLINE bool LightWeightDisabled(const int32 Index) const { return MLightWeightDisabled[Index]; }
 		FORCEINLINE bool& LightWeightDisabled(const int32 Index) { return MLightWeightDisabled[Index]; }
 
+
+		// Deprecated API
+		UE_DEPRECATED(5.3, "Use ConstraintGraphNode") const int32 ConstraintGraphIndex(const int32 Index) const { return INDEX_NONE; }
+		UE_DEPRECATED(5.3, "Use ConstraintGraphNode") int32& ConstraintGraphIndex(const int32 Index) { static int32 Dummy = INDEX_NONE; return Dummy; }
+
 private:
 		friend THandleType;
 		CHAOS_API void ResetWeakParticleHandle(const int32 Index)
@@ -686,7 +695,7 @@ public:
 		TArrayCollectionArray<FWeakParticleHandle> MWeakParticleHandle;
 		TArrayCollectionArray<FConstraintHandleArray> MParticleConstraints;
 		TArrayCollectionArray<FParticleCollisions> MParticleCollisions;
-		TArrayCollectionArray<int32> MGraphIndex;
+		TArrayCollectionArray<Private::FPBDIslandParticle*> MGraphNode;
 		TArrayCollectionArray<EResimType> MResimType;
 		TArrayCollectionArray<bool> MEnabledDuringResim;
 		TArrayCollectionArray<bool> MLightWeightDisabled;
diff --git a/Engine/Source/Runtime/Experimental/Chaos/Public/Chaos/Island/IslandGraph.h b/Engine/Source/Runtime/Experimental/Chaos/Public/Chaos/Island/IslandGraph.h
--- a/Engine/Source/Runtime/Experimental/Chaos/Public/Chaos/Island/IslandGraph.h
+++ b/Engine/Source/Runtime/Experimental/Chaos/Public/Chaos/Island/IslandGraph.h
@@ -1,436 +1,4 @@
 // Copyright Epic Games, Inc. All Rights Reserved.
 #pragma once
 
-#include "Containers/Map.h"
-#include "Containers/Queue.h"
-#include "Containers/Deque.h"
-#include "Containers/SparseArray.h"
-
-namespace Chaos
-{
-	/** Graph node structure */
-	template<typename NodeType>
-	struct CHAOS_API TIslandGraphNode
-	{
-		/** List of edges connected to the node */
-		TSparseArray<int32> NodeEdges;
-
-		/**
-		 * Node Island Index. Only used by valid nodes (Dynamic/Sleeping Particles).
-		 * Static/kinematic particles could belong to several islands and this is always INDEX_NONE.
-		 */
-		int32 IslandIndex = INDEX_NONE;
-
-		/** Check if a node is valid (checked for graph partitioning into islands).
-		 * In the context of particles (nodes) and constraints (edges), bValidNode is true when
-		 * a Particle is Dynamic or Sleeping, but not when Static or Kinematic. This is because we
-		 * do not need to merge islands that are only connected through a Static/Kinematic Particle node.
-		 */
-		bool bValidNode = true;
-
-		/** Check if a node is steady */
-		bool bStationaryNode = true;
-
-		/** Node counter to filter nodes already processed (@see SplitIsland) */
-		int32 NodeCounter = 0;
-
-		/** Node item that is stored per node. This is a ParticleHandle. */
-		NodeType NodeItem;
-
-		/** Node level index */
-		int32 LevelIndex = INDEX_NONE;
-
-		/** Set of used colors */
-		TSet<int32> ColorIndices;
-
-		/** Resim frame */
-		int32 ResimFrame = INDEX_NONE;
-	};
-
-	/** Graph edge structure */
-	template<typename EdgeType>
-	struct CHAOS_API TIslandGraphEdge
-	{
-		/** First node of the edge */
-		int32 FirstNode = INDEX_NONE;
-
-		/** Second node of the edge*/
-		int32 SecondNode = INDEX_NONE;
-
-		/** Current edge index in the list of edges of the first node */
-		int32 FirstEdge = INDEX_NONE;
-
-		/** Current edge index in the list of edges of the second node  */
-		int32 SecondEdge = INDEX_NONE;
-
-		/** Unique edge island index */
-		int32 IslandIndex = INDEX_NONE;
-
-		/** True if one or both nodes are valid (dynamic) */
-		bool bValidEdge = true;
-
-		/** Edge counter to filter edges already processed (@see SplitIsland) */
-		int32 EdgeCounter = 0;
-
-		/** Edge item that is stored per node */
-		EdgeType EdgeItem;
-
-		/** Item Container Id  */
-		int32 ItemContainer = 0;
-
-		/** Edge level index */
-		int32 LevelIndex = INDEX_NONE;
-
-		/** Edge Color index */
-		int32 ColorIndex = INDEX_NONE;
-	};
-
-	/** Graph island structure */
-	template<typename IslandType>
-	struct CHAOS_API TIslandGraphIsland
-	{
-		/** Number of edges per islands*/
-		int32 NumEdges = 0;
-
-		/** Number of valid nodes per islands (should be less than the ones in solver islands since it is only including the valid ones)*/
-		int32 NumNodes = 0;
-
-		/** Island counter to filter islands already processed */
-		int32 IslandCounter = 0;
-
-		/** Boolean to check if an island is persistent or not. A persistent island is one where no particles or constraints were added or removed */
-		bool bIsPersistent = true;
-
-		/** Boolean to check if an island is sleeping or not */
-		bool bIsSleeping = false;
-
-		/** Boolean to check if an island was sleeping last tick or not */
-		bool bWasSleeping = false;
-
-		/** List of children islands  to be merged */
-		// @todo(chaos): should we really have a set here? The implementation probably wastes a lot of memory
-		TSet<int32> ChildrenIslands;
-
-		/** Parent Island */
-		int32 ParentIsland = INDEX_NONE;
-
-		/** Island Item that is stored per island*/
-		IslandType IslandItem;
-
-		/** Max color per island */
-		int32 MaxColors = INDEX_NONE;
-	};
-
-/**
- * Island Graph.
- *
- * This template implements a graph that will be stored on the island manager
- * The goal here is to minimize memory allocation while doing graph operations. 
- *
- * @param NodeType The node type of the item that will be stored in each node (e.g., ParticleHandle)
- * @param EdgeType The edge type of the item that will be stored in each edge (e.g., ConstraintHandle)
- */
-template<typename NodeType, typename EdgeType, typename IslandType, typename OwnerType>
-class CHAOS_API FIslandGraph
-{
-public:
-	// NOTE: IslandGraph internal structs are non-member structs for better natvis behaviour/debugging. Natvis 
-	// has trouble with member classes of templated classes
-	using FGraphNode = TIslandGraphNode<NodeType>;
-	using FGraphEdge = TIslandGraphEdge<EdgeType>;
-	using FGraphIsland = TIslandGraphIsland<IslandType>;
-	using FGraphOwner = OwnerType;
-
-	static constexpr const int32 MaxCount = 100000;
-
-	/**
-	 * @brief Set the graph owner. This will allow the owner to cache node and edge indices with the node/edge items
-	*/
-	void SetOwner(FGraphOwner* InOwner)
-	{
-		Owner = InOwner;
-	}
-
-	/**
-	 * Set to true if we require deterministic behaviour.
-	 */
-	void SetIsDeterministic(const bool bInIsDeterministic);
-
-	/**
-	 * Update current graph to merge connected islands 
-	 */
-	void UpdateGraph();
-
-	/**
-	 * Perform end-of-tick cleanup if necessary
-	 */
-	 void EndTick();
-
-	/**
-	 * Reserve a number of nodes in memory before adding them later
-	 * @param NumNodes Number of nodes to be reserved in memory
-	 */
-	void ReserveNodes(const int32 NumNodes);
-
-	/**
-	 * Remove all nodes (and attached edges - so all of them) from the graph
-	 */
-	void ResetNodes();
-
-	/**
-	 * Given a node item, add a node to the graph nodes
-	 * @param NodeItem Node Item to be stored in the graph node
-	 * @param bValidNode Check if a node should be considered for graph partitioning (into islands)
-	 * @param IslandIndex Potential island index we want the node to belong to
-	 * @param bStationaryNode Boolean to check if a node is steady or not
-	 * @return Node index that has just been added
-	 */
-	int32 AddNode(const NodeType& NodeItem, const bool bValidNode = true, const int32 IslandIndex = INDEX_NONE, const bool bStationaryNode = false);
-
-	/**
-	* Given a node item, update the graph node information (valid,discard...)
-	* @param NodeIndex Sparse Index of the node to update
-	* @param bValidNode Check if a node should be considered for graph partitioning (into islands)
-	* @param bStationaryNode Boolean to check if a node is steady or not
-	* @return Node index that has just been added
-	*/
-	void UpdateNode(const int32 NodeIndex, const bool bValidNode, const bool bStationaryNode);
-
-	/**
-	 * Remove a node from the graph nodes list
-	 * @param NodeItem Item to be removed from the nodes list
-	 */
-	void RemoveNode(const NodeType& NodeItem);
-
-	/**
-	 * Get the level of the specified node item (distance from a kinematic node)
-	 */
-	int32 GetNodeItemLevel(const NodeType& NodeItem) const;
-
-	/**
-	 * Reserve a number of edges in memory before adding them later
-	 * @param NumEdges Number of edges to be reserved in memory
-	 */
-	void ReserveEdges(const int32 NumEdges);
-
-	/**
-	 * Reset the graph edges without deallocating memory
-	 */
-	void ResetEdges();
-
-	/**
-	 * Given an edge item, add an edge to the graph edges
-	 * @param EdgeItem Node Item to be stored in the graph edge
-	 * @param ItemContainer Item container id that will be stored on the graph edge
-	 * @param FirstNode First node index of the graph edge
-	 * @param SecondNode Second node index of the graph edge
-	 * @return Edge index that has just been added
-	 */
-	int32 AddEdge(const EdgeType& EdgeItem, const int32 ItemContainer, const int32 FirstNode,  const int32 SecondNode);
-
-	/**
-	 * Remove an edge from the graph edges list
-	 * @param EdgeIndex Index of the edge to be removed from the edges list
-	 */
-	void RemoveEdge(const int32 EdgeIndex);
-
-	/**
-	 * Update edge state based on current node state
-	 */
-	void UpdateEdge(const int32 EdgeIndex);
-	/**
-	 * Remove all non-sleeping edges from the graph
-	 */
-	void RemoveAllAwakeEdges();
-
-	/**
-	 * Get the maximum number of edges that are inside a graph (SparseArray size)
-	 */
-	FORCEINLINE int32 MaxNumEdges() const { return GraphEdges.GetMaxIndex(); }
-
-	/**
-	 * Get the maximum number of nodes that are inside a graph (SparseArray size)
-	 */
-	FORCEINLINE int32 MaxNumNodes() const { return GraphNodes.GetMaxIndex(); }
-
-	/**
-	 * Get the maximum number of islands that are inside a graph (SparseArray size)
-	 */
-	FORCEINLINE int32 MaxNumIslands() const { return GraphIslands.GetMaxIndex(); }
-
-	/**
-	 * Link two islands to each other for future island traversal
-	 * @param FirstIsland First island index to be linked
-	 * @param SecondIsland Second island index to be linked
-	 * @param bNonStationary Boolean to check if at least one of the 2 nodes are non stationary
-	 */
-	void ParentIslands(const int32 FirstIsland, const int32 SecondIsland,const bool bNonStationary);
-
-	/**
-	 * Given an edge index we try to attach islands. 
-	 * @param EdgeIndex Index of the edge used to attach the islands
-	 */
-	void AttachIslands(const int32 EdgeIndex);
-
-	/**
-	 * Merge 2 islands
-	 * @param ParentIndex Index of the parent island
-	 * @param ChildIndex Index of the child island
-	 */
-	void MergeIslands(const int32 ParentIndex, const int32 ChildIndex);
-
-	/**
-	* Merge all the graph islands
-	*/
-	void MergeIslands();
-
-	/**
-	* Split one island given a root index 
-	* @param NodeQueue node queue to use while splitting the island (avoid reallocation)
-	* @param RootIndex root index of the island graph
-	* @param IslandIndex index of the island we are splitting
-	*/
-	void SplitIsland(const int32 RootIndex, const int32 IslandIndex);
-
-	/**
-	* Split all the graph islands if not sleeping
-	*/
-	void SplitIslands();
-
-	/**
-	* Reassign the updated island index from the merging phase to the nodes/edges
-	*/
-	void ReassignIslands();
-
-	/**
-	* Init all the islands/edges/nodes levels and colors for sorting
-	*/
-	void InitSorting();
-
-	/**
-	* Compute all the islands/edges/nodes levels
-	*/
-	void ComputeLevels();
-	
-	/**
-	* Update all the islands/edges/nodes levels and push the connected nodes into the node queue
-	* @param NodeIndex Node index that we are currently iterating over
-	*/
-	void UpdateLevels(const int32 NodeIndex);
-
-	/**
-	* Compute all the islands/edges/nodes colors
-	* @param MinEdges Minimum number of edges to compute coloring
-	*/
-	void ComputeColors(const int32 MinEdges);
-
-	/**
-	* Update all the islands/edges/nodes colors and push the connected nodes into the node queue
-	* @param NodeIndex Node index that we are currently iterating over
-	* @param MinEdges Minimum number of edges to compute coloring
-	*/
-	void UpdateColors(const int32 NodeIndex, const int32 MinEdges);
-
-	/**
-	* Pick the first available color that is not used yet by the current graph node and the edge opposite one
-	* @param GraphNode Graph node that we currently iterate over
-	* @param OtherIndex Index of the edge opposite node
-	* @return First available color 
-	*/
-	int32 PickColor(const FGraphNode& GraphNode, const int32 OtherIndex);
-
-	/**
-	* Check if at least one of the 2 edge nodes is moving
-	* @return Boolean to check if the edge is moving or not
-	*/
-	bool IsEdgeMoving(const FGraphEdge& GraphEdge)
-	{
-		const bool bFirstNodeMoving = GraphNodes.IsValidIndex(GraphEdge.FirstNode) && !GraphNodes[GraphEdge.FirstNode].bStationaryNode;
-		const bool bSecondNodeMoving = GraphNodes.IsValidIndex(GraphEdge.SecondNode) && !GraphNodes[GraphEdge.SecondNode].bStationaryNode;
-		return bFirstNodeMoving || bSecondNodeMoving;
-	}
-
-	bool IsEdgeMoving(const int32 EdgeIndex)
-	{
-		if(GraphEdges.IsValidIndex(EdgeIndex))
-		{
-			return IsEdgeMoving(GraphEdges[EdgeIndex]);
-		}
-		return false;
-	}
-
-	int32 AddIsland(FGraphIsland&& Island)
-	{
-		const int32 IslandIndex = GraphIslands.Emplace(MoveTemp(Island));
-		if (Owner != nullptr)
-		{
-			Owner->GraphIslandAdded(IslandIndex);
-		}
-		return IslandIndex;
-	}
-
-	void SetNodeLevel(FGraphNode& GraphNode, const int32 LevelIndex)
-	{
-		GraphNode.LevelIndex = LevelIndex;
-		if (Owner != nullptr)
-		{
-			Owner->GraphNodeLevelAssigned(GraphNode);
-		}
-	}
-
-
-	void SetEdgeLevel(FGraphEdge& GraphEdge, const int32 LevelIndex)
-	{
-		GraphEdge.LevelIndex = LevelIndex;
-		if (Owner != nullptr)
-		{
-			Owner->GraphEdgeLevelAssigned(GraphEdge);
-		}
-	}
-
-	void SortFreeLists();
-
-	FGraphOwner* Owner = nullptr;
-
-	/** List of graph nodes */
-	TSparseArray<FGraphNode> GraphNodes;
-
-	/** List of graph edges */
-	TSparseArray<FGraphEdge> GraphEdges;
-
-	/** List of graph islands */
-	TSparseArray<FGraphIsland> GraphIslands;
-
-	/** Reverse list to find given an item the matching node */
-	TMap<NodeType,int32> ItemNodes;
-
-	/** Reverse list to find given an item the matching edge */
-	TMap<EdgeType,int32> ItemEdges;
-
-	/** Whether we care about determinism */
-	bool bIsDeterministic = false;
-
-	/** Graph counter used for graph traversal to check if an edge/node/island has already been processed */
-	int32 GraphCounter = 0;
-
-	/** List of nodes to be processed in the graph functions */
-	TDeque<int32> NodeQueue;
-};
-
-// For use with unit tests to pretend to be the owner of the graph
-template<typename NodeType, typename EdgeType>
-struct TNullIslandGraphOwner
-{
-	void GraphIslandAdded(const int32 IslandIndex) {}
-
-	void GraphNodeAdded(const NodeType& NodeItem, const int32 NodeIndex) {}
-	void GraphNodeRemoved(const NodeType& NodeItem) {}
-	void GraphNodeLevelAssigned(TIslandGraphNode<NodeType>& Node) {}
-
-	void GraphEdgeAdded(const EdgeType& EdgeItem, const int32 EdgeIndex) {}
-	void GraphEdgeRemoved(const EdgeType& EdgeItem) {}
-	void GraphEdgeLevelAssigned(TIslandGraphEdge<EdgeType>& Edge) {}
-};
-
-
-}
\ No newline at end of file
+// DEPRECATED - to be removed
diff --git a/Engine/Source/Runtime/Experimental/Chaos/Public/Chaos/Island/IslandGroup.h b/Engine/Source/Runtime/Experimental/Chaos/Public/Chaos/Island/IslandGroup.h
--- a/Engine/Source/Runtime/Experimental/Chaos/Public/Chaos/Island/IslandGroup.h
+++ b/Engine/Source/Runtime/Experimental/Chaos/Public/Chaos/Island/IslandGroup.h
@@ -4,16 +4,10 @@
 #include "Containers/Array.h"
 
 #include "Chaos/Evolution/ConstraintGroupSolver.h"
+#include "Chaos/Island/IslandManagerFwd.h"
 
 namespace Chaos
 {
-	namespace Private
-	{
-		class FPBDIsland;
-		class FPBDIslandManager;
-	}
-	class FPBDConstraintContainer;
-
 	namespace Private
 	{
 		/**
diff --git a/Engine/Source/Runtime/Experimental/Chaos/Public/Chaos/Island/IslandGroupManager.h b/Engine/Source/Runtime/Experimental/Chaos/Public/Chaos/Island/IslandGroupManager.h
--- a/Engine/Source/Runtime/Experimental/Chaos/Public/Chaos/Island/IslandGroupManager.h
+++ b/Engine/Source/Runtime/Experimental/Chaos/Public/Chaos/Island/IslandGroupManager.h
@@ -5,11 +5,6 @@
 
 namespace Chaos
 {
-	namespace Private
-	{
-		class FPBDIsland;
-		class FPBDIslandManager;
-	}
 	class FPBDConstraintContainer;
 
 #if CSV_PROFILER
diff --git a/Engine/Source/Runtime/Experimental/Chaos/Public/Chaos/Island/IslandManager.h b/Engine/Source/Runtime/Experimental/Chaos/Public/Chaos/Island/IslandManager.h
--- a/Engine/Source/Runtime/Experimental/Chaos/Public/Chaos/Island/IslandManager.h
+++ b/Engine/Source/Runtime/Experimental/Chaos/Public/Chaos/Island/IslandManager.h
@@ -1,388 +1,636 @@
 // Copyright Epic Games, Inc. All Rights Reserved.
 #pragma once
-
-#include "Containers/SparseArray.h"
-
-#include "Chaos/Real.h"
-#include "Chaos/Vector.h"
+#include "Chaos/Core.h"
+#include "Chaos/Framework/PoolBackedArray.h"
 #include "Chaos/Framework/Handles.h"
+#include "Chaos/Island/IslandManagerFwd.h"
 #include "Chaos/ParticleHandleFwd.h"
-#include "Chaos/GeometryParticlesfwd.h"
-#include "Chaos/Island/IslandGraph.h"
-#include "Chaos/Island/SolverIsland.h"
+#include "Chaos/Serializable.h"
 
 namespace Chaos
 {
+	class FPerParticleGravity;
+}
 
-/** Forward declaration */
-template<typename T, int D>
-class TPBDRigidParticles;
-template<class T>
-class TChaosPhysicsMaterial;
-template <typename T>
-class TSerializablePtr;
-template<typename T>
-class TArrayCollectionArray;
-template <typename T>
-class TParticleView;
-class FPBDRigidsSOAs;
-class FConstraintHandle;
-class FChaosPhysicsMaterial;
-
-using FPBDRigidParticles = TPBDRigidParticles<FReal, 3>;
-
-namespace Private
-{
-
-/** Island manager responsible to create the list of solver islands that will be persistent over time */
-class CHAOS_API FPBDIslandManager
+namespace Chaos::Private
 {
-public:
-	using GraphType = FIslandGraph<FGeometryParticleHandle*, FConstraintHandleHolder, FPBDIsland*, FPBDIslandManager>;
-	using FGraphNode = GraphType::FGraphNode;
-	using FGraphEdge = GraphType::FGraphEdge;
-	using FGraphIsland = GraphType::FGraphIsland;
+	class FPBDIslandMergeSet;
 
 	/**
-	  * Default island manager constructor
-	  */
-	FPBDIslandManager();
-	
-	/**
-	* Default island manager destructor
+	* A particle in the constraint graph.
 	*/
-	~FPBDIslandManager();
+	class CHAOS_API FPBDIslandParticle
+	{
+	public:
+		FPBDIslandParticle();
+		FPBDIslandParticle(FGeometryParticleHandle* InParticle);
+		~FPBDIslandParticle();
 
-	/**
-	 * Set to true if we require deterministic behaviour.
-	 */
-	void SetIsDeterministic(const bool bInIsDeterministic);
+		FGeometryParticleHandle* GetParticle()
+		{
+			return Particle;
+		}
 
-	/**
-	 * Clear all nodes and edges from the graph
-	*/
-	void Reset();
+		const FGeometryParticleHandle* GetParticle() const
+		{
+			return Particle;
+		}
 
-	/**
-	 * Remove all particles (and their constraints) from the graph
-	*/
-	void RemoveParticles();
+		// TPoolBackedItemAdapter for FPBDIslandManager::Nodes
+		const int32 GetArrayIndex() const { return ArrayIndex; }
+		void SetArrayIndex(const int32 InIndex) { ArrayIndex = InIndex; }
+		void Reuse(FGeometryParticleHandle* InParticle);
+		void Trash();
 
-	/**
-	* Remove all the constraints from the graph
-	*/
-	void RemoveConstraints();
+		// An integer ID for an island (for debug only). @see FPBDIsland::GetIslandId()
+		int32 GetIslandId() const;
 
-	/**
-	* Clear the manager and set up the particle-to-graph-node mapping for the specified particles
-	* Should be called before AddConstraint.
-	* @param Particles List of particles to be used to fill the graph nodes
-	*/
-	void InitializeGraph(const TParticleView<FPBDRigidParticles>& Particles);
+	private:
+		friend class FPBDIslandManager;
+		friend class FPBDIsland;
 
-	/**
-	 * @brief Called at the end of the frame to perform cleanup
-	*/
-	void EndTick();
+		union FFlags
+		{
+		public:
+			FFlags() : Bits(0) {}
+			void Reset() { Bits = 0; }
 
-	/**
-	  * Reserve space in the graph for NumConstraints additional constraints.
-	  * @param NumConstraints Number of constraints to be used to reserved memory space
-	  */
-	void ReserveConstraints(const int32 NumConstraints);
+			struct
+			{
+				uint32 bIsDynamic : 1;
+				uint32 bIsSleeping : 1;
+				uint32 bIsMoving : 1;
+				uint32 bNeedsResim : 1;
+			};
 
-	/**
-	  * Add a constraint to the graph for each constraint in the container.
-	  * @param ContainerId Contaner id the constraint is belonging to
-	  * @param ConstraintHandle Constraint Handle that will be used for the edge item
-	  * @param ConstrainedParticles List of 2 particles handles that are used within the constraint
-	  * @return Edge index within the sparse graph edges list
-	  */
-	void AddConstraint(const uint32 ContainerId, FConstraintHandle* ConstraintHandle, const TVector<FGeometryParticleHandle*, 2>& ConstrainedParticles);
+		private:
+			uint32 Bits;
+		};
+
+		// The particle we represent
+		FGeometryParticleHandle* Particle = nullptr;
+
+		// The island we are in (only if dynamic, always null for kinematics which may be in many islands)
+		FPBDIsland* Island = nullptr;
+
+		// All the edges connected to this node
+		TArray<FPBDIslandConstraint*> Edges;
+
+		// Used in a few places to track which nodes have been visited
+		int32 VisitEpoch = INDEX_NONE;
+
+		// Index into FPBDIslandManager::Nodes
+		int32 ArrayIndex = INDEX_NONE;
+
+		// Index into FPBDIsland::Nodes
+		int32 IslandArrayIndex = INDEX_NONE;
+
+		// All the flags
+		FFlags Flags;
+
+		// Distance to a kinematic through the graph
+		int32 Level = 0;
+
+		// Sleep thresholds
+		FRealSingle SleepLinearThresholdSq = 0;
+		FRealSingle SleepAngularThresholdSq = 0;
+		int32 SleepCounterThreshold = 0;
+
+		// Disable thresholds
+		FRealSingle DisableLinearThresholdSq = 0;
+		FRealSingle DisableAngularThresholdSq = 0;
+
+		int32 ResimFrame = INDEX_NONE;
+	};
 
 	/**
-	 * Utility function for use by the constraint containers to add all their constraints
-	 */
-	template<typename ConstraintContainerType>
-	void AddContainerConstraints(ConstraintContainerType& ConstraintContainer)
+	* A constraint in the constraint graph.
+	*/
+	class CHAOS_API FPBDIslandConstraint
 	{
-		const int32 ContainerId = ConstraintContainer.GetContainerId();
-		for (FConstraintHandle* ConstraintHandle : ConstraintContainer.GetConstraintHandles())
+	public:
+		FPBDIslandConstraint();
+		FPBDIslandConstraint(const int32 InContainerId, FConstraintHandle* InConstraint);
+		~FPBDIslandConstraint();
+
+		FConstraintHandle* GetConstraint() const
 		{
-			if ((ConstraintHandle != nullptr) && ConstraintHandle->IsEnabled())
-			{
-				AddConstraint(ContainerId, ConstraintHandle, ConstraintHandle->GetConstrainedParticles());
-			}
+			return Constraint;
 		}
-	}
 
-	/**
-	  * Remove a constraint from the graph
-	  * @param ContainerId Container id the constraint is belonging to
-	  * @param ConstraintHandle Constraint Handle that will be used for the edge item
-	  * @param ConstrainedParticles List of 2 particles handles that are used within the constraint
-	  */
-	void RemoveConstraint(const uint32 ContainerId, FConstraintHandle* ConstraintHandle);
+		uint64 GetSortKey() const
+		{
+			// NOTE: Level 0 means unasigned (used to be INDEX_NONE) so that this sort key still makes sense
+			return (uint64(Level) << 32) | uint64(LevelSortKey);
+		}
 
-	/**
-	 * Remove all constraints of the specified type.
-	 */
-	void RemoveConstraints(const uint32 ContainerId);
+		// TPoolBackedItemAdapter for FPBDIslandManager::Edges
+		const int32 GetArrayIndex() const { return ArrayIndex; }
+		void SetArrayIndex(const int32 InIndex) { ArrayIndex = InIndex; }
+		void Reuse(const int32 InContainerId, FConstraintHandle* InConstraint);
+		void Trash();
 
-	/**
-	 * @brief Remove all the constraints of the specified type (ContainerId) belonging to the Particle
-	 * @param Particle the particle to remove the constraints from
-	 * @param ContainerId Container id the constraint is belonging to
-	*/
-	void RemoveParticleConstraints(FGeometryParticleHandle* Particle, const uint32 ContainerId);
-	
-	/**
-	  * Preallocate buffers for \p Num particles.
-	  * @param NumParticles Number of particles to be used to reserve memory
-	  * @return Number of new slots created.
-	  */
-	int32 ReserveParticles(const int32 NumParticles);
+		// An integer ID for an island (for debug only). @see FPBDIsland::GetIslandId()
+		int32 GetIslandId() const;
 
-	/**
-	  * Add a particle to the graph
-	  * @param ParticleHandle Particle to be added to the graph
-	  * @param ParentParticleHandle Put the ParticleHandle in the same island as the parent with the same sleep state. May be null.
-	  * @return Sparse graph node index
-	  */
-	int32 AddParticle(FGeometryParticleHandle* ParticleHandle);
+	private:
+		friend class FPBDIslandManager;
+		friend class FPBDIsland;
 
-	/**
-	  * Remove a particle from the graph
-	  * @param Particle Particle to be removed from the islands
-	  */
-	void RemoveParticle(FGeometryParticleHandle* Particle);
+		union FFlags
+		{
+		public:
+			FFlags() : Bits(0) {}
+			void Reset() { Bits = 0; }
 
-	/**
-	 * Get the Level for the specified particle (distance to a kinmetic in the graph)
-	*/
-	int32 GetParticleLevel(FGeometryParticleHandle* Particle) const;
+			struct
+			{
+				uint32 bIsSleeping : 1;
+			};
+		private:
+			uint32 Bits;
+		};
+
+		FConstraintHandle* Constraint = nullptr;
+		FPBDIsland* Island = nullptr;
+		FPBDIslandParticle* Nodes[2] = { nullptr, nullptr };
 
+		// Used in a few places to track which nodes have been visited
+		int32 VisitEpoch = INDEX_NONE;
+
+		// Index into FPBDIslandManager::Edges
+		int32 ArrayIndex = INDEX_NONE;
+
+		// Index into FPBDIslandParticle::Edges for Node0 and Node1
+		int32 NodeArrayIndices[2] = { INDEX_NONE, INDEX_NONE };
+
+		// Indices into FPBDIsland::ContainerEdges
+		int32 ContainerIndex = INDEX_NONE;
+		int32 IslandArrayIndex = INDEX_NONE;
+
+		// Distance to a kinematic through the graph
+		int32 Level = 0;
+
+		// Used for consistent ordering within a level
+		uint32 LevelSortKey = 0;
+
+		FFlags Flags;
+	};
 
 	/**
-	  * Call the visitor on all constraints in an awake island. E.g., can be used to check if any constraints
-	  * in the graph are expired and should be removed. However, no constraints should be added or removed
-	  * from the visitor itself - build a list and remove them after.
-	  */
-	template<typename VisitorType>
-	void VisitConstraintsInAwakeIslands(const int32 ContainerId, const VisitorType& Visitor)
+	* A set of connected constraints and particles.
+	* A dynamic particle can only be in one island, and each island explicitly holds an array of
+	* its dynamic particles. Kinematic particles may be in multiple islands and we do not hold
+	* the set of these. To find kinematic particles used by an island you would need to visit all
+	* the constraints.
+	*/
+	class CHAOS_API FPBDIsland
 	{
-		for (FGraphEdge& GraphEdge : IslandGraph->GraphEdges)
+	public:
+		FPBDIsland();
+		~FPBDIsland();
+
+		int32 GetNumParticles() const
 		{
-			// If a constraint is not in any island (kinematic-kinematic for example) it is considered asleep.
-			if ((GraphEdge.ItemContainer == ContainerId) && (GraphEdge.IslandIndex != INDEX_NONE))
-			{
-				const FGraphIsland& GraphIsland = IslandGraph->GraphIslands[GraphEdge.IslandIndex];
-				if (!GraphIsland.bIsSleeping)
-				{
-					Visitor(GraphEdge.EdgeItem);
-				}
-			}
+			return Nodes.Num();
 		}
-	}
 
-	/**
-	  * Clear all the islands
-	  */
-	void ResetIslands(const TParticleView<FPBDRigidParticles>& PBDRigids);
+		int32 GetNumConstraints() const
+		{
+			return NumEdges;
+		}
 
-	/**
-	  * Pack all the constraints into their assigned islands in sorted order
-	  */
-	void UpdateIslands(FPBDRigidsSOAs& Particles, const bool IsResimming = false);
+		int32 GetNumContainerConstraints(const int32 ContainerId) const
+		{
+			return ContainerEdges[ContainerId].Num();
+		}
 
-	/**
-	  * @brief Put all particles and constraints in an island to sleep
-	  * @param Particles used in the islands
-	  * @param IslandIndex Island index to be woken up
-	*/
-	void SleepIsland(FPBDRigidsSOAs& Particles, const int32 IslandIndex);
+		const FPBDIslandParticle* GetNode(const int32 NodeIndex) const
+		{
+			return Nodes[NodeIndex];
+		}
 
-	/**
-	  * Put particles in inactive islands to sleep.
-	  * @param IslandIndex Island index
-	  * @param PerParticleMaterialAttributes Material attributes that has been set on each particles
-	  * @param SolverPhysicsMaterials Solver materials that could be used to find the sleeping thresholds
-	  * @return return true if the island is going to sleep
-	  */
-	bool SleepInactive(
-		const int32 IslandIndex, 
-		const TArrayCollectionArray<TSerializablePtr<FChaosPhysicsMaterial>>& PhysicsMaterials, 
-		const TArrayCollectionArray<TUniquePtr<FChaosPhysicsMaterial>>& PerParticlePhysicsMaterials,
-		const THandleArray<FChaosPhysicsMaterial>& SimMaterials);
+		bool IsSleeping() const
+		{
+			return Flags.bIsSleeping;
+		}
 
-	/**
-	 * Get the set of islands that the particle belongs to. Will be 0-1 island for dynamics and could be any number for non-dynamics.
-	 * This works for sleeping and waking particles (but we must iterate over constraints for this to work).
-	 */
-	TArray<int32> FindParticleIslands(const FGeometryParticleHandle* ParticleHandle) const;
+		int32 GetSleepCounter() const
+		{
+			return SleepCounter;
+		}
 
-	/**
-	 * Get the island index of the constraint
-	*/
-	int32 GetConstraintIsland(const FConstraintHandle* ConstraintHandle) const;
+		void SetSleepCounter(const int32 InSleepCounter)
+		{
+			SleepCounter = InSleepCounter;
+		}
 
-	/**
-	 * When resim is used, tells us whether we need to resolve island
-	 * @param IslandIndex Island index
-	 * @return If true the island need to be resim
-	 */
-	bool IslandNeedsResim(const int32 IslandIndex) const;
+		bool IsSleepAllowed() const
+		{
+			return Flags.bIsSleepAllowed;
+		}
 
-	/** Get the particle resim frame */
-	int32 GetParticleResimFrame(const FGeometryParticleHandle* ParticleHandle) const;
+		bool NeedsResim() const
+		{
+			return Flags.bNeedsResim;
+		}
 
-	/** Set the particle resim frame */
-	void SetParticleResimFrame(FGeometryParticleHandle* ParticleHandle, const int32 ResimFrame);
+		void SetNeedsResim(const bool bInNeedsResim)
+		{
+			Flags.bNeedsResim = bInNeedsResim;
+		}
 
-	/** Reset all the particle resim frames to INDEX_NONE */
-	void ResetParticleResimFrame(const int32 ResetFrame = INDEX_NONE);
+		int32 GetResimFrame() const
+		{
+			return ResimFrame;
+		}
 
-	/**
-	  * The number of islands in the graph.
-	  * @return the number of islands
-	  */
-	FORCEINLINE int32 NumIslands() const
-	{
-		return IslandIndexing.Num();
-	}
-	
-	/**
-	 * The number of islands in the graph.
-	 * @return Get the island graph
-	 */
-	FORCEINLINE const GraphType* GetIslandGraph() const
-	{
-		return IslandGraph.Get();
-	}
-	/**
-	* Get the island graph
-	* @return Island graph of the manager
-	*/
-	FORCEINLINE GraphType* GetIslandGraph()
-	{
-		return IslandGraph.Get();
-	}
+		void SetResimFrame(const int32 InResimFrame)
+		{
+			ResimFrame = InResimFrame;
+		}
 
-	/**
-	 * Get Max UniqueIdx of the particles.
-	 * Used to create arrays large enough to use Particle.UniqueIdx().Idx for indexing.
-	 * @return the max index of the particles in the islands
-	 */
-	inline int32 GetMaxParticleUniqueIdx() const 
-	{ 
-		return MaxParticleIndex;
-	}
+		bool IsUsingCache() const
+		{
+			return Flags.bIsUsingCache;
+		}
 
-	/**
-	 * Get the sparse island graph index given a dense index between 0 and NumIslands
-	 * @param IslandIndex Island index
-	 * @return the max index of the particles in the islands
-	 */
-	int32 GetGraphIndex(const int32 IslandIndex) const 
-	{
-		return IslandIndexing.IsValidIndex(IslandIndex) ? IslandIndexing[IslandIndex] : 0;
-	}
+		void SetIsUsingCache(const bool bInIsUsingCache)
+		{
+			Flags.bIsUsingCache = bInIsUsingCache;
+		}
 
-	/**
-	 * Register a constraint container with the island manager. Currently this is only used to determin
-	 * how many ContainerIds we need in order to presize arrays etc.
-	*/
-	void AddConstraintContainer(FPBDConstraintContainer& ConstraintContainer);
-
-	/** Accessors of the graph island */
-	const FPBDIsland* GetIsland(const int32 IslandIndex) const {return IslandIndexing.IsValidIndex(IslandIndex) ? Islands[IslandIndexing[IslandIndex]].Get() : nullptr; }
-	FPBDIsland* GetIsland(const int32 IslandIndex) {return IslandIndexing.IsValidIndex(IslandIndex) ? Islands[IslandIndexing[IslandIndex]].Get() : nullptr; }
-
-	/** Get the graph nodes */
-	// @todo(chaos): remove this - only used by testing
-	const TSparseArray<GraphType::FGraphNode>& GetGraphNodes() const { return IslandGraph->GraphNodes; }
-	
-	/** Get the graph edges */
-	// @todo(chaos): remove this - only used by testing
-	const TSparseArray<GraphType::FGraphEdge>& GetGraphEdges() const { return IslandGraph->GraphEdges; }
-
-	/** Get the particle nodes */
-	// @todo(chaos): remove this - only used by testing
-	const TMap<FGeometryParticleHandle*, int32>& GetParticleNodes() const { return IslandGraph->ItemNodes; }
-
-	// Graph Owner API Implementation. Allows us to store edge/node indices on constraints and particles
-	void GraphIslandAdded(const int32 IslandIndex);
-	void GraphNodeAdded(FGeometryParticleHandle* NodeItem, const int32 NodeIndex);
-	void GraphNodeRemoved(FGeometryParticleHandle* NodeItem);
-	void GraphNodeLevelAssigned(FGraphNode& GraphNode);
-	void GraphEdgeAdded(const FConstraintHandleHolder& EdgeItem, const int32 EdgeIndex);
-	void GraphEdgeRemoved(const FConstraintHandleHolder& EdgeItem);
-	void GraphEdgeLevelAssigned(FGraphEdge& GraphEdge);
-
-#if CHAOS_CONSTRAINTHANDLE_DEBUG_ENABLED
-	bool DebugCheckParticlesInGraph(const TParticleView<FPBDRigidParticles>& Particles) const;
-	bool DebugCheckParticleNotInGraph(const FGeometryParticleHandle* ParticleHandle) const;
-	bool DebugParticleConstraintsNotInGraph(const FGeometryParticleHandle* Particle, const int32 ContainerId) const;
-	bool DebugCheckConstraintNotInGraph(const FConstraintHandle* ConstraintHandle) const;
-#endif
-
-protected:
-
-#if CHAOS_CONSTRAINTHANDLE_DEBUG_ENABLED
-	bool DebugCheckIslandConstraints() const;
-	void DebugCheckParticleIslands(const FGeometryParticleHandle* ParticleHandle) const;
-	void DebugCheckConstraintIslands(const FConstraintHandle* ConstraintHandle, const int32 IslandIndex) const;
-	void DebugCheckIslands() const;
-#endif
-
-	// Add a particle to the graph if it is not already in it. Return the node index if present (added or existing)
-	int32 TryAddParticle(FGeometryParticleHandle* ParticleHandle, const int32 IslandIndex);
-
-	// Add a particle to the graph if it is dynamic and not already in it. Return the node index if present, even if not dynamic
-	int32 TryAddParticleIfDynamic(FGeometryParticleHandle* ParticleHandle, const int32 IslandIndex);
-
-	// Add a particle to the graph if it is kinematic and not already in it. Return the node index if present, even if dynamic
-	int32 TryAddParticleIfNotDynamic(FGeometryParticleHandle* ParticleHandle, const int32 IslandIndex);
-
-	void InitIslands();
+		TArrayView<FPBDIslandParticle*> GetParticles()
+		{
+			return MakeArrayView(Nodes);
+		}
 
-	/**
-	* Sync the Islands with the IslandGraph
-	*/
-	void SyncIslands(FPBDRigidsSOAs& Particles, const bool IsResimming);
+		TArrayView<FPBDIslandConstraint*> GetConstraints(const int32 ContainerId)
+		{
+			return MakeArrayView(ContainerEdges[ContainerId]);
+		}
+
+		// An integer ID for an island. For debug only (e.g., to assign a color in debug draw) - has no internal use.
+		// Not necessarily persistent between ticks.
+		int32 GetIslandId() const
+		{
+			return ArrayIndex;
+		}
 
+		// TPoolBackedItemAdapter for FPBDIslandManager::Islands
+		const int32 GetArrayIndex() const { return ArrayIndex; }
+		void SetArrayIndex(const int32 InIndex) { ArrayIndex = InIndex; }
+		void Reuse();
+		void Trash();
 
-	/**
-	* Set the island and all particles' sleep state (called from the graph update).
-	*/
-	void SetIslandSleeping(FPBDRigidsSOAs& Particles, const int32 IslandIndex, const bool bIsSleeping);
+	private:
+		friend class FPBDIslandManager;
+
+		void UpdateSyncState();
+
+		union FFlags
+		{
+		public:
+			FFlags() : Bits(0) {}
+			void Reset() { Bits = 0; }
+
+			struct
+			{
+				uint32 bItemsAdded : 1;
+				uint32 bItemsRemoved : 1;
+				uint32 bIsSleepAllowed : 1;
+				uint32 bIsSleeping : 1;
+				uint32 bWasSleeping : 1;
+				uint32 bCheckSleep : 1;
+				uint32 bNeedsResim : 1;
+				uint32 bIsUsingCache : 1;
+			};
 
-	void ValidateIslands() const;
+		private:
+			uint32 Bits;
+		};
 
-	TArray<FPBDConstraintContainer*>		ConstraintContainers;
+		// Nodes for dynamic particles in the Island (does not include kinematics)
+		// @todo(chaos): we only need to track allocated nodes for sleep logic when particles have no constraints - fix this
+		TArray<FPBDIslandParticle*> Nodes;
 
-	/** List of solver islands in the manager */
-	TSparseArray<TUniquePtr<FPBDIsland>>	Islands;
+		// All edges in the island for each constraint type
+		TConstraintTypeArray<TArray<FPBDIslandConstraint*>> ContainerEdges;
 
-	/** Island map to make the redirect an index between 0...NumIslands to the persistent island index  */
-	TArray<int32>							IslandIndexing;
+		// When edges are added, islands merge. The MergeSet tracks which islands are merging together
+		FPBDIslandMergeSet* MergeSet = nullptr;
 
-	/** Global graph to deal with merging and splitting of islands */
-	TUniquePtr<GraphType>					IslandGraph;
+		// Index into FPBDIslandMergeSet::Islands
+		int32 MergeSetIslandIndex = INDEX_NONE;
 
-	/** Maximum particle index that we have in the graph*/
-	int32									MaxParticleIndex = INDEX_NONE;
+		// Index into FPBDIslandManager::Islands
+		int32 ArrayIndex = INDEX_NONE;
 
-	/** Islands are only populated with constraints and particles during the solver phase. Trying to access them outside that is an error we can trap */
-	bool									bIslandsPopulated;
+		// Total number of edges in ContainerEdges[0..N]
+		int32 NumEdges = 0;
 
-	/** Make sure EndTick is called. Sanity check mostlyhelpful for unit tests */
-	bool									bEndTickCalled;
-};
+		// When SleepCounter hits a threshold the island will sleep
+		int32 SleepCounter = 0;
 
-}	// namespace Private
+		// When DisableCounter hits a threshold the island particle will be disabled
+		int32 DisableCounter = 0;
 
+		int32 ResimFrame = INDEX_NONE;
 
-	using FPBDIslandManager UE_DEPRECATED(5.2, "Internal class moved to Private namespace") = Private::FPBDIslandManager;
+		FFlags Flags;
+	};
+
+	/**
+	* A set of islands to be merged together.
+	*/
+	class CHAOS_API FPBDIslandMergeSet
+	{
+	public:
+		FPBDIslandMergeSet();
+		~FPBDIslandMergeSet();
+
+		// TPoolBackedItemAdapter for FPBDIslandManager::MergeSets
+		const int32 GetArrayIndex() const { return ArrayIndex; }
+		void SetArrayIndex(const int32 InIndex) { ArrayIndex = InIndex; }
+		void Reuse();
+		void Trash();
+
+	private:
+		friend class FPBDIslandManager;
+
+		TArray<FPBDIsland*> Islands;
+		int32 NumEdges = 0;
+
+		// Index into FPBDIslandManager::MergeSets
+		int32 ArrayIndex = INDEX_NONE;
+	};
+
+	/**
+	*
+	* Notes on the implementation:
+	*	- Particles are nodes in the graph
+	*	- Constraints are edges in the graph
+	*	- Each island knows about the constraints and dynamic particle in it
+	*	- Islands do not track their kinematic particles (kinematics may be in multiple islands)
+	*
+	* @todo(chaos): the sleep and disable logic feel like they don't belong in here. It requires
+	* that we know about materials, and the Disabled Particle concept sits at a higher level
+	* (in the Evolution, and in principle sleep does too since other system might need to know).
+	* The main reason it is here is so that we can cache the thresholds and don't have to check 
+	* the materials for every particle every tick.
+	*/
+	class CHAOS_API FPBDIslandManager
+	{
+	public:
+		FPBDIslandManager(FPBDRigidsSOAs& InParticles);
+		~FPBDIslandManager();
+
+		// The number of registered constraint containers
+		int32 GetNumConstraintContainers() const;
+
+		// Add/Remove particles
+		int32 GetNumParticles() const { return Nodes.Num(); }
+		void AddParticle(FGeometryParticleHandle* Particle);
+		void RemoveParticle(FGeometryParticleHandle* Particle);
+		int32 ReserveParticles(const int32 InNumParticles);
+		void UpdateParticleMaterial(FGeometryParticleHandle* Particle);
+		int32 GetParticleLevel(FGeometryParticleHandle* Particle) const;
+
+		// Add/Remove constraints
+		int32 GetNumConstraints() const { return Edges.Num(); }
+		void AddConstraint(const int32 ContainerId, FConstraintHandle* Constraint, const TVec2<FGeometryParticleHandle*>& ConstrainedParticles);
+		template<typename ConstraintContainerType> void AddContainerConstraints(ConstraintContainerType& ConstraintContainer);
+		void RemoveConstraint(FConstraintHandle* ConstraintHandle);
+		void RemoveParticleConstraints(FGeometryParticleHandle* Particle, int32 ContainerId);
+		void RemoveContainerConstraints(int32 ContainerId);
+
+		// Access to Islands
+		int32 GetNumIslands() const { return Islands.Num(); }
+		const FPBDIsland* GetIsland(const int32 IslandIndex) const { return Islands[IslandIndex]; }
+		const FPBDIsland* GetParticleIsland(const FGeometryParticleHandle* Particle) const;
+		const FPBDIsland* GetConstraintIsland(const FConstraintHandle* Constraint) const;
+
+		// Setup and update
+		void Reset();
+		void RemoveAllConstraints();
+		void AddConstraintContainer(const FPBDConstraintContainer& Container);
+		void RemoveConstraintContainer(const FPBDConstraintContainer& Container);
+		void SetMaterialContainers(const TArrayCollectionArray<TSerializablePtr<FChaosPhysicsMaterial>>* InPhysicsMaterials, const TArrayCollectionArray<TUniquePtr<FChaosPhysicsMaterial>>* InPerParticlePhysicsMaterials, const THandleArray<FChaosPhysicsMaterial>* InSimMaterials);
+		void SetGravityForces(const FPerParticleGravity* InGravity);
+		void SetDisableCounterThreshold(const int32 InDisableCounterThreshold);
+		void SetIsDeterministic(const bool bInIsDeterministic);
+		void UpdateParticles();
+		void UpdateIslands();
+		void UpdateSleep();
+		void UpdateDisable(TFunctionRef<void(FPBDRigidParticleHandle*)> ParticleDisableFunctor);
+		void EndTick();
+
+		// Levels and Colors for sorting an parallelization
+		int32 GetParticleLevel(const FPBDIslandParticle* Node) const;
+		int32 GetParticleColor(const FPBDIslandParticle* Node) const;
+		int32 GetConstraintLevel(const FPBDIslandConstraint* Edge) const;
+		int32 GetConstraintColor(const FPBDIslandConstraint* Edge) const;
+
+		/**
+		* Visit all the awake constraints from the specified container.
+		* @tparam VisitorType void(const FPBDIslandConstraint*)
+		*/
+		template<typename VisitorType> void VisitAwakeConstraints(const int32 ContainerId, const VisitorType& Visitor);
+		template<typename VisitorType> void VisitAwakeConstConstraints(const int32 ContainerId, const VisitorType& Visitor) const;
+
+		/**
+		* Visit all the constraints from the specified container
+		* @tparam VisitorType void(const FPBDIslandConstraint*)
+		*/
+		template<typename VisitorType> void VisitConstraints(const int32 ContainerId, const VisitorType& Visitor);
+		template<typename VisitorType> void VisitConstConstraints(const int32 ContainerId, const VisitorType& Visitor) const;
+
+		// Resim functionality
+		int32 GetParticleResimFrame(const FGeometryParticleHandle* Particle) const;
+		void SetParticleResimFrame(FGeometryParticleHandle* Particle, const int32 ResimFrame);
+		void ResetParticleResimFrame(const int32 ResetFrame = INDEX_NONE);
+
+		// The following is intended for testing, debugging and visualization only and is not considered part of the API
+		FPBDIsland* GetIsland(const int32 IslandIndex) { return Islands[IslandIndex]; }
+		int32 GetIslandIndex(const FPBDIsland* Island) const;
+		int32 GetIslandArrayIndex(const FPBDIslandConstraint* Edge) const;
+		TArray<const FPBDIsland*> FindParticleIslands(const FGeometryParticleHandle* Particle) const;
+		TArray<const FGeometryParticleHandle*> FindParticlesInIslands(const TArray<const FPBDIsland*> Islands) const;
+		TArray<const FConstraintHandle*> FindConstraintsInIslands(const TArray<const FPBDIsland*> Islands, int32 ContainerId) const;
+		void SetParticleIslandIsSleeping(FGeometryParticleHandle* Particle, const bool bInIsSleeping);
+
+		// Deprecated API
+		UE_DEPRECATED(5.3, "Use Reset") void RemoveConstraints() { Reset(); }
+		UE_DEPRECATED(5.3, "Use UpdateParticles") void InitializeGraph(const TParticleView<FPBDRigidParticles>& InParticles) { UpdateParticles(); }
+		UE_DEPRECATED(5.3, "Use GetNumIslands") int32 NumIslands() const { return GetNumIslands(); }
+		UE_DEPRECATED(5.3, "Use GetParticleIsland") const FPBDIsland* GetIsland(const FGeometryParticleHandle* Particle) const { return GetParticleIsland(Particle); }
+		UE_DEPRECATED(5.3, "Use RemoveConstraint without ContainerId") void RemoveConstraint(const int32 ContainerId, FConstraintHandle* Constraint) { RemoveConstraint(Constraint); }
+
+	private:
+		// @todo(chaos): for access to non-const GetIsland(). Try to remove these...
+		friend class FPBDIslandGroupManager;
+		friend class FPBDRigidsEvolutionGBF;
+
+		// Node management
+		FPBDIslandParticle* GetGraphNode(const FGeometryParticleHandle* Particle) const;
+		FPBDIslandParticle* GetGraphNode(const FTransientGeometryParticleHandle& Particle) const;
+		FPBDIslandParticle* CreateGraphNode(FGeometryParticleHandle* Particle);
+		FPBDIslandParticle* GetOrCreateGraphNode(FGeometryParticleHandle* Particle);
+		void DestroyGraphNode(FPBDIslandParticle* Node);
+
+		// Edge management
+		FPBDIslandConstraint* GetGraphEdge(const FConstraintHandle* Constraint);
+		const FPBDIslandConstraint* GetGraphEdge(const FConstraintHandle* Constraint) const;
+		FPBDIslandConstraint* CreateGraphEdge(const int32 ContainerId, FConstraintHandle* Constraint, const TVec2<FGeometryParticleHandle*>& ConstrainedParticles);
+		void DestroyGraphEdge(FPBDIslandConstraint* Edge);
+		void BindEdgeToNodes(FPBDIslandConstraint* Edge, FPBDIslandParticle* Node0, FPBDIslandParticle* Node1);
+		void UnbindEdgeFromNodes(FPBDIslandConstraint* Edge);
+		void UpdateGraphNodeSleepSettings(FPBDIslandParticle* Node);
+
+		// Per-tick state transfer to nodes/edges
+		void UpdateGraphNode(FPBDIslandParticle* Node);
+		void UpdateGraphEdge(FPBDIslandConstraint* Edge);
+
+		// Island managament
+		FPBDIsland* CreateIsland();
+		void DestroyIsland(FPBDIsland* Island);
+		void AssignNodeIsland(FPBDIslandParticle* Node);
+		void AssignEdgeIsland(FPBDIslandConstraint* Edge);
+		void AddNodeToIsland(FPBDIslandParticle* Node, FPBDIsland* Island);
+		void RemoveNodeFromIsland(FPBDIslandParticle* Node);
+		void AddEdgeToIsland(FPBDIslandConstraint* Edge, FPBDIsland* Island);
+		void RemoveEdgeFromIsland(FPBDIslandConstraint* Edge);
+		void EnqueueIslandCheckSleep(FPBDIsland* Island, const bool bIsSleepAllowed);
+
+		// Island merge set management
+		FPBDIsland* MergeNodeIslands(FPBDIslandParticle* Node0, FPBDIslandParticle* Node1);
+		FPBDIsland* MergeIslands(FPBDIsland* Island0, FPBDIsland* Island1);
+		FPBDIslandMergeSet* CreateMergeSet(FPBDIsland* Island0, FPBDIsland* Island1);
+		void DestroyMergeSet(FPBDIslandMergeSet* IslandMergeSet);
+		void AddIslandToMergeSet(FPBDIsland* Island, FPBDIslandMergeSet* IslandMergeSet);
+		void RemoveIslandFromMergeSet(FPBDIsland* Island);
+		FPBDIslandMergeSet* CombineMergeSets(FPBDIslandMergeSet* IslandMergeSetParent, FPBDIslandMergeSet* IslandMergeSetChild);
+		FPBDIsland* GetMergeSetParentIsland(FPBDIslandMergeSet* MergeSet, int32& OutNumEdges, const TArrayView<int32>& OutNumContainerEdges);
+
+		// Merge/split processing
+		void ProcessIslands();
+		void ProcessWakes();
+		void ProcessMerges();
+		void ProcessSplits();
+		void AssignLevels();
+		void ProcessIslandMerge(FPBDIsland* ParentIsland, FPBDIsland* ChildIsland);
+		void ProcessIslandSplits(FPBDIsland* Island);
+
+		// Sorting
+		void AssignIslandLevels(FPBDIsland* Island);
+		void SortIslandEdges(FPBDIsland* Island);
+
+		// Sleeping
+		void ProcessSleep();
+		void ProcessParticlesSleep();
+		void ProcessIslandSleep(FPBDIsland* Island);
+		void PropagateIslandSleep(FPBDIsland* Island);
+
+		// Disabling
+		void ProcessDisable(TFunctionRef<void(FPBDRigidParticleHandle*)> ParticleDisableFunctor);
+		void ProcessParticlesDisable(TArray<FPBDRigidParticleHandle*>& OutDisableParticles);
+		void ProcessIslandDisable(FPBDIsland* Island, TArray<FPBDRigidParticleHandle*>& OutDisableParticles);
+
+		void FinalizeIslands();
+		void ApplyDeterminism();
+		bool Validate() const;
+
+		int32 GetNextVisitEpoch();
+
+		// Ideally the island manager would not know about the particles, materials etc, but
+		// it is currently required for particle sleep management
+		FPBDRigidsSOAs& Particles;
+		const TArrayCollectionArray<TSerializablePtr<FChaosPhysicsMaterial>>* PhysicsMaterials;
+		const TArrayCollectionArray<TUniquePtr<FChaosPhysicsMaterial>>* PerParticlePhysicsMaterials;
+		const THandleArray<FChaosPhysicsMaterial>* SimMaterials;
+		const FPerParticleGravity* Gravity;
+		int32 DisableCounterThreshold = TNumericLimits<int32>::Max();
+
+		// The nodes, edges and islands in the graph
+		TPoolBackedArray<FPBDIslandParticle> Nodes;
+		TPoolBackedArray<FPBDIslandConstraint> Edges;
+		TPoolBackedArray<FPBDIsland> Islands;
+
+		// For tracking what islands need to be merged each tick
+		TPoolBackedArray<FPBDIslandMergeSet> MergeSets;
+
+		// All the constraint types registered with the graph
+		TArray<const FPBDConstraintContainer*> ConstraintContainers;
+
+		// A monotonically increasing key used to sort edges at the same level
+		int32 NextLevelSortKey = 0;
+
+		// A counter used to check if a node or edge has been visited in a loop
+		int32 NextVisitEpoch = 0;
+
+		// Whether we require determinism
+		bool bIsDeterministic = false;
+	};
+
+
+	template<typename ConstraintContainerType>
+	void FPBDIslandManager::AddContainerConstraints(ConstraintContainerType& ConstraintContainer)
+	{
+		const int32 ContainerId = ConstraintContainer.GetContainerId();
+		for (FConstraintHandle* ConstraintHandle : ConstraintContainer.GetConstraintHandles())
+		{
+			if ((ConstraintHandle != nullptr) && ConstraintHandle->IsEnabled())
+			{
+				AddConstraint(ContainerId, ConstraintHandle, ConstraintHandle->GetConstrainedParticles());
+			}
+		}
+	}
+
+	template<typename VisitorType>
+	void FPBDIslandManager::VisitAwakeConstConstraints(const int32 ContainerId, const VisitorType& Visitor) const
+	{
+		for (FPBDIsland* Island : Islands)
+		{
+			if (!Island->Flags.bIsSleeping)
+			{
+				for (FPBDIslandConstraint* Edge : Island->ContainerEdges[ContainerId])
+				{
+					Visitor(Edge);
+				}
+			}
+		}
+	}
+
+	template<typename VisitorType>
+	void FPBDIslandManager::VisitAwakeConstraints(const int32 ContainerId, const VisitorType& Visitor)
+	{
+		const_cast<FPBDIslandManager*>(this)->VisitAwakeConstConstraints(ContainerId, 
+			[&Visitor](const FPBDIslandConstraint* Edge)
+			{
+				Visitor(const_cast<FPBDIslandConstraint*>(Edge));
+			});
+	}
+
+	template<typename VisitorType>
+	void FPBDIslandManager::VisitConstConstraints(const int32 ContainerId, const VisitorType& Visitor) const
+	{
+		for (const FPBDIslandConstraint* Edge : Edges)
+		{
+			if (Edge->ContainerIndex == ContainerId)
+			{
+				Visitor(Edge);
+			}
+		}
+	}
+
+	template<typename VisitorType>
+	void FPBDIslandManager::VisitConstraints(const int32 ContainerId, const VisitorType& Visitor)
+	{
+		const_cast<FPBDIslandManager*>(this)->VisitConstConstraints(ContainerId,
+			[&Visitor](const FPBDIslandConstraint* Edge)
+			{
+				Visitor(const_cast<FPBDIslandConstraint*>(Edge));
+			});
+	}
 
-}	// namespace Chaos
\ No newline at end of file
+} // namespace Chaos::Private
diff --git a/Engine/Source/Runtime/Experimental/Chaos/Public/Chaos/Joint/PBDJointContainerSolver.h b/Engine/Source/Runtime/Experimental/Chaos/Public/Chaos/Joint/PBDJointContainerSolver.h
--- a/Engine/Source/Runtime/Experimental/Chaos/Public/Chaos/Joint/PBDJointContainerSolver.h
+++ b/Engine/Source/Runtime/Experimental/Chaos/Public/Chaos/Joint/PBDJointContainerSolver.h
@@ -31,7 +31,7 @@ namespace Chaos
 			virtual int32 GetNumConstraints() const override final { return ContainerConstraintIndices.Num(); }
 			virtual void Reset(const int32 InMaxCollisions) override final;
 			virtual void AddConstraints() override final;
-			virtual void AddConstraints(const TArrayView<Private::FPBDIslandConstraint>& IslandConstraints) override final;
+			virtual void AddConstraints(const TArrayView<Private::FPBDIslandConstraint*>& IslandConstraints) override final;
 			virtual void AddBodies(FSolverBodyContainer& SolverBodyContainer) override final;
 			virtual void GatherInput(const FReal Dt) override final;
 			virtual void GatherInput(const FReal Dt, const int32 BeginIndex, const int32 EndIndex) override final;
diff --git a/Engine/Source/Runtime/Experimental/Chaos/Public/Chaos/PBDConstraintContainer.h b/Engine/Source/Runtime/Experimental/Chaos/Public/Chaos/PBDConstraintContainer.h
--- a/Engine/Source/Runtime/Experimental/Chaos/Public/Chaos/PBDConstraintContainer.h
+++ b/Engine/Source/Runtime/Experimental/Chaos/Public/Chaos/PBDConstraintContainer.h
@@ -6,15 +6,10 @@
 #include "Chaos/Evolution/SolverBody.h"
 #include "Chaos/Evolution/SolverConstraintContainer.h"
 #include "Chaos/Evolution/ConstraintGroupSolver.h"
+#include "Chaos/Island/IslandManagerFwd.h"
 
 namespace Chaos
 {
-	namespace Private
-	{
-		class FPBDIslandManager;
-	}
-	class FConstraintContainerSolver;
-
 	/**
 	 * Base class for containers of constraints.
 	 * A Constraint Container holds an array of constraints and provides methods to allocate and deallocate constraints
diff --git a/Engine/Source/Runtime/Experimental/Chaos/Public/Chaos/PBDConstraintGraph.h b/Engine/Source/Runtime/Experimental/Chaos/Public/Chaos/PBDConstraintGraph.h
--- a/Engine/Source/Runtime/Experimental/Chaos/Public/Chaos/PBDConstraintGraph.h
+++ b/Engine/Source/Runtime/Experimental/Chaos/Public/Chaos/PBDConstraintGraph.h
@@ -3,9 +3,5 @@
 
 #include "Chaos/Island/IslandManager.h"
 
-namespace Chaos
-{
-	using FPBDConstraintGraph UE_DEPRECATED(5.2, "Internal class moved to Private namespace") = Private::FPBDIslandManager;
-}
-
+// DEPRECATED - to be removed
 
diff --git a/Engine/Source/Runtime/Experimental/Chaos/Public/Chaos/PBDJointConstraints.h b/Engine/Source/Runtime/Experimental/Chaos/Public/Chaos/PBDJointConstraints.h
--- a/Engine/Source/Runtime/Experimental/Chaos/Public/Chaos/PBDJointConstraints.h
+++ b/Engine/Source/Runtime/Experimental/Chaos/Public/Chaos/PBDJointConstraints.h
@@ -6,6 +6,7 @@
 #include "Chaos/ConstraintHandle.h"
 #include "Chaos/Collision/CollisionApplyType.h"
 #include "Chaos/Evolution/IndexedConstraintContainer.h"
+#include "Chaos/Island/IslandManagerFwd.h"
 #include "Chaos/ParticleHandleFwd.h"
 #include "Chaos/PBDJointConstraintTypes.h"
 #include "Chaos/PBDJointConstraintData.h"
@@ -16,7 +17,6 @@ namespace Chaos
 	namespace Private
 	{
 		class FPBDJointContainerSolver;
-		class FPBDIslandManager;
 	}
 
 	/**
diff --git a/Engine/Source/Runtime/Experimental/Chaos/Public/Chaos/PBDRigidsEvolution.h b/Engine/Source/Runtime/Experimental/Chaos/Public/Chaos/PBDRigidsEvolution.h
--- a/Engine/Source/Runtime/Experimental/Chaos/Public/Chaos/PBDRigidsEvolution.h
+++ b/Engine/Source/Runtime/Experimental/Chaos/Public/Chaos/PBDRigidsEvolution.h
@@ -3,7 +3,6 @@
 
 #include "Chaos/Character/CharacterGroundConstraintContainer.h"
 #include "Chaos/PBDCollisionConstraints.h"
-#include "Chaos/PBDConstraintGraph.h"
 #include "Chaos/PBDRigidClustering.h"
 #include "Chaos/PBDRigidParticles.h"
 #include "Chaos/ParticleHandle.h"
@@ -14,6 +13,7 @@
 #include "Chaos/PBDRigidsSOAs.h"
 #include "Chaos/SpatialAccelerationCollection.h"
 #include "Chaos/PBDRigidsEvolutionFwd.h"
+#include "Chaos/Island/IslandManager.h"
 #include "Chaos/Island/IslandGroupManager.h"
 #include "Chaos/Defines.h"
 #include "Chaos/PendingSpatialData.h"
@@ -464,6 +464,7 @@ public:
 		RemoveParticleFromAccelerationStructure(*Particle);
 		Particles.DisableParticle(Particle);
 		DisableConstraints(Particle);
+		DestroyTransientConstraints(Particle);
 		IslandManager.RemoveParticle(Particle);
 	}
 
@@ -581,6 +582,7 @@ public:
 
 		RemoveParticleFromAccelerationStructure(*Particle);
 		DisconnectConstraints(TSet<FGeometryParticleHandle*>({ Particle }));
+		DestroyTransientConstraints(Particle);
 		IslandManager.RemoveParticle(Particle);
 		Particles.DestroyParticle(Particle);
 	}
@@ -608,7 +610,7 @@ public:
 	{
 		if (ConstraintHandle->IsInConstraintGraph())
 		{
-			IslandManager.RemoveConstraint(ConstraintHandle->GetContainerId(), ConstraintHandle);
+			IslandManager.RemoveConstraint(ConstraintHandle);
 		}
 	}
 
@@ -629,7 +631,7 @@ public:
 	}
 
 	/** 
-	* Disconnect constraints from a set of particles to be destroyed. 
+	* Disconnect constraints (all types except collisions) from a set of particles to be destroyed. 
 	* this will set the constraints to Enabled = false and set their respective bodies handles to nullptr.
 	* Once this is done, the constraints cannot be re-enabled.
 	* @note This only applies to persistent constraints (joints etc), not transient constraints (collisions)
@@ -649,7 +651,7 @@ public:
 	}
 
 	/** 
-	* Disconnect constraints from a particle to be removed (or destroyed)
+	* Disconnect constraints (all types except collisions) from a particle to be removed (or destroyed)
 	* this will set the constraints to Enabled = false, but leave connections to the particles to support
 	* re-enabling at a later time.
 	* @note This only applies to persistent constraints (joints etc), not transient constraints (collisions)
@@ -669,7 +671,7 @@ public:
 	}
 
 	/** 
-	* Enable constraints from the enabled particles; constraints will only become enabled if their particle end points are valid.
+	* Enable constraints (all types except collisions) from the enabled particles; constraints will only become enabled if their particle end points are valid.
 	* @note This only applies to persistent constraints (joints etc), not transient constraints (collisons)
 	*/
 	CHAOS_API void EnableConstraints(FGeometryParticleHandle* ParticleHandle)
@@ -689,8 +691,8 @@ public:
 	CHAOS_API void ResetConstraints()
 	{
 		// Remove all the constraints from the graph
-		GetIslandManager().RemoveConstraints();
-		
+		GetIslandManager().Reset();
+
 		// Clear all particle lists of collisions and constraints
 		// (this could be performed by the constraint containers
 		// but it would be unnecessarily expensive to remove them
@@ -709,7 +711,7 @@ public:
 	}
 
 	/**
-	* Destroy all transient constraints (collisions) onvolving the specified particle.
+	* Destroy all transient constraints (collisions) involving the specified particle.
 	*/
 	virtual void DestroyTransientConstraints(FGeometryParticleHandle* Particle) {}
 
@@ -722,12 +724,14 @@ public:
 	CHAOS_API void SetPerParticlePhysicsMaterial(FGeometryParticleHandle* Particle, TUniquePtr<FChaosPhysicsMaterial> &InMaterial)
 	{
 		Particle->AuxilaryValue(PerParticlePhysicsMaterials) = MoveTemp(InMaterial);
+		IslandManager.UpdateParticleMaterial(Particle);
 	}
 
 	CHAOS_API void SetPhysicsMaterial(FGeometryParticleHandle* Particle, TSerializablePtr<FChaosPhysicsMaterial> InMaterial)
 	{
 		check(!Particle->AuxilaryValue(PerParticlePhysicsMaterials)); //shouldn't be setting non unique material if a unique one already exists
 		Particle->AuxilaryValue(PhysicsMaterials) = InMaterial;
+		IslandManager.UpdateParticleMaterial(Particle);
 	}
 
 	void PrepareTick()
@@ -987,7 +991,7 @@ protected:
 	{
 		// Update the current state of the graph based on existing particles and constraints.
 		// Any new particles (from this tick) should have been added when they were enabled.
-		IslandManager.InitializeGraph(Particles.GetNonDisabledDynamicView());
+		IslandManager.UpdateParticles();
 
 		// Add all constraints to the graph.
 		// NOTE: in PersistentGraph mode, only new constraints need to be added and expired ones should be removed.
@@ -1004,7 +1008,7 @@ protected:
 	void CreateIslands()
 	{
 		// Package the constraints and particles into islands
-		IslandManager.UpdateIslands(Particles, IsResimming());
+		IslandManager.UpdateIslands();
 	}
 
 	void FlushInternalAccelerationQueue();
@@ -1023,7 +1027,6 @@ protected:
 	Private::FPBDIslandGroupManager IslandGroupManager;
 	TArrayCollectionArray<TSerializablePtr<FChaosPhysicsMaterial>> PhysicsMaterials;
 	TArrayCollectionArray<TUniquePtr<FChaosPhysicsMaterial>> PerParticlePhysicsMaterials;
-	TArrayCollectionArray<int32> ParticleDisableCount;
 	TArrayCollectionArray<bool> Collided;
 
 	FPBDRigidsSOAs& Particles;
