Chaos - big island manager refactor

diff --git a/Engine/Source/Runtime/Experimental/Chaos/Public/Chaos/PBDRigidsEvolutionGBF.h b/Engine/Source/Runtime/Experimental/Chaos/Public/Chaos/PBDRigidsEvolutionGBF.h
--- a/Engine/Source/Runtime/Experimental/Chaos/Public/Chaos/PBDRigidsEvolutionGBF.h
+++ b/Engine/Source/Runtime/Experimental/Chaos/Public/Chaos/PBDRigidsEvolutionGBF.h
@@ -160,14 +160,6 @@ namespace Chaos
 		 */
 		void ReloadParticlesCache();
 
-		/**
-		 * Build the list of disables particles and update the sleeping flag on the island
-		 * @param Island Index of the island in which the cache will be used
-		 * @param DisabledParticles List of islands disabled particles
-		 * @param SleepedIslands List of islands sleeping state 
-		 */
-		void BuildDisabledParticles(const int32 Island, TArray<TArray<FPBDRigidParticleHandle*>>& DisabledParticles, TArray<bool>& SleepedIslands);
-
 		void DestroyParticleCollisionsInAllocator(FGeometryParticleHandle* Particle);
 
 		virtual void DestroyTransientConstraints(FGeometryParticleHandle* Particle) override;
diff --git a/Engine/Source/Runtime/Experimental/Chaos/Public/Chaos/ParticleHandle.h b/Engine/Source/Runtime/Experimental/Chaos/Public/Chaos/ParticleHandle.h
--- a/Engine/Source/Runtime/Experimental/Chaos/Public/Chaos/ParticleHandle.h
+++ b/Engine/Source/Runtime/Experimental/Chaos/Public/Chaos/ParticleHandle.h
@@ -415,7 +415,7 @@ protected:
 	TGeometryParticleHandleImp()
 		: TParticleHandleBase<T, d>()
 	{
-		SetConstraintGraphIndex(INDEX_NONE);
+		SetConstraintGraphNode(nullptr);
 	}
 
 	TGeometryParticleHandleImp(TSerializablePtr<TGeometryParticles<T, d>> InParticles, int32 InParticleIdx, int32 InHandleIdx, const FGeometryParticleParameters& Params)
@@ -427,7 +427,13 @@ protected:
 		GeometryParticleDefaultConstruct<T, d>(*this, Params);
 		SetHasBounds(false);
 		SetLightWeightDisabled(false);
-		SetConstraintGraphIndex(INDEX_NONE);
+		SetConstraintGraphNode(nullptr);
+	}
+
+	// For transient handle
+	TGeometryParticleHandleImp(TGeometryParticles<T, d>* InParticles, const int32 InParticleIdx)
+		: TParticleHandleBase<T, d>(InParticles, InParticleIdx)
+	{
 	}
 
 	template <typename TParticlesType, typename TParams>
@@ -446,14 +452,13 @@ public:
 		return TGeometryParticleHandleImp<T,d,bPersistent>::CreateParticleHandleHelper(InParticles, InParticleIdx, InHandleIdx, Params);
 	}
 
-
 	~TGeometryParticleHandleImp()
 	{
-		// If we weren't removed from the graph, invalid pointer dereferencing is possible
-		check(ConstraintGraphIndex() == INDEX_NONE);
-
-		if (bPersistent)
+		if constexpr (bPersistent)
 		{
+			// If we weren't removed from the graph, invalid pointer dereferencing is possible
+			check(GetConstraintGraphNode() == nullptr);
+
 			GeometryParticles->ResetWeakParticleHandle(ParticleIdx);
 			GeometryParticles->DestroyParticle(ParticleIdx);
 			if (static_cast<uint32>(ParticleIdx) < GeometryParticles->Size())
@@ -471,13 +476,13 @@ public:
 					}
 				}
 			}
-		}
 
-		// Zero the handle out to detect dangling handles and associated memory corruptions
-		HandleIdx = INDEX_NONE;
-		GeometryParticles = nullptr;
-		ParticleIdx = INDEX_NONE;
-		Type = EParticleType::Unknown;
+			// Zero the handle out to detect dangling handles and associated memory corruptions
+			HandleIdx = INDEX_NONE;
+			GeometryParticles = nullptr;
+			ParticleIdx = INDEX_NONE;
+			Type = EParticleType::Unknown;
+		}
 	}
 
 	template <typename T2, int d2>
@@ -485,6 +490,14 @@ public:
 	template <typename T2, int d2>
 	friend const TGeometryParticleHandle<T2, d2>* GetHandleHelper(const TTransientGeometryParticleHandle<T2, d2>* Handle);
 
+	friend class TGeometryParticleHandleImp<T, d, false>;
+	friend class TGeometryParticleHandleImp<T, d, true>;
+
+	TGeometryParticleHandleImp<T, d, false> AsTransient()
+	{
+		return TGeometryParticleHandleImp<T, d, false>(GeometryParticles, ParticleIdx);
+	}
+
 	TGeometryParticleHandleImp(const TGeometryParticleHandleImp&) = delete;
 
 	const TVector<T, d>& X() const { return GeometryParticles->X(ParticleIdx); }
@@ -727,21 +740,25 @@ public:
 		return GeometryParticles->ParticleCollisions(ParticleIdx);
 	}
 
-	int32 ConstraintGraphIndex() const
-	{ 
-		return GeometryParticles->ConstraintGraphIndex(ParticleIdx);
+	bool IsInConstraintGraph() const
+	{
+		return (GetConstraintGraphNode() != nullptr);
 	}
-	
-	void SetConstraintGraphIndex(const int32 InGraphIndex)
+
+	Private::FPBDIslandParticle* GetConstraintGraphNode() const
 	{ 
-		GeometryParticles->ConstraintGraphIndex(ParticleIdx) = InGraphIndex;
+		return GeometryParticles->ConstraintGraphNode(ParticleIdx);
 	}
 
-	bool IsInConstraintGraph() const
+	void SetConstraintGraphNode(Private::FPBDIslandParticle* InNode)
 	{
-		return ConstraintGraphIndex() != INDEX_NONE;
+		GeometryParticles->ConstraintGraphNode(ParticleIdx) = InNode;
 	}
 
+	// Deprecated API
+	UE_DEPRECATED(5.3, "Use GetConstraintGraphNode()") int32 ConstraintGraphIndex() const { return INDEX_NONE; }
+	UE_DEPRECATED(5.3, "Use GetConstraintGraphNode()") void SetConstraintGraphIndex(const int32 InGraphIndex) {}
+
 protected:
 
 	friend TGeometryParticleHandles<T, d>;
@@ -883,8 +900,9 @@ protected:
 		SetAngularAcceleration(TVector<T, d>(0));
 		SetObjectStateLowLevel(Params.bStartSleeping ? EObjectStateType::Sleeping : EObjectStateType::Dynamic);
 		SetPreObjectStateLowLevel(ObjectState());
-		SetIslandIndex(INDEX_NONE);
 		SetSleepType(ESleepType::MaterialSleep);
+		SetSleepCounter(0);
+		SetDisableCounter(0);
 		SetInvIConditioning(TVec3<FRealSingle>(1));
 	}
 public:
@@ -1086,10 +1104,6 @@ public:
 	T& MaxAngularSpeedSq() { return PBDRigidParticles->MaxAngularSpeedSq(ParticleIdx); }
 	void SetMaxAngularSpeedSq(const T& InMaxAngularSpeed) { PBDRigidParticles->MaxAngularSpeedSq(ParticleIdx) = InMaxAngularSpeed; }
 
-	int32 IslandIndex() const { return PBDRigidParticles->IslandIndex(ParticleIdx); }
-	int32& IslandIndex() { return PBDRigidParticles->IslandIndex(ParticleIdx); }
-	void SetIslandIndex(const int32 InIslandIndex) { PBDRigidParticles->IslandIndex(ParticleIdx) = InIslandIndex; }
-	
 	EObjectStateType ObjectState() const { return PBDRigidParticles->ObjectState(ParticleIdx); }
 	EObjectStateType PreObjectState() const { return PBDRigidParticles->PreObjectState(ParticleIdx); }
 
@@ -1097,6 +1111,7 @@ public:
 	void SetPreObjectStateLowLevel(EObjectStateType InState) { PBDRigidParticles->PreObjectState(ParticleIdx) = InState; }
 	
 	bool IsSleeping() const { return ObjectState() == EObjectStateType::Sleeping; }
+	bool WasSleeping() const { return PreObjectState() == EObjectStateType::Sleeping; }
 
 	bool Sleeping() const { return PBDRigidParticles->Sleeping(ParticleIdx); }
 	void SetSleeping(bool bSleeping) { PBDRigidParticles->SetSleeping(ParticleIdx, bSleeping); }
@@ -1216,8 +1231,33 @@ public:
 
 	void SetSleepType(ESleepType SleepType){ PBDRigidParticles->SetSleepType(ParticleIdx, SleepType); }
 
+	int8 SleepCounter() const
+	{
+		return PBDRigidParticles->SleepCounter(ParticleIdx);
+	}
+
+	void SetSleepCounter(int8 SleepCounter)
+	{
+		PBDRigidParticles->SleepCounter(ParticleIdx) = SleepCounter;
+	}
+
+	int8 DisableCounter() const
+	{
+		return PBDRigidParticles->DisableCounter(ParticleIdx);
+	}
+
+	void SetDisableCounter(int8 DisableCounter)
+	{
+		PBDRigidParticles->DisableCounter(ParticleIdx) = DisableCounter;
+	}
 
 	static constexpr EParticleType StaticType() { return EParticleType::Rigid; }
+
+	// Deprecated API
+	UE_DEPRECATED(5.3, "No longer supported") int32 IslandIndex() const { return INDEX_NONE; }
+	UE_DEPRECATED(5.3, "No longer supported") int32& IslandIndex() { static int32 Dummy = INDEX_NONE; return Dummy; }
+	UE_DEPRECATED(5.3, "No longer supported") void SetIslandIndex(const int32 InIslandIndex) {}
+
 };
 
 template <typename T, int d, bool bPersistent>
@@ -1489,6 +1529,15 @@ public:
 	bool IsKinematic() const { return (MHandle->ObjectState() == EObjectStateType::Kinematic); }
 	bool IsDynamic() const { return (MHandle->ObjectState() == EObjectStateType::Dynamic) || (MHandle->ObjectState() == EObjectStateType::Sleeping); }
 	bool IsSleeping() const { return (MHandle->ObjectState() == EObjectStateType::Sleeping); }
+	
+	bool WasSleeping() const 
+	{ 
+		if (FPBDRigidParticleHandle* Rigid = MHandle->CastToRigidParticle())
+		{
+			return Rigid->WasSleeping();
+		}
+		return IsSleeping();
+	}
 
 	const FKinematicGeometryParticleHandle* CastToKinematicParticle() const { return MHandle->CastToKinematicParticle(); }
 	FKinematicGeometryParticleHandle* CastToKinematicParticle() { return MHandle->CastToKinematicParticle(); }
@@ -2091,46 +2140,19 @@ public:
 	}
 #endif
 
-	/** Get the island index from the particle handle */
-	FORCEINLINE int32 IslandIndex() const
-	{
-		if (auto RigidHandle = MHandle->CastToRigidParticle())
-		{
-			if (IsDynamic())
-			{
-				return RigidHandle->IslandIndex();
-			}
-		}
-		return INDEX_NONE;
-	}
-
-	/** Set the island index onto the particle handle 
-	 * @param IslandIndex Island Index to be set onto the particle 
-	 */
-	FORCEINLINE void SetIslandIndex( const int32 IslandIndex) 
-	{
-		if (auto RigidHandle = MHandle->CastToRigidParticle())
-		{
-			if (IsDynamic())
-			{
-				RigidHandle->IslandIndex() = IslandIndex;
-			}
-		}
-	}
-
-	FORCEINLINE int32 ConstraintGraphIndex() const
+	FORCEINLINE bool IsInConstraintGraph() const
 	{
-		return MHandle->ConstraintGraphIndex();
+		return MHandle->IsInConstraintGraph();
 	}
 
-	FORCEINLINE void SetConstraintGraphIndex(const int32 InGraphIndex)
+	FORCEINLINE Private::FPBDIslandParticle* GetConstraintGraphNode() const
 	{
-		MHandle->SetConstraintGraphIndex(InGraphIndex);
+		return MHandle->GetConstraintGraphNode();
 	}
 
-	FORCEINLINE bool IsInConstraintGraph() const
+	FORCEINLINE void SetConstraintGraphNode(Private::FPBDIslandParticle* InNode)
 	{
-		return MHandle->IsInConstraintGraph();
+		MHandle->SetConstraintGraphNode(InNode);
 	}
 
 	const FShapesArray& ShapesArray() const { return MHandle->ShapesArray(); }
@@ -2141,6 +2163,12 @@ public:
 		return EParticleType::Unknown;
 	}
 
+	// Deprecated API
+	UE_DEPRECATED(5.3, "No longer supported") int32 IslandIndex() const { return INDEX_NONE; }
+	UE_DEPRECATED(5.3, "No longer supported") void SetIslandIndex(const int32 IslandIndex) {}
+	UE_DEPRECATED(5.3, "No longer supported") int32 ConstraintGraphIndex() const { return INDEX_NONE; }
+	UE_DEPRECATED(5.3, "No longer supported") void SetConstraintGraphIndex(const int32 InGraphIndex) {}
+
 private:
 	FGeometryParticleHandle* MHandle;
 
diff --git a/Engine/Source/Runtime/Experimental/Chaos/Public/Chaos/ParticleHandleFwd.h b/Engine/Source/Runtime/Experimental/Chaos/Public/Chaos/ParticleHandleFwd.h
--- a/Engine/Source/Runtime/Experimental/Chaos/Public/Chaos/ParticleHandleFwd.h
+++ b/Engine/Source/Runtime/Experimental/Chaos/Public/Chaos/ParticleHandleFwd.h
@@ -14,7 +14,6 @@ namespace Chaos
 
 	template <typename T, int d>
 	class TGeometryParticle;
-	//using TGeometryParticleFloat3 = TGeometryParticle<float, 3>;
 
 	template <typename T, int d, bool bProcessing>
 	class TGeometryParticleHandleImp;
@@ -22,9 +21,13 @@ namespace Chaos
 	template <typename T, int d>
 	using TGeometryParticleHandle = TGeometryParticleHandleImp<T, d, true>;
 
+	using FGeometryParticleHandle = TGeometryParticleHandle<FReal, 3>;
+
 	template <typename T, int d>
 	using TTransientGeometryParticleHandle = TGeometryParticleHandleImp<T, d, false>;
 
+	using FTransientGeometryParticleHandle = TTransientGeometryParticleHandle<FReal, 3>;
+
 	// TKinematicGeometryParticle
 
 	template <typename T, int d>
@@ -53,7 +56,6 @@ namespace Chaos
 
 	template <typename T, int d>
 	using TPBDRigidParticleHandle = TPBDRigidParticleHandleImp<T, d, true>;
-	//using TPBDRigidParticleHandleFloat3 = TPBDRigidParticleHandle<float, 3>;
 
 	using FPBDRigidParticleHandle = TPBDRigidParticleHandle<FReal, 3>;
 
@@ -69,7 +71,6 @@ namespace Chaos
 
 	template <typename T, int d>
 	using TPBDRigidClusteredParticleHandle = TPBDRigidClusteredParticleHandleImp<T, d, true>;
-	//using TPBDRigidClusteredParticleHandleFloat3 = TPBDRigidClusteredParticleHandle<float, 3>;
 
 	using FPBDRigidClusteredParticleHandle = TPBDRigidClusteredParticleHandle<FReal, 3>;
 
@@ -87,14 +88,16 @@ namespace Chaos
 
 	template <typename T, int d>
 	using TPBDGeometryCollectionParticleHandle = TPBDGeometryCollectionParticleHandleImp<T, d, true>;
-	//using TPBDGeometryCollectionParticleHandleFloat3 = TPBDGeometryCollectionParticleHandle<float, 3>;
 
 	using FPBDGeometryCollectionParticleHandle = TPBDGeometryCollectionParticleHandle<FReal, 3>;
 
 	template <typename T, int d>
 	using TTransientPBDGeometryCollectionParticleHandle = TPBDGeometryCollectionParticleHandleImp<T, d, false>;
 	
+	// Generic Particle Handle (syntactic sugar wrapping any particle handle type and providing unified API)
+
 	class FGenericParticleHandle;
+
 	class FConstGenericParticleHandle;
 
 	// TParticleIterator
@@ -105,11 +108,14 @@ namespace Chaos
 	template <typename TSOA>
 	class TConstParticleIterator;
 
+	// Game Thread particles
+
 	using FGeometryParticle = TGeometryParticle<FReal, 3>;
-	using FGeometryParticleHandle = TGeometryParticleHandle<FReal, 3>;
 	using FKinematicGeometryParticle = TKinematicGeometryParticle<FReal, 3>;
 	using FPBDRigidParticle = TPBDRigidParticle<FReal, 3>;
 
+	// Thread-templated particles
+
 	template<EThreadContext Id>
 	using TThreadParticle = std::conditional_t<Id == EThreadContext::External, FGeometryParticle, FGeometryParticleHandle>;
 
diff --git a/Engine/Source/Runtime/Experimental/Chaos/Public/Chaos/PhysicsMaterialUtilities.h b/Engine/Source/Runtime/Experimental/Chaos/Public/Chaos/PhysicsMaterialUtilities.h
--- a/Engine/Source/Runtime/Experimental/Chaos/Public/Chaos/PhysicsMaterialUtilities.h
+++ b/Engine/Source/Runtime/Experimental/Chaos/Public/Chaos/PhysicsMaterialUtilities.h
@@ -20,15 +20,15 @@ namespace Chaos
 			const FGeometryParticleHandle* Particle,
 			const FShapeInstance* Shape,
 			const int32 FaceIndex,
-			const TArrayCollectionArray<TSerializablePtr<FChaosPhysicsMaterial>>& PhysicsMaterials,
-			const TArrayCollectionArray<TUniquePtr<FChaosPhysicsMaterial>>& PerParticlePhysicsMaterials,
+			const TArrayCollectionArray<TSerializablePtr<FChaosPhysicsMaterial>>* PhysicsMaterials,
+			const TArrayCollectionArray<TUniquePtr<FChaosPhysicsMaterial>>* PerParticlePhysicsMaterials,
 			const THandleArray<FChaosPhysicsMaterial>* const SimMaterials);
 
 		// Get the first physical material on the particle (a multi-shape particle may have many)
 		extern CHAOS_API const FChaosPhysicsMaterial* GetFirstPhysicsMaterial(
 			const FGeometryParticleHandle* Particle,
-			const TArrayCollectionArray<TSerializablePtr<FChaosPhysicsMaterial>>& PhysicsMaterials,
-			const TArrayCollectionArray<TUniquePtr<FChaosPhysicsMaterial>>& PerParticlePhysicsMaterials,
+			const TArrayCollectionArray<TSerializablePtr<FChaosPhysicsMaterial>>* PhysicsMaterials,
+			const TArrayCollectionArray<TUniquePtr<FChaosPhysicsMaterial>>* PerParticlePhysicsMaterials,
 			const THandleArray<FChaosPhysicsMaterial>* const SimMaterials);
 	}
 }
diff --git a/Engine/Source/Runtime/Experimental/Chaos/Public/Chaos/RigidParticles.h b/Engine/Source/Runtime/Experimental/Chaos/Public/Chaos/RigidParticles.h
--- a/Engine/Source/Runtime/Experimental/Chaos/Public/Chaos/RigidParticles.h
+++ b/Engine/Source/Runtime/Experimental/Chaos/Public/Chaos/RigidParticles.h
@@ -108,8 +108,9 @@ public:
 		, MaxLinearSpeedsSq(MoveTemp(Other.MaxLinearSpeedsSq))
 		, MaxAngularSpeedsSq(MoveTemp(Other.MaxAngularSpeedsSq))
 		, MCollisionParticles(MoveTemp(Other.MCollisionParticles))
-		, MIslandIndex(MoveTemp(Other.MIslandIndex))
 		, MSleepType(MoveTemp(Other.MSleepType))
+		, MSleepCounter(MoveTemp(Other.MSleepCounter))
+		, MDisableCounter(MoveTemp(Other.MDisableCounter))
 	{
 		RegisterArrays();
 	}
@@ -136,8 +137,9 @@ public:
 		TArrayCollection::AddArray(&MaxLinearSpeedsSq);
 		TArrayCollection::AddArray(&MaxAngularSpeedsSq);
 		TArrayCollection::AddArray(&MCollisionParticles);
-		TArrayCollection::AddArray(&MIslandIndex);
 		TArrayCollection::AddArray(&MSleepType);
+		TArrayCollection::AddArray(&MSleepCounter);
+		TArrayCollection::AddArray(&MDisableCounter);
 
 	}
 
@@ -229,6 +231,12 @@ public:
 	FORCEINLINE ESleepType SleepType(const int32 Index) const { return MSleepType[Index]; }
 	FORCEINLINE ESleepType& SleepType(const int32 Index) { return MSleepType[Index]; }
 
+	FORCEINLINE int8 SleepCounter(const int32 Index) const { return MSleepCounter[Index]; }
+	FORCEINLINE int8& SleepCounter(const int32 Index) { return MSleepCounter[Index]; }
+
+	FORCEINLINE int8 DisableCounter(const int32 Index) const { return MDisableCounter[Index]; }
+	FORCEINLINE int8& DisableCounter(const int32 Index) { return MDisableCounter[Index]; }
+
 	// @todo(chaos): This data should be marshalled via the proxies like everything else. There is probably a particle recycling bug here.
 	FORCEINLINE TArray<TSleepData<T, d>>& GetSleepData() { return MSleepData; }
 	FORCEINLINE	void AddSleepData(TGeometryParticleHandle<T, d>* Particle, bool Sleeping)
@@ -261,16 +269,13 @@ public:
 
 	FORCEINLINE const bool HasInfiniteMass(const int32 Index) const { return MInvM[Index] == (T)0; }
 
-	FORCEINLINE const int32 IslandIndex(const int32 Index) const { return MIslandIndex[Index]; }
-	FORCEINLINE int32& IslandIndex(const int32 Index) { return MIslandIndex[Index]; }
-
 	FORCEINLINE FString ToString(int32 Index) const
 	{
 		FString BaseString = TKinematicGeometryParticles<T, d>::ToString(Index);
-		return FString::Printf(TEXT("%s, MAcceleration:%s, MAngularAcceleration:%s, MLinearImpulseVelocity:%s, MAngularImpulseVelocity:%s, MI:%s, MInvI:%s, MM:%f, MInvM:%f, MCenterOfMass:%s, MRotationOfMass:%s, MCollisionParticles(num):%d, MCollisionGroup:%d, MDisabled:%d, MSleeping:%d, MIslandIndex:%d"),
+		return FString::Printf(TEXT("%s, MAcceleration:%s, MAngularAcceleration:%s, MLinearImpulseVelocity:%s, MAngularImpulseVelocity:%s, MI:%s, MInvI:%s, MM:%f, MInvM:%f, MCenterOfMass:%s, MRotationOfMass:%s, MCollisionParticles(num):%d, MCollisionGroup:%d, MDisabled:%d, MSleeping:%d"),
 			*BaseString, *Acceleration(Index).ToString(), *AngularAcceleration(Index).ToString(), *LinearImpulseVelocity(Index).ToString(), *AngularImpulseVelocity(Index).ToString(),
 			*I(Index).ToString(), *InvI(Index).ToString(), M(Index), InvM(Index), *CenterOfMass(Index).ToString(), *RotationOfMass(Index).ToString(), CollisionParticlesSize(Index),
-			CollisionGroup(Index), Disabled(Index), Sleeping(Index), IslandIndex(Index));
+			CollisionGroup(Index), Disabled(Index), Sleeping(Index));
 	}
 
 	CHAOS_API virtual void Serialize(FChaosArchive& Ar) override
@@ -279,12 +284,17 @@ public:
 
 		// To avoid bumping file version, serialize to/from previous structures
 		// If we aren't loading (i.e., we are saving or copying) CoreData will be valid, so copy that to the legacy structure
+		// Also, Particles do not know their island index and it should not be serialized (but it used to be)
+		// @todo(chaos): I think its time to bump the version number and clean up particle serialization!
 		FLegacyData LegacyData;
+		TArrayCollectionArray<int32> LegacyIslandIndex;
+
 		if (!Ar.IsLoading())
 		{
 			LegacyData.CopyFromCoreData(CoreData);
+			LegacyIslandIndex.SetNumZeroed(MCollisionParticles.Num());
 		}
-		
+
 		Ar.UsingCustomVersion(FExternalPhysicsCustomObjectVersion::GUID);
 		if(Ar.CustomVer(FExternalPhysicsCustomObjectVersion::GUID) >= FExternalPhysicsCustomObjectVersion::KinematicCentersOfMass)
 		{
@@ -328,8 +338,8 @@ public:
 			Ar << MLinearEtherDrag << MAngularEtherDrag;
 		}
 
-		Ar << MCollisionParticles << LegacyData.MCollisionGroup << MIslandIndex << LegacyData.MDisabled << LegacyData.MObjectState << MSleepType;
-		// @todo(chaos): what about ControlFlags, TransientFlags, PreObjectState, ..?
+		Ar << MCollisionParticles << LegacyData.MCollisionGroup << LegacyIslandIndex << LegacyData.MDisabled << LegacyData.MObjectState << MSleepType;
+		// @todo(chaos): what about ControlFlags, TransientFlags, PreObjectState, SleepCounter, ..?
 
 		// If we loaded into the legacy structure, copy to CoreData
 		if (Ar.IsLoading())
@@ -338,6 +348,10 @@ public:
 		}
 	}
 
+	// Deprecated API
+	UE_DEPRECATED(5.3, "No longer supported") const int32 IslandIndex(const int32 Index) const { return INDEX_NONE; }
+	UE_DEPRECATED(5.3, "No longer supported") int32& IslandIndex(const int32 Index) { static int32 Dummy = INDEX_NONE; return Dummy; }
+
 private:
 	// Used during serialization to avoid bumping the file version as we switch to aggregated strunctures like FRigidParticleCoreData.
 	// Note: Only serialized data is needed here and not all data is serialized so some elements in the aggregates are not represented.
@@ -394,8 +408,9 @@ private:
 	TArrayCollectionArray<T> MaxLinearSpeedsSq;
 	TArrayCollectionArray<T> MaxAngularSpeedsSq;
 	TArrayCollectionArray<TUniquePtr<TBVHParticles<T, d>>> MCollisionParticles;
-	TArrayCollectionArray<int32> MIslandIndex;
 	TArrayCollectionArray<ESleepType> MSleepType;
+	TArrayCollectionArray<int8> MSleepCounter;
+	TArrayCollectionArray<int8> MDisableCounter;
 
 	TArray<TSleepData<T, d>> MSleepData;
 	FRWLock SleepDataLock;
diff --git a/Engine/Source/Runtime/Experimental/Chaos/Public/ChaosStats.h b/Engine/Source/Runtime/Experimental/Chaos/Public/ChaosStats.h
--- a/Engine/Source/Runtime/Experimental/Chaos/Public/ChaosStats.h
+++ b/Engine/Source/Runtime/Experimental/Chaos/Public/ChaosStats.h
@@ -18,6 +18,7 @@ DECLARE_STATS_GROUP(TEXT("ChaosCounters"), STATGROUP_ChaosCounters, STATCAT_Adva
 DECLARE_STATS_GROUP(TEXT("ChaosIterations"), STATGROUP_ChaosIterations, STATCAT_Advanced);
 DECLARE_STATS_GROUP(TEXT("ChaosCollisionCounters"), STATGROUP_ChaosCollisionCounters, STATCAT_Advanced);
 DECLARE_STATS_GROUP(TEXT("ChaosConstraintDetails"), STATGROUP_ChaosConstraintDetails, STATCAT_Advanced);
+DECLARE_STATS_GROUP(TEXT("ChaosIslands"), STATGROUP_ChaosIslands, STATCAT_Advanced);
 
 DECLARE_CYCLE_STAT_EXTERN(TEXT("Physics Tick"), STAT_ChaosTick, STATGROUP_Chaos, CHAOS_API);
 DECLARE_CYCLE_STAT_EXTERN(TEXT("Physics Advance"), STAT_PhysicsAdvance, STATGROUP_Chaos, CHAOS_API);
diff --git a/Engine/Source/Runtime/Experimental/ChaosCore/Public/Chaos/ObjectPool.h b/Engine/Source/Runtime/Experimental/ChaosCore/Public/Chaos/ObjectPool.h
--- a/Engine/Source/Runtime/Experimental/ChaosCore/Public/Chaos/ObjectPool.h
+++ b/Engine/Source/Runtime/Experimental/ChaosCore/Public/Chaos/ObjectPool.h
@@ -219,7 +219,7 @@ namespace Chaos
 		{
 			const int32 ItemsSize = NumItems * PaddedItemSize;
 			const int32 BlockSize = NumPerBlock * PaddedItemSize;
-			const int32 NumBlocks = FMath::CeilToInt(static_cast<float>(ItemsSize) / BlockSize);
+			const int32 NumBlocks = int32(FMath::CeilToInt(static_cast<double>(ItemsSize) / BlockSize));
 			ReserveBlocks(NumBlocks);
 		}
 
