Add GetV, GetW, and GetRigidParticle to physics object interface

diff --git a/Engine/Source/Runtime/Experimental/Chaos/Private/Chaos/PhysicsObjectInterface.cpp b/Engine/Source/Runtime/Experimental/Chaos/Private/Chaos/PhysicsObjectInterface.cpp
--- a/Engine/Source/Runtime/Experimental/Chaos/Private/Chaos/PhysicsObjectInterface.cpp
+++ b/Engine/Source/Runtime/Experimental/Chaos/Private/Chaos/PhysicsObjectInterface.cpp
@@ -233,6 +233,44 @@ namespace Chaos
 		return FQuat::Identity;
 	}
 
+	template<EThreadContext Id>
+	FVector FReadPhysicsObjectInterface<Id>::GetV(const FConstPhysicsObjectHandle Object)
+	{
+		if (!Object)
+		{
+			return FVector::Zero();
+		}
+
+		if (TThreadParticle<Id>* Particle = Object->GetParticle<Id>())
+		{
+			if (Chaos::TThreadRigidParticle<Id>* Rigid = Particle->CastToRigidParticle())
+			{
+				return Rigid->V();
+			}
+		}
+
+		return FVector::Zero();
+	}
+
+	template<EThreadContext Id>
+	FVector FReadPhysicsObjectInterface<Id>::GetW(const FConstPhysicsObjectHandle Object)
+	{
+		if (!Object)
+		{
+			return FVector::Zero();
+		}
+
+		if (TThreadParticle<Id>* Particle = Object->GetParticle<Id>())
+		{
+			if (Chaos::TThreadRigidParticle<Id>* Rigid = Particle->CastToRigidParticle())
+			{
+				return Rigid->W();
+			}
+		}
+
+		return FVector::Zero();
+	}
+
 	template<EThreadContext Id>
 	FSpatialAccelerationIdx FReadPhysicsObjectInterface<Id>::GetSpatialIndex(const FConstPhysicsObjectHandle Object)
 	{
@@ -259,6 +297,23 @@ namespace Chaos
 		return Handle->GetParticle<Id>();
 	}
 
+	template<EThreadContext Id>
+	TThreadRigidParticle<Id>* FReadPhysicsObjectInterface<Id>::GetRigidParticle(const FConstPhysicsObjectHandle Handle)
+	{
+		if (Handle)
+		{
+			if (TThreadParticle<Id>* Particle = Handle->GetParticle<Id>())
+			{
+				if (TThreadRigidParticle<Id>* RigidParticle = Particle->CastToRigidParticle())
+				{
+					return RigidParticle;
+				}
+			}
+		}
+
+		return nullptr;
+	}
+
 	template<EThreadContext Id>
 	TArray<TThreadParticle<Id>*> FReadPhysicsObjectInterface<Id>::GetAllParticles(TArrayView<const FConstPhysicsObjectHandle> InObjects)
 	{
diff --git a/Engine/Source/Runtime/Experimental/Chaos/Public/Chaos/PhysicsObjectInterface.h b/Engine/Source/Runtime/Experimental/Chaos/Public/Chaos/PhysicsObjectInterface.h
--- a/Engine/Source/Runtime/Experimental/Chaos/Public/Chaos/PhysicsObjectInterface.h
+++ b/Engine/Source/Runtime/Experimental/Chaos/Public/Chaos/PhysicsObjectInterface.h
@@ -67,9 +67,12 @@ namespace Chaos
 		FVector GetCoM(const FConstPhysicsObjectHandle Object);
 		FVector GetWorldCoM(const FConstPhysicsObjectHandle Object);
 		FQuat GetR(const FConstPhysicsObjectHandle Object);
+		FVector GetV(const FConstPhysicsObjectHandle Object);
+		FVector GetW(const FConstPhysicsObjectHandle Object);
 		FSpatialAccelerationIdx GetSpatialIndex(const FConstPhysicsObjectHandle Object);
 
 		TThreadParticle<Id>* GetParticle(const FConstPhysicsObjectHandle Object);
+		TThreadRigidParticle<Id>* GetRigidParticle(const FConstPhysicsObjectHandle Object);
 		TArray<TThreadParticle<Id>*> GetAllParticles(TArrayView<const FConstPhysicsObjectHandle> InObjects);
 		TArray<TThreadRigidParticle<Id>*> GetAllRigidParticles(TArrayView<const FConstPhysicsObjectHandle> InObjects);
 
