I/O - removed HTTP client due to platform specific includes

diff --git a/Engine/Source/Runtime/Experimental/IoStoreOnDemand/Private/CoreHttp/Client.h b/Engine/Source/Runtime/Experimental/IoStoreOnDemand/Private/CoreHttp/Client.h
deleted file mode 100644
diff --git a/Engine/Source/Runtime/Experimental/IoStoreOnDemand/Private/CoreHttp/Http.cpp b/Engine/Source/Runtime/Experimental/IoStoreOnDemand/Private/CoreHttp/Http.cpp
deleted file mode 100644
diff --git a/Engine/Source/Runtime/Experimental/IoStoreOnDemand/Private/OnDemandIoDispatcherBackend.cpp b/Engine/Source/Runtime/Experimental/IoStoreOnDemand/Private/OnDemandIoDispatcherBackend.cpp
--- a/Engine/Source/Runtime/Experimental/IoStoreOnDemand/Private/OnDemandIoDispatcherBackend.cpp
+++ b/Engine/Source/Runtime/Experimental/IoStoreOnDemand/Private/OnDemandIoDispatcherBackend.cpp
@@ -2,7 +2,6 @@
 
 #include "OnDemandIoDispatcherBackend.h"
 #include "Containers/StringView.h"
-#include "CoreHttp/Client.h"
 #include "EncryptionKeyManager.h"
 #include "HAL/Event.h"
 #include "HAL/Platform.h"
@@ -25,7 +24,10 @@
 #include "Tasks/Pipe.h"
 #include <atomic>
 
-#define WITH_HTTP_CLIENT 1
+#define WITH_HTTP_CLIENT 0
+#if WITH_HTTP_CLIENT
+#include "CoreHttp/Client.h"
+#endif
 
 namespace UE::IO::Private
 {
@@ -131,18 +133,19 @@ FIoHash GetChunkKey(const FIoHash& ChunkHash, const FIoOffsetAndLength& Range)
 class FHttpClient
 {
 public:
-	FHttpClient(UE::HTTP::FEventLoop& EventLoop);
+	FHttpClient();
 	void Get(FAnsiStringView Url, FIoReadCallback&& Callback, const TCHAR* DebugName = nullptr);
 	void Get(FAnsiStringView Url, const FIoOffsetAndLength& Range, FIoReadCallback&& Callback, const TCHAR* DebugName = nullptr);
 	bool Tick();
 
 private:
+#if WITH_HTTP_CLIENT
 	void Issue(UE::HTTP::FRequest&& Request, FIoReadCallback&& Callback, FAnsiStringView Url = FAnsiStringView(), const TCHAR* DebugName = nullptr);
-	UE::HTTP::FEventLoop& EventLoop;
+	UE::HTTP::FEventLoop EventLoop;
+#endif
 };
 
-FHttpClient::FHttpClient(UE::HTTP::FEventLoop& Loop)
-	: EventLoop(Loop)
+FHttpClient::FHttpClient()
 {
 }
 
@@ -165,9 +168,9 @@ void FHttpClient::Get(FAnsiStringView Url, FIoReadCallback&& Callback, const TCH
 #endif
 }
 
+#if WITH_HTTP_CLIENT
 void FHttpClient::Issue(UE::HTTP::FRequest&& Request, FIoReadCallback&& Callback, FAnsiStringView DebugUrl, const TCHAR* DebugName)
 {
-#if WITH_HTTP_CLIENT
 	using namespace UE::HTTP;
 	
 	FString Debug = DebugName ? FString(DebugName) : FString();
@@ -227,8 +230,8 @@ void FHttpClient::Issue(UE::HTTP::FRequest&& Request, FIoReadCallback&& Callback
 		};
 
 	EventLoop.Send(MoveTemp(Request), MoveTemp(Sink));
-#endif
 }
+#endif
 
 bool FHttpClient::Tick()
 {
@@ -586,8 +589,7 @@ class FOnDemandIoBackend final
 
 		FString DebugName;
 		UE::Tasks::FPipe Pipe;
-		UE::HTTP::FEventLoop HttpLoop;
-		FHttpClient Client{HttpLoop};
+		FHttpClient Client;
 	};
 
 	struct FBackendData
@@ -721,8 +723,10 @@ private:
 	TUniquePtr<FOnDemandIoStore> IoStore;
 	TArray<TUniquePtr<FHttpPipe>> HttpPipes;
 	FChunkRequests ChunkRequests;
+#if WITH_HTTP_CLIENT
 	UE::HTTP::FEventLoop HttpEventLoop;
 	TUniquePtr<FHttpClient> HttpClient;
+#endif
 	FIoRequestQueue CompletedRequests;
 	FOnDemandEndpoint CurrentEndpoint;
 	std::atomic_bool bStopRequested{false};
@@ -734,7 +738,9 @@ FOnDemandIoBackend::FOnDemandIoBackend(TSharedPtr<IIoCache> InCache)
 	: Cache(InCache)
 {
 	IoStore = MakeUnique<FOnDemandIoStore>();
+#if WITH_HTTP_CLIENT
 	HttpClient = MakeUnique<FHttpClient>(HttpEventLoop);
+#endif
 
 	for (uint32 Idx = 0; Idx < HttpWorkerCount; ++Idx)
 	{
@@ -973,11 +979,11 @@ void FOnDemandIoBackend::Mount(const FOnDemandEndpoint& Endpoint)
 	if (EnumHasAnyFlags(Endpoint.EndpointType, EOnDemandEndpointType::CDN))
 	{
 		UE_LOG(LogIoStoreOnDemand, Log, TEXT("Mounting CDN endpoint, Url='%s', Toc='%s'"), *Endpoint.Url, *Endpoint.TocPath);
-		using namespace UE::HTTP;
 
 		TAnsiStringBuilder<256> Url;
 		Url << Endpoint.Url << "/" << Endpoint.TocPath;
 
+#if WITH_HTTP_CLIENT
 		HttpClient->Get(Url.ToView(), [this, Url = Endpoint.Url](TIoStatusOr<FIoBuffer> Status)
 		{
 			if (Status.IsOk())
@@ -1000,6 +1006,7 @@ void FOnDemandIoBackend::Mount(const FOnDemandEndpoint& Endpoint)
 		});
 
 		while (HttpClient->Tick());
+#endif
 	}
 	else
 	{
