Filter out kinematics when buffering physics results in PT to save GT syncbodies time.

diff --git a/Engine/Source/Runtime/Experimental/Chaos/Private/PBDRigidsSolver.cpp b/Engine/Source/Runtime/Experimental/Chaos/Private/PBDRigidsSolver.cpp
--- a/Engine/Source/Runtime/Experimental/Chaos/Private/PBDRigidsSolver.cpp
+++ b/Engine/Source/Runtime/Experimental/Chaos/Private/PBDRigidsSolver.cpp
@@ -1795,7 +1795,10 @@ namespace Chaos
 						{
 							if(!bIsResim || DirtyParticle.SyncState() == ESyncState::HardDesync)
 							{
-								ActiveRigid.AddUnique((FSingleParticlePhysicsProxy*)Proxy);
+								if (!(Chaos::SyncKinematicOnGameThread == 0 && DirtyParticle.ObjectState() == EObjectStateType::Kinematic))
+								{
+									ActiveRigid.AddUnique((FSingleParticlePhysicsProxy*)Proxy);
+								}
 							}
 							break;
 						}
