Geometry collection optimization : Change GlobalMatrices array of FMatrix to ComponentSpaceTransforms array of FTransform
- This reduces the number of back and forth conversion between FTransform and FMatrix

diff --git a/Engine/Plugins/Experimental/ChaosNiagara/Source/ChaosNiagara/Private/NiagaraDataInterfaceChaosDestruction.cpp b/Engine/Plugins/Experimental/ChaosNiagara/Source/ChaosNiagara/Private/NiagaraDataInterfaceChaosDestruction.cpp
--- a/Engine/Plugins/Experimental/ChaosNiagara/Source/ChaosNiagara/Private/NiagaraDataInterfaceChaosDestruction.cpp
+++ b/Engine/Plugins/Experimental/ChaosNiagara/Source/ChaosNiagara/Private/NiagaraDataInterfaceChaosDestruction.cpp
@@ -2052,7 +2052,7 @@ void UNiagaraDataInterfaceChaosDestruction::HandleTrailingEvents(const Chaos::FT
 						}
 					}
 
-					const TArray<FMatrix>& GlobalMatrices = GeometryCollectionComponent->GetGlobalMatrices();
+					const TArray<FTransform>& ComponentSpaceTransforms = GeometryCollectionComponent->GetComponentSpaceTransforms();
 					const FTransform ActorTransform = GeometryCollectionComponent->GetComponentToWorld();
 
 					const FGeometryDynamicCollection* DynamicCollection = GeometryCollectionComponent->GetDynamicCollection();
@@ -2085,7 +2085,7 @@ void UNiagaraDataInterfaceChaosDestruction::HandleTrailingEvents(const Chaos::FT
 									{
 										Chaos::FTrailingDataExt TrailingData;
 
-										FTransform CurrTransform = FTransform(CollectionMassToLocal[Idx].ToMatrixWithScale() * GlobalMatrices[Idx] * ActorTransform.ToMatrixWithScale());
+										const FTransform CurrTransform = CollectionMassToLocal[Idx] * ComponentSpaceTransforms[Idx] * ActorTransform;
 										TrailingData.Location = CurrTransform.GetTranslation();
 
 										TrailingData.Velocity = (*LinearVelocity)[Idx];
diff --git a/Engine/Plugins/Experimental/ChaosNiagara/Source/ChaosNiagara/Private/NiagaraDataInterfaceGeometryCollection.cpp b/Engine/Plugins/Experimental/ChaosNiagara/Source/ChaosNiagara/Private/NiagaraDataInterfaceGeometryCollection.cpp
--- a/Engine/Plugins/Experimental/ChaosNiagara/Source/ChaosNiagara/Private/NiagaraDataInterfaceGeometryCollection.cpp
+++ b/Engine/Plugins/Experimental/ChaosNiagara/Source/ChaosNiagara/Private/NiagaraDataInterfaceGeometryCollection.cpp
@@ -178,7 +178,7 @@ void FNDIGeometryCollectionData::Update(UNiagaraDataInterfaceGeometryCollection*
 			const TManagedArray<FBox>& BoundingBoxes = Collection->BoundingBox;
 			const TManagedArray<int32>& TransformIndexArray = Collection->TransformIndex;
 			const TManagedArray<FTransform>& Transforms = Interface->GeometryCollectionActor->GetGeometryCollectionComponent()->GetTransformArray();
-			const TArray<FMatrix>& GlobalMatrices = Interface->GeometryCollectionActor->GetGeometryCollectionComponent()->GetGlobalMatrices();
+			const TArray<FTransform>& ComponentSpaceTransforms = Interface->GeometryCollectionActor->GetGeometryCollectionComponent()->GetComponentSpaceTransforms();
 			const TManagedArray<TSet<int32>>& Children = Collection->Children;
 			
 			int NumPieces = 0;
@@ -194,7 +194,7 @@ void FNDIGeometryCollectionData::Update(UNiagaraDataInterfaceGeometryCollection*
 				}
 			}
 
-			if (GlobalMatrices.Num() != Transforms.Num())
+			if (ComponentSpaceTransforms.Num() != Transforms.Num())
 			{
 				return;
 			}
@@ -232,13 +232,11 @@ void FNDIGeometryCollectionData::Update(UNiagaraDataInterfaceGeometryCollection*
 					FVector LocalTranslation = (CurrBox.Max + CurrBox.Min) * .5;
 					FTransform LocalOffset(LocalTranslation);
 
+					const FTransform3f CurrTransform(LocalOffset * ComponentSpaceTransforms[CurrTransformIndex] * ActorTransform);
+					CurrTransform.ToMatrixWithScale().To3x4MatrixTranspose(&AssetArrays->WorldTransformBuffer[TransformIndex].X);
 
-
-					FMatrix44f CurrTransform = FMatrix44f(LocalOffset.ToMatrixWithScale() * GlobalMatrices[CurrTransformIndex] * ActorTransform.ToMatrixWithScale());
-					CurrTransform.To3x4MatrixTranspose(&AssetArrays->WorldTransformBuffer[TransformIndex].X);
-
-					FMatrix44f CurrInverse = CurrTransform.Inverse();
-					CurrInverse.To3x4MatrixTranspose(&AssetArrays->WorldInverseTransformBuffer[TransformIndex].X);
+					const FTransform3f CurrInverse = CurrTransform.Inverse();
+					CurrInverse.ToMatrixWithScale().To3x4MatrixTranspose(&AssetArrays->WorldInverseTransformBuffer[TransformIndex].X);
 
 					PieceIndex++;
 				}
diff --git a/Engine/Plugins/MeshPainting/Source/MeshPaintingToolset/Private/MeshPaintGeometryCollectionAdapter.cpp b/Engine/Plugins/MeshPainting/Source/MeshPaintingToolset/Private/MeshPaintGeometryCollectionAdapter.cpp
--- a/Engine/Plugins/MeshPainting/Source/MeshPaintingToolset/Private/MeshPaintGeometryCollectionAdapter.cpp
+++ b/Engine/Plugins/MeshPainting/Source/MeshPaintingToolset/Private/MeshPaintGeometryCollectionAdapter.cpp
@@ -86,15 +86,15 @@ bool FMeshPaintGeometryCollectionComponentAdapter::InitializeVertexData()
 	const TManagedArray<int32>& BoneMap = Collection->BoneMap;
 	const TManagedArray<int32>& SimTypes = Collection->SimulationType;
 	const TManagedArray<bool>& Visible = Collection->Visible;
-	const TArray<FMatrix>& Matrices = GeometryCollectionComponent->GetGlobalMatrices();
+	const TArray<FTransform>& ComponentSpaceTransforms = GeometryCollectionComponent->GetComponentSpaceTransforms();
 	const int32 NumVertices = Vertices.Num();
 	MeshVertices.Reset();
 	MeshVertices.AddZeroed(NumVertices);
 	for (int32 Index = 0; Index < NumVertices; ++Index)
 	{
-		int32 BoneIndex = BoneMap[Index];
-		const FMatrix& Matrix = Matrices[BoneIndex];
-		const FVector Position = Matrix.TransformPosition((FVector)Vertices[Index]);
+		const int32 BoneIndex = BoneMap[Index];
+		const FTransform& Transform = ComponentSpaceTransforms[BoneIndex];
+		const FVector Position = Transform.TransformPosition((FVector)Vertices[Index]);
 		MeshVertices[Index] = Position;
 	}
 
diff --git a/Engine/Source/Runtime/Experimental/Chaos/Public/GeometryCollectionProxyData.h b/Engine/Source/Runtime/Experimental/Chaos/Public/GeometryCollectionProxyData.h
--- a/Engine/Source/Runtime/Experimental/Chaos/Public/GeometryCollectionProxyData.h
+++ b/Engine/Source/Runtime/Experimental/Chaos/Public/GeometryCollectionProxyData.h
@@ -301,7 +301,7 @@ public:
 
 	Chaos::FReal SolverDt;
 	TArray<FState> States;
-	TArray<FMatrix> GlobalTransforms;
+	TArray<FTransform> GlobalTransforms;
 	TArray<Chaos::FVec3> ParticleXs;
 	TArray<Chaos::FRotation3> ParticleRs;
 	TArray<Chaos::FVec3> ParticleVs;
diff --git a/Engine/Source/Runtime/Experimental/GeometryCollectionEngine/Private/GeometryCollection/GeometryCollectionComponent.cpp b/Engine/Source/Runtime/Experimental/GeometryCollectionEngine/Private/GeometryCollection/GeometryCollectionComponent.cpp
--- a/Engine/Source/Runtime/Experimental/GeometryCollectionEngine/Private/GeometryCollection/GeometryCollectionComponent.cpp
+++ b/Engine/Source/Runtime/Experimental/GeometryCollectionEngine/Private/GeometryCollection/GeometryCollectionComponent.cpp
@@ -470,10 +470,11 @@ void UGeometryCollectionComponent::GetLifetimeReplicatedProps(TArray<FLifetimePr
 
 namespace
 {
-	void UpdateGlobalMatricesWithExplodedVectors(TArray<FMatrix>& GlobalMatricesIn, FGeometryCollection& GeometryCollection)
+	template <typename TTransformType>
+	void UpdateGlobalMatricesWithExplodedVectors(TArray<TTransformType>& GlobalMatricesInOut, FGeometryCollection& GeometryCollection)
 	{
-		int32 NumMatrices = GlobalMatricesIn.Num();
-		if (GlobalMatricesIn.Num() > 0)
+		int32 NumMatrices = GlobalMatricesInOut.Num();
+		if (GlobalMatricesInOut.Num() > 0)
 		{
 			if (GeometryCollection.HasAttribute("ExplodedVector", FGeometryCollection::TransformGroup))
 			{
@@ -483,9 +484,20 @@ namespace
 				{
 					return;
 				}
-				for (int32 TransformIndex = 0; TransformIndex < ExplodedVectors.Num(); ++TransformIndex)
+				constexpr bool bIsMatrixType = std::is_same<TTransformType, FMatrix>::value;
+				if constexpr (bIsMatrixType)
 				{
-					GlobalMatricesIn[TransformIndex] = GlobalMatricesIn[TransformIndex].ConcatTranslation((FVector)ExplodedVectors[TransformIndex]);
+					for (int32 TransformIndex = 0; TransformIndex < ExplodedVectors.Num(); ++TransformIndex)
+					{
+						GlobalMatricesInOut[TransformIndex] = GlobalMatricesInOut[TransformIndex].ConcatTranslation((FVector)ExplodedVectors[TransformIndex]);
+					}
+				}
+				else
+				{
+					for (int32 TransformIndex = 0; TransformIndex < ExplodedVectors.Num(); ++TransformIndex)
+					{
+						GlobalMatricesInOut[TransformIndex].AddToTranslation((FVector)ExplodedVectors[TransformIndex]);
+					}
 				}
 			}
 		}
@@ -494,42 +506,48 @@ namespace
 	/**
 	* compute the bounding box from the bounding boxes stored in the geometry group
 	*/
+	template <typename TTransformType>
 	inline FBox ComputeBoundsFromGeometryBoundingBoxes(
 		const TManagedArray<int32>& TransformToGeometryIndex,
 		const TManagedArray<int32>& TransformIndices,
 		const TManagedArray<FBox>& BoundingBoxes,
-		const TArray<FMatrix>& GlobalMatrices,
-		const FMatrix& LocalToWorldWithScale)
+		const TArray<TTransformType>& GlobalMatrices,
+		const TTransformType& LocalToWorldWithScale)
 	{
 		FBox BoundingBox(ForceInit);
-		if (bChaos_BoxCalcBounds_ISPC_Enabled && !bGeometryCollectionSingleThreadedBoundsCalculation)
+		// todo(chaos ) implement ISPC function using FTransform 
+		constexpr bool bIsMatrixType = std::is_same<TTransformType, FMatrix>::value;
+		if constexpr (bIsMatrixType)
 		{
-			ensure(BoundingBoxes.Num() > 0);
-			ensure(TransformIndices.Num() == BoundingBoxes.Num());
-			ensure(TransformToGeometryIndex.Num() > 0);
-			ensure(TransformToGeometryIndex.Num() == GlobalMatrices.Num());
-			
+			if (bChaos_BoxCalcBounds_ISPC_Enabled && !bGeometryCollectionSingleThreadedBoundsCalculation)
+			{
+				ensure(BoundingBoxes.Num() > 0);
+				ensure(TransformIndices.Num() == BoundingBoxes.Num());
+				ensure(TransformToGeometryIndex.Num() > 0);
+				ensure(TransformToGeometryIndex.Num() == GlobalMatrices.Num());
+
 #if INTEL_ISPC
-			ispc::BoxCalcBoundsFromGeometryGroup(
-				(int32*)&TransformToGeometryIndex[0],
-				(int32*)&TransformIndices[0],
-				(ispc::FMatrix*)&GlobalMatrices[0],
-				(ispc::FBox*)&BoundingBoxes[0],
-				(ispc::FMatrix&)LocalToWorldWithScale,
-				(ispc::FBox&)BoundingBox,
-				BoundingBoxes.Num());
+				ispc::BoxCalcBoundsFromGeometryGroup(
+					(int32*)&TransformToGeometryIndex[0],
+					(int32*)&TransformIndices[0],
+					(ispc::FMatrix*)&GlobalMatrices[0],
+					(ispc::FBox*)&BoundingBoxes[0],
+					(ispc::FMatrix&)LocalToWorldWithScale,
+					(ispc::FBox&)BoundingBox,
+					BoundingBoxes.Num());
 #endif
+				return BoundingBox;
+			}
 		}
-		else
+
+		// non-ISPC path
+		for (int32 BoxIdx = 0; BoxIdx < BoundingBoxes.Num(); ++BoxIdx)
 		{
-			for (int32 BoxIdx = 0; BoxIdx < BoundingBoxes.Num(); ++BoxIdx)
+			const int32 TransformIndex = TransformIndices[BoxIdx];
+			
+			if (TransformToGeometryIndex[TransformIndex] != INDEX_NONE)
 			{
-				const int32 TransformIndex = TransformIndices[BoxIdx];
-
-				if (TransformToGeometryIndex[TransformIndex] != INDEX_NONE)
-				{
-					BoundingBox += BoundingBoxes[BoxIdx].TransformBy(GlobalMatrices[TransformIndex] * LocalToWorldWithScale);
-				}
+				BoundingBox += BoundingBoxes[BoxIdx].TransformBy(GlobalMatrices[TransformIndex] * LocalToWorldWithScale);
 			}
 		}
 		return BoundingBox;
@@ -539,56 +557,79 @@ namespace
 	* compute the bounding box from the bounding boxes stored in the transform group
 	* (used for nanite or when the geometry group data has been stripped on cook )
 	*/
+	template <typename TTransformType>
 	inline FBox ComputeBoundsFromTransformBoundingBoxes(
 		const TManagedArray<int32>& TransformToGeometryIndex,
 		const TManagedArray<FBox>& BoundingBoxes,
-		const TArray<FMatrix>& GlobalMatrices,
-		const FMatrix& LocalToWorldWithScale)
+		const TArray<TTransformType>& GlobalMatrices,
+		const TTransformType& LocalToWorldWithScale)
 	{
 		FBox BoundingBox(ForceInit);
-		if (bChaos_BoxCalcBounds_ISPC_Enabled && !bGeometryCollectionSingleThreadedBoundsCalculation)
+
+		// todo(chaos ) implement ISPC function using FTransform 
+		constexpr bool bIsMatrixType = std::is_same<TTransformType, FMatrix>::value;
+		if constexpr (bIsMatrixType)
 		{
+			if (bChaos_BoxCalcBounds_ISPC_Enabled && !bGeometryCollectionSingleThreadedBoundsCalculation)
+			{
 #if INTEL_ISPC
-			ispc::BoxCalcBoundsFromTransformGroup(
-				(int32*)&TransformToGeometryIndex[0],
-				(ispc::FMatrix*)&GlobalMatrices[0],
-				(ispc::FBox*)&BoundingBoxes[0],
-				(ispc::FMatrix&)LocalToWorldWithScale,
-				(ispc::FBox&)BoundingBox,
-				BoundingBoxes.Num());
+				ispc::BoxCalcBoundsFromTransformGroup(
+					(int32*)&TransformToGeometryIndex[0],
+					(ispc::FMatrix*)&GlobalMatrices[0],
+					(ispc::FBox*)&BoundingBoxes[0],
+					(ispc::FMatrix&)LocalToWorldWithScale,
+					(ispc::FBox&)BoundingBox,
+					BoundingBoxes.Num());
 #endif
+				return BoundingBox;
+			}
 		}
-		else
+
+		// non-ISPC path 
+		for (int32 TransformIndex = 0; TransformIndex < BoundingBoxes.Num(); ++TransformIndex)
 		{
-			for (int32 TransformIndex = 0; TransformIndex < BoundingBoxes.Num(); ++TransformIndex)
+			if (TransformToGeometryIndex[TransformIndex] != INDEX_NONE)
 			{
-				if (TransformToGeometryIndex[TransformIndex] != INDEX_NONE)
-				{
-					BoundingBox += BoundingBoxes[TransformIndex].TransformBy(GlobalMatrices[TransformIndex] * LocalToWorldWithScale);
-				}
+				BoundingBox += BoundingBoxes[TransformIndex].TransformBy(GlobalMatrices[TransformIndex] * LocalToWorldWithScale);
 			}
 		}
 		return BoundingBox;
 	}
+
+	template <typename TTransformType>
+	inline FBox ComputeBoundsFromTransforms(const UGeometryCollectionComponent& Component, const TTransformType& LocalToWorldWithScale, const TArray<TTransformType>& GlobalMatricesArray)
+	{
+		auto GeometryCollectionPtr = Component.GetRestCollection()->GetGeometryCollection();
+		const TManagedArray<FBox>* TransformBoundingBoxes = GeometryCollectionPtr->FindAttribute<FBox>("BoundingBox", FGeometryCollection::TransformGroup);
+		const TManagedArray<FBox>& GeometryBoundingBoxes = Component.GetBoundingBoxArray();
+		const TManagedArray<int32>& TransformToGeometryIndex = Component.GetTransformToGeometryIndexArray();
+
+		if (TransformBoundingBoxes)
+		{
+			return ComputeBoundsFromTransformBoundingBoxes(TransformToGeometryIndex, *TransformBoundingBoxes, GlobalMatricesArray, LocalToWorldWithScale);
+		}
+
+		const TManagedArray<int32>& TransformIndices = Component.GetTransformIndexArray();
+		return ComputeBoundsFromGeometryBoundingBoxes(TransformToGeometryIndex, TransformIndices, GeometryBoundingBoxes, GlobalMatricesArray, LocalToWorldWithScale);
+	}
 }
 
 FBox UGeometryCollectionComponent::ComputeBoundsFromGlobalMatrices(const FMatrix& LocalToWorldWithScale, const TArray<FMatrix>& GlobalMatricesArray) const
 {
-	auto GeometryCollectionPtr = RestCollection->GetGeometryCollection();
-	const TManagedArray<FBox>* TransformBoundingBoxes = GeometryCollectionPtr->FindAttribute<FBox>("BoundingBox", FGeometryCollection::TransformGroup);
-	const TManagedArray<FBox>& GeometryBoundingBoxes = GetBoundingBoxArray();
-	const TManagedArray<int32>& TransformToGeometryIndex = GetTransformToGeometryIndexArray();
-
-	if (TransformBoundingBoxes)
-	{
-		return ComputeBoundsFromTransformBoundingBoxes(TransformToGeometryIndex, *TransformBoundingBoxes, GlobalMatricesArray, LocalToWorldWithScale);
-	}
+	return ComputeBoundsFromTransforms<FMatrix>(*this, LocalToWorldWithScale, GlobalMatricesArray);
+}
 
-	const TManagedArray<int32>& TransformIndices = GetTransformIndexArray();
-	return ComputeBoundsFromGeometryBoundingBoxes(TransformToGeometryIndex, TransformIndices, GeometryBoundingBoxes, GlobalMatricesArray, LocalToWorldWithScale);
+FBox UGeometryCollectionComponent::ComputeBoundsFromComponentSpaceTransforms(const FTransform& LocalToWorldWithScale, const TArray<FTransform>& ComponentSpaceTransformsArray) const
+{
+	return ComputeBoundsFromTransforms<FTransform>(*this, LocalToWorldWithScale, ComponentSpaceTransformsArray);
 }
 
 FBox UGeometryCollectionComponent::ComputeBounds(const FMatrix& LocalToWorldWithScale) const
+{
+	return ComputeBounds(FTransform(LocalToWorldWithScale));
+}
+
+FBox UGeometryCollectionComponent::ComputeBounds(const FTransform& LocalToWorldWithScale) const
 {
 	FBox BoundingBox(ForceInit);
 	if (RestCollection)
@@ -597,7 +638,7 @@ FBox UGeometryCollectionComponent::ComputeBounds(const FMatrix& LocalToWorldWith
 		auto GeometryCollectionPtr = RestCollection->GetGeometryCollection();
 
 		const int32 NumElements = GeometryCollectionPtr->NumElements(FGeometryCollection::TransformGroup);
-		if (NumElements == 0 || GlobalMatrices.Num() != RestCollection->NumElements(FGeometryCollection::TransformGroup))
+		if (NumElements == 0 || ComponentSpaceTransforms.Num() != RestCollection->NumElements(FGeometryCollection::TransformGroup))
 		{
 			// #todo(dmp): we could do the bbox transform in parallel with a bit of reformulating		
 			// #todo(dmp):  there are some cases where the calcbounds function is called before the component
@@ -608,21 +649,21 @@ FBox UGeometryCollectionComponent::ComputeBounds(const FMatrix& LocalToWorldWith
 			const TManagedArray<FTransform>& Transforms = GetTransformArray();
 			const TManagedArray<int32>& ParentIndices = GetParentArray();
 
-			TArray<FMatrix> TmpGlobalMatrices;
-			GeometryCollectionAlgo::GlobalMatrices(Transforms, ParentIndices, TmpGlobalMatrices);
-			if (TmpGlobalMatrices.Num() == 0)
+			TArray<FTransform> TmpComponentSpaceTransforms;
+			GeometryCollectionAlgo::GlobalMatrices(Transforms, ParentIndices, TmpComponentSpaceTransforms);
+			if (TmpComponentSpaceTransforms.Num() == 0)
 			{
 				BoundingBox = FBox(ForceInitToZero);
 			}
 			else
 			{
-				UpdateGlobalMatricesWithExplodedVectors(TmpGlobalMatrices, *GeometryCollectionPtr);
-				BoundingBox = ComputeBoundsFromGlobalMatrices(LocalToWorldWithScale, TmpGlobalMatrices);
+				UpdateGlobalMatricesWithExplodedVectors(TmpComponentSpaceTransforms, *GeometryCollectionPtr);
+				BoundingBox = ComputeBoundsFromComponentSpaceTransforms(LocalToWorldWithScale, TmpComponentSpaceTransforms);
 			}
 		}
 		else
 		{
-			BoundingBox = ComputeBoundsFromGlobalMatrices(LocalToWorldWithScale, GlobalMatrices);
+			BoundingBox = ComputeBoundsFromComponentSpaceTransforms(LocalToWorldWithScale, ComponentSpaceTransforms);
 		}
 	}
 	return BoundingBox;
@@ -648,7 +689,7 @@ FBoxSphereBounds UGeometryCollectionComponent::CalcBounds(const FTransform& Loca
 		
 		if (NeedBoundsUpdate)
 		{
-			ComponentSpaceBounds = ComputeBounds(FMatrix::Identity);
+			ComponentSpaceBounds = ComputeBounds(FTransform::Identity);
 		}
 		else
 		{
@@ -658,8 +699,7 @@ FBoxSphereBounds UGeometryCollectionComponent::CalcBounds(const FTransform& Loca
 		return ComponentSpaceBounds.TransformBy(LocalToWorldIn);
 	}
 
-	const FMatrix LocalToWorldWithScale = LocalToWorldIn.ToMatrixWithScale();
-	return FBoxSphereBounds(ComputeBounds(LocalToWorldWithScale));
+	return FBoxSphereBounds(ComputeBounds(LocalToWorldIn));
 }
 
 int32 UGeometryCollectionComponent::GetNumElements(FName Group) const
@@ -683,7 +723,7 @@ void UGeometryCollectionComponent::SetDamageThreshold(const TArray<float>& InDam
 
 void UGeometryCollectionComponent::UpdateCachedBounds()
 {
-	ComponentSpaceBounds = ComputeBounds(FMatrix::Identity);
+	ComponentSpaceBounds = ComputeBounds(FTransform::Identity);
 	CalculateLocalBounds();
 	UpdateBounds();
 }
@@ -1141,7 +1181,7 @@ void UGeometryCollectionComponent::RefreshEmbeddedGeometry()
 	}
 
 	const TManagedArray<int32>& ExemplarIndexArray = GetExemplarIndexArray();
-	const int32 TransformCount = GlobalMatrices.Num();
+	const int32 TransformCount = ComponentSpaceTransforms.Num();
 	if (!ensureMsgf(TransformCount == ExemplarIndexArray.Num(), TEXT("GlobalMatrices (Num=%d) cached on GeometryCollectionComponent are not in sync with ExemplarIndexArray (Num=%d) on underlying GeometryCollection; likely missed a dynamic data update"), TransformCount, ExemplarIndexArray.Num()))
 	{
 		return;
@@ -1174,7 +1214,7 @@ void UGeometryCollectionComponent::RefreshEmbeddedGeometry()
 			{
 				if (!HideArray || !(*HideArray)[Idx])
 				{ 
-					InstanceTransforms.Add(FTransform(GlobalMatrices[Idx]));
+					InstanceTransforms.Add(ComponentSpaceTransforms[Idx]);
 #if WITH_EDITOR
 					int32 InstanceIndex = EmbeddedBoneMaps[ExemplarIndex].Add(Idx);
 					EmbeddedInstanceIndex[Idx] = InstanceIndex;
@@ -1247,9 +1287,9 @@ void UGeometryCollectionComponent::SetRestState(TArray<FTransform>&& InRestTrans
 	}
 
 	FGeometryCollectionDynamicData* DynamicData = GDynamicDataPool.Allocate();
-	DynamicData->SetPrevTransforms(GlobalMatrices);
+	DynamicData->SetPrevTransforms(ComponentSpaceTransforms);
 	CalculateGlobalMatrices();
-	DynamicData->SetTransforms(GlobalMatrices);
+	DynamicData->SetTransforms(ComponentSpaceTransforms);
 	DynamicData->IsDynamic = true;
 
 	if (SceneProxy)
@@ -1317,7 +1357,7 @@ void UGeometryCollectionComponent::GetResourceSizeEx(FResourceSizeEx& Cumulative
 		+ RestTransforms.GetAllocatedSize()
 		+ DisabledFlags.GetAllocatedSize()
 		+ CollisionProfilePerLevel.GetAllocatedSize()
-		+ GlobalMatrices.GetAllocatedSize()
+		+ ComponentSpaceTransforms.GetAllocatedSize()
 		+ EventsPlayed.GetAllocatedSize()
 		+ CopyOnWriteAttributeList.GetAllocatedSize()
 		+ EmbeddedGeometryComponents.GetAllocatedSize()
@@ -2169,17 +2209,17 @@ FGeometryCollectionDynamicData* UGeometryCollectionComponent::InitDynamicData(bo
 		DynamicData->IsLoading = GetIsObjectLoading();
 
 		// If we have no transforms stored in the dynamic data, then assign both prev and current to the same global matrices
-		if (GlobalMatrices.Num() == 0)
+		if (ComponentSpaceTransforms.Num() == 0)
 		{
 			// Copy global matrices over to DynamicData
 			CalculateGlobalMatrices();
 
-			DynamicData->SetAllTransforms(GlobalMatrices);
+			DynamicData->SetAllTransforms(ComponentSpaceTransforms);
 		}
 		else
 		{
 			// Copy existing global matrices into prev transforms
-			DynamicData->SetPrevTransforms(GlobalMatrices);
+			DynamicData->SetPrevTransforms(ComponentSpaceTransforms);
 
 			// Copy global matrices over to DynamicData
 			CalculateGlobalMatrices();
@@ -2187,14 +2227,14 @@ FGeometryCollectionDynamicData* UGeometryCollectionComponent::InitDynamicData(bo
 			bool bComputeChanges = true;
 
 			// if the number of matrices has changed between frames, then sync previous to current
-			if (GlobalMatrices.Num() != DynamicData->PrevTransforms.Num())
+			if (ComponentSpaceTransforms.Num() != DynamicData->PrevTransforms.Num())
 			{
-				DynamicData->SetPrevTransforms(GlobalMatrices);
-				DynamicData->ChangedCount = GlobalMatrices.Num();
+				DynamicData->SetPrevTransforms(ComponentSpaceTransforms);
+				DynamicData->ChangedCount = ComponentSpaceTransforms.Num();
 				bComputeChanges = false; // Optimization to just force all transforms as changed and skip comparison
 			}
 
-			DynamicData->SetTransforms(GlobalMatrices);
+			DynamicData->SetTransforms(ComponentSpaceTransforms);
 
 			// The number of transforms for current and previous should match now
 			check(DynamicData->PrevTransforms.Num() == DynamicData->Transforms.Num());
@@ -2289,9 +2329,9 @@ FTransform UGeometryCollectionComponent::GetSocketTransform(FName InSocketName,
 		if (TransformIndex != INDEX_NONE)
 		{
 			
-			if (GlobalMatrices.IsValidIndex(TransformIndex))
+			if (ComponentSpaceTransforms.IsValidIndex(TransformIndex))
 			{
-				const FTransform BoneComponentSpaceTransform = FTransform(GlobalMatrices[TransformIndex]);
+				const FTransform BoneComponentSpaceTransform = ComponentSpaceTransforms[TransformIndex];
 				switch (TransformSpace)
 				{
 				case RTS_World:
@@ -2314,9 +2354,9 @@ FTransform UGeometryCollectionComponent::GetSocketTransform(FName InSocketName,
 					{
 						const int32 ParentTransformIndex = Collection->Parent[TransformIndex];
 						FTransform ParentComponentSpaceTransform{ FTransform::Identity };
-						if (GlobalMatrices.IsValidIndex(ParentTransformIndex))
+						if (ComponentSpaceTransforms.IsValidIndex(ParentTransformIndex))
 						{
-							ParentComponentSpaceTransform = FTransform(GlobalMatrices[ParentTransformIndex]);
+							ParentComponentSpaceTransform = ComponentSpaceTransforms[ParentTransformIndex];
 						}
 						return BoneComponentSpaceTransform.GetRelativeTransform(ParentComponentSpaceTransform);
 					}
@@ -2990,7 +3030,7 @@ void UGeometryCollectionComponent::UpdateRemovalIfNeeded()
 							ShrinkRadius = -AccumulatedSphere.W;
 						}
 
-						const FQuat LocalRotation = (InverseComponentTransform * FTransform(GlobalMatrices[Idx]).Inverse()).GetRotation();
+						const FQuat LocalRotation = (InverseComponentTransform * ComponentSpaceTransforms[Idx].Inverse()).GetRotation();
 						const FVector LocalDown = LocalRotation.RotateVector(FVector(0.f, 0.f, ShrinkRadius));
 						const FVector CenterOfMass = DynamicCollection->MassToLocal[Idx].GetTranslation();
 						const FVector ScaleCenter = LocalDown + CenterOfMass;
@@ -4160,17 +4200,17 @@ AChaosSolverActor* UGeometryCollectionComponent::GetPhysicsSolverActor() const
 void UGeometryCollectionComponent::CalculateLocalBounds()
 {
 	LocalBounds.Init();
-	LocalBounds = ComputeBounds(FMatrix::Identity);
+	LocalBounds = ComputeBounds(FTransform::Identity);
 }
 
-#define GEOMETRY_COLLECTION_CHECK_FOR_NANS_IN_MATRICES 0
+#define GEOMETRY_COLLECTION_CHECK_FOR_NANS_IN_TRANSFORMS 0
 
-#if GEOMETRY_COLLECTION_CHECK_FOR_NANS_IN_MATRICES
-static void CheckForNaNs(const TArray<FMatrix>& Matrices)
+#if GEOMETRY_COLLECTION_CHECK_FOR_NANS_IN_TRANSFORMS
+static void CheckForNaNs(const TArray<FTransform>& Transforms)
 {
-	for (const FMatrix& Matrix : Matrices)
+	for (const FTransform& Transform : Transforms)
 	{
-		ensureAlways(false == Matrix.ContainsNaN());
+		ensureAlways(false == Transform.ContainsNaN());
 	}
 }
 #endif
@@ -4185,8 +4225,8 @@ void UGeometryCollectionComponent::CalculateGlobalMatrices()
 	if(NumTransforms > 0)
 	{
 		// Just calc from results
-		GlobalMatrices.Reset();
-		GlobalMatrices.Append(Results->GlobalTransforms);	
+		ComponentSpaceTransforms.Reset();
+		ComponentSpaceTransforms.Append(Results->GlobalTransforms);
 	}
 	else
 	{
@@ -4198,20 +4238,20 @@ void UGeometryCollectionComponent::CalculateGlobalMatrices()
 
 		if (!DynamicCollection && RestTransforms.Num() > 0)
 		{
-			GeometryCollectionAlgo::GlobalMatrices(RestTransforms, GetParentArray(), GlobalMatrices);
+			GeometryCollectionAlgo::GlobalMatrices(RestTransforms, GetParentArray(), ComponentSpaceTransforms);
 		}
 		else
 		{
-			GeometryCollectionAlgo::GlobalMatrices(GetTransformArray(), GetParentArray(), GlobalMatrices);
+			GeometryCollectionAlgo::GlobalMatrices(GetTransformArray(), GetParentArray(), ComponentSpaceTransforms);
 		}
 	}
 
 #if WITH_EDITOR
-	UpdateGlobalMatricesWithExplodedVectors(GlobalMatrices, *(RestCollection->GetGeometryCollection()));
+	UpdateGlobalMatricesWithExplodedVectors(ComponentSpaceTransforms, *(RestCollection->GetGeometryCollection()));
 #endif
 
-#if GEOMETRY_COLLECTION_CHECK_FOR_NANS_IN_MATRICES
-	CheckForNaNs(GlobalMatrices);
+#if GEOMETRY_COLLECTION_CHECK_FOR_NANS_IN_TRANSFORMS
+	CheckForNaNs(ComponentSpaceTransforms);
 #endif
 }
 
@@ -4439,11 +4479,11 @@ void UGeometryCollectionComponent::RefreshCustomRenderer()
 
 				if (bIsBroken)
 				{
-					RendererInterface->UpdateTransforms(*RestCollection, GlobalMatrices);
+					RendererInterface->UpdateTransforms(*RestCollection, ComponentSpaceTransforms);
 				}
 				else
 				{
-					const FTransform RootTransform = GlobalMatrices.IsValidIndex(RootIndex) ? FTransform(GlobalMatrices[RootIndex]) : FTransform::Identity;
+					const FTransform RootTransform = ComponentSpaceTransforms.IsValidIndex(RootIndex) ? ComponentSpaceTransforms[RootIndex] : FTransform::Identity;
 					RendererInterface->UpdateRootTransform(*RestCollection, RootTransform);
 				}
 			}
@@ -4868,9 +4908,9 @@ FTransform UGeometryCollectionComponent::GetRootCurrentTransform() const
 	if (RestCollection)
 	{
 		const int32 RootIndex = RestCollection->GetRootIndex();
-		if (GlobalMatrices.IsValidIndex(RootIndex))
+		if (ComponentSpaceTransforms.IsValidIndex(RootIndex))
 		{
-			RootInitialTransform = FTransform(GlobalMatrices[RootIndex]) * GetComponentTransform();
+			RootInitialTransform = ComponentSpaceTransforms[RootIndex] * GetComponentTransform();
 		}
 	}
 	return RootInitialTransform;
@@ -4944,6 +4984,18 @@ void UGeometryCollectionComponent::SetLocalRestTransforms(const TArray<FTransfor
 	}
 }
 
+TArray<FMatrix> UGeometryCollectionComponent::ComputeGlobalMatricesFromComponentSpaceTransforms() const
+{
+	TArray<FMatrix> ComponentSpaceMatrices;
+	ComponentSpaceMatrices.SetNumUninitialized(ComponentSpaceTransforms.Num());
+
+	for (int32 TransformIndex = 0; TransformIndex < ComponentSpaceTransforms.Num(); TransformIndex++)
+	{
+		ComponentSpaceMatrices[TransformIndex] = ComponentSpaceTransforms[TransformIndex].ToMatrixWithScale();
+	}
+	return ComponentSpaceMatrices;
+}
+
 void UGeometryCollectionComponent::GetMassAndExtents(int32 ItemIndex, float& OutMass, FBox& OutExtents)
 {
 	using FGeometryCollectionPtr = const TSharedPtr<FGeometryCollection, ESPMode::ThreadSafe>;
diff --git a/Engine/Source/Runtime/Experimental/GeometryCollectionEngine/Private/GeometryCollection/GeometryCollectionISMPoolRenderer.cpp b/Engine/Source/Runtime/Experimental/GeometryCollectionEngine/Private/GeometryCollection/GeometryCollectionISMPoolRenderer.cpp
--- a/Engine/Source/Runtime/Experimental/GeometryCollectionEngine/Private/GeometryCollection/GeometryCollectionISMPoolRenderer.cpp
+++ b/Engine/Source/Runtime/Experimental/GeometryCollectionEngine/Private/GeometryCollection/GeometryCollectionISMPoolRenderer.cpp
@@ -61,9 +61,9 @@ void UGeometryCollectionISMPoolRenderer::UpdateRootTransform(UGeometryCollection
 	UpdateMergedMeshTransforms(InRootTransform * ComponentTransform);
 }
 
-void UGeometryCollectionISMPoolRenderer::UpdateTransforms(UGeometryCollection const& InGeometryCollection, TArrayView<const FMatrix> InMatrices)
+void UGeometryCollectionISMPoolRenderer::UpdateTransforms(UGeometryCollection const& InGeometryCollection, TArrayView<const FTransform> InTransforms)
 {
-	UpdateInstanceTransforms(InGeometryCollection, ComponentTransform, InMatrices);
+	UpdateInstanceTransforms(InGeometryCollection, ComponentTransform, InTransforms);
 }
 
 UGeometryCollectionISMPoolComponent* UGeometryCollectionISMPoolRenderer::GetOrCreateISMPoolComponent()
@@ -177,7 +177,7 @@ void UGeometryCollectionISMPoolRenderer::UpdateMergedMeshTransforms(FTransform c
 	}
 }
 
-void UGeometryCollectionISMPoolRenderer::UpdateInstanceTransforms(UGeometryCollection const& InGeometryCollection, FTransform const& InBaseTransform, TArrayView<const FMatrix> InMatrices)
+void UGeometryCollectionISMPoolRenderer::UpdateInstanceTransforms(UGeometryCollection const& InGeometryCollection, FTransform const& InBaseTransform, TArrayView<const FTransform> InTransforms)
 {
 	if (InstancesGroup.GroupIndex == INDEX_NONE)
 	{
@@ -208,7 +208,7 @@ void UGeometryCollectionISMPoolRenderer::UpdateInstanceTransforms(UGeometryColle
 			const int32 AutoInstanceMeshIndex = InstancedMeshFacade.GetIndex(TransformIndex);
 			if (AutoInstanceMeshIndex == MeshIndex && Children[TransformIndex].Num() == 0)
 			{
-				InstanceTransforms.Add(FTransform(InMatrices[TransformIndex]) * InBaseTransform);
+				InstanceTransforms.Add(InTransforms[TransformIndex] * InBaseTransform);
 			}
 		}
 		
diff --git a/Engine/Source/Runtime/Experimental/GeometryCollectionEngine/Private/GeometryCollection/GeometryCollectionISMPoolRenderer.h b/Engine/Source/Runtime/Experimental/GeometryCollectionEngine/Private/GeometryCollection/GeometryCollectionISMPoolRenderer.h
--- a/Engine/Source/Runtime/Experimental/GeometryCollectionEngine/Private/GeometryCollection/GeometryCollectionISMPoolRenderer.h
+++ b/Engine/Source/Runtime/Experimental/GeometryCollectionEngine/Private/GeometryCollection/GeometryCollectionISMPoolRenderer.h
@@ -21,7 +21,7 @@ public:
 	virtual void OnUnregisterGeometryCollection() override;
 	virtual void UpdateState(UGeometryCollection const& InGeometryCollection, FTransform const& InComponentTransform, bool bInIsBroken, bool bInIsVisible) override;
 	virtual void UpdateRootTransform(UGeometryCollection const& InGeometryCollection, FTransform const& InRootTransform) override;
-	virtual void UpdateTransforms(UGeometryCollection const& InGeometryCollection, TArrayView<const FMatrix> InMatrices) override;
+	virtual void UpdateTransforms(UGeometryCollection const& InGeometryCollection, TArrayView<const FTransform> InTransforms) override;
 	//~ End IGeometryCollectionExternalRenderInterface Interface.
 
 	/** Description for a group of meshes that are added/updated together. */
@@ -48,6 +48,6 @@ private:
 	void InitMergedMeshFromGeometryCollection(UGeometryCollection const& InGeometryCollection);
 	void InitInstancesFromGeometryCollection(UGeometryCollection const& InGeometryCollection);
 	void UpdateMergedMeshTransforms(FTransform const& InBaseTransform);
-	void UpdateInstanceTransforms(UGeometryCollection const& InGeometryCollection, FTransform const& InBaseTransform, TArrayView<const FMatrix> InMatrices);
+	void UpdateInstanceTransforms(UGeometryCollection const& InGeometryCollection, FTransform const& InBaseTransform, TArrayView<const FTransform> InTransforms);
 	void ReleaseGroup(FISMPoolGroup& InOutGroup);
 };
diff --git a/Engine/Source/Runtime/Experimental/GeometryCollectionEngine/Private/GeometryCollection/GeometryCollectionSceneProxy.h b/Engine/Source/Runtime/Experimental/GeometryCollectionEngine/Private/GeometryCollection/GeometryCollectionSceneProxy.h
--- a/Engine/Source/Runtime/Experimental/GeometryCollectionEngine/Private/GeometryCollection/GeometryCollectionSceneProxy.h
+++ b/Engine/Source/Runtime/Experimental/GeometryCollectionEngine/Private/GeometryCollection/GeometryCollectionSceneProxy.h
@@ -48,7 +48,6 @@ public:
 	FShaderResourceViewRHIRef VertexBufferSRV;
 };
 
-
 inline void CopyTransformsWithConversionWhenNeeded(TArray<FMatrix44f>& DstTransforms, const TArray<FMatrix>& SrcTransforms)
 {
 	// LWC_TODO : we have no choice but to convert each element at this point to avoid changing GeometryCollectionAlgo::GlobalMatrices that is used all over the place
@@ -59,6 +58,15 @@ inline void CopyTransformsWithConversionWhenNeeded(TArray<FMatrix44f>& DstTransf
 	}
 }
 
+inline void CopyTransformsWithConversionWhenNeeded(TArray<FMatrix44f>& DstTransforms, const TArray<FTransform>& SrcTransforms)
+{
+	// LWC_TODO : we have no choice but to convert each element at this point to avoid changing GeometryCollectionAlgo::GlobalMatrices that is used all over the place
+	DstTransforms.SetNumUninitialized(SrcTransforms.Num());
+	for (int TransformIndex = 0; TransformIndex < SrcTransforms.Num(); ++TransformIndex)
+	{
+		DstTransforms[TransformIndex] = FTransform3f(SrcTransforms[TransformIndex]).ToMatrixWithScale(); // LWC_TODO: Perf pessimization
+	}
+}
 
 /** Mutable rendering data */
 struct FGeometryCollectionDynamicData
@@ -82,19 +90,43 @@ struct FGeometryCollectionDynamicData
 		IsLoading = false;
 	}
 
+	UE_DEPRECATED(5.3, "Use FTransform version of SetTransforms instead")
 	void SetTransforms(const TArray<FMatrix>& InTransforms)
 	{
 		// use for LWC as FMatrix and FMatrix44f are different when LWC is on 
 		CopyTransformsWithConversionWhenNeeded(Transforms, InTransforms);
 	}
 
+	void SetTransforms(const TArray<FTransform>& InTransforms)
+	{
+		// use for LWC as FMatrix and FMatrix44f are different when LWC is on 
+		CopyTransformsWithConversionWhenNeeded(Transforms, InTransforms);
+	}
+
+	UE_DEPRECATED(5.3, "Use FTransform version of SetPrevTransforms instead")
 	void SetPrevTransforms(const TArray<FMatrix>& InTransforms)
 	{
 		// use for LWC as FMatrix and FMatrix44f are different when LWC is on 
 		CopyTransformsWithConversionWhenNeeded(PrevTransforms, InTransforms);
 	}
 
+	void SetPrevTransforms(const TArray<FTransform>& InTransforms)
+	{
+		// use for LWC as FMatrix and FMatrix44f are different when LWC is on 
+		CopyTransformsWithConversionWhenNeeded(PrevTransforms, InTransforms);
+	}
+
+	UE_DEPRECATED(5.3, "Use FTransform version of SetAllTransforms instead")
 	void SetAllTransforms(const TArray<FMatrix>& InTransforms)
+	{
+		PRAGMA_DISABLE_DEPRECATION_WARNINGS
+		SetTransforms(InTransforms);
+		PRAGMA_ENABLE_DEPRECATION_WARNINGS
+		PrevTransforms = Transforms;
+		ChangedCount = Transforms.Num();
+	}
+
+	void SetAllTransforms(const TArray<FTransform>& InTransforms)
 	{
 		SetTransforms(InTransforms);
 		PrevTransforms = Transforms;
diff --git a/Engine/Source/Runtime/Experimental/GeometryCollectionEngine/Public/GeometryCollection/GeometryCollectionComponent.h b/Engine/Source/Runtime/Experimental/GeometryCollectionEngine/Public/GeometryCollection/GeometryCollectionComponent.h
--- a/Engine/Source/Runtime/Experimental/GeometryCollectionEngine/Public/GeometryCollection/GeometryCollectionComponent.h
+++ b/Engine/Source/Runtime/Experimental/GeometryCollectionEngine/Public/GeometryCollection/GeometryCollectionComponent.h
@@ -1050,7 +1050,10 @@ public:
 	void SetInitialClusterBreaks(const TArray<int32>& ReleaseIndices);
 
 	/** Used by Niagara DI to query global matrices rather than recalculating them again */
-	const TArray<FMatrix>& GetGlobalMatrices() { return GlobalMatrices; }
+	UE_DEPRECATED(5.3, "Use GetComponentSpaceTransforms instead")
+	TArray<FMatrix> GetGlobalMatrices() { return ComputeGlobalMatricesFromComponentSpaceTransforms(); }
+
+	const TArray<FTransform> GetComponentSpaceTransforms() { return ComponentSpaceTransforms; }
 
 	const FGeometryDynamicCollection* GetDynamicCollection() const;
 	FGeometryDynamicCollection* GetDynamicCollection();  // TEMP HACK?
@@ -1152,9 +1155,17 @@ protected:
 	Chaos::FPhysicsSolver* GetSolver(const UGeometryCollectionComponent& GeometryCollectionComponent);
 	void CalculateLocalBounds();
 	void CalculateGlobalMatrices();
+	
+	UE_DEPRECATED(5.3, "Use ComputeBoundsFromComponentSpaceTransforms instead")
 	FBox ComputeBoundsFromGlobalMatrices(const FMatrix& LocalToWorldWithScale, const TArray<FMatrix>& GlobalMatricesArray) const;
+
+	FBox ComputeBoundsFromComponentSpaceTransforms(const FTransform& LocalToWorldWithScale, const TArray<FTransform>& ComponentSpaceTransformsArray) const;
+
+	UE_DEPRECATED(5.3, "Use FTransform version of ComputeBounds instead")
 	FBox ComputeBounds(const FMatrix& LocalToWorldWithScale) const;
 
+	FBox ComputeBounds(const FTransform& LocalToWorldWithScale) const;
+
 	void RegisterForEvents();
 	void UpdateRBCollisionEventRegistration();
 	void UpdateBreakEventRegistration();
@@ -1261,7 +1272,11 @@ private:
 	TArray<int32> HighlightedBones;
 #endif
 
+	UE_DEPRECATED(5.3, "Use ComponentSpaceTransforms instead")
 	TArray<FMatrix> GlobalMatrices;
+
+	TArray<FTransform> ComponentSpaceTransforms;
+
 	FBox LocalBounds;
 
 	mutable FBoxSphereBounds ComponentSpaceBounds;
@@ -1339,6 +1354,9 @@ private:
 
 	bool HasVisibleGeometry() const;
 
+	/** backward compatibility method, until we can remove GlobalMatrices */
+	TArray<FMatrix> ComputeGlobalMatricesFromComponentSpaceTransforms() const;
+
 	/** The clusters we need to replicate */
 	TUniquePtr<TSet<Chaos::FPBDRigidClusteredParticleHandle*>> ClustersToRep;
 
diff --git a/Engine/Source/Runtime/Experimental/GeometryCollectionEngine/Public/GeometryCollection/GeometryCollectionExternalRenderInterface.h b/Engine/Source/Runtime/Experimental/GeometryCollectionEngine/Public/GeometryCollection/GeometryCollectionExternalRenderInterface.h
--- a/Engine/Source/Runtime/Experimental/GeometryCollectionEngine/Public/GeometryCollection/GeometryCollectionExternalRenderInterface.h
+++ b/Engine/Source/Runtime/Experimental/GeometryCollectionEngine/Public/GeometryCollection/GeometryCollectionExternalRenderInterface.h
@@ -24,5 +24,8 @@ public:
 	virtual void OnUnregisterGeometryCollection() = 0;
 	virtual void UpdateState(UGeometryCollection const& InGeometryCollection, FTransform const& InComponentTransform, bool bInIsBroken, bool bInIsVisible) = 0;
 	virtual void UpdateRootTransform(UGeometryCollection const& InGeometryCollection, FTransform const& InRootTransform) = 0;
-	virtual void UpdateTransforms(UGeometryCollection const& InGeometryCollection, TArrayView<const FMatrix> InMatrices) = 0;
+	virtual void UpdateTransforms(UGeometryCollection const& InGeometryCollection, TArrayView<const FTransform> InTransforms) = 0;
+
+	UE_DEPRECATED(5.3, "Use FTransform version of UpdateTransforms instead")
+	virtual void UpdateTransforms(UGeometryCollection const& InGeometryCollection, TArrayView<const FMatrix> InMatrices) {};
 };
