Fix compile error in PBDRigidsSolver.cpp

diff --git a/Engine/Source/Runtime/Experimental/Chaos/Public/PBDRigidsSolver.h b/Engine/Source/Runtime/Experimental/Chaos/Public/PBDRigidsSolver.h
--- a/Engine/Source/Runtime/Experimental/Chaos/Public/PBDRigidsSolver.h
+++ b/Engine/Source/Runtime/Experimental/Chaos/Public/PBDRigidsSolver.h
@@ -51,6 +51,7 @@ namespace ChaosTest
 */
 namespace Chaos
 {
+	CHAOS_API extern int32 SyncKinematicOnGameThread;
 	class FPersistentPhysicsTask;
 	class FChaosArchive;
 	class FCharacterGroundConstraint;
