applying external strain to a GC particle should apply even if that particle doesn't have children

diff --git a/Engine/Source/Runtime/Experimental/Chaos/Private/PhysicsProxy/GeometryCollectionPhysicsProxy.cpp b/Engine/Source/Runtime/Experimental/Chaos/Private/PhysicsProxy/GeometryCollectionPhysicsProxy.cpp
--- a/Engine/Source/Runtime/Experimental/Chaos/Private/PhysicsProxy/GeometryCollectionPhysicsProxy.cpp
+++ b/Engine/Source/Runtime/Experimental/Chaos/Private/PhysicsProxy/GeometryCollectionPhysicsProxy.cpp
@@ -2291,6 +2291,10 @@ static void ApplyToChildrenAtPointWithRadiusAndPropagation_Internal(
 			}
 		}
 	}
+	else
+	{
+		Action(&ClusteredHandle, 1.0);
+	}
 }
 
 
