Move over the extern declaration for SyncKinematicOnGameThread from PBDRigidsSolver header to the cpp.

diff --git a/Engine/Source/Runtime/Experimental/Chaos/Private/PBDRigidsSolver.cpp b/Engine/Source/Runtime/Experimental/Chaos/Private/PBDRigidsSolver.cpp
--- a/Engine/Source/Runtime/Experimental/Chaos/Private/PBDRigidsSolver.cpp
+++ b/Engine/Source/Runtime/Experimental/Chaos/Private/PBDRigidsSolver.cpp
@@ -355,6 +355,8 @@ namespace Chaos
 {
 	using namespace CVars;
 
+	CHAOS_API extern int32 SyncKinematicOnGameThread;
+
 	class AdvanceOneTimeStepTask : public FNonAbandonableTask
 	{
 		friend class FAutoDeleteAsyncTask<AdvanceOneTimeStepTask>;
diff --git a/Engine/Source/Runtime/Experimental/Chaos/Public/PBDRigidsSolver.h b/Engine/Source/Runtime/Experimental/Chaos/Public/PBDRigidsSolver.h
--- a/Engine/Source/Runtime/Experimental/Chaos/Public/PBDRigidsSolver.h
+++ b/Engine/Source/Runtime/Experimental/Chaos/Public/PBDRigidsSolver.h
@@ -51,7 +51,6 @@ namespace ChaosTest
 */
 namespace Chaos
 {
-	CHAOS_API extern int32 SyncKinematicOnGameThread;
 	class FPersistentPhysicsTask;
 	class FChaosArchive;
 	class FCharacterGroundConstraint;
