Chaos - handle FullAABB and EmptyAABB in AABB Transform functions

diff --git a/Engine/Source/Runtime/Experimental/Chaos/Private/Chaos/AABB.cpp b/Engine/Source/Runtime/Experimental/Chaos/Private/Chaos/AABB.cpp
--- a/Engine/Source/Runtime/Experimental/Chaos/Private/Chaos/AABB.cpp
+++ b/Engine/Source/Runtime/Experimental/Chaos/Private/Chaos/AABB.cpp
@@ -273,6 +273,12 @@ bool TAABB<T, d>::Raycast(const TVector<FReal, d>& StartPoint, const TVector<FRe
 template<typename T>
 inline TAABB<T, 3> TransformedAABBHelper(const TAABB<T, 3>& AABB, const FMatrix44& SpaceTransform)
 {
+	// Full and empty bounds do not transform
+	if (AABB.IsFull() || AABB.IsEmpty())
+	{
+		return AABB;
+	}
+
 	// Initialize to center
 	FVec3 Translation(SpaceTransform.M[3][0], SpaceTransform.M[3][1], SpaceTransform.M[3][2]);
 	FVec3 Min = Translation;
@@ -303,6 +309,12 @@ inline TAABB<T, 3> TransformedAABBHelper(const TAABB<T, 3>& AABB, const FMatrix4
 
 inline TAABB<FReal, 3> TransformedAABBHelperISPC(const TAABB<FReal, 3>& AABB, const FTransform& SpaceTransform)
 {
+	// Full and empty bounds do not transform
+	if (AABB.IsFull() || AABB.IsEmpty())
+	{
+		return AABB;
+	}
+
 	check(bRealTypeCompatibleWithISPC);
 #if INTEL_ISPC
 	TVector<Chaos::FReal, 3> NewMin, NewMax;
@@ -373,8 +385,14 @@ TAABB<T, d> TAABB<T, d>::TransformedAABB(const FTransform& SpaceTransform) const
 }
 
 template<typename T, int d>
-inline TAABB<T, d> TAABB<T, d>::InverseTransformedAABB(const FRigidTransform3& SpaceTransform) const
+TAABB<T, d> TAABB<T, d>::InverseTransformedAABB(const FRigidTransform3& SpaceTransform) const
 {
+	// Full and empty bounds do not transform
+	if (IsFull() || IsEmpty())
+	{
+		return *this;
+	}
+
 	TVector<T, d> CurrentExtents = Extents();
 	int32 Idx = 0;
 	const TVector<T, d> MinToNewSpace = SpaceTransform.InverseTransformPosition(FVector(Min()));
diff --git a/Engine/Source/Runtime/Experimental/Chaos/Public/Chaos/AABB.h b/Engine/Source/Runtime/Experimental/Chaos/Public/Chaos/AABB.h
--- a/Engine/Source/Runtime/Experimental/Chaos/Public/Chaos/AABB.h
+++ b/Engine/Source/Runtime/Experimental/Chaos/Public/Chaos/AABB.h
@@ -537,10 +537,34 @@ namespace Chaos
 		FORCEINLINE T GetMargin() const { return 0; }
 		FORCEINLINE T GetRadius() const { return 0; }
 
+		// A bounding box that will fail all overlap tests.
+		// NOTE: this bounds cannot be transformed (all transform return EmptyAABB)
 		FORCEINLINE static TAABB<T, d> EmptyAABB() { return TAABB<T, d>(TVector<T, d>(TNumericLimits<T>::Max()), TVector<T, d>(-TNumericLimits<T>::Max())); }
-		FORCEINLINE static TAABB<T, d> ZeroAABB() { return TAABB<T, d>(TVector<T, d>((T)0), TVector<T, d>((T)0)); }
+
+		// A bounding box that overlaps all space.
+		// NOTE: this bounds cannot be transformed (all transforms return FullAABB)
 		FORCEINLINE static TAABB<T, d> FullAABB() { return TAABB<T, d>(TVector<T, d>(-TNumericLimits<T>::Max()), TVector<T, d>(TNumericLimits<T>::Max())); }
 
+		// A single-point bounds at the origin
+		FORCEINLINE static TAABB<T, d> ZeroAABB() { return TAABB<T, d>(TVector<T, d>((T)0), TVector<T, d>((T)0)); }
+
+
+		// If this AABB covering all of space (all overlap tests will succeed)
+		// NOTE: This just checks if Max.X is Real::Max. The assumption is that max numeric values 
+		// are only used when explicitly set in EmptyAABB() and FullAABB().
+		FORCEINLINE bool IsFull() const
+		{
+			return (MMax.X == TNumericLimits<T>::Max());
+		}
+
+		// If this AABB empty (all overlap tests will fail)
+		// NOTE: This just checks if Min.X is Real::Max. The assumption is that max numeric values 
+		// are only used when explicitly set in EmptyAABB() and FullAABB().
+		FORCEINLINE bool IsEmpty() const
+		{
+			return (MMin.X == TNumericLimits<T>::Max());
+		}
+
 		FORCEINLINE void Serialize(FChaosArchive& Ar)
 		{
 			// LWC : this method exists in order to properly call the << operator using FChaosArchive on Chaos::TVector 
