Chaos - support disabled constraints in RBAN solver

diff --git a/Engine/Source/Runtime/Experimental/Chaos/Private/Chaos/Joint/PBDJointContainerSolver.cpp b/Engine/Source/Runtime/Experimental/Chaos/Private/Chaos/Joint/PBDJointContainerSolver.cpp
--- a/Engine/Source/Runtime/Experimental/Chaos/Private/Chaos/Joint/PBDJointContainerSolver.cpp
+++ b/Engine/Source/Runtime/Experimental/Chaos/Private/Chaos/Joint/PBDJointContainerSolver.cpp
@@ -120,14 +120,17 @@ namespace Chaos
 			}
 		}
 
-		void FPBDJointContainerSolver::FPBDJointContainerSolver::AddConstraints()
+		void FPBDJointContainerSolver::AddConstraints()
 		{
 			Reset(ConstraintContainer.GetNumConstraints());
 
 			// @todo(chaos): we could eliminate the index array if we're solving all constraints in the scene (RBAN)
 			for (int32 ContainerConstraintIndex = 0; ContainerConstraintIndex < ConstraintContainer.GetNumConstraints(); ++ContainerConstraintIndex)
 			{
-				AddConstraint(ContainerConstraintIndex);
+				if (ConstraintContainer.IsConstraintEnabled(ContainerConstraintIndex))
+				{
+					AddConstraint(ContainerConstraintIndex);
+				}
 			}
 		}
 
diff --git a/Engine/Source/Runtime/Experimental/Chaos/Public/Chaos/Evolution/SolverConstraintContainer.h b/Engine/Source/Runtime/Experimental/Chaos/Public/Chaos/Evolution/SolverConstraintContainer.h
--- a/Engine/Source/Runtime/Experimental/Chaos/Public/Chaos/Evolution/SolverConstraintContainer.h
+++ b/Engine/Source/Runtime/Experimental/Chaos/Public/Chaos/Evolution/SolverConstraintContainer.h
@@ -12,10 +12,18 @@ namespace Chaos
 
 	/** 
 	 * Base class for all the solver for a set of constraints of a specific type.
+	 * 
 	 * A SolverContainer is used to solve a set constraints in sequential order. There will
 	 * be one solver container for each thread on which we solve constraints (see FPBDIslandConstraintGroupSolver). 
 	 * How constraints are assigned to groups depends on the constraint type and settings, but
 	 * usually a group contains all constraints from one or more islands (unless we are coloring).
+	 * 
+	 * NOTE: there are two main use-case for FConstraintContainerSolver objects: the main physics scene and RBAN. 
+	 * The main scene uses the IslandManager to break the scene up into groups of constraints
+	 * that can be solved in parallel (Islands). Those islands are then put into IslandGroups, and each IslandGroup
+	 * is solved in a task, therefore we will have one FConstraintContainerSolver per type of constraint per IslandGroup.
+	 * RBAN does not attempt to partition its scene into islands and just solves all constraints on its main
+	 * thread, so it only has one FConstraintContainerSolver (per constraint type).
 	*/
 	class CHAOS_API FConstraintContainerSolver
 	{
@@ -55,10 +63,15 @@ namespace Chaos
 
 		virtual int32 GetNumConstraints() const = 0;
 
+		/**
+		 * RBAN API.
+		 * Add all (active) constraints to the solver.
+		 */
 		virtual void AddConstraints() = 0;
 
 		/**
-		 * Add a set of constraints to the solver container. This can be called multiple times: once for each island a group, but
+		 * Island API.
+		 * Add a set of constraints to the solver. This can be called multiple times: once for each island in an IslandGroup, but
 		 * there will never be more constraints added than specified in Reset().
 		 * NOTE: this should not do any actual data gathering - it should just add to the list of constraints in this group. All data
 		 * gathering is handled in GatherInput.
