Chaos - fix suspension constraints being left in the graph after particle destruction

diff --git a/Engine/Source/Runtime/Experimental/Chaos/Public/Chaos/PBDRigidsEvolution.h b/Engine/Source/Runtime/Experimental/Chaos/Public/Chaos/PBDRigidsEvolution.h
--- a/Engine/Source/Runtime/Experimental/Chaos/Public/Chaos/PBDRigidsEvolution.h
+++ b/Engine/Source/Runtime/Experimental/Chaos/Public/Chaos/PBDRigidsEvolution.h
@@ -617,16 +617,9 @@ public:
 	/** remove a list of constraints from the constraint graph */
 	CHAOS_API void RemoveConstraintsFromConstraintGraph(const FConstraintHandleArray& Constraints)
 	{
-		for (FConstraintHandle* BaseConstraintHandle : Constraints)
+		for (FConstraintHandle* ConstraintHandle : Constraints)
 		{
-			if (FPBDJointConstraintHandle* ConstraintHandle = BaseConstraintHandle->As<FPBDJointConstraintHandle>())
-			{
-				RemoveConstraintFromConstraintGraph(ConstraintHandle);
-			}
-			else if (FCharacterGroundConstraintHandle* CharacterGroundConstraint = BaseConstraintHandle->As<FCharacterGroundConstraintHandle>())
-			{
-				RemoveConstraintFromConstraintGraph(CharacterGroundConstraint);
-			}
+			RemoveConstraintFromConstraintGraph(ConstraintHandle);
 		}
 	}
 
