Chaos - fix collision for objects made of a large number of duplicated shapes

diff --git a/Engine/Source/Runtime/Experimental/Chaos/Private/Chaos/Collision/ParticlePairMidPhase.cpp b/Engine/Source/Runtime/Experimental/Chaos/Private/Chaos/Collision/ParticlePairMidPhase.cpp
--- a/Engine/Source/Runtime/Experimental/Chaos/Private/Chaos/Collision/ParticlePairMidPhase.cpp
+++ b/Engine/Source/Runtime/Experimental/Chaos/Private/Chaos/Collision/ParticlePairMidPhase.cpp
@@ -396,7 +396,7 @@ namespace Chaos
 				}
 			}
 
-			// If we get here, we did not activate hte constraint and it should be disabled for this tick
+			// If we get here, we did not activate the constraint and it should be disabled for this tick
 			Constraint->SetDisabled(true);
 		}
 
@@ -644,8 +644,8 @@ namespace Chaos
 		// overlapping shape pair on the two bodies via a callback to FindOrCreateConstraint
 		Collisions::ConstructConstraints(
 			Particle0, Particle1,
-			Implicit0, Shape0, BVHParticles0,
-			Implicit1, Shape1, BVHParticles1,
+			Implicit0, Shape0, BVHParticles0, 0,
+			Implicit1, Shape1, BVHParticles1, 0,
 			ParticleWorldTransform0, ShapeRelativeTransform0,
 			ParticleWorldTransform1, ShapeRelativeTransform1,
 			CullDistance, Dt, bEnableCCDSweep, LocalContext);
@@ -1078,8 +1078,8 @@ namespace Chaos
 				// Create the constraint for the shape pair if the bounds overlap
 				Collisions::ConstructConstraints(
 					Particle0, Particle1,
-					Object0.Geometry, Object0.ShapeInstance, BVHParticles0,
-					Object1.Geometry, Object1.ShapeInstance, BVHParticles1,
+					Object0.Geometry, Object0.ShapeInstance, BVHParticles0, Object0.LeafObjectIndex,
+					Object1.Geometry, Object1.ShapeInstance, BVHParticles1, Object1.LeafObjectIndex,
 					ParticleWorldTransform0, Object0.RelativeTransform,
 					ParticleWorldTransform1, Object1.RelativeTransform,
 					CullDistance, Dt, bEnableCCDSweep, LocalContext);
@@ -1300,7 +1300,7 @@ namespace Chaos
 			}
 		}
 
-		// If we get here, we did not activate hte constraint and it should be disabled for this tick
+		// If we get here, we did not activate the constraint and it should be disabled for this tick
 		Constraint->SetDisabled(true);
 	}
 
@@ -1364,15 +1364,12 @@ namespace Chaos
 		}
 #endif
 
-		// If we did get a hit but it's at TOI = 1, treat this constraint as a regular non-swept constraint (skip the rewind)
-		if ((!bDidSweep) || Constraint->GetCCDTimeOfImpact() == FReal(1))
+		// If we did not get a sweep hit (TOI > 1) or did not sweep (bDidSweep = false), we need to run standard collision detection at T=1.
+		// Likewise, if we did get a sweep hit but it's at TOI = 1, treat this constraint as a regular non-swept constraint and skip the rewind.
+		// NOTE: The sweep will report TOI==1 for "shallow" sweep hits below the CCD thresholds in the constraint.
+		if ((!bDidSweep) || (Constraint->GetCCDTimeOfImpact() >= FReal(1)))
 		{
 			Collisions::UpdateConstraint(*Constraint, Constraint->GetShapeWorldTransform0(), Constraint->GetShapeWorldTransform1(), Dt);
-		}
-
-		// If the sweep did not find a hit, or we are treating it like a regular contact, we will skip the CCD rewind step for this contact
-		if (Constraint->GetCCDTimeOfImpact() >= FReal(1))
-		{
 			Constraint->SetCCDSweepEnabled(false);
 		}
 
diff --git a/Engine/Source/Runtime/Experimental/Chaos/Private/Chaos/CollisionResolution.cpp b/Engine/Source/Runtime/Experimental/Chaos/Private/Chaos/CollisionResolution.cpp
--- a/Engine/Source/Runtime/Experimental/Chaos/Private/Chaos/CollisionResolution.cpp
+++ b/Engine/Source/Runtime/Experimental/Chaos/Private/Chaos/CollisionResolution.cpp
@@ -1678,10 +1678,12 @@ namespace Chaos
 			const FImplicitObject* Implicit0, 
 			const FPerShapeData* Shape0, 
 			const FBVHParticles* Simplicial0, 
+			const int32 ImplicitID0,
 			const FImplicitObject* Implicit1, 
 			const FPerShapeData* Shape1, 
 			const FBVHParticles* Simplicial1, 
-			const FRigidTransform3& ParticleWorldTransform0, 
+			const int32 ImplicitID1,
+			const FRigidTransform3& ParticleWorldTransform0,
 			const FRigidTransform3& LocalTransform0, 
 			const FRigidTransform3& ParticleWorldTransform1, 
 			const FRigidTransform3& LocalTransform1, 
@@ -1692,7 +1694,7 @@ namespace Chaos
 		{
 			INC_DWORD_STAT(STAT_ChaosCollisionCounter_NumParticlePairs);
 
-			ConstructConstraintsRecursive(Particle0, Particle1, Implicit0, Shape0->AsShapeInstance(), Simplicial0, 0, Implicit1, Shape1->AsShapeInstance(), Simplicial1, 0, ParticleWorldTransform0, LocalTransform0, ParticleWorldTransform1, LocalTransform1, CullDistance, dT, bEnableSweep, Context);
+			ConstructConstraintsRecursive(Particle0, Particle1, Implicit0, Shape0->AsShapeInstance(), Simplicial0, ImplicitID0, Implicit1, Shape1->AsShapeInstance(), Simplicial1, ImplicitID1, ParticleWorldTransform0, LocalTransform0, ParticleWorldTransform1, LocalTransform1, CullDistance, dT, bEnableSweep, Context);
 		}
 
 		// This is deprecated
diff --git a/Engine/Source/Runtime/Experimental/Chaos/Public/Chaos/CollisionResolution.h b/Engine/Source/Runtime/Experimental/Chaos/Public/Chaos/CollisionResolution.h
--- a/Engine/Source/Runtime/Experimental/Chaos/Public/Chaos/CollisionResolution.h
+++ b/Engine/Source/Runtime/Experimental/Chaos/Public/Chaos/CollisionResolution.h
@@ -64,7 +64,27 @@ namespace Chaos
 		bool CHAOS_API ShouldUseCCD(const FGeometryParticleHandle* Particle0, const FVec3& DeltaX0, const FGeometryParticleHandle* Particle1, const FVec3& DeltaX1, FVec3& Dir, FReal& Length);
 
 		// Create constraints for the particle pair. This could create multiple constraints: one for each potentially colliding shape pair in multi-shape particles.
-		void CHAOS_API ConstructConstraints(TGeometryParticleHandle<FReal, 3>* Particle0, TGeometryParticleHandle<FReal, 3>* Particle1, const FImplicitObject* Implicit0, const FPerShapeData* Shape0, const FBVHParticles* Simplicial0, const FImplicitObject* Implicit1, const FPerShapeData* Shape1, const FBVHParticles* Simplicial1, const FRigidTransform3& ParticleWorldTransform0, const FRigidTransform3& Transform0, const FRigidTransform3& ParticleWorldTransform1, const FRigidTransform3& Transform1, const FReal CullDistance, const FReal Dt, const bool bEnableSweep, const FCollisionContext& Context);
+		// @param ImplicitID0 Used to differentiate contacts when an implicit is used multiple times in a geometry hierarchy
+		// @param ImplicitID1 Used to differentiate contacts when an implicit is used multiple times in a geometry hierarchy
+		void CHAOS_API ConstructConstraints(
+			TGeometryParticleHandle<FReal, 3>* Particle0, 
+			TGeometryParticleHandle<FReal, 3>* Particle1, 
+			const FImplicitObject* Implicit0, 
+			const FPerShapeData* Shape0, 
+			const FBVHParticles* Simplicial0, 
+			const int32 ImplicitID0,
+			const FImplicitObject* Implicit1,
+			const FPerShapeData* Shape1, 
+			const FBVHParticles* Simplicial1, 
+			const int32 ImplicitID1,
+			const FRigidTransform3& ParticleWorldTransform0,
+			const FRigidTransform3& Transform0, 
+			const FRigidTransform3& ParticleWorldTransform1, 
+			const FRigidTransform3& Transform1, 
+			const FReal CullDistance, 
+			const FReal Dt, 
+			const bool bEnableSweep, 
+			const FCollisionContext& Context);
 
 		// @todo(chaos): this should not be exposed but is currently used in tests
 		void UpdateLevelsetLevelsetConstraint(const FRigidTransform3& WorldTransform0, const FRigidTransform3& WorldTransform1, const FReal Dt, FPBDCollisionConstraint& Constraint);
@@ -72,17 +92,22 @@ namespace Chaos
 		// Reset per-frame collision stat counters
 		void CHAOS_API ResetChaosCollisionCounters();
 
-
 		//
 		//
 		// DEPRECATED STUFF (see above)
 		//
 		//
 
+		UE_DEPRECATED(5.3, "Replaced with a version that takes a ImplicitIDs")
+		inline void CHAOS_API ConstructConstraints(TGeometryParticleHandle<FReal, 3>* Particle0, TGeometryParticleHandle<FReal, 3>* Particle1, const FImplicitObject* Implicit0, const FPerShapeData* Shape0, const FBVHParticles* Simplicial0, const FImplicitObject* Implicit1, const FPerShapeData* Shape1, const FBVHParticles* Simplicial1, const FRigidTransform3& ParticleWorldTransform0, const FRigidTransform3& Transform0, const FRigidTransform3& ParticleWorldTransform1, const FRigidTransform3& Transform1, const FReal CullDistance, const FReal Dt, const bool bEnableSweep, const FCollisionContext& Context)
+		{
+			ConstructConstraints(Particle0, Particle1, Implicit0, Shape0, Simplicial0, 0, Implicit1, Shape1, Simplicial1, 0, ParticleWorldTransform0, Transform0, ParticleWorldTransform1, Transform1, CullDistance, Dt, bEnableSweep, Context);
+		}
+
 		UE_DEPRECATED(5.3, "Replaced with a version that takes a bEnableSweep flag")
 		inline void CHAOS_API ConstructConstraints(TGeometryParticleHandle<FReal, 3>* Particle0, TGeometryParticleHandle<FReal, 3>* Particle1, const FImplicitObject* Implicit0, const FPerShapeData* Shape0, const FBVHParticles* Simplicial0, const FImplicitObject* Implicit1, const FPerShapeData* Shape1, const FBVHParticles* Simplicial1, const FRigidTransform3& ParticleWorldTransform0, const FRigidTransform3& Transform0, const FRigidTransform3& ParticleWorldTransform1, const FRigidTransform3& Transform1, const FReal CullDistance, const FReal Dt, const FCollisionContext& Context)
 		{
-			ConstructConstraints(Particle0, Particle1, Implicit0, Shape0, Simplicial0, Implicit1, Shape1, Simplicial1, ParticleWorldTransform0, Transform0, ParticleWorldTransform1, Transform1, CullDistance, Dt, false, Context);
+			ConstructConstraints(Particle0, Particle1, Implicit0, Shape0, Simplicial0, 0, Implicit1, Shape1, Simplicial1, 0, ParticleWorldTransform0, Transform0, ParticleWorldTransform1, Transform1, CullDistance, Dt, false, Context);
 		}
 
 		template<ECollisionUpdateType UpdateType>
