Geometry Collection : fix damage threshold being 0 when no connection is found , causing the parent to cluster autoamtically break at start
- setting the strength of the piece to maximum

diff --git a/Engine/Source/Runtime/Experimental/Chaos/Private/PhysicsProxy/GeometryCollectionPhysicsProxy.cpp b/Engine/Source/Runtime/Experimental/Chaos/Private/PhysicsProxy/GeometryCollectionPhysicsProxy.cpp
--- a/Engine/Source/Runtime/Experimental/Chaos/Private/PhysicsProxy/GeometryCollectionPhysicsProxy.cpp
+++ b/Engine/Source/Runtime/Experimental/Chaos/Private/PhysicsProxy/GeometryCollectionPhysicsProxy.cpp
@@ -1424,15 +1424,16 @@ float FGeometryCollectionPhysicsProxy::ComputeMaterialBasedDamageThreshold_Inter
 		// strain actually store the surface area 
 		TotalConnectionArea += Connection.GetArea();
 	}
-
-	// This model compute the total surface are from the connections and compute the maximum force 
-	if (const Chaos::FChaosPhysicsMaterial* ChaosMaterial = Parameters.PhysicalMaterialHandle.Get())
+	if (TotalConnectionArea > 0)
 	{
-		// force unit is Kg.cm/s2 here
-		const float ForceThreshold = ChaosMaterial->Strength.TensileStrength * TotalConnectionArea;
-		DamageThreshold = ForceThreshold;
+		// This model compute the total surface are from the connections and compute the maximum force 
+		if (const Chaos::FChaosPhysicsMaterial* ChaosMaterial = Parameters.PhysicalMaterialHandle.Get())
+		{
+			// force unit is Kg.cm/s2 here
+			const float ForceThreshold = ChaosMaterial->Strength.TensileStrength * TotalConnectionArea;
+			DamageThreshold = ForceThreshold;
+		}
 	}
-
 	return DamageThreshold;
 }
 
