I/O - move slab allocator to internal for reuse between backends

diff --git a/Engine/Source/Runtime/Core/Internal/IO/IoAllocators.h b/Engine/Source/Runtime/Core/Internal/IO/IoAllocators.h
new file mode 100644
--- /dev/null
+++ b/Engine/Source/Runtime/Core/Internal/IO/IoAllocators.h
@@ -0,0 +1,77 @@
+// Copyright Epic Games, Inc. All Rights Reserved.
+
+#pragma once
+
+#include "Templates/TypeCompatibleBytes.h"
+
+template <typename T, uint16 SlabSize = 4096>
+class TSingleThreadedSlabAllocator
+{
+public:
+	TSingleThreadedSlabAllocator()
+	{
+		CurrentSlab = new FSlab();
+	}
+
+	~TSingleThreadedSlabAllocator()
+	{
+		check(CurrentSlab->Allocated == CurrentSlab->Freed);
+		delete CurrentSlab;
+	}
+
+	template <typename... ArgsType>
+	T* Construct(ArgsType&&... Args)
+	{
+		return new(Alloc()) T(Forward<ArgsType>(Args)...);
+	}
+
+	void Destroy(T* Ptr)
+	{
+		Ptr->~T();
+		Free(Ptr);
+	}
+
+private:
+	struct FSlab;
+
+	struct FElement
+	{
+		TTypeCompatibleBytes<T> Data;
+		FSlab* Slab = nullptr;
+	};
+
+	struct FSlab
+	{
+		uint16 Allocated = 0;
+		uint16 Freed = 0;
+		FElement Elements[SlabSize];
+	};
+
+	T* Alloc()
+	{
+		uint16 ElementIndex = CurrentSlab->Allocated++;
+		check(ElementIndex < SlabSize);
+		FElement* Element = CurrentSlab->Elements + ElementIndex;
+		Element->Slab = CurrentSlab;
+		if (CurrentSlab->Allocated == SlabSize)
+		{
+			//TRACE_CPUPROFILER_EVENT_SCOPE(AllocSlab);
+			CurrentSlab = new FSlab();
+		}
+		return Element->Data.GetTypedPtr();
+	}
+
+	void Free(T* Ptr)
+	{
+		FElement* Element = reinterpret_cast<FElement*>(Ptr);
+		FSlab* Slab = Element->Slab;
+		if (++Slab->Freed == SlabSize)
+		{
+			//TRACE_CPUPROFILER_EVENT_SCOPE(FreeSlab);
+			check(Slab->Freed == Slab->Allocated);
+			delete Slab;
+		}
+	}
+
+	FSlab* CurrentSlab = nullptr;
+};
diff --git a/Engine/Source/Runtime/Experimental/IoStoreOnDemand/Private/OnDemandIoDispatcherBackend.cpp b/Engine/Source/Runtime/Experimental/IoStoreOnDemand/Private/OnDemandIoDispatcherBackend.cpp
--- a/Engine/Source/Runtime/Experimental/IoStoreOnDemand/Private/OnDemandIoDispatcherBackend.cpp
+++ b/Engine/Source/Runtime/Experimental/IoStoreOnDemand/Private/OnDemandIoDispatcherBackend.cpp
@@ -9,6 +9,7 @@
 #include "HAL/PreprocessorHelpers.h"
 #include "HAL/Runnable.h"
 #include "HAL/RunnableThread.h"
+#include "IO/IoAllocators.h"
 #include "IO/IoCache.h"
 #include "IO/IoDispatcher.h"
 #include "IO/IoOffsetLength.h"
@@ -45,80 +46,6 @@ TRACE_DECLARE_INT_COUNTER(PendingHttpRequestCount, TEXT("OnDemandIoBackend/HTTP/
 TRACE_DECLARE_INT_COUNTER(InflightHttpRequestCount, TEXT("OnDemandIoBackend/HTTP/InflightCount"));
 //TRACE_DECLARE_FLOAT_COUNTER(TotalRequestSeconds, TEXT("OnDemandIoBackend/Duration/TotalSeconds"));
 
-///////////////////////////////////////////////////////////////////////////////
-//TODO: Move this into core and reuse in file backend (see IoDispatcherFileBackendTypes.h) 
-template <typename T, uint16 SlabSize = 4096>
-class TSingleThreadedSlabAllocator
-{
-public:
-	TSingleThreadedSlabAllocator()
-	{
-		CurrentSlab = new FSlab();
-	}
-
-	~TSingleThreadedSlabAllocator()
-	{
-		check(CurrentSlab->Allocated == CurrentSlab->Freed);
-		delete CurrentSlab;
-	}
-
-	template <typename... ArgsType>
-	T* Construct(ArgsType&&... Args)
-	{
-		return new(Alloc()) T(Forward<ArgsType>(Args)...);
-	}
-
-	void Destroy(T* Ptr)
-	{
-		Ptr->~T();
-		Free(Ptr);
-	}
-
-private:
-	struct FSlab;
-
-	struct FElement
-	{
-		TTypeCompatibleBytes<T> Data;
-		FSlab* Slab = nullptr;
-	};
-
-	struct FSlab
-	{
-		uint16 Allocated = 0;
-		uint16 Freed = 0;
-		FElement Elements[SlabSize];
-	};
-
-	T* Alloc()
-	{
-		uint16 ElementIndex = CurrentSlab->Allocated++;
-		check(ElementIndex < SlabSize);
-		FElement* Element = CurrentSlab->Elements + ElementIndex;
-		Element->Slab = CurrentSlab;
-		if (CurrentSlab->Allocated == SlabSize)
-		{
-			//TRACE_CPUPROFILER_EVENT_SCOPE(AllocSlab);
-			CurrentSlab = new FSlab();
-		}
-		return Element->Data.GetTypedPtr();
-	}
-
-	void Free(T* Ptr)
-	{
-		FElement* Element = reinterpret_cast<FElement*>(Ptr);
-		FSlab* Slab = Element->Slab;
-		if (++Slab->Freed == SlabSize)
-		{
-			//TRACE_CPUPROFILER_EVENT_SCOPE(FreeSlab);
-			check(Slab->Freed == Slab->Allocated);
-			delete Slab;
-		}
-	}
-
-	FSlab* CurrentSlab = nullptr;
-};
-
 ///////////////////////////////////////////////////////////////////////////////
 FIoHash GetChunkKey(const FIoHash& ChunkHash, const FIoOffsetAndLength& Range)
 {
diff --git a/Engine/Source/Runtime/PakFile/Internal/IoDispatcherFileBackendTypes.h b/Engine/Source/Runtime/PakFile/Internal/IoDispatcherFileBackendTypes.h
--- a/Engine/Source/Runtime/PakFile/Internal/IoDispatcherFileBackendTypes.h
+++ b/Engine/Source/Runtime/PakFile/Internal/IoDispatcherFileBackendTypes.h
@@ -2,6 +2,7 @@
 
 #pragma once
 
+#include "IO/IoAllocators.h"
 #include "IO/IoStore.h"
 #include "IO/IoDispatcherBackend.h"
 #include "Async/MappedFileHandle.h"
@@ -701,76 +702,7 @@ private:
 };
 
 template <typename T, uint16 SlabSize = 4096>
-class TIoDispatcherSingleThreadedSlabAllocator
-{
-public:
-	TIoDispatcherSingleThreadedSlabAllocator()
-	{
-		CurrentSlab = new FSlab();
-	}
-
-	~TIoDispatcherSingleThreadedSlabAllocator()
-	{
-		check(CurrentSlab->Allocated == CurrentSlab->Freed);
-		delete CurrentSlab;
-	}
-
-	template <typename... ArgsType>
-	T* Construct(ArgsType&&... Args)
-	{
-		return new(Alloc()) T(Forward<ArgsType>(Args)...);
-	}
-
-	void Destroy(T* Ptr)
-	{
-		Ptr->~T();
-		Free(Ptr);
-	}
-
-private:
-	struct FSlab;
-
-	struct FElement
-	{
-		TTypeCompatibleBytes<T> Data;
-		FSlab* Slab = nullptr;
-	};
-
-	struct FSlab
-	{
-		uint16 Allocated = 0;
-		uint16 Freed = 0;
-		FElement Elements[SlabSize];
-	};
-
-	T* Alloc()
-	{
-		uint16 ElementIndex = CurrentSlab->Allocated++;
-		check(ElementIndex < SlabSize);
-		FElement* Element = CurrentSlab->Elements + ElementIndex;
-		Element->Slab = CurrentSlab;
-		if (CurrentSlab->Allocated == SlabSize)
-		{
-			//TRACE_CPUPROFILER_EVENT_SCOPE(AllocSlab);
-			CurrentSlab = new FSlab();
-		}
-		return Element->Data.GetTypedPtr();
-	}
-
-	void Free(T* Ptr)
-	{
-		FElement* Element = reinterpret_cast<FElement*>(Ptr);
-		FSlab* Slab = Element->Slab;
-		if (++Slab->Freed == SlabSize)
-		{
-			//TRACE_CPUPROFILER_EVENT_SCOPE(FreeSlab);
-			check(Slab->Freed == Slab->Allocated);
-			delete Slab;
-		}
-	}
-
-	FSlab* CurrentSlab = nullptr;
-};
+using TIoDispatcherSingleThreadedSlabAllocator = TSingleThreadedSlabAllocator<T, SlabSize>;
 
 struct PAKFILE_API FFileIoStoreReadRequestLink
 {
