Restored IoStoreOnDemandIoStore's client
#rnx

diff --git a/Engine/Source/Runtime/Experimental/IoStoreOnDemand/Private/CoreHttp/Client.h b/Engine/Source/Runtime/Experimental/IoStoreOnDemand/Private/CoreHttp/Client.h
new file mode 100644
--- /dev/null
+++ b/Engine/Source/Runtime/Experimental/IoStoreOnDemand/Private/CoreHttp/Client.h
@@ -0,0 +1,198 @@
+// Copyright Epic Games, Inc. All Rights Reserved.
+
+#pragma once
+
+#if !defined(NO_UE_INCLUDES)
+#include "Containers/Array.h"
+#include "Containers/StringView.h"
+#include "HAL/CriticalSection.h"
+#include "Memory/MemoryView.h"
+#endif
+
+#include <atomic>
+
+#if !defined(COREHTTP_API)
+#	define COREHTTP_API IOSTOREONDEMAND_API
+#endif
+
+////////////////////////////////////////////////////////////////////////////////
+class FIoBuffer;
+
+namespace UE::HTTP
+{
+
+////////////////////////////////////////////////////////////////////////////////
+enum class EMimeType
+{
+	Unknown = 0,
+	Text,
+	Binary,
+	Json,
+	Xml,
+	/* UE_CUSTOM_MIME_TYPES
+	CbObject,
+	CbPackage,
+	CompressedBuffer,
+	*/
+	Count
+};
+
+////////////////////////////////////////////////////////////////////////////////
+enum class EStatusCodeClass
+{
+	Informational,
+	Successful,
+	Redirection,
+	ClientError,
+	ServerError,
+	Unknown,
+};
+
+////////////////////////////////////////////////////////////////////////////////
+using	FTicket = uint64;
+struct	FActivity;
+
+////////////////////////////////////////////////////////////////////////////////
+class COREHTTP_API FConnectionPool
+{
+public:
+	struct FParams
+	{
+		int32				SetHostFromUrl(FAnsiStringView Url);
+		uint32				ConnectionCount;
+		struct {
+			FAnsiStringView	Name;
+			uint32			Port = 80;
+			/* uint8		Protocol */;
+		}					Host;
+		/*
+		enum class ProxyType { Http, Socks4 };
+		Proxy = { ip, port, type }
+		 */
+	};
+
+							FConnectionPool() = default;
+							FConnectionPool(const FParams& Params);
+							~FConnectionPool();
+							FConnectionPool(FConnectionPool&& Rhs)	{ *this = MoveTemp(Rhs); }
+	FConnectionPool&		operator = (FConnectionPool&& Rhs)		{ Swap(Ptr, Rhs.Ptr); return *this; }
+
+private:
+	friend					class FEventLoop;
+	class FSocketPool*		Ptr = nullptr;
+
+private:
+							FConnectionPool(const FConnectionPool&) = delete;
+	FConnectionPool&		operator = (const FConnectionPool&) = delete;
+};
+
+////////////////////////////////////////////////////////////////////////////////
+class COREHTTP_API FRequest
+{
+public:
+						~FRequest();
+						FRequest(FRequest&& Rhs);
+	bool				IsValid() const { return Ptr != nullptr; }
+	FRequest&&			Accept(EMimeType MimeType);
+	FRequest&&			Accept(FAnsiStringView MimeType);
+	FRequest&&			Header(FAnsiStringView Key, FAnsiStringView Value);
+	void				Content(const void* Data, SIZE_T Size, EMimeType MimeType);
+	void				Content(const void* Data, SIZE_T Size, FAnsiStringView MimeType);
+
+private:
+	friend				class FEventLoop;
+						FRequest() = default;
+	FActivity*			Ptr = nullptr;
+
+private:
+						FRequest(const FRequest&) = delete;
+	FRequest&			operator = (const FRequest&) = delete;
+	FRequest&			operator = (FRequest&&) = delete;
+};
+
+////////////////////////////////////////////////////////////////////////////////
+class COREHTTP_API FResponse
+{
+public:
+	EStatusCodeClass	GetStatus() const;
+	uint32				GetStatusCode() const;
+	FAnsiStringView 	GetStatusMessage() const;
+	int64				GetContentLength() const;
+	EMimeType			GetContentType() const;
+	void				GetContentType(FAnsiStringView& Out) const;
+	FAnsiStringView 	GetHeader(FAnsiStringView Name) const;
+	void				SetDestination(FIoBuffer* Buffer);
+
+private:
+						FResponse() = delete;
+						FResponse(const FResponse&) = delete;
+						FResponse(FResponse&&) = delete;
+	FResponse&			operator = (const FResponse&) = delete;
+	FResponse&			operator = (FResponse&&) = delete;
+};
+
+////////////////////////////////////////////////////////////////////////////////
+class COREHTTP_API FTicketStatus
+{
+public:
+	enum class EId : uint8 { Response, Content, Cancelled, Error };
+	EId					GetId() const;
+	UPTRINT				GetParam() const;
+	FTicket				GetTicket() const;
+	uint32				GetIndex() const;
+	FResponse&			GetResponse() const;		// if GetId() == EId::Response
+	const FIoBuffer&	GetContent() const;			// if GetId() == EId::Content
+	uint32				GetContentLength() const;	//  |
+	const char*			GetErrorReason() const;		// if GetId() == EId::Error
+
+private:
+						FTicketStatus() = delete;
+						FTicketStatus(const FTicketStatus&) = delete;
+						FTicketStatus(FTicketStatus&&) = delete;
+	FTicketStatus&		operator = (const FTicketStatus&) = delete;
+	FTicketStatus&		operator = (FTicketStatus&&) = delete;
+};
+
+using FTicketSink = TFunction<void (const FTicketStatus&)>;
+
+////////////////////////////////////////////////////////////////////////////////
+class COREHTTP_API FEventLoop
+{
+public:
+	struct FRequestParams
+	{
+		uint32	BufferSize	= 256;
+	 // uint32	PageSize	= 2 << 10;
+	};
+
+	template <typename... T> [[nodiscard]] FRequest Get(T&&... t)  { return Request("GET",  Forward<T&&>(t)...); }
+	template <typename... T> [[nodiscard]] FRequest Post(T&&... t) { return Request("POST", Forward<T&&>(t)...); }
+
+							FEventLoop() = default;
+							~FEventLoop();
+	uint32					Tick();
+	bool					IsIdle() const;
+	void					Cancel(FTicket Ticket);
+	[[nodiscard]] FRequest	Request(FAnsiStringView Method, FAnsiStringView Url, const FRequestParams* Params=nullptr);
+	[[nodiscard]] FRequest	Request(FAnsiStringView Method, FAnsiStringView Path, FConnectionPool& Pool, const FRequestParams* Params=nullptr);
+	FTicket					Send(FRequest&& Request, FTicketSink Sink, UPTRINT Param=0);
+
+private:
+	FRequest				Request(FAnsiStringView Method, FAnsiStringView Path, FActivity* Activity);
+	FCriticalSection		Lock;
+	std::atomic<uint64>		FreeSlots		= ~0ull;
+	std::atomic<uint64>		Cancels			= 0;
+	uint64					PrevFreeSlots	= ~0ull;
+	TArray<FActivity*>		Pending;
+	TArray<FActivity*>		Active;
+
+private:
+							FEventLoop(const FEventLoop&)	= delete;
+							FEventLoop(FEventLoop&&)		= delete;
+	FEventLoop&				operator = (const FEventLoop&)	= delete;
+	FEventLoop&				operator = (FEventLoop&&)		= delete;
+};
+
+} // namespace UE::HTTP
+
+/* vim: set noet : */
diff --git a/Engine/Source/Runtime/Experimental/IoStoreOnDemand/Private/CoreHttp/Http.cpp b/Engine/Source/Runtime/Experimental/IoStoreOnDemand/Private/CoreHttp/Http.cpp
new file mode 100644
--- /dev/null
+++ b/Engine/Source/Runtime/Experimental/IoStoreOnDemand/Private/CoreHttp/Http.cpp
@@ -0,0 +1,2343 @@
+// Copyright Epic Games, Inc. All Rights Reserved.
+
+#include "CoreHttp/Client.h"
+
+#if !defined(NO_UE_INCLUDES)
+#include "IO/IoBuffer.h"
+#include "Math/UnrealMathUtility.h"
+#include "Misc/ScopeExit.h"
+#include "Misc/ScopeLock.h"
+#include "Misc/StringBuilder.h"
+#include "ProfilingDebugging/CpuProfilerTrace.h"
+#include "Tasks/Task.h"
+#endif
+
+// {{{1 platforms ..............................................................
+
+#if PLATFORM_MICROSOFT
+#	if !defined(NO_UE_INCLUDES)
+#	include "Windows/AllowWindowsPlatformTypes.h"
+#		include <winsock2.h>
+#		include <ws2tcpip.h>
+#	include "Windows/HideWindowsPlatformTypes.h"
+#	else
+#		include <winsock2.h>
+#		include <ws2tcpip.h>
+#		pragma comment(lib, "Ws2_32.lib")
+#	endif // NO_UE_INCLUDES
+	using SocketType	= SOCKET;
+	using MsgFlagType	= int;
+
+	// Winsock defines "PF_MAX" indicating the protocol family count. This
+	// however competes with UE definitions related to pixel formats.
+	#if defined(PF_MAX)
+	#	undef PF_MAX
+	#endif
+
+#elif PLATFORM_APPLE | PLATFORM_UNIX | PLATFORM_ANDROID
+#	include <arpa/inet.h>
+#	include <fcntl.h>
+#	include <netdb.h>
+#	include <sys/select.h>
+#	include <sys/socket.h>
+#	include <unistd.h>
+	using SocketType	= int;
+	using MsgFlagType	= int;
+	static int32 closesocket(int32 Socket) { return close(Socket); }
+#else
+#	include "CoreHttp/Http.inl"
+#endif
+
+static_assert(sizeof(sockaddr_in::sin_addr) == sizeof(uint32));
+
+#if PLATFORM_MICROSOFT
+	static int32 LastSocketResult() { return WSAGetLastError(); }
+#	define IsSocketResult(err)		(LastSocketResult() == WSA##err)
+#else
+	static int32 LastSocketResult() { return errno; }
+#	define IsSocketResult(err)		(LastSocketResult() == err)
+	static_assert(EWOULDBLOCK == EAGAIN);
+#endif
+
+enum : SocketType { InvalidSocket = -1 };
+
+// }}}
+
+namespace UE::HTTP
+{
+
+// {{{1 misc ...................................................................
+
+////////////////////////////////////////////////////////////////////////////////
+template <typename LambdaType>
+static void EnumerateHeaders(FAnsiStringView Headers, LambdaType&& Lambda)
+{
+	// NB. here we are assuming that we will be dealing with servers that will
+	// not be returning headers with "obsolete line folding".
+
+	auto IsOws = [] (int32 c) { return (c == ' ') | (c == '\t'); };
+
+	const char* Cursor = Headers.GetData();
+	const char* End = Cursor + Headers.Len();
+	do
+	{
+		int32 ColonIndex = 0;
+		for (; Cursor + ColonIndex < End; ++ColonIndex)
+		{
+			if (Cursor[ColonIndex] == ':')
+			{
+				break;
+			}
+		}
+
+		Cursor += ColonIndex;
+
+		const char* Right = Cursor + 1;
+		while (Right < End)
+		{
+			if (Right[0] != '\r' || Right[1] != '\n')
+			{
+				Right += 1 + (Right[1] != '\r');
+				continue;
+			}
+
+			FAnsiStringView Name(Cursor - ColonIndex, ColonIndex);
+
+			const char* Left = Cursor + 1;
+			for (; IsOws(Left[0]); ++Left);
+
+			Cursor = Right;
+			for (; Cursor > Left + 1 && IsOws(Cursor[-1]); --Cursor);
+
+			FAnsiStringView Value (Left, ptrdiff_t(Cursor - Left));
+
+			if (!Lambda(Name, Value))
+			{
+				Right = End;
+			}
+
+			break;
+		}
+
+		Cursor = Right + 2;
+	} while (Cursor < End);
+}
+
+////////////////////////////////////////////////////////////////////////////////
+static int32 FindMessageTerminal(const char* Data, uint32 Length)
+{
+	for (uint32 i = 4; i <= Length; ++i)
+	{
+		uint32 Candidate;
+		::memcpy(&Candidate, Data + i - 4, sizeof(Candidate));
+		if (Candidate == 0x0a0d0a0d)
+		{
+			return i;
+		}
+
+		i += (Data[i - 1] > 0x0d) ? 3 : 0;
+	}
+
+	return -1;
+}
+
+////////////////////////////////////////////////////////////////////////////////
+static int64 CrudeToInt(FAnsiStringView View)
+{
+	// FCStringAnsi::* is not used to mitigate any locale hiccups. By
+	// initialising 'Value' with MSB set we can detect cases where View did not
+	// start with digits. This works as we won't be using this on huge numbers.
+	int64 Value = 0x8000'0000'0000'0000ll;
+	for (int32 c : View)
+	{
+		uint32 Digit = c - '0';
+		if (Digit > 9u)
+		{
+			break;
+		}
+		Value *= 10;
+		Value += Digit;
+	}
+	return Value;
+};
+
+////////////////////////////////////////////////////////////////////////////////
+struct FMessageOffsets
+{
+	uint8	StatusCode;
+	uint8	Message;
+	uint16	Headers;
+};
+
+static int32 ParseMessage(FAnsiStringView Message, FMessageOffsets& Out)
+{
+	const FAnsiStringView Protocol("HTTP/1.1 ");
+
+	// Check there's enough data
+	if (Message.Len() < Protocol.Len() + 1) // "+1" accounts for at least one digit
+	{
+		return -1;
+	}
+
+	const char* Cursor = Message.GetData();
+
+	// Check for the expected protocol
+	if (FAnsiStringView(Cursor, 9) != Protocol)
+	{
+		return -1;
+	}
+	int32 i = Protocol.Len();
+
+	// Trim left and tightly reject anything adventurous
+	for (int n = 32; Cursor[i] == ' ' && i < n; ++i);
+	Out.StatusCode = i;
+
+	// At least one status line digit. (Note to self; expect exactly three)
+	for (int n = 32; uint32(Cursor[i] - 0x30) <= 9 && i < n; ++i);
+	if (uint32(i - Out.StatusCode - 1) > 32)
+	{
+		return -1;
+	}
+
+	// Trim left
+	for (int n = 32; Cursor[i] == ' ' && i < n; ++i);
+	Out.Message = i;
+
+	// Extra conservative length allowance
+	if (i > 32)
+	{
+		return -1;
+	}
+
+	// Find \r\n
+	for (; Cursor[i] != '\r'; ++i)
+	{
+		if (i >= 2048)
+		{
+			return -1;
+		}
+	}
+	if (Cursor[i + 1] != '\n')
+	{
+		return -1;
+	}
+	Out.Headers = i + 2;
+
+	return 1;
+}
+
+////////////////////////////////////////////////////////////////////////////////
+struct FUrlOffsets
+{
+	struct Slice
+	{
+		FAnsiStringView	Get(FAnsiStringView Url) const { return Url.Mid(Off, Len); }
+						operator bool () const { return Len != 0; }
+		uint8			Off;
+		uint8			Len;
+	};
+	Slice				UserInfo;
+	Slice				HostName;
+	Slice				Port;
+	uint8				Path;
+	uint8				SchemeLength;
+};
+
+static int32 ParseUrl(FAnsiStringView Url, FUrlOffsets& Out)
+{
+	static const int32 LengthLimit = 127;
+
+	Out = {};
+
+	const char* Start = Url.GetData();
+	const char* Cursor = Start;
+	const char* End = Start + Url.Len();
+
+	// Scheme
+	int32 i = 0;
+	for (; i < 5; ++i)
+	{
+		if (uint32(Cursor[i] - 'a') > uint32('z' - 'a'))
+		{
+			break;
+		}
+	}
+
+	Out.SchemeLength = uint8(i);
+	FAnsiStringView Scheme = Url.Left(i);
+	if (Scheme != "http" && Scheme != "https")
+	{
+		return -1;
+	}
+
+	// Separator and authority
+	if (Cursor[i] != ':' || Cursor[i + 1] != '/' || Cursor[i + 2] != '/')
+	{
+		return -1;
+	}
+	Cursor += i + 3;
+
+	struct { int32 c; int32 i; } Seps[2];
+	int32 SepCount = 0;
+	for (i = 0; i < LengthLimit - 8; ++i) // '8' is roughly "http[s]://"
+	{
+		int32 c = Cursor[i];
+		if (c < '-')							break;
+		if (c != ':' && c != '@' && c != '/')	continue;
+		if (c == '/' || SepCount >= 2)			break;
+
+		if (c == '@' && SepCount)
+		{
+			SepCount -= (Seps[SepCount - 1].c == ':');
+		}
+		Seps[SepCount++] = { c, i };
+	}
+
+	if (int32 c = Cursor[i]; c)
+	{
+		if (c != '/')
+		{
+			return -1;
+		}
+		Out.Path = uint8(ptrdiff_t(Cursor + i - Start));
+	}
+
+	Out.HostName = { uint8(Scheme.Len() + 3), uint8(i) };
+
+	switch (SepCount)
+	{
+	case 0:
+		break;
+
+	case 1:
+		if (Seps[0].c == ':')
+		{
+			Out.Port = { uint8(Out.HostName.Off + Seps[0].i + 1), uint8(i - Seps[0].i - 1) };
+			Out.HostName.Len = Seps[0].i;
+		}
+		else
+		{
+			Out.UserInfo = { Out.HostName.Off, uint8(Seps[0].i) };
+			Out.HostName.Off += Seps[0].i + 1;
+			Out.HostName.Len -= Seps[0].i + 1;
+		}
+		break;
+
+	case 2:
+		if ((Seps[0].c != '@') | (Seps[1].c != ':'))
+		{
+			return -1;
+		}
+		Out.UserInfo = { Out.HostName.Off, uint8(Seps[0].i) };
+		Out.Port = Out.HostName;
+		Out.Port.Off += Seps[1].i + 1;
+		Out.Port.Len -= Seps[1].i + 1;
+		Out.HostName.Off += Out.UserInfo.Len + 1;
+		Out.HostName.Len -= Out.UserInfo.Len + Out.Port.Len + 2;
+		break;
+
+	default:
+		return -1;
+	}
+
+	bool Bad = false;
+	Bad |= (Out.HostName.Len == 0);
+	Bad |= (Out.UserInfo.Off != 0) & (Out.UserInfo.Len == 0);
+
+	if (Out.Port.Off)
+	{
+		Bad |= (Out.Port.Len == 0);
+		for (int32 j = 0, n = Out.Port.Len; j < n; ++j)
+		{
+			Bad |= (uint32(Start[Out.Port.Off + j] - '0') > 9);
+		}
+	}
+
+	return Bad ? -1 : 1;
+}
+
+
+
+// {{{1 buffer .................................................................
+
+////////////////////////////////////////////////////////////////////////////////
+class alignas(16) FBuffer
+{
+public:
+	struct FMutableSection
+	{
+		char*		Data;
+		uint32		Size;
+	};
+
+								FBuffer() = default;
+								FBuffer(char* InData, uint32 InMax);
+								~FBuffer();
+	FBuffer&					operator = (FBuffer&& Rhs);
+	void						Fix();
+	void						Reset();
+	const char*					GetData() const;
+	uint32						GetSize() const;
+	uint32						GetCapacity() const;
+	template <typename T> T*	Alloc(uint32 Count=1);
+	FMutableSection				GetMutableFree(uint32 MinSize, uint32 PageSize=0);
+	void						AdvanceUsed(uint32 Delta);
+
+private:
+	char*						GetDataPtr();
+	void						Extend(uint32 AtLeast, uint32 PageSize=1024);
+	union
+	{
+		struct
+		{
+			UPTRINT				Data : 63;
+			UPTRINT				Inline : 1;
+		};
+		UPTRINT					DataInline = 0;
+	};
+	uint32						Max = 0;
+	uint32						Used = 0;
+
+private:
+								FBuffer(FBuffer&&) = delete;
+								FBuffer(const FBuffer&) = delete;
+	FBuffer&					operator = (const FBuffer&) = delete;
+};
+
+////////////////////////////////////////////////////////////////////////////////
+FBuffer::FBuffer(char* InData, uint32 InMax)
+: Data(UPTRINT(InData))
+, Inline(1)
+, Max(InMax)
+{
+}
+
+////////////////////////////////////////////////////////////////////////////////
+FBuffer::~FBuffer()
+{
+	if (Data && !Inline)
+	{
+		FMemory::Free(GetDataPtr());
+	}
+}
+
+////////////////////////////////////////////////////////////////////////////////
+FBuffer& FBuffer::operator = (FBuffer&& Rhs)
+{
+	Swap(DataInline, Rhs.DataInline);
+	Swap(Max, Rhs.Max);
+	Swap(Used, Rhs.Used);
+	return *this;
+}
+
+////////////////////////////////////////////////////////////////////////////////
+void FBuffer::Fix()
+{
+	check(Inline);
+	Data += Used;
+	Max -= Used;
+	Used = 0;
+}
+
+////////////////////////////////////////////////////////////////////////////////
+void FBuffer::Reset()
+{
+	Used = 0;
+}
+
+////////////////////////////////////////////////////////////////////////////////
+char* FBuffer::GetDataPtr()
+{
+	return (char*)Data;
+}
+
+////////////////////////////////////////////////////////////////////////////////
+const char* FBuffer::GetData() const
+{
+	return (char*)Data;
+}
+
+////////////////////////////////////////////////////////////////////////////////
+uint32 FBuffer::GetSize() const
+{
+	return Used;
+}
+
+////////////////////////////////////////////////////////////////////////////////
+uint32 FBuffer::GetCapacity() const
+{
+	return Max;
+}
+
+////////////////////////////////////////////////////////////////////////////////
+template <typename T>
+T* FBuffer::Alloc(uint32 Count)
+{
+	uint32 AlignBias = uint32(Data) & (alignof(T) - 1);
+	if (AlignBias)
+	{
+		AlignBias = alignof(T) - AlignBias;
+	}
+
+	uint32 PotentialUsed = Used + AlignBias + (sizeof(T) * Count);
+	if (PotentialUsed > Max)
+	{
+		Extend(PotentialUsed);
+	}
+
+	void* Ret = GetDataPtr() + Used + AlignBias;
+	Used = PotentialUsed;
+	return (T*)Ret;
+}
+
+////////////////////////////////////////////////////////////////////////////////
+FBuffer::FMutableSection FBuffer::GetMutableFree(uint32 MinSize, uint32 PageSize)
+{
+	MinSize = (MinSize == 0 && Used == Max) ? PageSize : MinSize;
+
+	uint32 PotentialUsed = Used + MinSize;
+	if (PotentialUsed > Max)
+	{
+		Extend(PotentialUsed);
+	}
+
+	return FMutableSection{ GetDataPtr() + Used, Max - Used };
+}
+
+////////////////////////////////////////////////////////////////////////////////
+void FBuffer::AdvanceUsed(uint32 Delta)
+{
+	Used += Delta;
+	check(Used <= Max);
+}
+
+////////////////////////////////////////////////////////////////////////////////
+void FBuffer::Extend(uint32 AtLeast, uint32 PageSize)
+{
+	checkSlow((PageSize - 1 & PageSize) == 0);
+
+	--PageSize;
+	Max = (AtLeast + PageSize) & ~PageSize;
+
+	if (!Inline)
+	{
+		Data = UPTRINT(FMemory::Realloc(GetDataPtr(), Max, alignof(FBuffer)));
+		return;
+	}
+
+	const char* PrevData = GetDataPtr();
+	Data = UPTRINT(FMemory::Malloc(Max, alignof(FBuffer)));
+	::memcpy(GetDataPtr(), PrevData, Used);
+	Inline = 0;
+}
+
+
+
+////////////////////////////////////////////////////////////////////////////////
+class FMessageBuilder
+{
+public:
+						FMessageBuilder(FBuffer& Ruffeb);
+	FMessageBuilder&	operator << (FAnsiStringView Lhs);
+
+private:
+	FBuffer&		Buffer;
+};
+
+////////////////////////////////////////////////////////////////////////////////
+FMessageBuilder::FMessageBuilder(FBuffer& Ruffeb)
+: Buffer(Ruffeb)
+{
+}
+
+////////////////////////////////////////////////////////////////////////////////
+FMessageBuilder& FMessageBuilder::operator << (FAnsiStringView Lhs)
+{
+	uint32 Length = uint32(Lhs.Len());
+	FBuffer::FMutableSection Section = Buffer.GetMutableFree(Length);
+	::memcpy(Section.Data, Lhs.GetData(), Length);
+	Buffer.AdvanceUsed(Length);
+	return *this;
+}
+
+
+
+// {{{1 connection-pool ........................................................
+
+////////////////////////////////////////////////////////////////////////////////
+class FSocketPool
+{
+public:
+	enum class EState : uint8 { Unresolved, Busy, Resolved, Error };
+
+					FSocketPool(FAnsiStringView InHostName, uint32 InPort, uint32 InMaxLeases);
+					~FSocketPool();
+	static uint32	GetAllocSize(uint32 MaxLeases);
+	bool			LeaseSocket(SocketType& Out);
+	void			ReturnLease(SocketType Socket);
+	bool			AddIpAddress(uint32 Address);
+	uint32			GetIpAddress() const	{ return IpAddresses[0]; }
+	FAnsiStringView	GetHostName() const		{ return HostName; }
+	uint32			GetPort() const			{ return Port; }
+	EState			GetState() const		{ return State; }
+	void			SetState(EState Value)	{ State = Value; }
+
+private:
+	FAnsiStringView	HostName;
+	uint32			IpAddresses[4] = {};
+	uint16			Port;
+	uint8			LeaseCount = 0;
+	uint8			MaxLeases : 6;
+	EState			State : 2;
+	SocketType		Sockets[1/*...N*/]; // this should be the last member
+};
+
+////////////////////////////////////////////////////////////////////////////////
+FSocketPool::FSocketPool(FAnsiStringView InHostName, uint32 InPort, uint32 InMaxLeases)
+: HostName(InHostName)
+, Port(InPort)
+, MaxLeases(InMaxLeases)
+, State(EState::Unresolved)
+{
+	check(MaxLeases == InMaxLeases); // field overflow
+
+	for (uint32 i = 0; i < InMaxLeases; ++i)
+	{
+		Sockets[i] = InvalidSocket;
+	}
+}
+
+////////////////////////////////////////////////////////////////////////////////
+FSocketPool::~FSocketPool()
+{
+	for (uint32 i = 0; i < MaxLeases; ++i)
+	{
+		if (Sockets[i] == InvalidSocket)
+		{
+			continue;
+		}
+
+		closesocket(Sockets[i]);
+	}
+}
+
+////////////////////////////////////////////////////////////////////////////////
+uint32 FSocketPool::GetAllocSize(uint32 MaxLeases)
+{
+	return sizeof(FSocketPool) + (sizeof(Sockets[0]) * (MaxLeases - 1));
+}
+
+////////////////////////////////////////////////////////////////////////////////
+bool FSocketPool::LeaseSocket(SocketType& Out)
+{
+	check(LeaseCount <= MaxLeases);
+
+	if (LeaseCount == MaxLeases)
+	{
+		return false;
+	}
+
+	Out = Sockets[LeaseCount];
+	++LeaseCount;
+
+	return true;
+}
+
+////////////////////////////////////////////////////////////////////////////////
+void FSocketPool::ReturnLease(SocketType Socket)
+{
+	check(LeaseCount > 0);
+	--LeaseCount;
+	Sockets[LeaseCount] = Socket;
+}
+
+////////////////////////////////////////////////////////////////////////////////
+bool FSocketPool::AddIpAddress(uint32 Address)
+{
+	for (uint32& Entry : IpAddresses)
+	{
+		if (Entry != 0)
+		{
+			continue;
+		}
+
+		Entry = Address;
+		return true;
+	}
+
+	return false;
+}
+
+
+
+////////////////////////////////////////////////////////////////////////////////
+int32 FConnectionPool::FParams::SetHostFromUrl(FAnsiStringView Url)
+{
+	FUrlOffsets Offsets;
+	if (ParseUrl(Url, Offsets) < 0)
+	{
+		return -1;
+	}
+
+	Host.Name = Offsets.HostName.Get(Url);
+
+	if (Offsets.Port)
+	{
+		FAnsiStringView PortView = Offsets.Port.Get(Url);
+		Host.Port = CrudeToInt(PortView);
+	}
+
+	return Offsets.Path;
+}
+
+////////////////////////////////////////////////////////////////////////////////
+FConnectionPool::FConnectionPool(const FParams& Params)
+{
+	check(Params.ConnectionCount - 1 <= 63u);
+	check(Params.Host.Port - 1 <= 0xfffeu);
+
+	// Alloc a new internal object
+	uint32 PoolAllocSize = FSocketPool::GetAllocSize(Params.ConnectionCount);
+	uint32 HostNameLen = Params.Host.Name.Len();
+	uint32 AllocSize = PoolAllocSize + (HostNameLen + 1);
+	auto* Internal = (FSocketPool*)FMemory::Malloc(AllocSize, alignof(FSocketPool));
+
+	// Copy host
+	char* HostDest = (char*)Internal + PoolAllocSize;
+	memcpy(HostDest, Params.Host.Name.GetData(), HostNameLen);
+	HostDest[HostNameLen] = '\0';
+
+	// Init internal object
+	new (Internal) FSocketPool(
+		FAnsiStringView(HostDest, HostNameLen),
+		Params.Host.Port,
+		Params.ConnectionCount
+	);
+	Ptr = Internal;
+}
+
+////////////////////////////////////////////////////////////////////////////////
+FConnectionPool::~FConnectionPool()
+{
+	if (Ptr != nullptr)
+	{
+		FMemory::Free(Ptr);
+	}
+}
+
+
+
+// {{{1 activity ...............................................................
+
+////////////////////////////////////////////////////////////////////////////////
+struct FResponseInternal
+{
+	FIoBuffer*		Dest;
+	FMessageOffsets Offsets;
+	int32			ContentLength;
+	uint16			MessageLength;
+	mutable int16	Code;
+	uint32			_Unused;
+	const char		Data[];
+};
+
+////////////////////////////////////////////////////////////////////////////////
+struct alignas(16) FActivity
+{
+	enum class EWait : uint8 { None, Read, Write, Pool };
+	enum class EState : uint8
+	{
+		Build,
+		Resolve,
+		Connect,
+		Send,
+		RecvMessage,
+		RecvContent,
+		RecvStream,
+		Completed,
+		Cancelled,
+		Failed,
+	};
+
+	int8				Slot = -1;
+	EState				State;
+	EWait				SocketWait;
+	uint8				IsKeepAlive : 1;
+	uint8				_Unused0 : 7;
+	uint32				StateParam = 0;
+
+	FSocketPool*		Pool;
+	const char*			ErrorReason;
+	UPTRINT				SinkParam;
+	FTicketSink			Sink;
+
+	SocketType			Socket = InvalidSocket;
+
+	FBuffer				Buffer;
+};
+
+////////////////////////////////////////////////////////////////////////////////
+static FActivity* Activity_Alloc(uint32 BufferSize)
+{
+	BufferSize = (BufferSize + 15) & ~15;
+
+	uint32 Size = BufferSize + sizeof(FActivity);
+	auto* Activity = (FActivity*)FMemory::Malloc(Size, alignof(FActivity));
+
+	new (Activity) FActivity();
+
+	auto* Scratch = (char*)(Activity + 1);
+	uint32 ScratchSize = BufferSize;
+	Activity->Buffer = FBuffer(Scratch, ScratchSize);
+
+	Activity->State = FActivity::EState::Build;
+	return Activity;
+}
+
+////////////////////////////////////////////////////////////////////////////////
+static void Activity_Free(FActivity* Activity)
+{
+	if (Activity->State > FActivity::EState::Connect)
+	{
+		SocketType Socket = Activity->Socket;
+		if (!Activity->IsKeepAlive && Socket != InvalidSocket)
+		{
+			closesocket(Socket);
+			Activity->Socket = InvalidSocket;
+		}
+
+		if (Activity->Pool->GetState() == FSocketPool::EState::Resolved)
+		{
+			Activity->Pool->ReturnLease(Activity->Socket);
+		}
+	}
+
+	Activity->~FActivity();
+	FMemory::Free(Activity);
+}
+
+////////////////////////////////////////////////////////////////////////////////
+static void Activity_SetError(FActivity* Activity, const char* Reason)
+{
+	Activity->IsKeepAlive = 0;
+	Activity->SocketWait = FActivity::EWait::None;
+	Activity->ErrorReason = Reason;
+	Activity->State = FActivity::EState::Failed;
+	Activity->StateParam = LastSocketResult();
+}
+
+// {{{1 request ................................................................
+
+////////////////////////////////////////////////////////////////////////////////
+FRequest::~FRequest()
+{
+	if (Ptr == nullptr)
+	{
+		return;
+	}
+
+	Activity_Free(Ptr);
+}
+
+////////////////////////////////////////////////////////////////////////////////
+FRequest::FRequest(FRequest&& Rhs)
+{
+	this->~FRequest();
+	Swap(Ptr, Rhs.Ptr);
+}
+
+////////////////////////////////////////////////////////////////////////////////
+FRequest&& FRequest::Accept(EMimeType MimeType)
+{
+	switch (MimeType)
+	{
+	case EMimeType::Text:				return Header("Accept", "text/html");
+	case EMimeType::Binary:				return Header("Accept", "application/octet-stream");
+	case EMimeType::Json:				return Header("Accept", "application/json");
+	case EMimeType::Xml:				return Header("Accept", "application/xml");
+	/* UE_CUSTOM_MIME_TYPES
+	case EMimeType::CbObject:			return Header("Accept", "application/x-ue-cb");
+	case EMimeType::CbPackage:			return Header("Accept", "application/x-ue-pkg");
+	case EMimeType::CompressedBuffer:	return Header("Accept", "application/x-ue-comp");
+	*/
+	}
+
+	return MoveTemp(*this);
+}
+
+////////////////////////////////////////////////////////////////////////////////
+FRequest&& FRequest::Accept(FAnsiStringView MimeType)
+{
+	return Header("Accept", MimeType);
+}
+
+////////////////////////////////////////////////////////////////////////////////
+FRequest&& FRequest::Header(FAnsiStringView Key, FAnsiStringView Value)
+{
+	checkSlow(Ptr->State == FActivity::EState::Build);
+	FMessageBuilder(Ptr->Buffer) << Key << ": " << Value << "\r\n";
+	return MoveTemp(*this);
+}
+
+
+
+// {{{1 response ...............................................................
+
+////////////////////////////////////////////////////////////////////////////////
+static FResponseInternal& ToResponseInternal(FResponse* Addr)
+{
+	auto* Activity = (FActivity*)Addr;
+	return *(FResponseInternal*)(Activity->Buffer.GetData());
+}
+
+////////////////////////////////////////////////////////////////////////////////
+static const FResponseInternal& ToResponseInternal(const FResponse* Addr)
+{
+	const auto* Activity = (const FActivity*)Addr;
+	return *(const FResponseInternal*)(Activity->Buffer.GetData());
+}
+
+////////////////////////////////////////////////////////////////////////////////
+EStatusCodeClass FResponse::GetStatus() const
+{
+	uint32 Code = GetStatusCode();
+	if (Code <= 199) return EStatusCodeClass::Informational;
+	if (Code <= 299) return EStatusCodeClass::Successful;
+	if (Code <= 399) return EStatusCodeClass::Redirection;
+	if (Code <= 499) return EStatusCodeClass::ClientError;
+	if (Code <= 599) return EStatusCodeClass::ServerError;
+	return EStatusCodeClass::Unknown;
+}
+
+////////////////////////////////////////////////////////////////////////////////
+uint32 FResponse::GetStatusCode() const
+{
+	const FResponseInternal& Internal = ToResponseInternal(this);
+
+	if (Internal.Code < 0)
+	{
+		const char* CodePtr = Internal.Data + Internal.Offsets.StatusCode;
+		Internal.Code = uint16(CrudeToInt(FAnsiStringView(CodePtr, 3)));
+	}
+
+	return Internal.Code;
+}
+
+////////////////////////////////////////////////////////////////////////////////
+FAnsiStringView FResponse::GetStatusMessage() const
+{
+	const FResponseInternal& Internal = ToResponseInternal(this);
+	return FAnsiStringView(
+		Internal.Data + Internal.Offsets.Message,
+		Internal.Offsets.Headers - Internal.Offsets.Message
+	);
+}
+
+////////////////////////////////////////////////////////////////////////////////
+int64 FResponse::GetContentLength() const
+{
+	return ToResponseInternal(this).ContentLength;
+}
+
+////////////////////////////////////////////////////////////////////////////////
+EMimeType FResponse::GetContentType() const
+{
+    FAnsiStringView Value;
+    GetContentType(Value);
+
+    if (Value == "text/html")                   return EMimeType::Text;
+    if (Value == "application/octet-stream")    return EMimeType::Binary;
+    if (Value == "application/json")            return EMimeType::Json;
+    if (Value == "application/xml")             return EMimeType::Xml;
+	/* UE_CUSTOM_MIME_TYPES
+    if (Value == "application/x-ue-cb")         return EMimeType::CbObject;
+    if (Value == "application/x-ue-pkg")        return EMimeType::CbPackage;
+    if (Value == "application/x-ue-comp")       return EMimeType::CompressedBuffer;
+    */
+
+	return EMimeType::Unknown;
+}
+
+////////////////////////////////////////////////////////////////////////////////
+void FResponse::GetContentType(FAnsiStringView& Out) const
+{
+    Out = GetHeader("Accept");
+
+	int32 SemiColon;
+	if (Out.FindChar(';', SemiColon))
+	{
+		Out = Out.Mid(SemiColon).TrimEnd();
+	}
+}
+
+////////////////////////////////////////////////////////////////////////////////
+FAnsiStringView FResponse::GetHeader(FAnsiStringView Name) const
+{
+	const FResponseInternal& Internal = ToResponseInternal(this);
+
+	FAnsiStringView Result, Headers(
+		Internal.Data + Internal.Offsets.Headers,
+		Internal.MessageLength - Internal.Offsets.Headers
+	);
+
+	EnumerateHeaders(Headers, [&Result, Name] (FAnsiStringView Candidate, FAnsiStringView Value)
+	{
+		if (Candidate != Name)
+		{
+			return true;
+		}
+
+		Result = Value;
+		return false;
+	});
+	return Result;
+}
+
+////////////////////////////////////////////////////////////////////////////////
+void FResponse::SetDestination(FIoBuffer* Buffer)
+{
+	ToResponseInternal(this).Dest = Buffer;
+}
+
+
+
+// {{{1 ticket-status ..........................................................
+
+////////////////////////////////////////////////////////////////////////////////
+FTicketStatus::EId FTicketStatus::GetId() const
+{
+	const auto* Activity = (FActivity*)this;
+	switch (Activity->State)
+	{
+	case FActivity::EState::RecvMessage:	return EId::Response;
+	case FActivity::EState::RecvContent:	return EId::Content;
+	case FActivity::EState::Cancelled:		return EId::Cancelled;
+	case FActivity::EState::Failed:			return EId::Error;
+	default:								check(false);
+	}
+	return EId::Error;
+}
+
+////////////////////////////////////////////////////////////////////////////////
+UPTRINT FTicketStatus::GetParam() const
+{
+	const auto* Activity = (FActivity*)this;
+	return Activity->SinkParam;
+}
+
+////////////////////////////////////////////////////////////////////////////////
+FTicket FTicketStatus::GetTicket() const
+{
+	const auto* Activity = (FActivity*)this;
+	return 1ull << Activity->Slot;
+}
+
+////////////////////////////////////////////////////////////////////////////////
+uint32 FTicketStatus::GetIndex() const
+{
+	const auto* Activity = (FActivity*)this;
+	return Activity->Slot;
+}
+
+////////////////////////////////////////////////////////////////////////////////
+FResponse& FTicketStatus::GetResponse() const
+{
+	check(GetId() <= EId::Content);
+	const auto* Activity = (FActivity*)this;
+	return *(FResponse*)Activity;
+}
+
+////////////////////////////////////////////////////////////////////////////////
+uint32 FTicketStatus::GetContentLength() const
+{
+	check(GetId() <= EId::Content);
+	const auto* Activity = (FActivity*)this;
+	auto& Response = *(FResponseInternal*)(Activity->Buffer.GetData());
+	return Response.ContentLength;
+}
+
+////////////////////////////////////////////////////////////////////////////////
+const FIoBuffer& FTicketStatus::GetContent() const
+{
+	check(GetId() == EId::Content);
+	const auto* Activity = (FActivity*)this;
+	auto& Response = *(FResponseInternal*)(Activity->Buffer.GetData());
+	return *(Response.Dest);
+}
+
+////////////////////////////////////////////////////////////////////////////////
+const char* FTicketStatus::GetErrorReason() const
+{
+	check(GetId() == EId::Error);
+	const auto* Activity = (FActivity*)this;
+	return Activity->ErrorReason;
+}
+
+
+
+// {{{1 event-loop-int .........................................................
+
+////////////////////////////////////////////////////////////////////////////////
+static int32 Select(SocketType Socket, int32 Mode, int32 TimeoutMs=0)
+{
+	timeval TimeVal = {};
+	timeval* TimeValPtr = (TimeoutMs >= 0 ) ? &TimeVal : nullptr;
+	if (TimeoutMs > 0)
+	{
+		TimeVal = { TimeoutMs >> 10, TimeoutMs & ((1 << 10) - 1) };
+	}
+
+	fd_set FdSet;
+	FD_ZERO(&FdSet);
+	FD_SET(Socket, &FdSet);
+
+	fd_set FdSetExcept;
+	FD_ZERO(&FdSetExcept);
+	FD_SET(Socket, &FdSetExcept);
+
+	fd_set* FdSets[] = { &FdSet, nullptr };
+	if (Mode == 1)
+	{
+		Swap(FdSets[0], FdSets[1]);
+	}
+	int32 Result = select(1, FdSets[0], FdSets[1], &FdSetExcept, TimeValPtr);
+	return Result;
+}
+
+
+
+////////////////////////////////////////////////////////////////////////////////
+class FEventLoopInternal
+{
+public:
+	static bool		DoIsReady(FActivity* Activity);
+	static int32	DoResolve(FActivity* Activity);
+	static int32	DoConnect(FActivity* Activity);
+	static int32	DoSend(FActivity* Activity);
+	static int32	DoRecvMessage(FActivity* Activity);
+	static int32	DoRecvContent(FActivity* Activity);
+	static int32	DoRecvStream(FActivity* Activity);
+	static void		Cancel(FActivity* Activity);
+};
+
+
+
+////////////////////////////////////////////////////////////////////////////////
+bool FEventLoopInternal::DoIsReady(FActivity* Activity)
+{
+	TRACE_CPUPROFILER_EVENT_SCOPE(CoreHttp::DoIsReady);
+
+	int32 Result = 0;
+	switch (Activity->SocketWait)
+	{
+	case FActivity::EWait::Pool:  Result = (Activity->Pool->GetState() >= FSocketPool::EState::Resolved);
+	case FActivity::EWait::Read:  Result = Select(Activity->Socket, 0);	break;
+	case FActivity::EWait::Write: Result = Select(Activity->Socket, 1);	break;
+	case FActivity::EWait::None:  return true;
+	}
+
+	if (Result != 0)
+	{
+		Activity->SocketWait = FActivity::EWait::None;
+		return true;
+	}
+
+	return false;
+}
+
+////////////////////////////////////////////////////////////////////////////////
+int32 FEventLoopInternal::DoResolve(FActivity* Activity)
+{
+	TRACE_CPUPROFILER_EVENT_SCOPE(CoreHttp::DoResolve);
+
+	// todo: GetAddrInfoW() for async resolve on Windows
+
+	// There could be many activities using the same socket pool. Resolving only
+	// needs to happen once, the first activity in can do the honours. Everyone
+	// else can wait.
+	FSocketPool* Pool = Activity->Pool;
+	switch (Pool->GetState())
+	{
+	case FSocketPool::EState::Error:
+	case FSocketPool::EState::Resolved:
+		Activity->State = FActivity::EState::Connect;
+		return 0;
+
+	case FSocketPool::EState::Busy:
+		Activity->SocketWait = FActivity::EWait::Pool;
+		Activity->State = FActivity::EState::Connect;
+		return 1;
+	}
+
+	Pool->SetState(FSocketPool::EState::Busy);
+
+	addrinfo* Info = nullptr;
+	ON_SCOPE_EXIT { if (Info != nullptr) freeaddrinfo(Info); };
+
+	const FAnsiStringView& HostName = Pool->GetHostName();
+
+	addrinfo Hints = {};
+	Hints.ai_family = AF_INET;
+	Hints.ai_socktype = SOCK_STREAM;
+	Hints.ai_protocol = IPPROTO_TCP;
+	auto Result = getaddrinfo(HostName.GetData(), nullptr, &Hints, &Info);
+	if (uint32(Result) || Info == nullptr)
+	{
+		Pool->SetState(FSocketPool::EState::Error);
+		Activity_SetError(Activity, "Error encountered resolving");
+		return -1;
+	}
+
+	if (Info->ai_family != AF_INET)
+	{
+		Pool->SetState(FSocketPool::EState::Error);
+		Activity_SetError(Activity, "Unexpected address family during resolve");
+		return -1;
+	}
+
+	uint32 AddressCount = 0;
+	for (const addrinfo* Cursor = Info; Cursor != nullptr; Cursor = Cursor->ai_next)
+	{
+		const auto* AddrInet = (sockaddr_in*)(Cursor->ai_addr);
+		if (AddrInet->sin_family != AF_INET)
+		{
+			continue;
+		}
+
+		uint32 IpAddress = 0;
+		memcpy(&IpAddress, &(AddrInet->sin_addr), sizeof(uint32));
+
+		if (IpAddress == 0)
+		{
+			break;
+		}
+
+		if (!Pool->AddIpAddress(IpAddress))
+		{
+			break;
+		}
+
+		++AddressCount;
+	}
+
+	auto NextState = AddressCount
+		? FSocketPool::EState::Resolved
+		: FSocketPool::EState::Error;
+	Pool->SetState(NextState);
+
+	Activity->State = FActivity::EState::Connect;
+	return 0;
+}
+
+////////////////////////////////////////////////////////////////////////////////
+int32 FEventLoopInternal::DoConnect(FActivity* Activity)
+{
+	TRACE_CPUPROFILER_EVENT_SCOPE(CoreHttp::DoConnect);
+
+	FSocketPool* Pool = Activity->Pool;
+	if (Pool->GetState() == FSocketPool::EState::Error)
+	{
+		Activity_SetError(Activity, "Unable to resolve host");
+		return -1;
+	}
+	check(Pool->GetState() == FSocketPool::EState::Resolved);
+
+	// Claim an existing socket from the pool.
+	SocketType Candidate;
+	if (!Pool->LeaseSocket(Candidate))
+	{
+		// none available at this time
+		return 1;
+	}
+
+	if (Candidate != InvalidSocket)
+	{
+		Activity->Socket = Candidate;
+		Activity->SocketWait = FActivity::EWait::None;
+		Activity->State = FActivity::EState::Send;
+		Activity->StateParam = 0;
+		return 0;
+	}
+
+	// Leased socket isn't valid so we'll create and connect one
+	Candidate = socket(AF_INET, SOCK_STREAM, IPPROTO_TCP);
+	if (Candidate == InvalidSocket)
+	{
+		Activity_SetError(Activity, "Failed to create socket");
+		return -1;
+	}
+	ON_SCOPE_EXIT { if (Candidate != InvalidSocket) closesocket(Candidate); };
+
+	// make the socket non-blocking
+	{
+		bool Success = false;
+
+#if PLATFORM_MICROSOFT
+		unsigned long NonBlockingMode = 1;
+		if (ioctlsocket(Candidate, FIONBIO, &NonBlockingMode) != SOCKET_ERROR)
+		{
+			Success = true;
+		}
+#else
+		int32 Flags = fcntl(Candidate, F_GETFL, 0);
+		if (Flags != -1)
+		{
+			Flags |= Flags | int32(O_NONBLOCK);
+			if (fcntl(Candidate, F_SETFL, Flags) >= 0)
+			{
+				Success = true;
+			}
+		}
+#endif
+
+		if (!Success)
+		{
+			Activity_SetError(Activity, "Unable to set socket non-blocking");
+			return -1;
+		}
+	}
+
+	uint32 IpAddress = Activity->Pool->GetIpAddress();
+	if (IpAddress == 0)
+	{
+		Activity_SetError(Activity, "No IP address to connect to");
+		return -1;
+	}
+
+	// connect
+	sockaddr_in AddrInet = { sizeof(sockaddr_in) };
+	AddrInet.sin_family = AF_INET;
+	AddrInet.sin_port = htons(Activity->Pool->GetPort());
+	memcpy(&(AddrInet.sin_addr), &IpAddress, sizeof(IpAddress));
+	{
+		int Result = connect(Candidate, &(sockaddr&)AddrInet, sizeof(AddrInet));
+		if (Result < 0 && !(IsSocketResult(EWOULDBLOCK) | IsSocketResult(EINPROGRESS)))
+		{
+			Activity_SetError(Activity, "Socket connect failed");
+			return -1;
+		}
+	}
+
+	Activity->Socket = Candidate;
+	Activity->SocketWait = FActivity::EWait::Write;
+	Activity->State = FActivity::EState::Send;
+	Activity->StateParam = 0;
+	Candidate = InvalidSocket;
+	return 1;
+}
+
+////////////////////////////////////////////////////////////////////////////////
+int32 FEventLoopInternal::DoSend(FActivity* Activity)
+{
+	TRACE_CPUPROFILER_EVENT_SCOPE(CoreHttp::DoSend);
+
+	enum { PackBits = 4 };
+	uint32 Index = Activity->StateParam & ((1 << PackBits) - 1);
+	uint32 Remaining = Activity->StateParam >> PackBits;
+
+	FBuffer& Buffer = Activity->Buffer;
+
+	const char* SendData = nullptr;
+	uint32 SendSize = 0;
+	switch (Index)
+	{
+	case 0:
+		SendData = Buffer.GetData();
+		SendSize = Buffer.GetSize();
+		break;
+
+	case 1: {
+		SendData = "\r\n";
+		SendSize = 2;
+		break;
+		}
+	}
+
+	SendData += Remaining;
+	SendSize -= Remaining;
+
+	if (SendSize == 0)
+	{
+		// It is expected there will be enough space for a RespInt object
+		Buffer.Reset();
+		Buffer.AdvanceUsed(sizeof(FResponseInternal));
+
+		Activity->StateParam = 0;
+		Activity->State = FActivity::EState::RecvMessage;
+		Activity->SocketWait = FActivity::EWait::Read;
+		return 1;
+	}
+
+	int32 Result = send(Activity->Socket, SendData, SendSize, MsgFlagType(0));
+	if (Result < 0)
+	{
+		if (IsSocketResult(ENOTCONN))
+		{
+			int32 Error = 0;
+			socklen_t ErrorSize = sizeof(Error);
+			Result = getsockopt(Activity->Socket, SOL_SOCKET, SO_ERROR, (char*)&Error, &ErrorSize);
+			if (Result < 0 || Error != 0)
+			{
+				Activity_SetError(Activity, "Connection error");
+				return -1;
+			}
+			return 1;
+		}
+
+		if (!IsSocketResult(EWOULDBLOCK))
+		{
+			Activity_SetError(Activity, "Error returned from socket send");
+			return -1;
+		}
+	}
+
+	if (Result == 0)
+	{
+		Activity_SetError(Activity, "ATH0.Send");
+		return -1;
+	}
+
+	Remaining = SendSize - Result;
+	if (Remaining != 0)
+	{
+		Activity->StateParam = Index | (Remaining << PackBits);
+		Activity->SocketWait = FActivity::EWait::Write;
+		return 1;
+	}
+
+	Activity->StateParam = Index + 1;
+	return DoSend(Activity);
+}
+
+////////////////////////////////////////////////////////////////////////////////
+int32 FEventLoopInternal::DoRecvMessage(FActivity* Activity)
+{
+	TRACE_CPUPROFILER_EVENT_SCOPE(CoreHttp::DoRecvMessage);
+
+	static const uint32 PageSize = 2048;
+
+	FBuffer& Buffer = Activity->Buffer;
+
+	const char* MessageRight;
+	while (true)
+	{
+		auto [Dest, DestSize] = Buffer.GetMutableFree(0, PageSize);
+
+		int32 Result = recv(Activity->Socket, Dest, DestSize, MsgFlagType(0));
+		if (Result < 0)
+		{
+			if (IsSocketResult(EWOULDBLOCK))
+			{
+				Activity->SocketWait = FActivity::EWait::Read;
+				return 1;
+			}
+
+			Activity_SetError(Activity, "Error occurred on socket recv");
+			return -1;
+		}
+
+		Buffer.AdvanceUsed(Result);
+
+		// Rewind a little to cover cases where the terminal is fragmented across
+		// recv() calls
+		uint32 DestBias = 0;
+		if (Dest - 3 >= Buffer.GetData() + sizeof(FResponseInternal))
+		{
+			Dest -= (DestBias = 3);
+		}
+
+		int32 MessageEnd = FindMessageTerminal(Dest, Result + DestBias);
+		if (MessageEnd < 0)
+		{
+			if (Buffer.GetSize() > (8 << 10))
+			{
+				Activity_SetError(Activity, "Headers have grown larger than expected");
+				return -1;
+			}
+
+			if (Result == 0)
+			{
+				Activity_SetError(Activity, "ATH0.RecvMessage");
+				return -1;
+			}
+
+			continue;
+		}
+
+		MessageRight = Dest + MessageEnd;
+		break;
+	}
+
+	// Fill out the internal response object
+	auto& Internal = *(FResponseInternal*)(Buffer.GetData());
+	Internal.MessageLength = uint16(ptrdiff_t(MessageRight - Internal.Data));
+
+	FAnsiStringView ResponseView(Internal.Data, Internal.MessageLength);
+	if (ParseMessage(ResponseView, Internal.Offsets) < 0)
+	{
+		Activity_SetError(Activity, "Failed to parse message status");
+		return -1;
+	}
+
+	// Parse headers
+	FAnsiStringView Headers = ResponseView.Mid(
+		Internal.Offsets.Headers,
+		Internal.MessageLength - Internal.Offsets.Headers - 2 // "-2" trims off '\r\n' that signals end of headers
+	);
+
+	bool IsKeepAlive = true;
+	int32 ContentLength = -1;
+	EnumerateHeaders(
+		Headers,
+		[&ContentLength, &IsKeepAlive] (FAnsiStringView Name, FAnsiStringView Value)
+		{
+			// todo; may need smarter value handling; ;/, seperated options & key-value pairs (ex. in rfc2068)
+
+			// "Keep-Alive"			- deprecated
+			// "Transfer-Encoding"	- may be required later
+
+			if (Name.Equals("Content-Length", ESearchCase::IgnoreCase))
+			{
+				ContentLength = CrudeToInt(Value);
+				return true;
+			}
+			
+			else if (Name.Equals("Connection", ESearchCase::IgnoreCase))
+			{
+				IsKeepAlive = !Value.Equals("close");
+				return true;
+			}
+
+			return true;
+		}
+	);
+
+	Activity->IsKeepAlive &= IsKeepAlive;
+
+	if (ContentLength < 0)
+	{
+		if (ContentLength == -1)
+		{
+			// todo; query Transfer-Encoding for chunked
+		}
+
+		Activity_SetError(Activity, "Unknown content length value");
+		return -1;
+	}
+
+	// Call out to the sink to get a content destination
+	Internal.Dest = nullptr;
+	Internal.Code = -1;
+	Internal.ContentLength = ContentLength;
+	{
+		FTicketStatus& SinkArg = *(FTicketStatus*)Activity;
+		Activity->Sink(SinkArg);
+	}
+
+	if (Internal.Dest == nullptr)
+	{
+		Activity_SetError(Activity, "User did not provide a destination buffer");
+		return -1;
+	}
+
+	// The user seems to have forgotten something. Let's help them along
+	if (Internal.Dest->GetSize() == 0)
+	{
+		*Internal.Dest = FIoBuffer(ContentLength);
+	}
+
+	bool Streamed = (Internal.Dest->GetSize() < ContentLength);
+
+	// Perhaps we have some of the content already?
+	const char* BufferRight = Buffer.GetData() + Buffer.GetSize();
+	uint32 AlreadyReceived = uint32(ptrdiff_t(BufferRight - MessageRight));
+	if (AlreadyReceived > uint32(ContentLength))
+	{
+		Activity_SetError(Activity, "More data recevied that expected");
+		return -1;
+	}
+
+	Activity->State = Streamed ? FActivity::EState::RecvStream : FActivity::EState::RecvContent;
+	Activity->StateParam = AlreadyReceived;
+
+	if (AlreadyReceived == 0)
+	{
+		return 0;
+	}
+
+	FMutableMemoryView DestView = Internal.Dest->GetMutableView();
+	const char* Cursor = BufferRight - AlreadyReceived;
+	if (!Streamed || AlreadyReceived < DestView.GetSize())
+	{
+		::memcpy(DestView.GetData(), Cursor, AlreadyReceived);
+		return 0;
+	}
+
+#if 1
+	check(false); // not implemented yet
+#else
+	do
+	{
+		uint32 Size = FMath::Min<uint32>(DestView.GetSize(), AlreadyReceived);
+		::memcpy(DestView.GetData(), Cursor, Size);
+
+		/* send sink */
+
+		AlreadyReceived -= Size;
+		Cursor += Size;
+	}
+	while (AlreadyReceived);
+#endif
+
+	return 0;
+}
+
+////////////////////////////////////////////////////////////////////////////////
+int32 FEventLoopInternal::DoRecvContent(FActivity* Activity)
+{
+	TRACE_CPUPROFILER_EVENT_SCOPE(CoreHttp::DoRecvContent);
+
+	auto& Response = *(FResponseInternal*)(Activity->Buffer.GetData());
+
+	FMutableMemoryView DestView = Response.Dest->GetMutableView();
+	while (true)
+	{
+		uint32 Size = (Response.ContentLength - Activity->StateParam);
+		if (Size == 0)
+		{
+			break;
+		}
+
+		char* Cursor = (char*)(DestView.GetData()) + Activity->StateParam;
+		int32 Result = recv(Activity->Socket, Cursor, Size, MsgFlagType(0));
+		if (Result < 0)
+		{
+			if (IsSocketResult(EWOULDBLOCK))
+			{
+				Activity->SocketWait = FActivity::EWait::Read;
+				return 1;
+			}
+
+			Activity_SetError(Activity, "Socket error while receiving content");
+			return -1;
+		}
+
+		if (Result == 0 && (Activity->StateParam != Response.ContentLength))
+		{
+			Activity_SetError(Activity, "ATH0.RecvContent");
+			return -1;
+		}
+
+		Activity->StateParam += Result;
+	}
+
+	// Notify the user we've received everything
+	FTicketStatus& SinkArg = *(FTicketStatus*)Activity;
+	Activity->Sink(SinkArg);
+
+	Activity->State = FActivity::EState::Completed;
+	return 0;
+}
+
+////////////////////////////////////////////////////////////////////////////////
+int32 FEventLoopInternal::DoRecvStream(FActivity* Activity)
+{
+	TRACE_CPUPROFILER_EVENT_SCOPE(CoreHttp::DoRecvStream);
+
+	check(false); // not implemented yet
+	return -1;
+}
+
+////////////////////////////////////////////////////////////////////////////////
+void FEventLoopInternal::Cancel(FActivity* Activity)
+{
+	if (Activity->State >= FActivity::EState::Completed)
+	{
+		return;
+	}
+
+	Activity->State = FActivity::EState::Cancelled;
+
+	FTicketStatus& SinkArg = *(FTicketStatus*)Activity;
+	Activity->Sink(SinkArg);
+}
+
+
+
+// {{{1 event-loop .............................................................
+
+////////////////////////////////////////////////////////////////////////////////
+static const FEventLoop::FRequestParams GDefaultParams;
+
+////////////////////////////////////////////////////////////////////////////////
+FEventLoop::~FEventLoop()
+{
+	for (FActivity* Activity : Active)
+	{
+		Activity_Free(Activity);
+	}
+}
+
+////////////////////////////////////////////////////////////////////////////////
+FRequest FEventLoop::Request(
+	FAnsiStringView Method,
+	FAnsiStringView Url,
+	const FRequestParams* Params)
+{
+	// Parse the URL into its components
+	FUrlOffsets UrlOffsets;
+	if (ParseUrl(Url, UrlOffsets) < 0)
+	{
+		return FRequest();
+	}
+
+	FAnsiStringView HostName = UrlOffsets.HostName.Get(Url);
+
+	uint32 Port = 80;
+	if (UrlOffsets.SchemeLength == 5)
+	{
+		//Port = 443;
+		//Protocol = Protocol::Tls;
+		return FRequest();
+	}
+
+	if (UrlOffsets.Port)
+	{
+		FAnsiStringView PortView = UrlOffsets.Port.Get(Url);
+		Port = CrudeToInt(PortView);
+	}
+
+	FAnsiStringView Path;
+	if (UrlOffsets.Path > 0)
+	{
+		Path = Url.Mid(UrlOffsets.Path);
+	}
+
+	// Create an activity and an emphemeral socket pool
+	Params = (Params != nullptr) ? Params : &GDefaultParams;
+
+	uint32 BufferSize = Params->BufferSize;
+	BufferSize = (BufferSize >= 128) ? BufferSize : 128;
+	BufferSize += sizeof(FSocketPool) + HostName.Len();
+	FActivity* Activity = Activity_Alloc(BufferSize);
+
+	FBuffer& Buffer = Activity->Buffer;
+
+	FSocketPool* Pool = Activity->Pool = Buffer.Alloc<FSocketPool>();
+	Activity->IsKeepAlive = 0;
+
+	uint32 HostNameLength = HostName.Len();
+	char* HostNamePtr = Buffer.Alloc<char>(HostNameLength + 1);
+
+	Buffer.Fix();
+
+	memcpy(HostNamePtr, HostName.GetData(), HostNameLength);
+	HostNamePtr[HostNameLength] = '\0';
+	HostName = FAnsiStringView(HostNamePtr, HostNameLength);
+
+	new (Pool) FSocketPool(HostName, Port, 1);
+
+	return Request(Method, Path, Activity);
+}
+
+////////////////////////////////////////////////////////////////////////////////
+FRequest FEventLoop::Request(
+	FAnsiStringView Method,
+	FAnsiStringView Path,
+	FConnectionPool& Pool,
+	const FRequestParams* Params)
+{
+	check(Pool.Ptr != nullptr);
+
+	Params = (Params != nullptr) ? Params : &GDefaultParams;
+
+	uint32 BufferSize = Params->BufferSize;
+	BufferSize = (BufferSize >= 128) ? BufferSize : 128;
+	FActivity* Activity = Activity_Alloc(BufferSize);
+
+	Activity->Pool = Pool.Ptr;
+	Activity->IsKeepAlive = 1;
+
+	return Request(Method, Path, Activity);
+}
+
+////////////////////////////////////////////////////////////////////////////////
+FRequest FEventLoop::Request(
+	FAnsiStringView Method,
+	FAnsiStringView Path,
+	FActivity* Activity)
+{
+	if (Path.Len() == 0)
+	{
+		Path = "/";
+	}
+
+	FMessageBuilder Builder(Activity->Buffer);
+
+	Builder << Method << " " << Path << " HTTP/1.1\r\n"
+		"Host: " << Activity->Pool->GetHostName() << "\r\n";
+
+	// HTTP/1.1 is persistent by default thus "Connection" header isn't required
+	if (!Activity->IsKeepAlive)
+	{
+		Builder << "Connection: close\r\n";
+	}
+
+	FRequest Ret;
+	Ret.Ptr = Activity;
+	return Ret;
+}
+
+////////////////////////////////////////////////////////////////////////////////
+FTicket FEventLoop::Send(FRequest&& Request, FTicketSink Sink, UPTRINT SinkParam)
+{
+	FActivity* Activity = nullptr;
+	Swap(Activity, Request.Ptr);
+	Activity->State = FActivity::EState::Resolve;
+	Activity->SinkParam = SinkParam;
+	Activity->Sink = Sink;
+
+	uint64 Slot;
+	{
+		FScopeLock _(&Lock);
+
+		for (;; FPlatformProcess::SleepNoStats(0.0f))
+		{
+			uint64 FreeSlotsLoad = FreeSlots.load(std::memory_order_relaxed);
+			if (!FreeSlotsLoad)
+			{
+				// we don't handle oversubscription at the moment. Could return
+				// activity to Reqeust and return a 0 ticket.
+				check(false);
+			}
+			Slot = -int64(FreeSlotsLoad) & FreeSlotsLoad;
+			if (FreeSlots.compare_exchange_weak(FreeSlotsLoad, FreeSlotsLoad - Slot, std::memory_order_relaxed))
+			{
+				break;
+			}
+		}
+		Activity->Slot = int8(63 - FMath::CountLeadingZeros64(Slot));
+
+		Pending.Add(Activity);
+	}
+
+    return Slot;
+}
+
+////////////////////////////////////////////////////////////////////////////////
+bool FEventLoop::IsIdle() const
+{
+	return FreeSlots.load(std::memory_order_relaxed) == ~0ull;
+}
+
+////////////////////////////////////////////////////////////////////////////////
+void FEventLoop::Cancel(FTicket Ticket)
+{
+	Cancels.fetch_or(Ticket, std::memory_order_relaxed);
+}
+
+////////////////////////////////////////////////////////////////////////////////
+uint32 FEventLoop::Tick()
+{
+	TRACE_CPUPROFILER_EVENT_SCOPE(CoreHttp::Tick);
+
+    // Collect activity changes
+	uint64 FreeSlotsLoad = FreeSlots.load(std::memory_order_relaxed);
+    if (FreeSlots != PrevFreeSlots)
+    {
+		TArray<FActivity*> NewActive;
+		{
+			FScopeLock _(&Lock);
+			NewActive = MoveTemp(Pending);
+		}
+
+		for (FActivity* Activity : NewActive)
+		{
+			Active.Add(Activity);
+		}
+
+        PrevFreeSlots = FreeSlotsLoad;
+    }
+
+	// Tick activities
+	uint64 CancelsLoad = Cancels.load(std::memory_order_relaxed);
+	uint32 BusyCount = 0;
+	for (FActivity* Activity : Active)
+	{
+		if ((1ull << Activity->Slot) & CancelsLoad)
+		{
+			FEventLoopInternal::Cancel(Activity);
+			continue;
+		}
+
+		++BusyCount;
+
+		int32 Result = 1;
+		if (!FEventLoopInternal::DoIsReady(Activity))
+		{
+			continue;
+		}
+
+		switch (Activity->State)
+		{
+		case FActivity::EState::Resolve:
+			Result = FEventLoopInternal::DoResolve(Activity);
+			if (Result)
+				break;
+
+		case FActivity::EState::Connect:
+			Result = FEventLoopInternal::DoConnect(Activity);
+			if (Result)
+				break;
+
+		case FActivity::EState::Send:
+			Result = FEventLoopInternal::DoSend(Activity);
+			if (Result)
+				break;
+
+		case FActivity::EState::RecvMessage:
+			Result = FEventLoopInternal::DoRecvMessage(Activity);
+			if (Result)
+				break;
+
+		case FActivity::EState::RecvContent:
+		case FActivity::EState::RecvStream:
+			if (Activity->State == FActivity::EState::RecvContent)
+			{
+				Result = FEventLoopInternal::DoRecvContent(Activity);
+			}
+			else
+			{
+				Result = FEventLoopInternal::DoRecvStream(Activity);
+			}
+			if (Result)
+				break;
+
+		case FActivity::EState::Completed:
+			--BusyCount;
+			break;
+		}
+
+		if (Result == -1)
+		{
+			check(Activity->State == FActivity::EState::Failed);
+
+			FTicketStatus& SinkArg = *(FTicketStatus*)Activity;
+			Activity->Sink(SinkArg);
+
+			--BusyCount;
+		}
+	}
+
+	// Reap done activities
+	uint64 ReturnedSlots = 0;
+	if (int32 n = Active.Num(); BusyCount != n)
+	{
+		int32 i = 0;
+		for (--n; i <= n;)
+		{
+			FActivity* Activity = Active[i];
+			if (Activity->State < FActivity::EState::Completed)
+			{
+				++i;
+				continue;
+			}
+
+			ReturnedSlots += (1ull << Activity->Slot);
+
+			Activity_Free(Activity);
+
+			Active[i] = Active[n];
+			--n;
+		}
+		Active.SetNum(i, false);
+	}
+
+	if (ReturnedSlots)
+	{
+		PrevFreeSlots += ReturnedSlots;
+		FreeSlots.fetch_add(ReturnedSlots, std::memory_order_relaxed);
+	}
+
+	if (CancelsLoad)
+	{
+		Cancels.fetch_and(~CancelsLoad, std::memory_order_relaxed);
+	}
+
+	return BusyCount;
+}
+
+
+
+// {{{1 test ...................................................................
+
+#if !(UE_BUILD_SHIPPING|UE_BUILD_TEST)
+
+////////////////////////////////////////////////////////////////////////////////
+static void MiscTest()
+{
+#define CRLF "\r\n"
+	struct {
+		FAnsiStringView Input;
+		int32 Output;
+	} FmtTestCases[] = {
+		{ "", -1 },
+		{ "abcd", -1 },
+		{ "abcd\r", -1 },
+		{ CRLF "\r\r", -1 },
+		{ CRLF CRLF, 4 },
+		{ "abc" CRLF CRLF, 7 },
+	};
+	for (const auto [Input, Output] : FmtTestCases)
+	{
+		check(FindMessageTerminal(Input.GetData(), Input.Len()) == Output);
+	}
+
+	FMessageOffsets MsgOut;
+	check(ParseMessage("", MsgOut) == -1);
+	check(ParseMessage("MR", MsgOut) == -1);
+	check(ParseMessage("HTTP/1.1", MsgOut) == -1);
+	check(ParseMessage("HTTP/1.1 ", MsgOut) == -1);
+	check(ParseMessage("HTTP/1.1 1" CRLF, MsgOut) > 0);
+	check(ParseMessage("HTTP/1.1    1" CRLF, MsgOut) > 0);
+	check(ParseMessage("HTTP/1.1 100 " CRLF, MsgOut) > 0);
+	check(ParseMessage("HTTP/1.1 100  Message of some sort    " CRLF, MsgOut) > 0);
+	check(ParseMessage("HTTP/1.1 100 _Message with a \r in it" CRLF, MsgOut) == -1);
+
+	bool AllIsWell = true;
+	auto NotExpectedToBeCalled = [&] (auto, auto) { return AllIsWell = false; };
+	EnumerateHeaders("",		NotExpectedToBeCalled); check(AllIsWell);
+	EnumerateHeaders(CRLF,		NotExpectedToBeCalled); check(AllIsWell);
+	EnumerateHeaders("foo",		NotExpectedToBeCalled); check(AllIsWell);
+	EnumerateHeaders(" foo",	NotExpectedToBeCalled); check(AllIsWell);
+	EnumerateHeaders(" foo ",	NotExpectedToBeCalled); check(AllIsWell);
+	EnumerateHeaders("foo:bar",	NotExpectedToBeCalled); check(AllIsWell);
+
+	auto IsBar = [&] (auto, auto Value) { return AllIsWell = (Value == "bar"); };
+	EnumerateHeaders("foo: bar" CRLF,		IsBar); check(AllIsWell);
+	EnumerateHeaders("foo: bar \t" CRLF,	IsBar); check(AllIsWell);
+	EnumerateHeaders("foo:\tbar " CRLF,		IsBar); check(AllIsWell);
+	EnumerateHeaders("foo:bar " CRLF,		IsBar); check(AllIsWell);
+	EnumerateHeaders("foo:bar" CRLF "!",	IsBar); check(AllIsWell);
+	EnumerateHeaders("foo:bar" CRLF " ",	IsBar); check(AllIsWell);
+	EnumerateHeaders("foo:bar" CRLF "n:ej",	IsBar); check(AllIsWell);
+
+	check(CrudeToInt("") < 0);
+	check(CrudeToInt("X") < 0);
+	check(CrudeToInt("/") < 0);
+	check(CrudeToInt(":") < 0);
+	check(CrudeToInt("-1") < -1);
+	check(CrudeToInt("0") == 0);
+	check(CrudeToInt("9") == 9);
+	check(CrudeToInt("493") == 493);
+
+	FUrlOffsets UrlOut;
+
+	check(ParseUrl("", UrlOut) == -1);
+	check(ParseUrl("abc://asd/", UrlOut) == -1);
+	check(ParseUrl("http://", UrlOut) == -1);
+	check(ParseUrl("http://:/", UrlOut) == -1);
+	check(ParseUrl("http://@:/", UrlOut) == -1);
+	check(ParseUrl("http://foo:ba:r/", UrlOut) == -1);
+	check(ParseUrl("http://foo@ba:r/", UrlOut) == -1);
+	check(ParseUrl("http://foo@ba:/", UrlOut) == -1);
+	check(ParseUrl("http://foo@ba@9/", UrlOut) == -1);
+	check(ParseUrl("http://@ba:9/", UrlOut) == -1);
+	check(ParseUrl("http://zzzzzzzzzzzzzzzzzzzzzzzzzzzzzzzzzzzzzzzzzzzzzzzzzzzzzzzzzzzzzzzzzzzzzzzzzzzzzzzzzzzzzzzzzzzzzzzzzzzzzzzzzzzzzzzzzzzzzzzzzzzzzzzzzzz-hello-zz.com/", UrlOut) == -1);
+
+	check(ParseUrl("http://ab-c.com/", UrlOut) > 0);
+	check(ParseUrl("http:[REDACTED_EMAIL]/", UrlOut) > 0);
+	check(ParseUrl("https://abc.com", UrlOut) > 0);
+	check(ParseUrl("https://abc.com:999", UrlOut) > 0);
+	check(ParseUrl("https://abc.com:999/", UrlOut) > 0);
+
+	FAnsiStringView Url = "http://abc:[REDACTED_EMAIL]:999/";
+	check(ParseUrl(Url, UrlOut) > 0);
+	check(UrlOut.SchemeLength == 4);
+	check(UrlOut.UserInfo.Get(Url) == "abc:123");
+	check(UrlOut.HostName.Get(Url) == "bc.com");
+	check(UrlOut.Port.Get(Url) == "999");
+	check(UrlOut.Path == 25);
+#undef CRLF
+}
+
+////////////////////////////////////////////////////////////////////////////////
+COREHTTP_API void CoreHttpTest()
+{
+	MiscTest();
+
+#if PLATFORM_WINDOWS
+	WSADATA WsaData;
+	if (WSAStartup(MAKEWORD(2, 2), &WsaData) == 0x0a9e0493)
+		return;
+	ON_SCOPE_EXIT { WSACleanup(); };
+#endif
+
+#define HOST_NAME "10.24.101.89"
+	FAnsiStringView TestUrl = "http://" HOST_NAME ":9493/data";
+	FAnsiStringView TestHostName = "localhost";
+
+	struct
+	{
+		FIoBuffer Dest;
+		uint64 Hash = 0;
+	} Content[64];
+
+	auto HashSink = [&] (const UE::HTTP::FTicketStatus& Status) -> FIoBuffer*
+	{
+		check(Status.GetId() != FTicketStatus::EId::Error);
+
+		uint32 Index = Status.GetIndex();
+
+		if (Status.GetId() == FTicketStatus::EId::Response)
+		{
+			FResponse& Response = Status.GetResponse();
+			check(Response.GetStatus() == EStatusCodeClass::Successful);
+			check(Response.GetStatusCode() == 200);
+			check(Response.GetContentLength() == Status.GetContentLength());
+
+			FAnsiStringView HashView = Response.GetHeader("X-TestServer-Hash");
+			Content[Index].Hash = CrudeToInt(HashView);
+			check(int64(Content[Index].Hash) > 0);
+
+			Content[Index].Dest = FIoBuffer();
+			Response.SetDestination(&(Content[Index].Dest));
+			return nullptr;
+		}
+
+		uint32 ReceivedHash = 0x493;
+		FMemoryView ContentView = Content[Index].Dest.GetView();
+		check(ContentView.GetSize() == Status.GetContentLength());
+		for (uint32 i = 0; i < Status.GetContentLength(); ++i)
+		{
+			uint8 c = ((const uint8*)(ContentView.GetData()))[i];
+			ReceivedHash = (ReceivedHash + c) * 0x493;
+		}
+		check(Content[Index].Hash == ReceivedHash);
+		Content[Index].Hash = 0;
+
+ 		return nullptr;
+	};
+
+	auto NullSink = [] (const FTicketStatus&) {};
+
+	auto NoErrorSink = [&] (const FTicketStatus& Status)
+	{
+		check(Status.GetId() != FTicketStatus::EId::Error);
+		if (Status.GetId() != FTicketStatus::EId::Response)
+		{
+			return;
+		}
+
+		uint32 Index = Status.GetIndex();
+
+		FResponse& Response = Status.GetResponse();
+		Response.SetDestination(&(Content[Index].Dest));
+	};
+
+	FEventLoop Loop;
+	volatile bool LoopStop = false;
+	volatile bool LoopTickDelay = false;
+	auto LoopTask = UE::Tasks::Launch(TEXT("CoreHttpTest.Loop"), [&] () {
+		uint32 DelaySeed = 493;
+		while (!LoopStop)
+		{
+			while (Loop.Tick())
+			{
+				if (!LoopTickDelay)
+				{
+					continue;
+				}
+
+				float DelayFloat = float(DelaySeed % 75) / 1000.0f;
+				FPlatformProcess::SleepNoStats(DelayFloat);
+				DelaySeed *= 0xa93;
+			}
+
+			FPlatformProcess::SleepNoStats(0.005f);
+		}
+	});
+
+	auto WaitForLoopIdle = [&Loop] ()
+	{
+		FPlatformProcess::SleepNoStats(0.25f);
+		while (!Loop.IsIdle())
+		{
+			FPlatformProcess::SleepNoStats(0.03f);
+		}
+	};
+
+	// unused request
+	{
+		FRequest Request = Loop.Request("GET", TestUrl);
+	}
+
+	// foundational
+	{
+		FRequest Request = Loop.Request("GET", "http://" HOST_NAME ":9493/seed/493");
+		Request.Accept(EMimeType::Json);
+
+		FTicket Ticket = Loop.Send(MoveTemp(Request), NullSink);
+
+		WaitForLoopIdle();
+	}
+
+	// convenience
+	{
+		FRequest Request = Loop.Get(TestUrl).Accept(EMimeType::Json);
+
+		FTicket Tickets[] = {
+			Loop.Send(MoveTemp(Request), HashSink),
+			Loop.Send(Loop.Get(TestUrl).Accept(EMimeType::Json), HashSink),
+			Loop.Send(Loop.Get("http://httpbin.org/get"), NoErrorSink),
+		};
+		WaitForLoopIdle();
+	}
+
+	// convenience
+	{
+		FRequest Request = Loop.Get(TestUrl).Accept(EMimeType::Json);
+		FTicket Ticket = Loop.Send(MoveTemp(Request), HashSink);
+		WaitForLoopIdle();
+	}
+
+	// no connect
+	{
+		FRequest Request[] = {
+			Loop.Request("GET", "http://" HOST_NAME ":10930"),
+			Loop.Request("GET", "http://thisdoesnotexistihope/"),
+		};
+		Loop.Send(MoveTemp(Request[0]), NullSink);
+		Loop.Send(MoveTemp(Request[1]), NullSink);
+		WaitForLoopIdle();
+	}
+
+	// stress 1
+	{
+		const uint32 StressLoad = 32;
+
+		struct {
+			FAnsiStringView Url;
+			bool Disconnect;
+		} StressUrls[] = {
+			{ "http://" HOST_NAME ":9494/data",				false },
+			{ "http://" HOST_NAME ":9494/data?disconnect",	true },
+		};
+
+		uint64 Errors = 0;
+		auto ErrorSink = [&] (const UE::HTTP::FTicketStatus& Status)
+		{
+			FTicket Ticket = Status.GetTicket();
+			uint32 Index = 63 - FMath::CountLeadingZeros64(uint64(Ticket));
+
+			if (Status.GetId() == FTicketStatus::EId::Error)
+			{
+				Errors |= 1ull << Index;
+				return;
+			}
+
+			else if (Status.GetId() == FTicketStatus::EId::Response)
+			{
+				FResponse& Response = Status.GetResponse();
+				Content[Index].Dest = FIoBuffer();
+				Response.SetDestination(&(Content[Index].Dest));
+				return;
+			}
+
+			check(false);
+		};
+
+		for (const auto& [StressUrl, ExpectDisconnect] : StressUrls)
+		{
+			FTicketSink Sink;
+			if (ExpectDisconnect)
+			{
+				Sink = ErrorSink;
+			}
+			else
+			{
+				Sink = HashSink;
+			}
+
+			for (bool AddDelay : {false, true})
+			{
+				FTicket Tickets[StressLoad];
+				for (FTicket& Ticket : Tickets)
+				{
+					Ticket = Loop.Send(Loop.Get(StressUrl).Header("Accept", "*/*"), Sink);
+				}
+
+				LoopTickDelay = AddDelay;
+
+				WaitForLoopIdle();
+			}
+
+			LoopTickDelay = false;
+		}
+	}
+
+	// stress 2
+	{
+		const uint32 StressLoad = 3;
+		const uint32 StressTaskCount = 7;
+		static_assert(StressLoad * StressTaskCount <= 32);
+
+		FAnsiStringView Url = "http://" HOST_NAME ":9494/data";
+
+		auto StressTaskEntry = [&] {
+			for (uint32 i = 0; i < StressLoad; ++i)
+			{
+				FTicket Ticket = Loop.Send(Loop.Get(Url), HashSink);
+				if (!Ticket)
+				{
+					FPlatformProcess::SleepNoStats(0.01f);
+					--i;
+				}
+			}
+		};
+
+		UE::Tasks::FTask StressTasks[StressTaskCount];
+		for (auto& StressTask : StressTasks)
+		{
+			StressTask = UE::Tasks::Launch(TEXT("StressTask"), [&] { StressTaskEntry(); });
+		}
+		for (auto& StressTask : StressTasks)
+		{
+			StressTask.Wait();
+		}
+
+		WaitForLoopIdle();
+	}
+
+	LoopStop = true;
+	LoopTask.Wait();
+
+	check(Loop.IsIdle());
+
+	// pre-generated headers
+	// request-with-body
+	// proxy
+	// chunked transfer encoding
+	// gzip / deflate
+	// redirects
+	// loop multi-req.
+	// tls
+	// http pipelining
+	// url auth credentials
+	// transfer-file / splice / sendfile
+	// (header field parser)
+	// (form-data)
+	// (cookies)
+	// (cache)
+	// (websocket)
+	// (ipv6)
+	// (utf-8 host names)
+}
+
+#endif // !SHIP|TEST
+
+// }}}
+
+} // namespace UE::HTTP
+
+/* vim: set noet : */
diff --git a/Engine/Source/Runtime/Experimental/IoStoreOnDemand/Private/OnDemandIoDispatcherBackend.cpp b/Engine/Source/Runtime/Experimental/IoStoreOnDemand/Private/OnDemandIoDispatcherBackend.cpp
--- a/Engine/Source/Runtime/Experimental/IoStoreOnDemand/Private/OnDemandIoDispatcherBackend.cpp
+++ b/Engine/Source/Runtime/Experimental/IoStoreOnDemand/Private/OnDemandIoDispatcherBackend.cpp
@@ -25,7 +25,7 @@
 #include "Tasks/Pipe.h"
 #include <atomic>
 
-#define WITH_HTTP_CLIENT 0
+#define WITH_HTTP_CLIENT 1
 #if WITH_HTTP_CLIENT
 #include "CoreHttp/Client.h"
 #endif
@@ -60,7 +60,7 @@ FIoHash GetChunkKey(const FIoHash& ChunkHash, const FIoOffsetAndLength& Range)
 class FHttpClient
 {
 public:
-	FHttpClient();
+	FHttpClient(UE::HTTP::FEventLoop& EventLoop);
 	void Get(FAnsiStringView Url, FIoReadCallback&& Callback, const TCHAR* DebugName = nullptr);
 	void Get(FAnsiStringView Url, const FIoOffsetAndLength& Range, FIoReadCallback&& Callback, const TCHAR* DebugName = nullptr);
 	bool Tick();
@@ -68,11 +68,12 @@ public:
 private:
 #if WITH_HTTP_CLIENT
 	void Issue(UE::HTTP::FRequest&& Request, FIoReadCallback&& Callback, FAnsiStringView Url = FAnsiStringView(), const TCHAR* DebugName = nullptr);
-	UE::HTTP::FEventLoop EventLoop;
+	UE::HTTP::FEventLoop& EventLoop;
 #endif
 };
 
-FHttpClient::FHttpClient()
+FHttpClient::FHttpClient(UE::HTTP::FEventLoop& Loop)
+	: EventLoop(Loop)
 {
 }
 
@@ -516,7 +517,8 @@ class FOnDemandIoBackend final
 
 		FString DebugName;
 		UE::Tasks::FPipe Pipe;
-		FHttpClient Client;
+		UE::HTTP::FEventLoop HttpLoop;
+		FHttpClient Client{HttpLoop};
 	};
 
 	struct FBackendData
