Chaos: Add stat ID to async callbacks to better identify which callbacks are running and for how long during callback application pre-simulate

diff --git a/Engine/Plugins/Experimental/ChaosUserDataPT/Source/ChaosUserDataPT/Public/ChaosUserDataPT.h b/Engine/Plugins/Experimental/ChaosUserDataPT/Source/ChaosUserDataPT/Public/ChaosUserDataPT.h
--- a/Engine/Plugins/Experimental/ChaosUserDataPT/Source/ChaosUserDataPT/Public/ChaosUserDataPT.h
+++ b/Engine/Plugins/Experimental/ChaosUserDataPT/Source/ChaosUserDataPT/Public/ChaosUserDataPT.h
@@ -208,6 +208,12 @@ namespace Chaos
 			return false;
 		}
 
+		virtual FName GetFNameForStatId() const override
+		{
+			const static FLazyName StaticName("TUserDataManagerPT");
+			return StaticName;
+		}
+
 	protected:
 
 		virtual void OnPreSimulate_Internal() override
diff --git a/Engine/Plugins/Experimental/ChaosVehiclesPlugin/Source/ChaosVehicles/Private/ChaosVehicleManagerAsyncCallback.cpp b/Engine/Plugins/Experimental/ChaosVehiclesPlugin/Source/ChaosVehicles/Private/ChaosVehicleManagerAsyncCallback.cpp
--- a/Engine/Plugins/Experimental/ChaosVehiclesPlugin/Source/ChaosVehicles/Private/ChaosVehicleManagerAsyncCallback.cpp
+++ b/Engine/Plugins/Experimental/ChaosVehiclesPlugin/Source/ChaosVehicles/Private/ChaosVehicleManagerAsyncCallback.cpp
@@ -11,6 +11,12 @@ extern FVehicleDebugParams GVehicleDebugParams;
 
 DECLARE_CYCLE_STAT(TEXT("AsyncCallback:OnPreSimulate_Internal"), STAT_AsyncCallback_OnPreSimulate, STATGROUP_ChaosVehicleManager);
 
+FName FChaosVehicleManagerAsyncCallback::GetFNameForStatId() const
+{
+	const static FLazyName StaticName("FChaosVehicleManagerAsyncCallback");
+	return StaticName;
+}
+
 /**
  * Callback from Physics thread
  */
diff --git a/Engine/Plugins/Experimental/ChaosVehiclesPlugin/Source/ChaosVehicles/Public/ChaosVehicleManagerAsyncCallback.h b/Engine/Plugins/Experimental/ChaosVehiclesPlugin/Source/ChaosVehicles/Public/ChaosVehicleManagerAsyncCallback.h
--- a/Engine/Plugins/Experimental/ChaosVehiclesPlugin/Source/ChaosVehicles/Public/ChaosVehicleManagerAsyncCallback.h
+++ b/Engine/Plugins/Experimental/ChaosVehiclesPlugin/Source/ChaosVehicles/Public/ChaosVehicleManagerAsyncCallback.h
@@ -180,6 +180,8 @@ struct FChaosVehicleManagerAsyncOutput : public Chaos::FSimCallbackOutput
  */
 class CHAOSVEHICLES_API FChaosVehicleManagerAsyncCallback : public Chaos::TSimCallbackObject<FChaosVehicleManagerAsyncInput, FChaosVehicleManagerAsyncOutput, Chaos::ESimCallbackOptions::Presimulate | Chaos::ESimCallbackOptions::ContactModification>
 {
+public:
+	virtual FName GetFNameForStatId() const override;
 private:
 	virtual void OnPreSimulate_Internal() override;
 	virtual void OnContactModification_Internal(Chaos::FCollisionContactModifier& Modifications) override;
diff --git a/Engine/Plugins/Experimental/Water/Source/Runtime/Private/BuoyancyManager.cpp b/Engine/Plugins/Experimental/Water/Source/Runtime/Private/BuoyancyManager.cpp
--- a/Engine/Plugins/Experimental/Water/Source/Runtime/Private/BuoyancyManager.cpp
+++ b/Engine/Plugins/Experimental/Water/Source/Runtime/Private/BuoyancyManager.cpp
@@ -302,6 +302,12 @@ void ABuoyancyManager::EndPlay(const EEndPlayReason::Type EndPlayReason)
 	}
 }
 
+FName FBuoyancyManagerAsyncCallback::GetFNameForStatId() const
+{
+	const static FLazyName StaticName("FBuoyancyManagerAsyncCallback");
+	return StaticName;
+}
+
 void FBuoyancyManagerAsyncCallback::CreateAsyncAux_External(Chaos::FUniqueIdx HandleIndex, TUniquePtr<FBuoyancyComponentAsyncAux>&& AsyncAux)
 {
 	GetSolver()->EnqueueCommandImmediate([this, HandleIndex, AsyncAux = MoveTemp(AsyncAux)]() mutable
diff --git a/Engine/Plugins/Experimental/Water/Source/Runtime/Public/BuoyancyManager.h b/Engine/Plugins/Experimental/Water/Source/Runtime/Public/BuoyancyManager.h
--- a/Engine/Plugins/Experimental/Water/Source/Runtime/Public/BuoyancyManager.h
+++ b/Engine/Plugins/Experimental/Water/Source/Runtime/Public/BuoyancyManager.h
@@ -17,8 +17,9 @@ class AWaterBody;
 
 class FBuoyancyManagerAsyncCallback : public Chaos::TSimCallbackObject<FBuoyancyManagerAsyncInput, FBuoyancyManagerAsyncOutput>
 {
-
 public:
+	virtual FName GetFNameForStatId() const override;
+
 	void CreateAsyncAux_External(Chaos::FUniqueIdx UniqueIndex, TUniquePtr<FBuoyancyComponentAsyncAux>&& Aux);
 	void ClearAsyncAux_External(Chaos::FUniqueIdx UniqueIndex);
 
diff --git a/Engine/Source/Runtime/Engine/Private/CharacterMovementComponentAsync.cpp b/Engine/Source/Runtime/Engine/Private/CharacterMovementComponentAsync.cpp
--- a/Engine/Source/Runtime/Engine/Private/CharacterMovementComponentAsync.cpp
+++ b/Engine/Source/Runtime/Engine/Private/CharacterMovementComponentAsync.cpp
@@ -4557,6 +4557,12 @@ void FCharacterAsyncInput::OnMovementModeChanged(EMovementMode PrevMovementMode,
 	//MovementModeChangedDelegate.Broadcast(this, PrevMovementMode, PrevCustomMode);
 }
 
+FName FCharacterMovementComponentAsyncCallback::GetFNameForStatId() const
+{
+	const static FLazyName StaticName("FCharacterMovementComponentAsyncCallback");
+	return StaticName;
+}
+
 void FCharacterMovementComponentAsyncCallback::OnPreSimulate_Internal()
 {
 	PreSimulateImpl<FCharacterMovementComponentAsyncInput, FCharacterMovementComponentAsyncOutput>(*this);
diff --git a/Engine/Source/Runtime/Engine/Private/PhysicsEngine/Experimental/PhysScene_Chaos.cpp b/Engine/Source/Runtime/Engine/Private/PhysicsEngine/Experimental/PhysScene_Chaos.cpp
--- a/Engine/Source/Runtime/Engine/Private/PhysicsEngine/Experimental/PhysScene_Chaos.cpp
+++ b/Engine/Source/Runtime/Engine/Private/PhysicsEngine/Experimental/PhysScene_Chaos.cpp
@@ -148,16 +148,30 @@ public:
 		const FReal DeltaTime = GetDeltaTime_Internal();
 		const FReal SimTime = GetSimTime_Internal();
 		//TODO: handle case where callbacks modify AsyncPhysicsTickComponents or AsyncPhysicsTickActors
-		for (UActorComponent* Component : AsyncPhysicsTickComponents)
 		{
-			Component->AsyncPhysicsTickComponent(DeltaTime, SimTime);
+			QUICK_SCOPE_CYCLE_COUNTER(STAT_AsyncPhys_TickComponents);
+			for(UActorComponent* Component : AsyncPhysicsTickComponents)
+			{
+				FScopeCycleCounterUObject ComponentScope(Component);
+				Component->AsyncPhysicsTickComponent(DeltaTime, SimTime);
+			}
 		}
 
-		for(AActor* Actor : AsyncPhysicsTickActors)
 		{
-			Actor->AsyncPhysicsTickActor(DeltaTime, SimTime);
+			QUICK_SCOPE_CYCLE_COUNTER(STAT_AsyncPhys_TickActors);
+			for(AActor* Actor : AsyncPhysicsTickActors)
+			{
+				FScopeCycleCounterUObject ActorScope(Actor);
+				Actor->AsyncPhysicsTickActor(DeltaTime, SimTime);
+			}
 		}
 	}
+
+	virtual FName GetFNameForStatId() const override
+	{
+		const static FLazyName StaticName("FAsyncPhysicsTickCallback");
+		return StaticName;
+	}
 };
 
 DEFINE_LOG_CATEGORY_STATIC(LogFPhysScene_ChaosSolver, Log, All);
diff --git a/Engine/Source/Runtime/Engine/Private/PhysicsEngine/PhysicsReplication.cpp b/Engine/Source/Runtime/Engine/Private/PhysicsEngine/PhysicsReplication.cpp
--- a/Engine/Source/Runtime/Engine/Private/PhysicsEngine/PhysicsReplication.cpp
+++ b/Engine/Source/Runtime/Engine/Private/PhysicsEngine/PhysicsReplication.cpp
@@ -96,6 +96,16 @@ struct FAsyncPhysicsRepCallbackData : public Chaos::FSimCallbackInput
 
 class FPhysicsReplicationAsyncCallback final : public Chaos::TSimCallbackObject<FAsyncPhysicsRepCallbackData>
 {
+public:
+
+	virtual FName GetFNameForStatId() const override
+	{
+		const static FLazyName StaticName("FPhysicsReplicationAsyncCallback");
+		return StaticName;
+	}
+
+private:
+
 	virtual void OnPreSimulate_Internal() override
 	{
 		FPhysicsReplication::ApplyAsyncDesiredState(GetDeltaTime_Internal(), GetConsumerInput_Internal());
diff --git a/Engine/Source/Runtime/Engine/Public/CharacterMovementComponentAsync.h b/Engine/Source/Runtime/Engine/Public/CharacterMovementComponentAsync.h
--- a/Engine/Source/Runtime/Engine/Public/CharacterMovementComponentAsync.h
+++ b/Engine/Source/Runtime/Engine/Public/CharacterMovementComponentAsync.h
@@ -621,6 +621,8 @@ struct ENGINE_API FCharacterMovementComponentAsyncInput : public Chaos::FSimCall
 
 class FCharacterMovementComponentAsyncCallback : public Chaos::TSimCallbackObject<FCharacterMovementComponentAsyncInput, FCharacterMovementComponentAsyncOutput>
 {
+public:
+	virtual FName GetFNameForStatId() const override;
 private:
 	virtual void OnPreSimulate_Internal() override;
 };
diff --git a/Engine/Source/Runtime/Experimental/Chaos/Private/Chaos/Framework/ChaosResultsManager.cpp b/Engine/Source/Runtime/Experimental/Chaos/Private/Chaos/Framework/ChaosResultsManager.cpp
--- a/Engine/Source/Runtime/Experimental/Chaos/Private/Chaos/Framework/ChaosResultsManager.cpp
+++ b/Engine/Source/Runtime/Experimental/Chaos/Private/Chaos/Framework/ChaosResultsManager.cpp
@@ -5,6 +5,9 @@
 #include "PhysicsProxy/SingleParticlePhysicsProxy.h"
 #include "PhysicsProxy/GeometryCollectionPhysicsProxy.h"
 #include "Chaos/ChaosMarshallingManager.h"
+#include "Chaos/Framework/PhysicsSolverBase.h"
+
+DECLARE_CYCLE_STAT(TEXT("Process Rigid Resim Targets"), STAT_RigidResimTargets, STATGROUP_Chaos);
 
 namespace Chaos
 {	
@@ -124,23 +127,37 @@ namespace Chaos
 	{
 		//clear results
 		const int32 Timestamp = PullData.SolverTimestamp;
+		
+		{
+			SCOPE_CYCLE_COUNTER(STAT_ProcessSingleProxy);
+			SetPrevNextDataHelperTyped<Mode>(PullData.DirtyRigids, Results.RigidInterpolations);
+		}
 
-		SetPrevNextDataHelperTyped<Mode>(PullData.DirtyRigids, Results.RigidInterpolations);
-		SetPrevNextDataHelperTyped<Mode>(PullData.DirtyGeometryCollections, Results.GeometryCollectionInterpolations);
-		SetPrevNextDataHelperTyped<Mode>(PullData.DirtyClusterUnions, Results.ClusterUnionInterpolations);
+		{
+			SCOPE_CYCLE_COUNTER(STAT_ProcessGCProxy);
+			SetPrevNextDataHelperTyped<Mode>(PullData.DirtyGeometryCollections, Results.GeometryCollectionInterpolations);
+		}
+
+		{
+			SCOPE_CYCLE_COUNTER(STAT_ProcessClusterUnionProxy);
+			SetPrevNextDataHelperTyped<Mode>(PullData.DirtyClusterUnions, Results.ClusterUnionInterpolations);
+		}
 
 		// update resim target for rigids
-		for (const FDirtyRigidParticleData& Data : PullData.DirtyRigids)
 		{
-			if (FSingleParticlePhysicsProxy* Proxy = Data.GetProxy())
+			SCOPE_CYCLE_COUNTER(STAT_RigidResimTargets);
+			for(const FDirtyRigidParticleData& Data : PullData.DirtyRigids)
 			{
-				// only if the proxy is associated with this channel 
-				if (Proxy->GetInterpolationData().GetInterpChannel_External() == ChannelIdx)
+				if(FSingleParticlePhysicsProxy* Proxy = Data.GetProxy())
 				{
-					//update leash target
-					if(FDirtyRigidParticleData* ResimTarget = ParticleToResimTarget.Find(Proxy))
+					// only if the proxy is associated with this channel 
+					if(Proxy->GetInterpolationData().GetInterpChannel_External() == ChannelIdx)
 					{
-						*ResimTarget = Data;
+						//update leash target
+						if(FDirtyRigidParticleData* ResimTarget = ParticleToResimTarget.Find(Proxy))
+						{
+							*ResimTarget = Data;
+						}
 					}
 				}
 			}
diff --git a/Engine/Source/Runtime/Experimental/Chaos/Private/Chaos/Framework/PhysicsSolverBase.cpp b/Engine/Source/Runtime/Experimental/Chaos/Private/Chaos/Framework/PhysicsSolverBase.cpp
--- a/Engine/Source/Runtime/Experimental/Chaos/Private/Chaos/Framework/PhysicsSolverBase.cpp
+++ b/Engine/Source/Runtime/Experimental/Chaos/Private/Chaos/Framework/PhysicsSolverBase.cpp
@@ -14,6 +14,14 @@
 #include "ChaosVisualDebugger/ChaosVisualDebuggerTrace.h"
 #include "ChaosVisualDebugger/ChaosVDContextProvider.h"
 
+DEFINE_STAT(STAT_AsyncPullResults);
+DEFINE_STAT(STAT_AsyncInterpolateResults);
+DEFINE_STAT(STAT_SyncPullResults);
+DEFINE_STAT(STAT_ProcessSingleProxy);
+DEFINE_STAT(STAT_ProcessGCProxy);
+DEFINE_STAT(STAT_ProcessClusterUnionProxy);
+DEFINE_STAT(STAT_PullConstraints);
+
 namespace Chaos
 {	
 	extern int GSingleThreadedPhysics;
diff --git a/Engine/Source/Runtime/Experimental/Chaos/Private/Chaos/PBDCollisionConstraints.cpp b/Engine/Source/Runtime/Experimental/Chaos/Private/Chaos/PBDCollisionConstraints.cpp
--- a/Engine/Source/Runtime/Experimental/Chaos/Private/Chaos/PBDCollisionConstraints.cpp
+++ b/Engine/Source/Runtime/Experimental/Chaos/Private/Chaos/PBDCollisionConstraints.cpp
@@ -579,6 +579,7 @@ namespace Chaos
 
 			for(ISimCallbackObject* ModifierCallback : CollisionModifiers)
 			{
+				FScopedTraceSolverCallback ScopedCallback(ModifierCallback);
 				ModifierCallback->ContactModification_Internal(Modifier);
 			}
 
diff --git a/Engine/Source/Runtime/Experimental/Chaos/Public/Chaos/Framework/PhysicsSolverBase.h b/Engine/Source/Runtime/Experimental/Chaos/Public/Chaos/Framework/PhysicsSolverBase.h
--- a/Engine/Source/Runtime/Experimental/Chaos/Public/Chaos/Framework/PhysicsSolverBase.h
+++ b/Engine/Source/Runtime/Experimental/Chaos/Public/Chaos/Framework/PhysicsSolverBase.h
@@ -24,6 +24,18 @@ DECLARE_MULTICAST_DELEGATE_OneParam(FSolverPreBuffer, Chaos::FReal);
 DECLARE_MULTICAST_DELEGATE_OneParam(FSolverPostAdvance, Chaos::FReal);
 DECLARE_MULTICAST_DELEGATE(FSolverTeardown);
 
+#ifndef UE_CHAOS_CALLBACK_TRACESTATS
+#define UE_CHAOS_CALLBACK_TRACESTATS STATS && !UE_BUILD_SHIPPING
+#endif
+
+DECLARE_CYCLE_STAT_EXTERN(TEXT("Async Pull Results"), STAT_AsyncPullResults, STATGROUP_Chaos, CHAOS_API);
+DECLARE_CYCLE_STAT_EXTERN(TEXT("Async Interpolate Results"), STAT_AsyncInterpolateResults, STATGROUP_Chaos, CHAOS_API);
+DECLARE_CYCLE_STAT_EXTERN(TEXT("Sync Pull Results"), STAT_SyncPullResults, STATGROUP_Chaos, CHAOS_API);
+DECLARE_CYCLE_STAT_EXTERN(TEXT("Process Single Particle Proxies"), STAT_ProcessSingleProxy, STATGROUP_Chaos, CHAOS_API);
+DECLARE_CYCLE_STAT_EXTERN(TEXT("Process Geometry Collections Proxies"), STAT_ProcessGCProxy, STATGROUP_Chaos, CHAOS_API);
+DECLARE_CYCLE_STAT_EXTERN(TEXT("Process Cluster Union Proxies"), STAT_ProcessClusterUnionProxy, STATGROUP_Chaos, CHAOS_API);
+DECLARE_CYCLE_STAT_EXTERN(TEXT("Pull Constraints"), STAT_PullConstraints, STATGROUP_Chaos, CHAOS_API);
+
 namespace Chaos
 {
 	class FPhysicsSolverBase;
@@ -222,6 +234,34 @@ namespace Chaos
 		FSolverTeardown EventTeardown;
 	};
 
+#if UE_CHAOS_CALLBACK_TRACESTATS
+	class FScopedTraceSolverCallback : public FCycleCounter
+	{
+	public:
+		FScopedTraceSolverCallback(ISimCallbackObject* InCallback)
+		{
+			if(InCallback)
+			{
+				if(TStatId CallbackStatId = InCallback->GetStatId(); FThreadStats::IsCollectingData(CallbackStatId))
+				{
+					Start(CallbackStatId);
+				}
+			}
+		}
+
+		~FScopedTraceSolverCallback()
+		{
+			Stop();
+		}
+	};
+#else
+	class FScopedTraceSolverCallback
+	{
+	public:
+		FScopedTraceSolverCallback(ISimCallbackObject* InCallback)
+		{}
+	};
+#endif
 
 	class CHAOS_API FPhysicsSolverBase : public FPhysicsSolverEvents
 	{
@@ -480,6 +520,8 @@ namespace Chaos
 				{
 					if (Callback->HasOption(ESimCallbackOptions::RunOnFrozenGameThread) == bGameThreadFrozen)
 					{
+						FScopedTraceSolverCallback TraceCallback(Callback);
+
 						Callback->SetSimAndDeltaTime_Internal(SimTime, MLastDt);
 						Callback->PreSimulate_Internal();
 					}
diff --git a/Engine/Source/Runtime/Experimental/Chaos/Public/Chaos/PhysicsSolverBaseImpl.h b/Engine/Source/Runtime/Experimental/Chaos/Public/Chaos/PhysicsSolverBaseImpl.h
--- a/Engine/Source/Runtime/Experimental/Chaos/Public/Chaos/PhysicsSolverBaseImpl.h
+++ b/Engine/Source/Runtime/Experimental/Chaos/Public/Chaos/PhysicsSolverBaseImpl.h
@@ -50,6 +50,8 @@ namespace Chaos
 		FPullPhysicsData* LatestData = nullptr;
 		if (IsUsingAsyncResults() && UseAsyncInterpolation)
 		{
+			SCOPE_CYCLE_COUNTER(STAT_AsyncInterpolateResults);
+
 			const FReal ResultsTime = GetPhysicsResultsTime_External();
 			//we want to interpolate between prev and next. There are a few cases to consider:
 			//case 1: dirty data exists in both prev and next. In this case continuous data is interpolated, state data is a step function from prev to next
@@ -58,7 +60,13 @@ namespace Chaos
 			//case 4: prev has no dirty data and next does. In this case interpolate from gt data to next
 			//case 5: prev has no dirty data and next was overwritten. In this case do nothing as the overwritten data wins, and also particle may be deleted
 
-			TArray<const FChaosInterpolationResults*> ResultsPerChannel = PullResultsManager->PullAsyncPhysicsResults_External(ResultsTime);
+			TArray<const FChaosInterpolationResults*> ResultsPerChannel;
+			
+			{
+				SCOPE_CYCLE_COUNTER(STAT_AsyncPullResults);
+				ResultsPerChannel = PullResultsManager->PullAsyncPhysicsResults_External(ResultsTime);
+			}
+
 			for (int32 ChannelIdx = 0; ChannelIdx < ResultsPerChannel.Num(); ++ChannelIdx)
 			{
 				const FChaosInterpolationResults& Results = *ResultsPerChannel[ChannelIdx];
@@ -67,37 +75,46 @@ namespace Chaos
 				const int32 SolverTimestamp = Results.Next ? Results.Next->SolverTimestamp : INDEX_NONE;
 
 				// single particles
-				for (const FChaosRigidInterpolationData& RigidInterp : Results.RigidInterpolations)
 				{
-					if (FSingleParticlePhysicsProxy* Proxy = RigidInterp.Prev.GetProxy())
+					SCOPE_CYCLE_COUNTER(STAT_ProcessSingleProxy);
+					for(const FChaosRigidInterpolationData& RigidInterp : Results.RigidInterpolations)
 					{
-						if (Proxy->PullFromPhysicsState(RigidInterp.Prev, SolverTimestamp, &RigidInterp.Next, &Results.Alpha))
+						if(FSingleParticlePhysicsProxy* Proxy = RigidInterp.Prev.GetProxy())
 						{
-							Private::TPullPhysicsStateDispatchHelper<TDispatcher>::Apply(Dispatcher, Proxy);
+							if(Proxy->PullFromPhysicsState(RigidInterp.Prev, SolverTimestamp, &RigidInterp.Next, &Results.Alpha))
+							{
+								Private::TPullPhysicsStateDispatchHelper<TDispatcher>::Apply(Dispatcher, Proxy);
+							}
 						}
 					}
 				}
 
 				// geometry collections
-				for (const FChaosGeometryCollectionInterpolationData& GCInterp : Results.GeometryCollectionInterpolations)
 				{
-					if (FGeometryCollectionPhysicsProxy* Proxy = GCInterp.Prev.GetProxy())
+					SCOPE_CYCLE_COUNTER(STAT_ProcessGCProxy);
+					for(const FChaosGeometryCollectionInterpolationData& GCInterp : Results.GeometryCollectionInterpolations)
 					{
-						if (Proxy->PullFromPhysicsState(GCInterp.Prev, SolverTimestamp, &GCInterp.Next, &Results.Alpha))
+						if(FGeometryCollectionPhysicsProxy* Proxy = GCInterp.Prev.GetProxy())
 						{
-							Private::TPullPhysicsStateDispatchHelper<TDispatcher>::Apply(Dispatcher, Proxy);
+							if(Proxy->PullFromPhysicsState(GCInterp.Prev, SolverTimestamp, &GCInterp.Next, &Results.Alpha))
+							{
+								Private::TPullPhysicsStateDispatchHelper<TDispatcher>::Apply(Dispatcher, Proxy);
+							}
 						}
 					}
 				}
 
 				// cluster unions
-				for (const FChaosClusterUnionInterpolationData& ClusterInterp : Results.ClusterUnionInterpolations)
 				{
-					if (FClusterUnionPhysicsProxy* Proxy = ClusterInterp.Prev.GetProxy())
+					SCOPE_CYCLE_COUNTER(STAT_ProcessClusterUnionProxy)
+					for(const FChaosClusterUnionInterpolationData& ClusterInterp : Results.ClusterUnionInterpolations)
 					{
-						if (Proxy->PullFromPhysicsState(ClusterInterp.Prev, SolverTimestamp, &ClusterInterp.Next, &Results.Alpha))
+						if(FClusterUnionPhysicsProxy* Proxy = ClusterInterp.Prev.GetProxy())
 						{
-							Private::TPullPhysicsStateDispatchHelper<TDispatcher>::Apply(Dispatcher, Proxy);
+							if(Proxy->PullFromPhysicsState(ClusterInterp.Prev, SolverTimestamp, &ClusterInterp.Next, &Results.Alpha))
+							{
+								Private::TPullPhysicsStateDispatchHelper<TDispatcher>::Apply(Dispatcher, Proxy);
+							}
 						}
 					}
 				}
@@ -105,6 +122,8 @@ namespace Chaos
 		}
 		else
 		{
+			SCOPE_CYCLE_COUNTER(STAT_SyncPullResults);
+
 			//no interpolation so just use latest, in non-substepping modes this will just be the next result
 			// available in the queue - however if we substepped externally we need to consume the whole
 			// queue by telling the sync pull that we expect multiple results.
@@ -115,37 +134,46 @@ namespace Chaos
 			const int32 SolverTimestamp = Results.Next ? Results.Next->SolverTimestamp : INDEX_NONE;
 
 			// Single particles
-			for (const FChaosRigidInterpolationData& RigidInterp : Results.RigidInterpolations)
 			{
-				if (FSingleParticlePhysicsProxy* Proxy = RigidInterp.Prev.GetProxy())
+				SCOPE_CYCLE_COUNTER(STAT_ProcessSingleProxy);
+				for(const FChaosRigidInterpolationData& RigidInterp : Results.RigidInterpolations)
 				{
-					if (Proxy->PullFromPhysicsState(RigidInterp.Next, SolverTimestamp))
+					if(FSingleParticlePhysicsProxy* Proxy = RigidInterp.Prev.GetProxy())
 					{
-						Private::TPullPhysicsStateDispatchHelper<TDispatcher>::Apply(Dispatcher, Proxy);
+						if(Proxy->PullFromPhysicsState(RigidInterp.Next, SolverTimestamp))
+						{
+							Private::TPullPhysicsStateDispatchHelper<TDispatcher>::Apply(Dispatcher, Proxy);
+						}
 					}
 				}
 			}
 
 			// geometry collections
-			for (const FChaosGeometryCollectionInterpolationData& GCInterp : Results.GeometryCollectionInterpolations)
 			{
-				if (FGeometryCollectionPhysicsProxy* Proxy = GCInterp.Prev.GetProxy())
+				SCOPE_CYCLE_COUNTER(STAT_ProcessGCProxy);
+				for(const FChaosGeometryCollectionInterpolationData& GCInterp : Results.GeometryCollectionInterpolations)
 				{
-					if (Proxy->PullFromPhysicsState(GCInterp.Next, SolverTimestamp))
+					if(FGeometryCollectionPhysicsProxy* Proxy = GCInterp.Prev.GetProxy())
 					{
-						Private::TPullPhysicsStateDispatchHelper<TDispatcher>::Apply(Dispatcher, Proxy);
+						if(Proxy->PullFromPhysicsState(GCInterp.Next, SolverTimestamp))
+						{
+							Private::TPullPhysicsStateDispatchHelper<TDispatcher>::Apply(Dispatcher, Proxy);
+						}
 					}
 				}
 			}
 
 			// cluster unions
-			for (const FChaosClusterUnionInterpolationData& ClusterInterp : Results.ClusterUnionInterpolations)
 			{
-				if (FClusterUnionPhysicsProxy* Proxy = ClusterInterp.Prev.GetProxy())
+				SCOPE_CYCLE_COUNTER(STAT_ProcessClusterUnionProxy)
+				for(const FChaosClusterUnionInterpolationData& ClusterInterp : Results.ClusterUnionInterpolations)
 				{
-					if (Proxy->PullFromPhysicsState(ClusterInterp.Prev, SolverTimestamp))
+					if(FClusterUnionPhysicsProxy* Proxy = ClusterInterp.Prev.GetProxy())
 					{
-						Private::TPullPhysicsStateDispatchHelper<TDispatcher>::Apply(Dispatcher, Proxy);
+						if(Proxy->PullFromPhysicsState(ClusterInterp.Prev, SolverTimestamp))
+						{
+							Private::TPullPhysicsStateDispatchHelper<TDispatcher>::Apply(Dispatcher, Proxy);
+						}
 					}
 				}
 			}
@@ -154,6 +182,7 @@ namespace Chaos
 		//no interpolation for GC or joints at the moment
 		if (LatestData)
 		{
+			SCOPE_CYCLE_COUNTER(STAT_PullConstraints);
 			const int32 SyncTimestamp = LatestData->SolverTimestamp;
 
 			//
diff --git a/Engine/Source/Runtime/Experimental/Chaos/Public/Chaos/SimCallbackObject.h b/Engine/Source/Runtime/Experimental/Chaos/Public/Chaos/SimCallbackObject.h
--- a/Engine/Source/Runtime/Experimental/Chaos/Public/Chaos/SimCallbackObject.h
+++ b/Engine/Source/Runtime/Experimental/Chaos/Public/Chaos/SimCallbackObject.h
@@ -6,6 +6,13 @@
 #include "Chaos/SimCallbackInput.h"
 #include "Chaos/CollisionResolutionTypes.h"
 #include "Chaos/GeometryParticlesfwd.h"
+#include "ChaosStats.h"
+
+// Enable through build or just here in code to cause an untracked callback to fail to compile
+// Untracked callbacks will show up a such in profiling sessions and this can help track them down
+#ifndef UE_CHAOS_UNTRACKED_CALLBACK_IS_ERROR
+#define UE_CHAOS_UNTRACKED_CALLBACK_IS_ERROR 0
+#endif
 
 namespace Chaos
 {
@@ -112,15 +119,54 @@ public:
 	 */
 	virtual void FreeInputData_Internal(FSimCallbackInput* Input) = 0;
 
+#if STATS
+	/**
+	 * Get a stat to describe this callback. Uses GetFNameForStatId() which derived classes should implement to describe themselves
+	 */
+	TStatId GetStatId() const
+	{
+		static TStatId StatId = FDynamicStats::CreateStatId<FStatGroup_STATGROUP_Chaos>(GetFNameForStatId());
+		return StatId;
+	};
+#endif
 
-	FPhysicsSolverBase* GetSolver() { return Solver; }
+	/**
+	 * Get a stat name to describe this callback in external profiling and STATS when Named events are enabled
+	 */
+#if UE_CHAOS_UNTRACKED_CALLBACK_IS_ERROR
+	virtual FName GetFNameForStatId() const = 0;
+#else
+	virtual FName GetFNameForStatId() const
+	{
+		const static FLazyName StaticName("Untracked Physics Callback");
+		return StaticName;
+	}
+#endif
 
-	const FPhysicsSolverBase* GetSolver() const { return Solver; }
+	FPhysicsSolverBase* GetSolver()
+	{ 
+		return Solver; 
+	}
+	
+	const FPhysicsSolverBase* GetSolver() const 
+	{ 
+		return Solver; 
+	}
 
 	// Rewind API
-	virtual int32 TriggerRewindIfNeeded_Internal(int32 LastCompletedStep) { ensure(false); return INDEX_NONE; }
-	virtual void ApplyCorrections_Internal(int32 PhysicsStep, FSimCallbackInput* Input) { ensure(false); }
-	virtual void FirstPreResimStep_Internal(int32 PhysicsStep) { }
+	virtual int32 TriggerRewindIfNeeded_Internal(int32 LastCompletedStep)
+	{ 
+		ensure(false); 
+		return INDEX_NONE; 
+	}
+
+	virtual void ApplyCorrections_Internal(int32 PhysicsStep, FSimCallbackInput* Input)
+	{
+		ensure(false);
+	}
+	
+	virtual void FirstPreResimStep_Internal(int32 PhysicsStep)
+	{}
 
 	bool HasOption(const ESimCallbackOptions Option) const
 	{
@@ -250,7 +296,10 @@ private:
 	FPhysicsSolverBase* Solver;
 
 	//putting this here so that user classes don't have to bother with non-default constructor
-	void SetSolver_External(FPhysicsSolverBase* InSolver) { Solver = InSolver;}
+	void SetSolver_External(FPhysicsSolverBase* InSolver)
+	{
+		Solver = InSolver;
+	}
 	
 	UE_DEPRECATED(5.1, "Do not change options after creation of the callback object - instead, specify them using the TOptions template parameter.")
 	void SetContactModification(bool InContactModification)
@@ -261,7 +310,11 @@ private:
 protected:
 	FSimCallbackOutput* CurrentOutput_Internal;	//the output currently being written to in this sim step
 
-	const FSimCallbackInput* GetCurrentInput_Internal() const { return CurrentInput_Internal; }
+	const FSimCallbackInput* GetCurrentInput_Internal() const 
+	{ 
+		return CurrentInput_Internal; 
+	}
+
 private:
 	FSimCallbackInput* CurrentInput_Internal;	        //the input associated with the step we are executing.
 
@@ -289,6 +342,12 @@ public:
 		check(false);
 	}
 
+	virtual FName GetFNameForStatId() const override
+	{
+		const static FLazyName StaticName("FSimCallbackCommandObject");
+		return StaticName;
+	}
+
 private:
 
 	virtual FSimCallbackInput* AllocateInputData_External()
diff --git a/Engine/Source/Runtime/PhysicsCore/Private/ChaosScene.cpp b/Engine/Source/Runtime/PhysicsCore/Private/ChaosScene.cpp
--- a/Engine/Source/Runtime/PhysicsCore/Private/ChaosScene.cpp
+++ b/Engine/Source/Runtime/PhysicsCore/Private/ChaosScene.cpp
@@ -320,6 +320,12 @@ void FChaosSceneSimCallback::OnPreSimulate_Internal()
 	}
 }
 
+FName FChaosSceneSimCallback::GetFNameForStatId() const
+{
+	const static FLazyName StaticName("FChaosSceneSimCallback");
+	return StaticName;
+}
+
 void FChaosScene::SetGravity(const Chaos::FVec3& Acceleration)
 {
 	SimCallback->GetProducerInputData_External()->Gravity = Acceleration;
diff --git a/Engine/Source/Runtime/PhysicsCore/Public/Chaos/ChaosScene.h b/Engine/Source/Runtime/PhysicsCore/Public/Chaos/ChaosScene.h
--- a/Engine/Source/Runtime/PhysicsCore/Public/Chaos/ChaosScene.h
+++ b/Engine/Source/Runtime/PhysicsCore/Public/Chaos/ChaosScene.h
@@ -73,6 +73,7 @@ struct FChaosSceneCallbackInput : public Chaos::FSimCallbackInput
 struct FChaosSceneSimCallback : public Chaos::TSimCallbackObject<FChaosSceneCallbackInput>
 {
 	virtual void OnPreSimulate_Internal() override;
+	virtual FName GetFNameForStatId() const override;
 };
 
 /**
