Geometry collection : optimize passing transform to the custom renderer when the root is not yet broken
- avoid calling CalculateGlobalMatrices in that special case

diff --git a/Engine/Source/Runtime/Experimental/GeometryCollectionEngine/Private/GeometryCollection/GeometryCollectionComponent.cpp b/Engine/Source/Runtime/Experimental/GeometryCollectionEngine/Private/GeometryCollection/GeometryCollectionComponent.cpp
--- a/Engine/Source/Runtime/Experimental/GeometryCollectionEngine/Private/GeometryCollection/GeometryCollectionComponent.cpp
+++ b/Engine/Source/Runtime/Experimental/GeometryCollectionEngine/Private/GeometryCollection/GeometryCollectionComponent.cpp
@@ -4469,21 +4469,25 @@ void UGeometryCollectionComponent::RefreshCustomRenderer()
 		{
 			if (RestCollection != nullptr)
 			{
-				CalculateGlobalMatrices();
-				
 				const FTransform ComponentTransform = GetComponentTransform();
 				const int32 RootIndex = GetRootIndex();
+
 				const bool bIsBroken = DynamicCollection ? !DynamicCollection->Active[RootIndex] : false;
 
 				RendererInterface->UpdateState(*RestCollection, ComponentTransform, bIsBroken, !bHiddenInGame);
 
 				if (bIsBroken)
 				{
+					CalculateGlobalMatrices();
 					RendererInterface->UpdateTransforms(*RestCollection, ComponentSpaceTransforms);
 				}
 				else
 				{
-					const FTransform RootTransform = ComponentSpaceTransforms.IsValidIndex(RootIndex) ? ComponentSpaceTransforms[RootIndex] : FTransform::Identity;
+					// No need to compute the component space transform in that case, since the root is always in component space 
+					const FTransform RootTransform = 
+						(DynamicCollection && DynamicCollection->Transform.IsValidIndex(RootIndex))
+						? DynamicCollection->Transform[RootIndex]
+						: FTransform::Identity;
 					RendererInterface->UpdateRootTransform(*RestCollection, RootTransform);
 				}
 			}
