Chaos - fix FPBDRigidsEvolution::InvalidateParticle

diff --git a/Engine/Source/Programs/HeadlessChaos/Private/HeadlessChaosTestGraphEvolution.cpp b/Engine/Source/Programs/HeadlessChaos/Private/HeadlessChaosTestGraphEvolution.cpp
--- a/Engine/Source/Programs/HeadlessChaos/Private/HeadlessChaosTestGraphEvolution.cpp
+++ b/Engine/Source/Programs/HeadlessChaos/Private/HeadlessChaosTestGraphEvolution.cpp
@@ -218,6 +218,41 @@ namespace ChaosTest
 		EXPECT_EQ(Test.IslandManager->FindParticleIslands(Test.ParticleHandles[1]).Num(), 2);
 	}
 
+	// Start with an island containing 2 particle connected in a chain, then invalidate one of the particles.
+	// Check that the particles and constraints are removed from the graph and then re-added at the next tick.
+	//		(d=dynamic, s=sleeping, k=kinematic)
+	//		Ad - Bd
+	// =>	Ad   Bd
+	// =>	Ad - Bd
+	//
+	GTEST_TEST(GraphEvolutionTests, TestConstraintGraph_Invalidate)
+	{
+		FGraphEvolutionTest Test(2);
+		Test.MakeChain();
+
+		Test.Advance();
+
+		EXPECT_EQ(Test.IslandManager->GetNumIslands(), 1);
+		EXPECT_TRUE(Test.ParticleHandles[0]->IsInConstraintGraph());
+		EXPECT_TRUE(Test.ParticleHandles[1]->IsInConstraintGraph());
+		EXPECT_TRUE(Test.ConstraintHandles[0]->IsInConstraintGraph());
+
+		// Invalidate B
+		Test.Evolution.InvalidateParticle(Test.ParticleHandles[1]);
+
+		// Everything was kicked from the graph
+		EXPECT_FALSE(Test.ParticleHandles[0]->IsInConstraintGraph());
+		EXPECT_FALSE(Test.ParticleHandles[1]->IsInConstraintGraph());
+		EXPECT_FALSE(Test.ConstraintHandles[0]->IsInConstraintGraph());
+
+		Test.Advance();
+
+		// Everything was added back to the graph
+		EXPECT_TRUE(Test.ParticleHandles[0]->IsInConstraintGraph());
+		EXPECT_TRUE(Test.ParticleHandles[1]->IsInConstraintGraph());
+		EXPECT_TRUE(Test.ConstraintHandles[0]->IsInConstraintGraph());
+	}
+
 	// An isolated stationary particle with no gravity should go to sleep
 	GTEST_TEST(GraphEvolutionTests, TestConstraintGraph_ParticleSleep_Isolated)
 	{
diff --git a/Engine/Source/Runtime/Experimental/Chaos/Private/Chaos/Island/IslandManager.cpp b/Engine/Source/Runtime/Experimental/Chaos/Private/Chaos/Island/IslandManager.cpp
--- a/Engine/Source/Runtime/Experimental/Chaos/Private/Chaos/Island/IslandManager.cpp
+++ b/Engine/Source/Runtime/Experimental/Chaos/Private/Chaos/Island/IslandManager.cpp
@@ -657,11 +657,16 @@ namespace Chaos::Private
 	{
 		if (FPBDIslandParticle* Node = GetGraphNode(Particle))
 		{
-			// All constraints for this particle should have been removed first
-			check(Node->Edges.Num() == 0);
-			check(Node->Island == nullptr);
+			// All nodes in the graph have at least one edge
+			check(Node->Edges.Num() > 0);
+
+			// Remove all the constraints attached to the particle. This will
+			// also destroy the particle's node (and possibly others) 
+			// because empty nodes are removed as they occur
+			RemoveParticleConstraints(Particle);
 
-			DestroyGraphNode(Node);
+			// We should not be in the graph anymore
+			check(GetGraphNode(Particle) == nullptr);
 		}
 	}
 
@@ -752,7 +757,25 @@ namespace Chaos::Private
 		}
 	}
 
-	void FPBDIslandManager::RemoveParticleConstraints(FGeometryParticleHandle* Particle, int32 ContainerId)
+	void FPBDIslandManager::RemoveParticleConstraints(FGeometryParticleHandle* Particle)
+	{
+		if (FPBDIslandParticle* Node = GetGraphNode(Particle))
+		{
+			// Reverse loop because we remove as we iterate
+			for (int32 EdgeIndex = Node->Edges.Num() - 1; EdgeIndex >= 0; --EdgeIndex)
+			{
+				FPBDIslandConstraint* Edge = Node->Edges[EdgeIndex];
+
+				// NOTE: Destroys the edge and possibly this and other nodes (if this was their last constraint)
+				RemoveConstraint(Edge->Constraint);
+			}
+
+			// A node with no edges will have been removed
+			check(GetGraphNode(Particle) == nullptr);
+		}
+	}
+
+	void FPBDIslandManager::RemoveParticleContainerConstraints(FGeometryParticleHandle* Particle, const int32 ContainerId)
 	{
 		if (FPBDIslandParticle* Node = GetGraphNode(Particle))
 		{
@@ -769,7 +792,7 @@ namespace Chaos::Private
 		}
 	}
 
-	void FPBDIslandManager::RemoveContainerConstraints(int32 ContainerId)
+	void FPBDIslandManager::RemoveContainerConstraints(const int32 ContainerId)
 	{
 		// Reverse loop because we remove as we iterate
 		for (int32 EdgeIndex = Edges.Num() - 1; EdgeIndex >= 0; --EdgeIndex)
@@ -799,9 +822,9 @@ namespace Chaos::Private
 	{
 		if (FPBDIslandParticle* Node = GetGraphNode(Particle))
 		{
-			if ((Node->Island != nullptr) && !Node->Island->Flags.bIsSleeping)
+			if ((Node->Island != nullptr) && (Node->Island->Flags.bIsSleeping != bInIsSleeping))
 			{
-				Node->Island->Flags.bIsSleeping = true;
+				Node->Island->Flags.bIsSleeping = bInIsSleeping;
 				Node->Island->SleepCounter = 0;
 				PropagateIslandSleep(Node->Island);
 			}
@@ -1467,6 +1490,8 @@ namespace Chaos::Private
 		// Find the largest island and count the nodes and edges from all islands
 		for (FPBDIsland* Island : MergeSet->Islands)
 		{
+			check(Island != nullptr);
+
 			if ((LargestIsland == nullptr) || (Island->NumEdges > LargestIsland->NumEdges))
 			{
 				LargestIsland = Island;
diff --git a/Engine/Source/Runtime/Experimental/Chaos/Private/Chaos/PBDRigidsEvolutionGBF.cpp b/Engine/Source/Runtime/Experimental/Chaos/Private/Chaos/PBDRigidsEvolutionGBF.cpp
--- a/Engine/Source/Runtime/Experimental/Chaos/Private/Chaos/PBDRigidsEvolutionGBF.cpp
+++ b/Engine/Source/Runtime/Experimental/Chaos/Private/Chaos/PBDRigidsEvolutionGBF.cpp
@@ -908,7 +908,7 @@ void FPBDRigidsEvolutionGBF::DestroyTransientConstraints(FGeometryParticleHandle
 	if (Particle != nullptr)
 	{
 		// Remove all the particle's collisions from the graph
-		GetIslandManager().RemoveParticleConstraints(Particle, CollisionConstraints.GetContainerId());
+		GetIslandManager().RemoveParticleContainerConstraints(Particle, CollisionConstraints.GetContainerId());
 
 		// Mark all collision constraints for destruction
 		DestroyParticleCollisionsInAllocator(Particle);
diff --git a/Engine/Source/Runtime/Experimental/Chaos/Public/Chaos/Island/IslandManager.h b/Engine/Source/Runtime/Experimental/Chaos/Public/Chaos/Island/IslandManager.h
--- a/Engine/Source/Runtime/Experimental/Chaos/Public/Chaos/Island/IslandManager.h
+++ b/Engine/Source/Runtime/Experimental/Chaos/Public/Chaos/Island/IslandManager.h
@@ -402,8 +402,9 @@ namespace Chaos::Private
 		void AddConstraint(const int32 ContainerId, FConstraintHandle* Constraint, const TVec2<FGeometryParticleHandle*>& ConstrainedParticles);
 		template<typename ConstraintContainerType> void AddContainerConstraints(ConstraintContainerType& ConstraintContainer);
 		void RemoveConstraint(FConstraintHandle* ConstraintHandle);
-		void RemoveParticleConstraints(FGeometryParticleHandle* Particle, int32 ContainerId);
-		void RemoveContainerConstraints(int32 ContainerId);
+		void RemoveParticleConstraints(FGeometryParticleHandle* Particle);
+		void RemoveParticleContainerConstraints(FGeometryParticleHandle* Particle, const int32 ContainerId);
+		void RemoveContainerConstraints(const int32 ContainerId);
 
 		// Access to Islands
 		int32 GetNumIslands() const { return Islands.Num(); }
diff --git a/Engine/Source/Runtime/Experimental/Chaos/Public/Chaos/PBDRigidsEvolution.h b/Engine/Source/Runtime/Experimental/Chaos/Public/Chaos/PBDRigidsEvolution.h
--- a/Engine/Source/Runtime/Experimental/Chaos/Public/Chaos/PBDRigidsEvolution.h
+++ b/Engine/Source/Runtime/Experimental/Chaos/Public/Chaos/PBDRigidsEvolution.h
@@ -470,26 +470,15 @@ public:
 
 	/**
 	* To be called when a particle geometry changes. We must clear collisions and anything else that may reference the prior shapes.
-	* @note This is also called during particle creation before we even know if the particle is static or kinematic, so we don't (re)add to 
-	* the graph unless it was already in there. There should be a call to either RegisterParticle or EnableParticle
-	* after creating a particle which will handle initially adding to the graph.
 	*/
 	CHAOS_API void InvalidateParticle(FGeometryParticleHandle* Particle)
 	{
+		// Remove all constraints (collisions, joints etc) from the graph
+		IslandManager.RemoveParticleConstraints(Particle);
+
 		// Destroy all the transient constraints (collisions) because the particle has changed somehow (e.g. new shapes) and they may have cached the previous state
+		// @todo(chaos): if any other types depended on geometry, they would also need to be notified here. Maybe make this more specific and tell all types...
 		DestroyTransientConstraints(Particle);
-
-		// Remove all persistent constraints (joints etc) from the graph too
-		// @todo(chaos): do we still need to do this? Only if some non-collision constraints hold refs to shapes...
-		if (Particle->IsInConstraintGraph())
-		{
-			// Remove the particle from the constraint graph. This should remove all the constraints too
-			IslandManager.RemoveParticle(Particle);
-
-			// Re-add the particle to the constraint graph
-			// (we could add a RemoveParticleConstraints method to the graph, but removing and adding a particle isn't too bad)
-			IslandManager.AddParticle(Particle);
-		}
 	}
 	
 	CHAOS_API void FlushExternalAccelerationQueue(FAccelerationStructure& Acceleration,FPendingSpatialDataQueue& ExternalQueue);
