Added a helper function to get overlap AABB intersections

diff --git a/Engine/Source/Runtime/Experimental/Chaos/Private/Chaos/PhysicsObjectCollisionInterface.cpp b/Engine/Source/Runtime/Experimental/Chaos/Private/Chaos/PhysicsObjectCollisionInterface.cpp
--- a/Engine/Source/Runtime/Experimental/Chaos/Private/Chaos/PhysicsObjectCollisionInterface.cpp
+++ b/Engine/Source/Runtime/Experimental/Chaos/Private/Chaos/PhysicsObjectCollisionInterface.cpp
@@ -75,6 +75,28 @@ namespace Chaos
 		);
 	}
 
+	template<EThreadContext Id>
+	bool FPhysicsObjectCollisionInterface<Id>::PhysicsObjectOverlapWithAABBIntersections(const FConstPhysicsObjectHandle ObjectA, const FTransform& InTransformA, const FConstPhysicsObjectHandle ObjectB, const FTransform& InTransformB, bool bTraceComplex, const FVector& Tolerance, TArray<FBox>& OutIntersections)
+	{
+		TRACE_CPUPROFILER_EVENT_SCOPE(FPhysicsObjectCollisionInterface<Id>::PhysicsObjectOverlapWithAABBIntersections);
+		OutIntersections.Reset();
+		return PairwiseShapeOverlapHelper(
+			ObjectA,
+			InTransformA,
+			ObjectB,
+			InTransformB,
+			bTraceComplex,
+			false,
+			Tolerance,
+			[&OutIntersections, &Tolerance](const FShapeOverlapData& ShapeA, const FShapeOverlapData& ShapeB, const FMTDInfo&)
+			{
+				const FAABB3 Intersection = ShapeA.BoundingBox.GetIntersection(ShapeB.BoundingBox);
+				OutIntersections.Emplace(Intersection.Min(), Intersection.Max());
+				return true;
+			}
+		);
+	}
+
 	template<EThreadContext Id>
 	bool FPhysicsObjectCollisionInterface<Id>::PairwiseShapeOverlapHelper(const FConstPhysicsObjectHandle ObjectA, const FTransform& InTransformA, const FConstPhysicsObjectHandle ObjectB, const FTransform& InTransformB, bool bTraceComplex, bool bComputeMTD, const FVector& Tolerance, const TFunction<bool(const FShapeOverlapData&, const FShapeOverlapData&, const FMTDInfo&)>& Lambda)
 	{
diff --git a/Engine/Source/Runtime/Experimental/Chaos/Public/Chaos/PhysicsObjectCollisionInterface.h b/Engine/Source/Runtime/Experimental/Chaos/Public/Chaos/PhysicsObjectCollisionInterface.h
--- a/Engine/Source/Runtime/Experimental/Chaos/Public/Chaos/PhysicsObjectCollisionInterface.h
+++ b/Engine/Source/Runtime/Experimental/Chaos/Public/Chaos/PhysicsObjectCollisionInterface.h
@@ -57,6 +57,7 @@ namespace Chaos
 		// This function does the same as GetPhysicsObjectOverlap but also computes the AABB overlap metric.
 		bool PhysicsObjectOverlapWithAABB(const FConstPhysicsObjectHandle ObjectA, const FTransform& InTransformA, const FConstPhysicsObjectHandle ObjectB, const FTransform& InTransformB, bool bTraceComplex, const FVector& Tolerance, FBox& OutOverlap);
 		bool PhysicsObjectOverlapWithAABBSize(const FConstPhysicsObjectHandle ObjectA, const FTransform& InTransformA, const FConstPhysicsObjectHandle ObjectB, const FTransform& InTransformB, bool bTraceComplex, const FVector& Tolerance, FVector& OutOverlapSize);
+		bool PhysicsObjectOverlapWithAABBIntersections(const FConstPhysicsObjectHandle ObjectA, const FTransform& InTransformA, const FConstPhysicsObjectHandle ObjectB, const FTransform& InTransformB, bool bTraceComplex, const FVector& Tolerance, TArray<FBox>& Intersections);
 
 		template<typename TRaycastHit>
 		bool LineTrace(TArrayView<const FConstPhysicsObjectHandle> InObjects, const FVector& WorldStart, const FVector& WorldEnd, bool bTraceComplex, TRaycastHit& OutBestHit)
