Chaos - fix sleep for particles in the air with high sleep thresholds

diff --git a/Engine/Source/Runtime/Experimental/Chaos/Private/Chaos/Island/IslandManager.cpp b/Engine/Source/Runtime/Experimental/Chaos/Private/Chaos/Island/IslandManager.cpp
--- a/Engine/Source/Runtime/Experimental/Chaos/Private/Chaos/Island/IslandManager.cpp
+++ b/Engine/Source/Runtime/Experimental/Chaos/Private/Chaos/Island/IslandManager.cpp
@@ -1941,52 +1941,39 @@ namespace Chaos::Private
 		{
 			if (Rigid.IsDynamic() && !Rigid.IsInConstraintGraph())
 			{
-				// If we have gravity we will not stop unless constrained
-				// so we only care about particles with zero gravity
-				bool bHasGravity = Rigid.GravityEnabled();
-				if (bHasGravity && (Gravity != nullptr))
+				// @todo(chaos): this is very expensive because we have to search for a material
+				// We should probably cache a particle's sleep (and disable) thresholds somewhere.
+				FRealSingle SleepLinearThreshold, SleepAngularThreshold;
+				int32 SleepCounterThreshold;
+				GetParticleSleepThresholds(Rigid.Handle(), PhysicsMaterials, PerParticlePhysicsMaterials, SimMaterials, SleepLinearThreshold, SleepAngularThreshold, SleepCounterThreshold);
+
+				// Check for sleep
+				if ((SleepLinearThreshold > 0) || (SleepAngularThreshold > 0))
 				{
-					const int32 GravityGroup = Rigid.GravityGroupIndex();
-					if (GravityGroup != INDEX_NONE)
+					int32 SleepCounter = 0;
+					if (SleepCounterThreshold < TNumericLimits<int32>::Max())
 					{
-						bHasGravity = !Gravity->GetAcceleration(GravityGroup).IsNearlyZero();
-					}
-				}
-
-				if (!bHasGravity)
-				{
-					FRealSingle SleepLinearThreshold, SleepAngularThreshold;
-					int32 SleepCounterThreshold;
-					GetParticleSleepThresholds(Rigid.Handle(), PhysicsMaterials, PerParticlePhysicsMaterials, SimMaterials, SleepLinearThreshold, SleepAngularThreshold, SleepCounterThreshold);
+						// Isolated particles have a max sleep counter of 127 (to reduce counter space in the particle)
+						SleepCounterThreshold = FMath::Min(SleepCounterThreshold, TNumericLimits<int8>::Max());
 
-					// Check for sleep
-					if ((SleepLinearThreshold > 0) || (SleepAngularThreshold > 0))
-					{
-						int32 SleepCounter = 0;
-						if (SleepCounterThreshold < TNumericLimits<int32>::Max())
+						// Did we exceed the velocity threshold?
+						if ((Rigid.VSmooth().SizeSquared() > FMath::Square(SleepLinearThreshold)) 
+							|| (Rigid.WSmooth().SizeSquared() > FMath::Square(SleepAngularThreshold)))
 						{
-							// Isolated particles have a max sleep counter of 127 (to reduce counter space in the particle)
-							SleepCounterThreshold = FMath::Min(SleepCounterThreshold, TNumericLimits<int8>::Max());
-
-							// Did we exceed the velocity threshold?
-							if ((Rigid.VSmooth().SizeSquared() > FMath::Square(SleepLinearThreshold)) 
-								|| (Rigid.WSmooth().SizeSquared() > FMath::Square(SleepAngularThreshold)))
-							{
-								continue;
-							}
+							continue;
+						}
 
-							// If we get here we want to sleep
-							// Update the counter and sleep if we exceed it
-							static_assert(sizeof(decltype(Rigid.SleepCounter())) == 1, "Expected int8 for SleepCounter(). Update clamp below");
-							SleepCounter = FMath::Min(int32(Rigid.SleepCounter()) + 1, int32(TNumericLimits<int8>::Max()));
-							if (SleepCounter > SleepCounterThreshold)
-							{
-								SleptParticles.Add(Rigid.Handle());
-								SleepCounter = 0;
-							}
+						// If we get here we want to sleep
+						// Update the counter and sleep if we exceed it
+						static_assert(sizeof(decltype(Rigid.SleepCounter())) == 1, "Expected int8 for SleepCounter(). Update clamp below");
+						SleepCounter = FMath::Min(int32(Rigid.SleepCounter()) + 1, int32(TNumericLimits<int8>::Max()));
+						if (SleepCounter > SleepCounterThreshold)
+						{
+							SleptParticles.Add(Rigid.Handle());
+							SleepCounter = 0;
 						}
-						Rigid.SetSleepCounter(int8(SleepCounter));
 					}
+					Rigid.SetSleepCounter(int8(SleepCounter));
 				}
 			}
 		}
@@ -2138,6 +2125,9 @@ namespace Chaos::Private
 		{
 			if (Particle.IsDynamic() && !Particle.IsInConstraintGraph())
 			{
+				// @todo(chaos): this is very expensive becuase we have to search materials.
+				// However this system is only used by Fields via PerParticlePhysicsMaterials, and not exposed to the editor.
+				// Maybe we should just skip looking further if no PerParticlePhysicsMaterials is assigned.
 				FRealSingle DisableLinearThreshold, DisableAngularThreshold;
 				GetParticleDisableThresholds(Particle.Handle(), PhysicsMaterials, PerParticlePhysicsMaterials, SimMaterials, DisableLinearThreshold, DisableAngularThreshold);
 
