fix incorrect early-out for empty selection in geometry collection convex hull merging

diff --git a/Engine/Source/Runtime/Experimental/Chaos/Private/GeometryCollection/GeometryCollectionConvexUtility.cpp b/Engine/Source/Runtime/Experimental/Chaos/Private/GeometryCollection/GeometryCollectionConvexUtility.cpp
--- a/Engine/Source/Runtime/Experimental/Chaos/Private/GeometryCollection/GeometryCollectionConvexUtility.cpp
+++ b/Engine/Source/Runtime/Experimental/Chaos/Private/GeometryCollection/GeometryCollectionConvexUtility.cpp
@@ -1689,7 +1689,7 @@ void FGeometryCollectionConvexUtility::GenerateClusterConvexHullsFromLeafOrChild
 	TArray<int32> TransformsToProcess;
 	if (bOnlySubset)
 	{
-		if (TransformsToProcess.IsEmpty())
+		if (TransformSubset.IsEmpty())
 		{
 			return;
 		}
