Chaos - add duration to CCD debug draw commands

diff --git a/Engine/Source/Runtime/Experimental/Chaos/Private/Chaos/ChaosDebugDraw.cpp b/Engine/Source/Runtime/Experimental/Chaos/Private/Chaos/ChaosDebugDraw.cpp
--- a/Engine/Source/Runtime/Experimental/Chaos/Private/Chaos/ChaosDebugDraw.cpp
+++ b/Engine/Source/Runtime/Experimental/Chaos/Private/Chaos/ChaosDebugDraw.cpp
@@ -296,10 +296,14 @@ namespace Chaos
 
 		void DrawShape(const FRigidTransform3& ShapeTransform, const FImplicitObject* Implicit, const FShapeOrShapesArray& Shapes, const FColor& Color, const FChaosDebugDrawSettings* Settings)
 		{
-			// Note: At the time of commenting this function does nothing as DrawShapesImpl does not handle null particle.
 			DrawShapesImpl(nullptr, ShapeTransform, Implicit, Shapes, 0.0f, Color, 0.0f, GetChaosDebugDrawSettings(Settings));
 		}
 
+		void DrawShape(const FRigidTransform3& ShapeTransform, const FImplicitObject* Implicit, const FShapeOrShapesArray& Shapes, const FColor& Color, const float Duration, const FChaosDebugDrawSettings* Settings)
+		{
+			DrawShapesImpl(nullptr, ShapeTransform, Implicit, Shapes, 0.0f, Color, Duration, GetChaosDebugDrawSettings(Settings));
+		}
+
 		void DrawShapesConvexImpl(const TGeometryParticleHandle<FReal, 3>* Particle, const FRigidTransform3& ShapeTransform, const FConvex* Shape, const FReal InMargin, const FColor& Color, const FRealSingle Duration, const FChaosDebugDrawSettings& Settings)
 		{
 			if (Shape->HasStructureData())
diff --git a/Engine/Source/Runtime/Experimental/Chaos/Private/Chaos/Collision/ParticlePairMidPhase.cpp b/Engine/Source/Runtime/Experimental/Chaos/Private/Chaos/Collision/ParticlePairMidPhase.cpp
--- a/Engine/Source/Runtime/Experimental/Chaos/Private/Chaos/Collision/ParticlePairMidPhase.cpp
+++ b/Engine/Source/Runtime/Experimental/Chaos/Private/Chaos/Collision/ParticlePairMidPhase.cpp
@@ -20,13 +20,21 @@ extern bool Chaos_Collision_NarrowPhase_AABBBoundsCheck;
 
 namespace Chaos
 {
+#if CHAOS_DEBUG_DRAW
 	namespace CVars
 	{
 		extern int32 ChaosSolverDrawCCDInteractions;
-#if CHAOS_DEBUG_DRAW
 		extern DebugDraw::FChaosDebugDrawSettings ChaosSolverDebugDebugDrawSettings;
+	}
+	namespace DebugDraw
+	{
+		extern float ChaosDebugDrawCCDDuration;
+	}
 #endif
 
+	namespace CVars
+	{
+
 		bool bChaos_Collision_MidPhase_EnableBoundsChecks = true;
 		FAutoConsoleVariableRef CVarChaos_Collision_EnableBoundsChecks(TEXT("p.Chaos.Collision.EnableBoundsChecks"), bChaos_Collision_MidPhase_EnableBoundsChecks, TEXT(""));
 
@@ -468,13 +476,13 @@ namespace Chaos
 			{
 				if (FConstGenericParticleHandle(Constraint->GetParticle0())->CCDEnabled())
 				{
-					DebugDraw::DrawShape(CCDShapeWorldTransform0, Constraint->GetImplicit0(), Shape0, FColor::Black, &CVars::ChaosSolverDebugDebugDrawSettings);
-					DebugDraw::DrawShape(ShapeWorldTransform0, Constraint->GetImplicit0(), Shape0, FColor::White, &CVars::ChaosSolverDebugDebugDrawSettings);
+					DebugDraw::DrawShape(CCDShapeWorldTransform0, Constraint->GetImplicit0(), Shape0, FColor::Black, DebugDraw::ChaosDebugDrawCCDDuration, &CVars::ChaosSolverDebugDebugDrawSettings);
+					DebugDraw::DrawShape(ShapeWorldTransform0, Constraint->GetImplicit0(), Shape0, FColor::White, DebugDraw::ChaosDebugDrawCCDDuration, &CVars::ChaosSolverDebugDebugDrawSettings);
 				}
 				if (FConstGenericParticleHandle(Constraint->GetParticle1())->CCDEnabled())
 				{
-					DebugDraw::DrawShape(CCDShapeWorldTransform1, Constraint->GetImplicit1(), Shape1, FColor::Black, &CVars::ChaosSolverDebugDebugDrawSettings);
-					DebugDraw::DrawShape(ShapeWorldTransform1, Constraint->GetImplicit1(), Shape1, FColor::White, &CVars::ChaosSolverDebugDebugDrawSettings);
+					DebugDraw::DrawShape(CCDShapeWorldTransform1, Constraint->GetImplicit1(), Shape1, FColor::Black, DebugDraw::ChaosDebugDrawCCDDuration, &CVars::ChaosSolverDebugDebugDrawSettings);
+					DebugDraw::DrawShape(ShapeWorldTransform1, Constraint->GetImplicit1(), Shape1, FColor::White, DebugDraw::ChaosDebugDrawCCDDuration, &CVars::ChaosSolverDebugDebugDrawSettings);
 				}
 			}
 #endif
@@ -917,13 +925,13 @@ namespace Chaos
 		{
 			if (FConstGenericParticleHandle(Constraint->GetParticle0())->CCDEnabled())
 			{
-				DebugDraw::DrawShape(CCDShapeWorldTransform0, Constraint->GetImplicit0(), Constraint->GetShape0(), FColor::Black, &CVars::ChaosSolverDebugDebugDrawSettings);
-				DebugDraw::DrawShape(ShapeWorldTransform0, Constraint->GetImplicit0(), Constraint->GetShape0(), FColor::White, &CVars::ChaosSolverDebugDebugDrawSettings);
+				DebugDraw::DrawShape(CCDShapeWorldTransform0, Constraint->GetImplicit0(), Constraint->GetShape0(), FColor::Black, DebugDraw::ChaosDebugDrawCCDDuration, &CVars::ChaosSolverDebugDebugDrawSettings);
+				DebugDraw::DrawShape(ShapeWorldTransform0, Constraint->GetImplicit0(), Constraint->GetShape0(), FColor::White, DebugDraw::ChaosDebugDrawCCDDuration, &CVars::ChaosSolverDebugDebugDrawSettings);
 			}
 			if (FConstGenericParticleHandle(Constraint->GetParticle1())->CCDEnabled())
 			{
-				DebugDraw::DrawShape(CCDShapeWorldTransform1, Constraint->GetImplicit1(), Constraint->GetShape1(), FColor::Black, &CVars::ChaosSolverDebugDebugDrawSettings);
-				DebugDraw::DrawShape(ShapeWorldTransform1, Constraint->GetImplicit1(), Constraint->GetShape1(), FColor::White, &CVars::ChaosSolverDebugDebugDrawSettings);
+				DebugDraw::DrawShape(CCDShapeWorldTransform1, Constraint->GetImplicit1(), Constraint->GetShape1(), FColor::Black, DebugDraw::ChaosDebugDrawCCDDuration, &CVars::ChaosSolverDebugDebugDrawSettings);
+				DebugDraw::DrawShape(ShapeWorldTransform1, Constraint->GetImplicit1(), Constraint->GetShape1(), FColor::White, DebugDraw::ChaosDebugDrawCCDDuration, &CVars::ChaosSolverDebugDebugDrawSettings);
 			}
 		}
 #endif
@@ -1353,13 +1361,13 @@ namespace Chaos
 		{
 			if (FConstGenericParticleHandle(Constraint->GetParticle0())->CCDEnabled())
 			{
-				DebugDraw::DrawShape(CCDShapeWorldTransform0, Constraint->GetImplicit0(), Constraint->GetShape0(), FColor::Black, &CVars::ChaosSolverDebugDebugDrawSettings);
-				DebugDraw::DrawShape(ShapeWorldTransform0, Constraint->GetImplicit0(), Constraint->GetShape0(), FColor::White, &CVars::ChaosSolverDebugDebugDrawSettings);
+				DebugDraw::DrawShape(CCDShapeWorldTransform0, Constraint->GetImplicit0(), Constraint->GetShape0(), FColor::Black, DebugDraw::ChaosDebugDrawCCDDuration, &CVars::ChaosSolverDebugDebugDrawSettings);
+				DebugDraw::DrawShape(ShapeWorldTransform0, Constraint->GetImplicit0(), Constraint->GetShape0(), FColor::White, DebugDraw::ChaosDebugDrawCCDDuration, &CVars::ChaosSolverDebugDebugDrawSettings);
 			}
 			if (FConstGenericParticleHandle(Constraint->GetParticle1())->CCDEnabled())
 			{
-				DebugDraw::DrawShape(CCDShapeWorldTransform1, Constraint->GetImplicit1(), Constraint->GetShape1(), FColor::Black, &CVars::ChaosSolverDebugDebugDrawSettings);
-				DebugDraw::DrawShape(ShapeWorldTransform1, Constraint->GetImplicit1(), Constraint->GetShape1(), FColor::White, &CVars::ChaosSolverDebugDebugDrawSettings);
+				DebugDraw::DrawShape(CCDShapeWorldTransform1, Constraint->GetImplicit1(), Constraint->GetShape1(), FColor::Black, DebugDraw::ChaosDebugDrawCCDDuration, &CVars::ChaosSolverDebugDebugDrawSettings);
+				DebugDraw::DrawShape(ShapeWorldTransform1, Constraint->GetImplicit1(), Constraint->GetShape1(), FColor::White, DebugDraw::ChaosDebugDrawCCDDuration, &CVars::ChaosSolverDebugDebugDrawSettings);
 			}
 		}
 #endif
diff --git a/Engine/Source/Runtime/Experimental/Chaos/Public/Chaos/ChaosDebugDraw.h b/Engine/Source/Runtime/Experimental/Chaos/Public/Chaos/ChaosDebugDraw.h
--- a/Engine/Source/Runtime/Experimental/Chaos/Public/Chaos/ChaosDebugDraw.h
+++ b/Engine/Source/Runtime/Experimental/Chaos/Public/Chaos/ChaosDebugDraw.h
@@ -221,6 +221,7 @@ namespace Chaos
 		CHAOS_API void DrawCharacterGroundConstraints(const FRigidTransform3& SpaceTransform, const FCharacterGroundConstraintContainer& Constraints, const FChaosDebugDrawSettings* Settings = nullptr);
 		CHAOS_API void DrawSimulationSpace(const FSimulationSpace& SimSpace, const FChaosDebugDrawSettings* Settings = nullptr);
 		CHAOS_API void DrawShape(const FRigidTransform3& ShapeTransform, const FImplicitObject* Implicit, const FShapeOrShapesArray& Shapes, const FColor& Color, const FChaosDebugDrawSettings* Settings = nullptr);
+		CHAOS_API void DrawShape(const FRigidTransform3& ShapeTransform, const FImplicitObject* Implicit, const FShapeOrShapesArray& Shapes, const FColor& Color, const float Duration, const FChaosDebugDrawSettings* Settings = nullptr);
 		CHAOS_API void DrawConstraintGraph(const FRigidTransform3& ShapeTransform, const Private::FPBDIslandManager& Graph, const FChaosDebugDrawSettings* Settings = nullptr);
 		CHAOS_API void DrawConnectionGraph(const FRigidClustering& Clustering, const FChaosDebugDrawSettings* Settings = nullptr);
 		CHAOS_API void DrawCollidingShapes(const FRigidTransform3& SpaceTransform, const FPBDCollisionConstraints& Collisions, FRealSingle ColorScale, const FRealSingle Duration = 0.f, const FChaosDebugDrawSettings* Settings = nullptr);
