Updating the break events to pass along the component.

diff --git a/Engine/Source/Runtime/Experimental/GeometryCollectionEngine/Private/GeometryCollection/GeometryCollectionComponent.cpp b/Engine/Source/Runtime/Experimental/GeometryCollectionEngine/Private/GeometryCollection/GeometryCollectionComponent.cpp
--- a/Engine/Source/Runtime/Experimental/GeometryCollectionEngine/Private/GeometryCollection/GeometryCollectionComponent.cpp
+++ b/Engine/Source/Runtime/Experimental/GeometryCollectionEngine/Private/GeometryCollection/GeometryCollectionComponent.cpp
@@ -2914,7 +2914,7 @@ void UGeometryCollectionComponent::OnPostPhysicsSync()
 		RequestUpdateRepData();
 	}
 
-	// Onc ethe GC is broken, removal feature will need the tick to properly update the timers 
+	// Once the GC is broken, removal feature will need the tick to properly update the timers 
 	// even if the physics does not get any updates
 	if (IsRootBroken())
 	{
@@ -2926,6 +2926,7 @@ void UGeometryCollectionComponent::OnPostPhysicsSync()
 				{
 					FChaosBreakEvent Event;
 					Event.Index = GetRootIndex();
+					Event.Component = this;
 					OnRootBreakEvent.Broadcast(Event);
 				}
 			}
