pass proper parameters to FindGeometryOpposingNormal in SweepQuery

diff --git a/Engine/Source/Runtime/Experimental/Chaos/Public/Chaos/GeometryQueries.h b/Engine/Source/Runtime/Experimental/Chaos/Public/Chaos/GeometryQueries.h
--- a/Engine/Source/Runtime/Experimental/Chaos/Public/Chaos/GeometryQueries.h
+++ b/Engine/Source/Runtime/Experimental/Chaos/Public/Chaos/GeometryQueries.h
@@ -302,9 +302,9 @@ namespace Chaos
 		//put back into world space
 		if (bResult && (OutTime > 0 || bComputeMTD))
 		{
-			OutNormal = ATM.TransformVectorNoScale(LocalNormal);
+			OutNormal = ATM.TransformVectorNoScale(LocalNormal).GetSafeNormal();
 			OutPosition = ATM.TransformPositionNoScale(LocalPosition);
-			OutFaceNormal = ATM.TransformVectorNoScale(A.FindGeometryOpposingNormal(Dir, OutFaceIndex, OutNormal));
+			OutFaceNormal = ATM.TransformVectorNoScale(A.FindGeometryOpposingNormal(LocalDir, OutFaceIndex, OutNormal));
 		}
 
 		return bResult;
