cluster unions should be able to leverage intercluster edges to allow pieces of geometry collections to remain in the cluster union when it breaks off

- All cluster union structs need to also know what their unique index is.
- Added a Chaos::IsInterclusterEdge to check if an edge spans two separate parent cluster particles.
- Modify FRigidClustering::ReleaseClusterParticlesImpl to cause a particle to be added into an adjacent cluster union instead of an internal cluster if it has an intercluster edge and its sibling is in a cluster union.
- Changed the places where we forcibly set max strain (FRigidClustering::ReleaseClusterParticles, FRigidClustering::BreakCluster, and FRigidClustering::BreakClustersByProxy) to also set the same internal strain to prevent the propagation of an infinite strain.
- Changed FRigidClustering::ForceReleaseChildParticleAndParents to call the right ReleaseClusterParticles override - the previous call used to force the parent particle to release all its children (not desired e.g. in the case of replication).

diff --git a/Engine/Source/Runtime/Engine/Classes/PhysicsEngine/ClusterUnionComponent.h b/Engine/Source/Runtime/Engine/Classes/PhysicsEngine/ClusterUnionComponent.h
--- a/Engine/Source/Runtime/Engine/Classes/PhysicsEngine/ClusterUnionComponent.h
+++ b/Engine/Source/Runtime/Engine/Classes/PhysicsEngine/ClusterUnionComponent.h
@@ -107,6 +107,9 @@ struct FClusterUnionParticleCandidateData
 	int32 BoneId = INDEX_NONE;
 };
 
+DECLARE_DYNAMIC_MULTICAST_DELEGATE_ThreeParams(FOnClusterUnionAddedComponent, UPrimitiveComponent*, Component, const TSet<int32>&, BoneIds, bool, bIsNew);
+DECLARE_DYNAMIC_MULTICAST_DELEGATE_OneParam(FOnClusterUnionRemovedComponent, UPrimitiveComponent*, Component);
+
 /**
  * This does the bulk of the work exposing a physics cluster union to the game thread.
  * This component needs to be a primitive component primarily because of how physics
@@ -163,6 +166,12 @@ public:
 	bool LineTraceComponent(TArray<FHitResult>& OutHit, const FVector Start, const FVector End, ECollisionChannel TraceChannel, const struct FCollisionQueryParams& Params, const struct FCollisionResponseParams& ResponseParams, const struct FCollisionObjectQueryParams& ObjectParams);
 	bool SweepComponent(TArray<FHitResult>& OutHit, const FVector Start, const FVector End, const FQuat& ShapeWorldRotation, const FPhysicsGeometry& Geometry, ECollisionChannel TraceChannel, const struct FCollisionQueryParams& Params, const struct FCollisionResponseParams& ResponseParams, const struct FCollisionObjectQueryParams& ObjectParams);
 
+	UPROPERTY(BlueprintAssignable, Category = "Events")
+	FOnClusterUnionAddedComponent OnComponentAddedEvent;
+
+	UPROPERTY(BlueprintAssignable, Category = "Events")
+	FOnClusterUnionRemovedComponent OnComponentRemovedEvent;
+
 	friend class UClusterUnionReplicatedProxyComponent;
 protected:
 
diff --git a/Engine/Source/Runtime/Engine/Private/PhysicsEngine/ClusterUnionComponent.cpp b/Engine/Source/Runtime/Engine/Private/PhysicsEngine/ClusterUnionComponent.cpp
--- a/Engine/Source/Runtime/Engine/Private/PhysicsEngine/ClusterUnionComponent.cpp
+++ b/Engine/Source/Runtime/Engine/Private/PhysicsEngine/ClusterUnionComponent.cpp
@@ -544,7 +544,6 @@ void UClusterUnionComponent::HandleAddOrModifiedClusteredComponent(UPrimitiveCom
 		ComponentData.bWasReplicating = ChangedComponent->GetIsReplicated();
 		if (IsAuthority())
 		{
-			ChangedComponent->SetIsReplicated(false);
 			if (AActor* Owner = ChangedComponent->GetOwner())
 			{
 				// Create a replicated proxy component and add it to the actor being added to the cluster.
@@ -561,9 +560,11 @@ void UClusterUnionComponent::HandleAddOrModifiedClusteredComponent(UPrimitiveCom
 				ComponentData.ReplicatedProxyComponent = ReplicatedProxy;
 			}
 		}
+
 	}
 
 	PerBoneChildToParent.GetKeys(ComponentData.BoneIds);
+	OnComponentAddedEvent.Broadcast(ChangedComponent, ComponentData.BoneIds, bIsNew);
 
 	if (IsAuthority() && ComponentData.ReplicatedProxyComponent.IsValid())
 	{
@@ -626,8 +627,6 @@ void UClusterUnionComponent::HandleRemovedClusteredComponent(UPrimitiveComponent
 
 		if (IsAuthority())
 		{
-			ChangedComponent->SetIsReplicated(ComponentData->bWasReplicating);
-
 			if (bDestroyReplicatedProxy && ensure(ComponentData->ReplicatedProxyComponent.IsValid()))
 			{
 				UClusterUnionReplicatedProxyComponent* ProxyComponent = ComponentData->ReplicatedProxyComponent.Get();
@@ -636,6 +635,7 @@ void UClusterUnionComponent::HandleRemovedClusteredComponent(UPrimitiveComponent
 		}
 
 		ComponentToPhysicsObjects.Remove(ChangedComponent);
+		OnComponentRemovedEvent.Broadcast(ChangedComponent);
 	}
 
 
diff --git a/Engine/Source/Runtime/Experimental/Chaos/Private/Chaos/ClusterUnionManager.cpp b/Engine/Source/Runtime/Experimental/Chaos/Private/Chaos/ClusterUnionManager.cpp
--- a/Engine/Source/Runtime/Experimental/Chaos/Private/Chaos/ClusterUnionManager.cpp
+++ b/Engine/Source/Runtime/Experimental/Chaos/Private/Chaos/ClusterUnionManager.cpp
@@ -24,6 +24,7 @@ namespace Chaos
 		check(NewIndex > 0);
 
 		FClusterUnion NewUnion;
+		NewUnion.InternalIndex = NewIndex;
 		NewUnion.ExplicitIndex = ClusterUnionParameters.ExplicitIndex;
 		NewUnion.SharedGeometry = ForceRecreateClusterUnionSharedGeometry(NewUnion);
 		NewUnion.InternalCluster = MClustering.CreateClusterParticle(-NewIndex, {}, Parameters, NewUnion.SharedGeometry, nullptr, ClusterUnionParameters.UniqueIndex);
@@ -268,7 +269,7 @@ namespace Chaos
 
 		for (FPBDRigidParticleHandle* Handle : Particles)
 		{
-			if (!Handle || Handle->Disabled())
+			if (!Handle)
 			{
 				continue;
 			}
diff --git a/Engine/Source/Runtime/Experimental/Chaos/Private/Chaos/PBDRigidClustering.cpp b/Engine/Source/Runtime/Experimental/Chaos/Private/Chaos/PBDRigidClustering.cpp
--- a/Engine/Source/Runtime/Experimental/Chaos/Private/Chaos/PBDRigidClustering.cpp
+++ b/Engine/Source/Runtime/Experimental/Chaos/Private/Chaos/PBDRigidClustering.cpp
@@ -36,6 +36,9 @@ namespace Chaos
 	int32 UseConnectivity = 1;
 	FAutoConsoleVariableRef CVarUseConnectivity(TEXT("p.UseConnectivity"), UseConnectivity, TEXT("Whether to use connectivity graph when breaking up clusters"));
 
+	bool bCheckForInterclusterEdgesOnRelease = false;
+	FAutoConsoleVariableRef CVarCheckForInterclusterEdgesOnRelease(TEXT("p.Chaos.CheckForInterclusterEdgesOnRelease"), bCheckForInterclusterEdgesOnRelease, TEXT("Whether to check for intercluster edges when removing a child from its parent cluster so that we can add the particle back into a cluster union."));
+
 	int32 ComputeClusterCollisionStrains = 1;
 	FAutoConsoleVariableRef CVarComputeClusterCollisionStrains(TEXT("p.ComputeClusterCollisionStrains"), ComputeClusterCollisionStrains, TEXT("Whether to use collision constraints when processing clustering."));
 
@@ -88,6 +91,31 @@ namespace Chaos
 		}
 		return nullptr;
 	}
+
+	namespace
+	{
+		FPBDRigidClusteredParticleHandle* GetActiveParentParticle(FPBDRigidParticleHandle* Particle)
+		{
+			if (!Particle)
+			{
+				return nullptr;
+			}
+
+			FPBDRigidClusteredParticleHandle* Current = Particle->CastToClustered();
+
+			while (Current)
+			{
+				if (!Current->Disabled())
+				{
+					break;
+				}
+
+				Current = Current->Parent();
+			}
+
+			return Current;
+		}
+	}
 	
 	//==========================================================================
 	// TPBDRigidClustering
@@ -230,16 +258,29 @@ namespace Chaos
 		TSet<FPBDRigidParticleHandle*> InChildrenSet(InChildren);
 		for (FPBDRigidParticleHandle* Handle : InChildren)
 		{
-			MEvolution.DisableParticle(Handle);
-			MEvolution.GetParticles().MarkTransientDirtyParticle(Handle);
-
 			if (FPBDRigidClusteredParticleHandle* ClusteredChild = Handle->CastToClustered())
 			{
+				if (FPBDRigidClusteredParticleHandle* ExistingParent = ClusteredChild->Parent())
+				{
+					if (ExistingParent == Cluster)
+					{
+						// This shouldn't really ever be the case probably?
+						continue;
+					}
+
+					// This is needed in the case where we use intercluster edges with geometry collections that need to then get stuck into a cluster union.
+					// It's possible due to replication ordering that we create an internal cluster surrounding the piece that we want to add to the cluster union first.
+					RemoveChildFromParentAndChildrenArray(ClusteredChild, ExistingParent);
+				}
+
 				TopLevelClusterParents.Remove(ClusteredChild);
 				TopLevelClusterParentsStrained.Remove(ClusteredChild);
 
 				ClusteredChild->ClusterIds().Id = Cluster;
 			}
+
+			MEvolution.DisableParticle(Handle);
+			MEvolution.GetParticles().MarkTransientDirtyParticle(Handle);
 		}
 
 		// Note that we want to compute the internal strain on the cluster the same if we build it up incrementally as well as if we
@@ -580,6 +621,7 @@ namespace Chaos
 			{
 				continue;
 			}
+
 			TArray<FPBDRigidParticleHandle*>& Island = Islands.AddDefaulted_GetRef();
 
 			TArray<FPBDRigidParticleHandle*> ProcessingQueue;
@@ -593,6 +635,11 @@ namespace Chaos
 					Island.Add(ChildToProcess);
 					for (const TConnectivityEdge<FReal>& Edge : ChildToProcess->CastToClustered()->ConnectivityEdges())
 					{
+						if (IsInterclusterEdge(*ChildToProcess, Edge))
+						{
+							continue;
+						}
+
 						if (!ProcessedChildren.Contains(Edge.Sibling))
 						{
 							ProcessingQueue.Add(Edge.Sibling);
@@ -633,7 +680,22 @@ namespace Chaos
 			Child->SetPreV(Child->PreV() + ClusteredParent->PreV());
 			Child->SetPreW(Child->PreW() + ClusteredParent->PreW());
 		}
-	};
+	}
+
+	void FRigidClustering::RemoveChildFromParentAndChildrenArray(FPBDRigidParticleHandle* Child, const FPBDRigidClusteredParticleHandle* ClusteredParent)
+	{
+		RemoveChildFromParent(Child, ClusteredParent);
+
+		// Also need to remove it from the children array.
+		if (FRigidHandleArray* Children = MChildren.Find(ClusteredParent))
+		{
+			const int32 Index = Children->Find(Child);
+			if (Index != INDEX_NONE)
+			{
+				Children->RemoveAtSwap(Index, 1, false);
+			}
+		}
+	}
 
 	TArray<FPBDRigidParticleHandle*> FRigidClustering::CreateClustersFromNewIslands(
 		TArray<FParticleIsland>& Islands,
@@ -791,7 +853,6 @@ namespace Chaos
 		bool bUseDamagePropagation = false;
 		float BreakDamagePropagationFactor = 0.0f;
 		float ShockDamagePropagationFactor = 0.0f;
-		EDamageEvaluationModel DamageEvaluationModel = EDamageEvaluationModel::StrainFromDamageThreshold;
 		const FGeometryCollectionPhysicsProxy* ConcreteGCProxy = GetConcreteProxy<FGeometryCollectionPhysicsProxy>(ClusteredParticle);
 		if (ConcreteGCProxy)
 		{
@@ -799,7 +860,6 @@ namespace Chaos
 			bUseDamagePropagation = SimParams.bUseDamagePropagation;
 			BreakDamagePropagationFactor = SimParams.BreakDamagePropagationFactor;
 			ShockDamagePropagationFactor = SimParams.ShockDamagePropagationFactor;
-			DamageEvaluationModel = SimParams.DamageEvaluationModel;
 		}
 
 		TArray<FPBDRigidParticleHandle*>& Children = MChildren[ClusteredParticle];
@@ -949,58 +1009,134 @@ namespace Chaos
 			if (UseConnectivity)
 			{
 				// The cluster may have contained forests, so find the connected pieces and cluster them together.
-
 				//first update the connected graph of the children we already removed
 				for (FPBDRigidParticleHandle* Child : ActivatedChildren)
 				{
 					RemoveNodeConnections(Child);
 				}
+				HandleConnectivityOnReleaseClusterParticle(ClusteredParticle, bCreateNewClusters);
+			}
+
+			for (FPBDRigidParticleHandle* Child : ActivatedChildren)
+			{
+				UpdateKinematicProperties(Child, MChildren, MEvolution);
+			}
+
+			// Disable the cluster only if we're not a cluster union. Cluster unions will handle themselves separately.
+			if (!bIsClusterUnion)
+			{
+				DisableCluster(ClusteredParticle);
+			}
+		}
+
+		return ActivatedChildren;
+	}
+
+	DECLARE_CYCLE_STAT(TEXT("TPBDRigidClustering<>::HandleConnectivityOnReleaseClusterParticle"), HandleConnectivityOnReleaseClusterParticle, STATGROUP_Chaos);
+	TSet<FPBDRigidParticleHandle*> FRigidClustering::HandleConnectivityOnReleaseClusterParticle(FPBDRigidClusteredParticleHandle* ClusteredParticle, bool bCreateNewClusters)
+	{
+		SCOPE_CYCLE_COUNTER(HandleConnectivityOnReleaseClusterParticle);
+		TSet<FPBDRigidParticleHandle*> ActivatedChildren;
+		TArray<FPBDRigidParticleHandle*>& Children = MChildren[ClusteredParticle];
 
-				// If we're breaking a geometry collection, we'll need to create internal clusters to parent the remaining particles.
-				// However, we do not need to do this if we're currently operating on a cluster union! Its remaining particles should stay
-				// attached to the cluster union because they can handle particles being dynamically added/removed.
-				if (Children.Num() && !bIsClusterUnion)
+		// If we're breaking a geometry collection, we'll need to create internal clusters to parent the remaining particles.
+		// However, we do not need to do this if we're currently operating on a cluster union! Its remaining particles should stay
+		// attached to the cluster union because they can handle particles being dynamically added/removed.
+		const bool bIsClusterUnion = ClusterUnionManager.IsClusterUnionParticle(ClusteredParticle);
+		if (Children.Num() && !bIsClusterUnion)
+		{
+			TArray<FParticleIsland> Islands = FindIslandsInChildren(ClusteredParticle);
+
+			TArray<const FConnectivityEdge*> IslandInterclusterEdges;
+			IslandInterclusterEdges.Reserve(Islands.Num());
+
+			// By default, all these islands will all just start simulating independently. However, if bCheckForInterclusterEdgesOnRelease
+			// is true, we're going to want to check if any of these islands contain a particle with an intercluster edge that connects to
+			// something else! In that case, keep the particles in the island connected. For simplicity
+			// we'll use a kinematic target for each group of particles that needs to stay attached . This kinematic target
+			// will be driven by simulation relative to whatever it should be attached to.
+			//
+			// We need to assume that only the server will have the intercluster edges necessary to drive this kinematic target. Therefore, this
+			// information is authoritative on the PT on the server while is authoritative on the GT on the client(s) via GC replication. So we
+			// also need to make sure we store this information in a way that's replicatable to the client via GC replication.
+			EDamageEvaluationModel DamageEvaluationModel = EDamageEvaluationModel::StrainFromDamageThreshold;
+			const FGeometryCollectionPhysicsProxy* ConcreteGCProxy = GetConcreteProxy<FGeometryCollectionPhysicsProxy>(ClusteredParticle);
+			if (ConcreteGCProxy)
+			{
+				const FSimulationParameters& SimParams = ConcreteGCProxy->GetSimParameters();
+				DamageEvaluationModel = SimParams.DamageEvaluationModel;
+			}
+
+			int32 IslandIndex = 0;
+			TArray<int32> IslandIndicesToRemove;
+
+			for (const FParticleIsland& Island : Islands)
+			{
+				const bool bNeedRecomputeConnectivityStrain = ConcreteGCProxy && DamageEvaluationModel == EDamageEvaluationModel::StrainFromMaterialStrengthAndConnectivity;
+
+				const FClusterUnion* AttachedClusterUnion = nullptr;
+
+				// Protect a potentially work-intensive loop behind the conditions that actually require us to step through every child in every island.
+				if (bNeedRecomputeConnectivityStrain || bCheckForInterclusterEdgesOnRelease)
 				{
-					TArray<FParticleIsland> Islands = FindIslandsInChildren(ClusteredParticle);
-					for (const FParticleIsland& Island: Islands)
+					for (FPBDRigidParticleHandle* ChildParticle : Island)
 					{
-						if (Island.Num() == 1) //need to break single pieces first
+						if (FPBDRigidClusteredParticleHandle* ClusteredChild = ChildParticle->CastToClustered())
 						{
-							FPBDRigidParticleHandle* Child = Island[0];
-							RemoveChildFromParent(Child, ClusteredParticle);
-							ActivatedChildren.Add(Child);
-						}
+							// recompute the strain as its connectivity has changed
+							if (bNeedRecomputeConnectivityStrain)
+							{
+								ConcreteGCProxy->ComputeMaterialBasedDamageThreshold_Internal(*ClusteredChild);
+							}
 
-						if (ConcreteGCProxy && DamageEvaluationModel == EDamageEvaluationModel::StrainFromMaterialStrengthAndConnectivity)
-						{
-							for (FPBDRigidParticleHandle* ChildParticle: Island)
+							if (bCheckForInterclusterEdgesOnRelease && !AttachedClusterUnion)
 							{
-								// recompute ethe strain as the connection have changed
-								if (FPBDRigidClusteredParticleHandle* ClusteredChild = ChildParticle->CastToClustered())
+								// Check if any edge is connected to "something else". This entire island
+								// will be connected to that "something else". For now, we're enforcing that it's a cluster union.
+								for (const FConnectivityEdge& Edge : ClusteredChild->ConnectivityEdges())
 								{
-									ConcreteGCProxy->ComputeMaterialBasedDamageThreshold_Internal(*ClusteredChild);
+									if (IsInterclusterEdge(*ClusteredChild, Edge))
+									{
+										if (FClusterUnion* ClusterUnion = ClusterUnionManager.FindClusterUnionFromParticle(GetActiveParentParticle(Edge.Sibling)))
+										{
+											AttachedClusterUnion = ClusterUnion;
+										}
+										break;
+									}
 								}
 							}
 						}
 					}
+				}
 
-					if (bCreateNewClusters)
+				if (AttachedClusterUnion)
+				{
+					for (FPBDRigidParticleHandle* Child : Island)
 					{
-						TArray<FPBDRigidParticleHandle*> NewClusters = CreateClustersFromNewIslands(Islands, ClusteredParticle); 
-						ActivatedChildren.Append(MoveTemp(NewClusters));
+						RemoveChildFromParentAndChildrenArray(Child, ClusteredParticle);
 					}
+					ClusterUnionManager.AddPendingClusterIndexOperation(AttachedClusterUnion->InternalIndex, EClusterUnionOperation::Add, Island);
+					IslandIndicesToRemove.Add(IslandIndex);
 				}
+				else if (Island.Num() == 1 && !AttachedClusterUnion) //need to break single pieces first
+				{
+					FPBDRigidParticleHandle* Child = Island[0];
+					RemoveChildFromParent(Child, ClusteredParticle);
+					ActivatedChildren.Add(Child);
+				}
+
+				++IslandIndex;
 			}
 
-			for (FPBDRigidParticleHandle* Child : ActivatedChildren)
+			for (int32 RemoveIndex = IslandIndicesToRemove.Num() - 1; RemoveIndex >= 0; --RemoveIndex)
 			{
-				UpdateKinematicProperties(Child, MChildren, MEvolution);
+				Islands.RemoveAtSwap(IslandIndicesToRemove[RemoveIndex], 1, false);
 			}
 
-			// Disable the cluster only if we're not a cluster union. Cluster unions will handle themselves separately.
-			if (!bIsClusterUnion)
+			if (bCreateNewClusters)
 			{
-				DisableCluster(ClusteredParticle);
+				TArray<FPBDRigidParticleHandle*> NewClusters = CreateClustersFromNewIslands(Islands, ClusteredParticle);
+				ActivatedChildren.Append(MoveTemp(NewClusters));
 			}
 		}
 
@@ -1033,6 +1169,8 @@ namespace Chaos
 						if (ensure(!ClusterHandle || ClusteredChildHandle->ClusterIds().Id == ClusterHandle))
 						{
 							SetExternalStrain(ClusteredChildHandle, TNumericLimits<FRealSingle>::Max());
+							// This way we won't try to propagate this infinite strain.
+							ClusteredChildHandle->SetInternalStrains(ClusteredChildHandle->GetExternalStrain());
 							ClusterHandle = ClusteredChildHandle->ClusterIds().Id;
 						}
 						else
@@ -1060,25 +1198,23 @@ namespace Chaos
 		{
 			// make sure we set unbreakable to false so that the children can be released
 			ChildClusteredParticle->SetUnbreakable(false);
-			if (ChildClusteredParticle->Disabled())
+			
+			// first release any parent if any
+			if (FPBDRigidClusteredParticleHandle* ParentCluster = ChildClusteredParticle->Parent())
 			{
-				// first release any parent if any
-				if (FPBDRigidClusteredParticleHandle* ParentCluster = ChildClusteredParticle->Parent())
+				if (const IPhysicsProxyBase* ParentProxy = ParentCluster->PhysicsProxy())
 				{
-					if (const IPhysicsProxyBase* ParentProxy = ParentCluster->PhysicsProxy())
+					// we shoudl not break cluster union parent and stop recursion there
+					if (ParentProxy->GetType() != FClusterUnionPhysicsProxy::ConcreteType())
 					{
-						// we shoudl not break cluster union parent and stop recursion there
-						if (ParentProxy->GetType() != FClusterUnionPhysicsProxy::ConcreteType())
-						{
-							// we need now to force parents to break
-							ForceReleaseChildParticleAndParents(ParentCluster, bTriggerBreakEvents);
-
-							SetExternalStrain(ChildClusteredParticle, TNumericLimits<FRealSingle>::Max());
-							ReleaseClusterParticles(ParentCluster, bTriggerBreakEvents);
-						}
+						// we need now to force parents to break
+						ForceReleaseChildParticleAndParents(ParentCluster, bTriggerBreakEvents);
 					}
 				}
 			}
+
+			// Trigger a release.
+			ReleaseClusterParticles(TArray<FPBDRigidParticleHandle*>{ ChildClusteredParticle }, bTriggerBreakEvents);
 		}
 	}
 
@@ -1169,8 +1305,7 @@ namespace Chaos
 					FClusterUnion* ClusterUnion = ClusterUnionManager.FindClusterUnionFromParticle(ActiveCluster);
 					if (ClusterUnion && ClusterUnion->InternalCluster != ActiveCluster)
 					{
-						// Need to pre-emptively remove the particle from the cluster union otherwise we won't be passing the disabled check.
-						//ClusterUnionManager.HandleRemoveOperationWithClusterLookup({ ActiveCluster }, EClusterUnionOperationTiming::Defer);
+						// Need to ignore the disabled check since the particle might still be inside a cluster union (and thus disabled).
 						bIgnoreDisabledCheck = true;
 					}
 
@@ -1178,19 +1313,27 @@ namespace Chaos
 					{
 						if (ActiveCluster->ClusterIds().NumChildren > 0) //active index is a cluster
 						{
+							bool bNeedToProcessActiveCluster = true;
+
 							if(ClusterUnion)
 							{
 								if(ClusterUnion->InternalCluster == ActiveCluster)
 								{
 									// ActiveCluster is itself a cluster union, so loop over its children and add those
 									// to process for breaking.
+									bNeedToProcessActiveCluster = false;
 									for(FPBDRigidParticleHandle* ChildParticle : ClusterUnion->ChildParticles)
 									{
 										if(ChildParticle)
 										{
 											if(FPBDRigidClusteredParticleHandle* ClusteredChild = ChildParticle->CastToClustered())
 											{
-												if(ClusteredChild->ClusterIds().NumChildren > 0)
+												const bool bChildHasChildren = ClusteredChild->ClusterIds().NumChildren > 0;
+
+												// If for some reason, a child of the cluster union doesn't have children then we need to process the cluster union particle instead.
+												bNeedToProcessActiveCluster |= !bChildHasChildren;
+
+												if(bChildHasChildren)
 												{
 													ProcessClusteredParticle(ClusteredChild);
 												}
@@ -1198,13 +1341,9 @@ namespace Chaos
 										}
 									}
 								}
-								else
-								{
-									// Clustered is inside a clustered union, but not a clustered union itself
-									ProcessClusteredParticle(ActiveCluster);
-								}
 							}
-							else
+							
+							if (bNeedToProcessActiveCluster)
 							{
 								ProcessClusteredParticle(ActiveCluster);
 							}
@@ -1232,6 +1371,14 @@ namespace Chaos
 				{
 					BreakingModel();
 				}
+
+				if (bCheckForInterclusterEdgesOnRelease)
+				{
+					// In this case the breaking model might have tried to *add* into cluster unions so we need to flush those operations here as well
+					// so those changes are immediately available for marshaling back to the GT.
+					ClusterUnionManager.FlushPendingOperations();
+				}
+
 			} // end if MCollisionImpulseArrayDirty
 		}
 		Timer.Stop();
@@ -1831,6 +1978,8 @@ namespace Chaos
 				{
 					ClusteredChildHandle->SetExternalStrains(MaxStrain);
 					SetExternalStrain(ClusteredChildHandle, MaxStrain);
+					// This way we won't try to propagate this infinite strain.
+					ClusteredChildHandle->SetInternalStrains(ClusteredChildHandle->GetExternalStrain());
 				}
 			}
 			if (ChildrenHandles->Num() > 0)
@@ -1895,6 +2044,8 @@ namespace Chaos
 						if (Chaos::FPBDRigidClusteredParticleHandle* ClusteredChildHandle = ChildHandle->CastToClustered())
 						{
 							SetExternalStrain(ClusteredChildHandle, MaxStrain);
+							// This way we won't try to propagate this infinite strain.
+							ClusteredChildHandle->SetInternalStrains(ClusteredChildHandle->GetExternalStrain());
 						}
 					}
 					if (Children->Num() > 0)
diff --git a/Engine/Source/Runtime/Experimental/Chaos/Public/Chaos/ClusterUnionManager.h b/Engine/Source/Runtime/Experimental/Chaos/Public/Chaos/ClusterUnionManager.h
--- a/Engine/Source/Runtime/Experimental/Chaos/Public/Chaos/ClusterUnionManager.h
+++ b/Engine/Source/Runtime/Experimental/Chaos/Public/Chaos/ClusterUnionManager.h
@@ -51,6 +51,9 @@ namespace Chaos
 		// All the particles that belong to this cluster.
 		TArray<FPBDRigidParticleHandle*> ChildParticles;
 
+		// The internal index used to reference the cluster union.
+		FClusterUnionIndex InternalIndex;
+
 		// An explicit index set by the user if any.
 		FClusterUnionExplicitIndex ExplicitIndex;
 
diff --git a/Engine/Source/Runtime/Experimental/Chaos/Public/Chaos/PBDRigidClusteredParticles.h b/Engine/Source/Runtime/Experimental/Chaos/Public/Chaos/PBDRigidClusteredParticles.h
--- a/Engine/Source/Runtime/Experimental/Chaos/Public/Chaos/PBDRigidClusteredParticles.h
+++ b/Engine/Source/Runtime/Experimental/Chaos/Public/Chaos/PBDRigidClusteredParticles.h
@@ -71,6 +71,16 @@ struct TConnectivityEdge
 typedef TConnectivityEdge<FReal> FConnectivityEdge;
 typedef TArray<FConnectivityEdge> FConnectivityEdgeArray;
 
+template<typename T>
+bool IsInterclusterEdge(const TPBDRigidParticleHandle<T, 3>& Particle, const TConnectivityEdge<T>& Edge)
+{
+	if (!Edge.Sibling)
+	{
+		return false;
+	}
+	return Particle.PhysicsProxy() != Edge.Sibling->PhysicsProxy();
+}
+
 class FRigidClusteredFlags
 {
 public:
diff --git a/Engine/Source/Runtime/Experimental/Chaos/Public/Chaos/PBDRigidClustering.h b/Engine/Source/Runtime/Experimental/Chaos/Public/Chaos/PBDRigidClustering.h
--- a/Engine/Source/Runtime/Experimental/Chaos/Public/Chaos/PBDRigidClustering.h
+++ b/Engine/Source/Runtime/Experimental/Chaos/Public/Chaos/PBDRigidClustering.h
@@ -440,6 +440,7 @@ public:
 	void RemoveNodeConnections(FPBDRigidClusteredParticleHandle* Child);
 
 	void RemoveChildFromParent(FPBDRigidParticleHandle* Child, const FPBDRigidClusteredParticleHandle* ClusteredParent);
+	void RemoveChildFromParentAndChildrenArray(FPBDRigidParticleHandle* Child, const FPBDRigidClusteredParticleHandle* ClusteredParent);
 
 	// When a body has broken due to contact resolution, record an entry in a set
 	// for the collision and the particle who's momentum should be restored.
@@ -456,6 +457,11 @@ public:
 		FPBDRigidClusteredParticleHandle* ClusteredParticle, 
 		bool bForceRelease,
 		bool bCreateNewClusters);
+
+	/**
+	 * Handles leveraging the connectivity edges on the children of the clustered particle to produce the desired effects.
+	 */
+	TSet<FPBDRigidParticleHandle*> HandleConnectivityOnReleaseClusterParticle(FPBDRigidClusteredParticleHandle* ClusteredParticle, bool bCreateNewClusters);
 	
 	using FParticleIsland = TArray<FPBDRigidParticleHandle*>;
 	TArray<FParticleIsland> FindIslandsInChildren(const FPBDRigidClusteredParticleHandle* ClusteredParticle);
