add ability to set collision profile per particle to geometry collections

diff --git a/Engine/Source/Runtime/Experimental/GeometryCollectionEngine/Private/GeometryCollection/GeometryCollectionComponent.cpp b/Engine/Source/Runtime/Experimental/GeometryCollectionEngine/Private/GeometryCollection/GeometryCollectionComponent.cpp
--- a/Engine/Source/Runtime/Experimental/GeometryCollectionEngine/Private/GeometryCollection/GeometryCollectionComponent.cpp
+++ b/Engine/Source/Runtime/Experimental/GeometryCollectionEngine/Private/GeometryCollection/GeometryCollectionComponent.cpp
@@ -3084,6 +3084,29 @@ void UGeometryCollectionComponent::SetPerLevelCollisionProfileNames(const TArray
 	LoadCollisionProfiles();
 }
 
+void UGeometryCollectionComponent::SetPerParticleCollisionProfileName(const TArray<int32>& BoneIds, FName ProfileName)
+{
+	if (!RestCollection)
+	{
+		return;
+	}
+
+	if (CollisionProfilePerParticle.IsEmpty())
+	{
+		CollisionProfilePerParticle.SetNumZeroed(RestCollection->GetGeometryCollection()->Transform.Num());
+	}
+
+	for (int32 Id : BoneIds)
+	{
+		if (CollisionProfilePerParticle.IsValidIndex(Id))
+		{
+			CollisionProfilePerParticle[Id] = ProfileName;
+		}
+	}
+
+	LoadCollisionProfiles();
+}
+
 void UGeometryCollectionComponent::LoadCollisionProfiles()
 {
 	if (!PhysicsProxy)
@@ -3147,7 +3170,11 @@ void UGeometryCollectionComponent::LoadCollisionProfiles()
 
 			TArrayView<Chaos::FPhysicsObjectHandle> ParticleView{ &PhysicsObjects[ParticleIndex], 1 };
 			const CollisionProfileDataCache* Data = nullptr;
-			if (AbandonedData && Level >= ReplicationAbandonAfterLevel + 1)
+			if (!CollisionProfilePerParticle.IsEmpty() && CollisionProfilePerParticle[ParticleIndex] != NAME_None)
+			{
+				Data = CreateOrGetCollisionProfileData(CollisionProfilePerParticle[ParticleIndex]);
+			}
+			else if (AbandonedData && Level >= ReplicationAbandonAfterLevel + 1)
 			{
 				Data = AbandonedData;
 			}
diff --git a/Engine/Source/Runtime/Experimental/GeometryCollectionEngine/Public/GeometryCollection/GeometryCollectionComponent.h b/Engine/Source/Runtime/Experimental/GeometryCollectionEngine/Public/GeometryCollection/GeometryCollectionComponent.h
--- a/Engine/Source/Runtime/Experimental/GeometryCollectionEngine/Public/GeometryCollection/GeometryCollectionComponent.h
+++ b/Engine/Source/Runtime/Experimental/GeometryCollectionEngine/Public/GeometryCollection/GeometryCollectionComponent.h
@@ -690,6 +690,9 @@ public:
 	UFUNCTION(BlueprintCallable, Category = "ChaosPhysics")
 	void SetPerLevelCollisionProfileNames(const TArray<FName>& ProfileNames);
 
+	UFUNCTION(BlueprintCallable, Category = "ChaosPhysics")
+	void SetPerParticleCollisionProfileName(const TArray<int32>& BoneIds, FName ProfileName);
+
 	/** API for getting at geometry collection data */
 	int32 GetNumElements(FName Group) const;
 
@@ -1199,6 +1202,11 @@ protected:
 	UPROPERTY(EditAnywhere, BlueprintReadWrite, BlueprintSetter=SetPerLevelCollisionProfileNames, Category=Physics)
 	TArray<FName> CollisionProfilePerLevel;
 
+	/**
+	 * A per-particle collision profile name. If the per-particle collision profile name exists, it will override the per-level profile name.
+	 */
+	TArray<FName> CollisionProfilePerParticle;
+
 	/**
 	 * If replicating - the cluster level to stop sending corrections for geometry collection chunks.
 	 * recommended for smaller leaf levels when the size of the objects means they are no longer
