ISMPool: If ISM component has no component transform then we can ignore the bWorldSpace flag. This makes the maths simpler in the ISM Update call.
This is a quick fix for a perf hot spot, but we can still improve the underlying ISM code here to avoid multiple transform inverts when updating multiple instances.

diff --git a/Engine/Source/Runtime/Experimental/GeometryCollectionEngine/Private/GeometryCollection/GeometryCollectionISMPoolComponent.cpp b/Engine/Source/Runtime/Experimental/GeometryCollectionEngine/Private/GeometryCollection/GeometryCollectionISMPoolComponent.cpp
--- a/Engine/Source/Runtime/Experimental/GeometryCollectionEngine/Private/GeometryCollection/GeometryCollectionISMPoolComponent.cpp
+++ b/Engine/Source/Runtime/Experimental/GeometryCollectionEngine/Private/GeometryCollection/GeometryCollectionISMPoolComponent.cpp
@@ -193,15 +193,18 @@ bool FGeometryCollectionISMPool::BatchUpdateInstancesTransforms(FGeometryCollect
 		const FInstanceGroups::FInstanceGroupRange& InstanceGroup = ISM.InstanceGroups.GroupRanges[MeshInfo.InstanceGroupIndex];
 		ensure((StartInstanceIndex + NewInstancesTransforms.Num()) <= InstanceGroup.Count);
 
+		// If ISM component has identity transform (the common case) then we can skip world space to component space maths inside BatchUpdateInstancesTransforms()
+		bWorldSpace &= !ISM.ISMComponent->GetComponentTransform().Equals(FTransform::Identity, 0.f);
+
 		int32 StartIndex = ISM.InstanceIndexToRenderIndex[InstanceGroup.Start];
 		int32 TransformIndex = 0;
 		int32 BatchCount = 1;
-		TArray<FTransform> BatchTransforms; //< Can't use TArrayView because blueprint function doesn't support that 
+		TArray<FTransform> BatchTransforms; // Can't use TArrayView because blueprint function doesn't support that 
 		BatchTransforms.Reserve(NewInstancesTransforms.Num());
 		BatchTransforms.Add(NewInstancesTransforms[TransformIndex++]);
 		for (int InstanceIndex = StartInstanceIndex + 1; InstanceIndex < NewInstancesTransforms.Num(); ++InstanceIndex)
 		{
-			// flush batch?
+			// Flush batch for non-sequential instances.
 			int32 RenderIndex = ISM.InstanceIndexToRenderIndex[InstanceGroup.Start + InstanceIndex];
 			if (RenderIndex != (StartIndex + BatchCount))
 			{
