[Navmesh] Account for bIsDynamicObstacle flag from UNavCollisionBase when exporting UGeometryCollectionComponent.

diff --git a/Engine/Source/Runtime/Experimental/GeometryCollectionEngine/Private/GeometryCollection/GeometryCollectionComponent.cpp b/Engine/Source/Runtime/Experimental/GeometryCollectionEngine/Private/GeometryCollection/GeometryCollectionComponent.cpp
--- a/Engine/Source/Runtime/Experimental/GeometryCollectionEngine/Private/GeometryCollection/GeometryCollectionComponent.cpp
+++ b/Engine/Source/Runtime/Experimental/GeometryCollectionEngine/Private/GeometryCollection/GeometryCollectionComponent.cpp
@@ -1008,6 +1008,11 @@ bool UGeometryCollectionComponent::DoCustomNavigableGeometryExport(FNavigableGeo
 				{
 					if (const UNavCollisionBase* NavCollision = ProxyMesh->GetNavCollision())
 					{
+						if (NavCollision->IsDynamicObstacle())
+						{
+							continue;
+						}
+						
 						bHasData = NavCollision->ExportGeometry(CompToWorld, GeomExport) || bHasData;
 					}
 				}
