Fix crash in geometry collection convex utility on null external collision data + consider the external collision data as missing if it is all null
+ also change another function to take FManagedArrayCollection instead of GeometryCollection

diff --git a/Engine/Source/Runtime/Experimental/Chaos/Private/GeometryCollection/GeometryCollectionConvexUtility.cpp b/Engine/Source/Runtime/Experimental/Chaos/Private/GeometryCollection/GeometryCollectionConvexUtility.cpp
--- a/Engine/Source/Runtime/Experimental/Chaos/Private/GeometryCollection/GeometryCollectionConvexUtility.cpp
+++ b/Engine/Source/Runtime/Experimental/Chaos/Private/GeometryCollection/GeometryCollectionConvexUtility.cpp
@@ -117,7 +117,7 @@ static const Chaos::FVec3f IcoHemisphere_Subdiv1[] =
 constexpr int32 IcoHemisphere_Subdiv1_Num = sizeof(IcoHemisphere_Subdiv1) / sizeof(Chaos::FVec3f);
 
 
-TOptional<FGeometryCollectionConvexUtility::FGeometryCollectionConvexData> FGeometryCollectionConvexUtility::GetConvexHullDataIfPresent(FGeometryCollection* GeometryCollection)
+TOptional<FGeometryCollectionConvexUtility::FGeometryCollectionConvexData> FGeometryCollectionConvexUtility::GetConvexHullDataIfPresent(FManagedArrayCollection* GeometryCollection)
 {
 	check(GeometryCollection);
 
@@ -1467,12 +1467,25 @@ void FGeometryCollectionConvexUtility::GenerateLeafConvexHulls(FGeometryCollecti
 {
 	using FSharedImplicit = TSharedPtr<Chaos::FImplicitObject, ESPMode::ThreadSafe>;
 
+	int32 NumTransforms = Collection.NumElements(FGeometryCollection::TransformGroup);
+
 	const TManagedArrayAccessor<FSharedImplicit> ExternalCollisionAttribute(Collection, "ExternalCollisions", FGeometryCollection::TransformGroup);
-	bool bUseExternal = ExternalCollisionAttribute.IsValid() && (GenerateMethod == EGenerateConvexMethod::ExternalCollision || GenerateMethod == EGenerateConvexMethod::IntersectExternalWithComputed);
+	bool bHasExternalCollision = ExternalCollisionAttribute.IsValid();
+	if (bHasExternalCollision)
+	{
+		bool bAnyValid = false;
+		for (int32 TransformIdx = 0; TransformIdx < NumTransforms; ++TransformIdx)
+		{
+			bAnyValid = bAnyValid || ExternalCollisionAttribute[TransformIdx].IsValid();
+		}
+		if (!bAnyValid)
+		{
+			bHasExternalCollision = false;
+		}
+	}
+	bool bUseExternal = bHasExternalCollision && (GenerateMethod == EGenerateConvexMethod::ExternalCollision || GenerateMethod == EGenerateConvexMethod::IntersectExternalWithComputed);
 	bool bUseGenerated = !bUseExternal || GenerateMethod == EGenerateConvexMethod::IntersectExternalWithComputed;
 	bool bUseIntersect = bUseExternal && GenerateMethod == EGenerateConvexMethod::IntersectExternalWithComputed;
-
-	int32 NumTransforms = Collection.NumElements(FGeometryCollection::TransformGroup);
 	
 	UE::GeometryCollectionConvexUtility::FConvexHulls ComputedHulls;
 	TArray<FTransformedConvex> ExternalHulls;
@@ -1525,6 +1538,10 @@ void FGeometryCollectionConvexUtility::GenerateLeafConvexHulls(FGeometryCollecti
 
 			// convert the external collisions to convex hulls
 			const FSharedImplicit ExternalCollision = ExternalCollisionAttribute[SourceTransformIdx];
+			if (!ExternalCollision.IsValid())
+			{
+				continue;
+			}
 
 			const int32 ExternalHullsStart = ExternalHulls.Num();
 			ConvertImplicitToConvexArray(*ExternalCollision, FTransform::Identity, ExternalHulls);
diff --git a/Engine/Source/Runtime/Experimental/Chaos/Public/GeometryCollection/GeometryCollectionConvexUtility.h b/Engine/Source/Runtime/Experimental/Chaos/Public/GeometryCollection/GeometryCollectionConvexUtility.h
--- a/Engine/Source/Runtime/Experimental/Chaos/Public/GeometryCollection/GeometryCollectionConvexUtility.h
+++ b/Engine/Source/Runtime/Experimental/Chaos/Public/GeometryCollection/GeometryCollectionConvexUtility.h
@@ -94,7 +94,7 @@ public:
 	static FGeometryCollectionConvexData GetValidConvexHullData(FGeometryCollection* GeometryCollection);
 
 	/** Get convex hull data for the Geometry Collection if it is present */
-	static TOptional<FGeometryCollectionConvexData> GetConvexHullDataIfPresent(FGeometryCollection* GeometryCollection);
+	static TOptional<FGeometryCollectionConvexData> GetConvexHullDataIfPresent(FManagedArrayCollection* GeometryCollection);
 
 	/** @return true if convex hull data is present */
 	static bool HasConvexHullData(const FManagedArrayCollection* GeometryCollection);
