Add autowire to dataflow

diff --git a/Engine/Source/Runtime/Experimental/Dataflow/Engine/Private/Dataflow/DataflowEdNode.cpp b/Engine/Source/Runtime/Experimental/Dataflow/Engine/Private/Dataflow/DataflowEdNode.cpp
--- a/Engine/Source/Runtime/Experimental/Dataflow/Engine/Private/Dataflow/DataflowEdNode.cpp
+++ b/Engine/Source/Runtime/Experimental/Dataflow/Engine/Private/Dataflow/DataflowEdNode.cpp
@@ -406,6 +406,47 @@ void UDataflowEdNode::GetPinHoverText(const UEdGraphPin& Pin, FString& HoverText
 	}
 }
 
+void UDataflowEdNode::AutowireNewNode(UEdGraphPin* FromPin)
+{
+	if (DataflowGraph && FromPin)
+	{
+		if (UEdGraphNode* FromGraphNode = FromPin->GetOwningNode())
+		{
+			if (UDataflowEdNode* FromDataflowGraphNode = Cast<UDataflowEdNode>(FromPin->GetOwningNode()))
+			{
+				const TSharedPtr<FDataflowNode> FromDataFlowNode = FromDataflowGraphNode->GetDataflowNode();
+				if (FromDataFlowNode)
+				{
+					if (FDataflowOutput* FromOutput = FromDataFlowNode->FindOutput(FromPin->PinName))
+					{
+						const FName OutputType = FromOutput->GetType();
+
+						const TSharedPtr<FDataflowNode> ToDataFlowNode = this->GetDataflowNode();
+						for (UEdGraphPin* InputPin : this->GetAllPins())
+						{
+							if (FDataflowInput* ToInput = ToDataFlowNode->FindInput(InputPin->PinName))
+							{
+								if (ToInput->GetType() == OutputType)
+								{
+									if (const UEdGraph* EdGraph = this->GetGraph())
+									{
+										if (EdGraph->GetSchema()->TryCreateConnection(FromPin, InputPin))
+										{
+											FromGraphNode->NodeConnectionListChanged();
+											this->NodeConnectionListChanged();
+											break;
+										}
+									}
+								}
+							}
+						}
+					}
+				}
+			}
+		}
+	}
+}
+
 #endif
 
 
diff --git a/Engine/Source/Runtime/Experimental/Dataflow/Engine/Public/Dataflow/DataflowEdNode.h b/Engine/Source/Runtime/Experimental/Dataflow/Engine/Public/Dataflow/DataflowEdNode.h
--- a/Engine/Source/Runtime/Experimental/Dataflow/Engine/Public/Dataflow/DataflowEdNode.h
+++ b/Engine/Source/Runtime/Experimental/Dataflow/Engine/Public/Dataflow/DataflowEdNode.h
@@ -45,6 +45,7 @@ public:
 	virtual FText GetTooltipText() const override;
 	virtual void GetPinHoverText(const UEdGraphPin& Pin, FString& HoverTextOut) const override;
 	virtual FText GetPinDisplayName(const UEdGraphPin* Pin) const override;
+	virtual void AutowireNewNode(UEdGraphPin* FromPin) override;
 #endif // WITH_EDITOR
 	// End of UEdGraphNode interface
 
