Fix bug from CL 25453643 where root component space transform  is no longer up to date

diff --git a/Engine/Source/Runtime/Experimental/GeometryCollectionEngine/Private/GeometryCollection/GeometryCollectionComponent.cpp b/Engine/Source/Runtime/Experimental/GeometryCollectionEngine/Private/GeometryCollection/GeometryCollectionComponent.cpp
--- a/Engine/Source/Runtime/Experimental/GeometryCollectionEngine/Private/GeometryCollection/GeometryCollectionComponent.cpp
+++ b/Engine/Source/Runtime/Experimental/GeometryCollectionEngine/Private/GeometryCollection/GeometryCollectionComponent.cpp
@@ -4521,6 +4521,7 @@ void UGeometryCollectionComponent::RefreshCustomRenderer()
 						(DynamicCollection && DynamicCollection->Transform.IsValidIndex(RootIndex))
 						? DynamicCollection->Transform[RootIndex]
 						: FTransform::Identity;
+					ComponentSpaceTransforms[RootIndex] = RootTransform;
 					RendererInterface->UpdateRootTransform(*RestCollection, RootTransform);
 				}
 			}
