Fix unity build from CL 25483372

diff --git a/Engine/Source/Runtime/Experimental/Dataflow/Engine/Private/Dataflow/DataflowEdNode.cpp b/Engine/Source/Runtime/Experimental/Dataflow/Engine/Private/Dataflow/DataflowEdNode.cpp
--- a/Engine/Source/Runtime/Experimental/Dataflow/Engine/Private/Dataflow/DataflowEdNode.cpp
+++ b/Engine/Source/Runtime/Experimental/Dataflow/Engine/Private/Dataflow/DataflowEdNode.cpp
@@ -13,6 +13,7 @@
 
 #if WITH_EDITOR
 #include "EdGraph/EdGraphPin.h"
+#include "EdGraph/EdGraphSchema.h"
 #endif
 
 #define LOCTEXT_NAMESPACE "DataflowEdNode"
