Replicate anchored state of GemoetryCollection's root particle

diff --git a/Engine/Source/Runtime/Experimental/GeometryCollectionEngine/Private/GeometryCollection/GeometryCollectionComponent.cpp b/Engine/Source/Runtime/Experimental/GeometryCollectionEngine/Private/GeometryCollection/GeometryCollectionComponent.cpp
--- a/Engine/Source/Runtime/Experimental/GeometryCollectionEngine/Private/GeometryCollection/GeometryCollectionComponent.cpp
+++ b/Engine/Source/Runtime/Experimental/GeometryCollectionEngine/Private/GeometryCollection/GeometryCollectionComponent.cpp
@@ -41,6 +41,7 @@
 #include "PhysicsField/PhysicsFieldComponent.h"
 #include "PhysicsProxy/GeometryCollectionPhysicsProxy.h"
 #include "PhysicsSolver.h"
+#include "Chaos/PBDRigidClusteringAlgo.h"
 
 #include "Algo/RemoveIf.h"
 
@@ -210,6 +211,8 @@ bool FGeometryCollectionRepData::NetSerialize(FArchive& Ar, class UPackageMap* M
 	int32 NumClusters = Clusters.Num();
 	Ar << NumClusters;
 
+	Ar << bIsRootAnchored;
+
 	if(Ar.IsLoading())
 	{
 		Clusters.SetNum(NumClusters);
@@ -1652,11 +1655,25 @@ void UGeometryCollectionComponent::UpdateRepData()
 
 		const TManagedArray<int32>* InitialLevels = PhysicsProxy->GetPhysicsCollection().FindAttribute<int32>("InitialLevel", FGeometryCollection::TransformGroup);
 		const TManagedArray<TSet<int32>>& InitialChildren = PhysicsProxy->GetPhysicsCollection().Children;
+		const TArray<FPBDRigidClusteredParticleHandle*>& ParticleHandles = PhysicsProxy->GetParticles();
+
+		// Replicate the anchored state of the root particle
+		const int32 InitialRootIndex = PhysicsProxy->GetSimParameters().InitialRootIndex;
+		if (FPBDRigidClusteredParticleHandle* RootHandle = ParticleHandles[InitialRootIndex])
+		{
+			const bool bIsRootAnchored = RootHandle->IsAnchored();
+			if (bFirstUpdate || bIsRootAnchored != RepData.bIsRootAnchored)
+			{
+				RepData.bIsRootAnchored = bIsRootAnchored;
+				bClustersChanged = true;
+			}
+		}
+		
 
 		//see if we have any new clusters that are enabled
 		TSet<FPBDRigidClusteredParticleHandle*> Processed;
 
-		for (FPBDRigidClusteredParticleHandle* Particle : PhysicsProxy->GetParticles())
+		for (FPBDRigidClusteredParticleHandle* Particle : ParticleHandles)
 		{
 			// Particle can be null if we have embedded geometry 
 			if (Particle)
@@ -1957,6 +1974,31 @@ bool UGeometryCollectionComponent::ProcessRepData(const float DeltaTime, const f
 		return false;
 	}
 
+	// Update the anchored state of the root particle
+	const int32 InitialRootIndex = PhysicsProxy->GetSimParameters().InitialRootIndex;
+	const TArray<FPBDRigidClusteredParticleHandle*>& ParticleHandles = PhysicsProxy->GetParticles();
+	if (FPBDRigidClusteredParticleHandle* RootHandle = ParticleHandles[InitialRootIndex])
+	{
+		if (RootHandle->Parent() || !RootHandle->Disabled())
+		{
+			if (RepData.bIsRootAnchored != RootHandle->IsAnchored())
+			{
+				RootHandle->SetIsAnchored(RepData.bIsRootAnchored);
+
+				if (FPBDRigidsSolver* Solver = PhysicsProxy->GetSolver<Chaos::FPBDRigidsSolver>())
+				{
+					if (FPBDRigidsEvolutionGBF* Evolution = Solver->GetEvolution())
+					{
+						FRigidClustering& RigidClustering = Solver->GetEvolution()->GetRigidClustering();
+						// NOTE: We have to start out with the particle as kinematic and let UpdateKinematicProperties correct it if this is wrong
+						Evolution->SetParticleObjectState(RootHandle, Chaos::EObjectStateType::Kinematic);
+						Chaos::UpdateKinematicProperties(RootHandle, RigidClustering.GetChildrenMap(), *Evolution);
+					}
+				}
+			}
+		}
+	}
+
 	bool bHardSnap = false;
 	const int64 CurrentTimeInMs = FPlatformTime::ToMilliseconds64(FPlatformTime::Cycles64());
 
diff --git a/Engine/Source/Runtime/Experimental/GeometryCollectionEngine/Public/GeometryCollection/GeometryCollectionComponent.h b/Engine/Source/Runtime/Experimental/GeometryCollectionEngine/Public/GeometryCollection/GeometryCollectionComponent.h
--- a/Engine/Source/Runtime/Experimental/GeometryCollectionEngine/Public/GeometryCollection/GeometryCollectionComponent.h
+++ b/Engine/Source/Runtime/Experimental/GeometryCollectionEngine/Public/GeometryCollection/GeometryCollectionComponent.h
@@ -386,7 +386,7 @@ struct FGeometryCollectionRepData
 	GENERATED_BODY()
 
 	FGeometryCollectionRepData()
-		: Version(0), ServerFrame(0)
+		: Version(0), ServerFrame(0), bIsRootAnchored(false)
 	{
 
 	}
@@ -407,6 +407,9 @@ struct FGeometryCollectionRepData
 	// The sim-time that this rep data was received
 	TOptional<float> RepDataReceivedTime;
 
+	// Is the root particle of the GC currently anchored
+	bool bIsRootAnchored;
+
 	// Just test version to skip having to traverse the whole pose array for replication
 	bool Identical(const FGeometryCollectionRepData* Other, uint32 PortFlags) const;
 	bool NetSerialize(FArchive& Ar, class UPackageMap* Map, bool& bOutSuccess);
