Iris - Fixed issue with marking not replicated object as dirty.

#rnx

diff --git a/Engine/Source/Runtime/Experimental/Iris/Core/Private/Iris/ReplicationSystem/ReplicationWriter.cpp b/Engine/Source/Runtime/Experimental/Iris/Core/Private/Iris/ReplicationSystem/ReplicationWriter.cpp
--- a/Engine/Source/Runtime/Experimental/Iris/Core/Private/Iris/ReplicationSystem/ReplicationWriter.cpp
+++ b/Engine/Source/Runtime/Experimental/Iris/Core/Private/Iris/ReplicationSystem/ReplicationWriter.cpp
@@ -334,12 +334,13 @@ void FReplicationWriter::QueueNetObjectAttachments(FInternalNetRefIndex OwnerInt
 		return;
 	}
 
-	ObjectsWithDirtyChanges.SetBit(SubObjectInternalIndex);
 
 	FReplicationInfo& TargetInfo = GetReplicationInfo(TargetIndex);
 	TargetInfo.HasAttachments = 1;
 
-	if (OwnerInternalIndex != SubObjectInternalIndex)
+	ObjectsWithDirtyChanges.SetBit(TargetIndex);
+
+	if (OwnerInternalIndex != TargetIndex)
 	{
 		ObjectsWithDirtyChanges.SetBit(OwnerInternalIndex);
 		FReplicationInfo& OwnerInfo = GetReplicationInfo(OwnerInternalIndex);
@@ -400,7 +401,7 @@ void FReplicationWriter::StartReplication(uint32 InternalIndex)
 {
 	FReplicationInfo& Info = GetReplicationInfo(InternalIndex);
 
-	check(ReplicatedObjects[InternalIndex].GetState() == EReplicatedObjectState::Invalid);
+	check(Info.GetState() == EReplicatedObjectState::Invalid);
 
 	// Reset info
 	Info = FReplicationInfo();
@@ -474,11 +475,12 @@ void FReplicationWriter::StartReplication(uint32 InternalIndex)
 		FNetRefHandleManager::FReplicatedObjectData& ObjectData = NetRefHandleManager->GetReplicatedObjectDataNoCheck(InternalIndex);
 		if (ensureAlways(ObjectData.SubObjectRootIndex != FNetRefHandleManager::InvalidInternalIndex))
 		{
-			// If owner is not pending destroy we mark the state of the SubObject to SubObjectPendingDestroy and mark owner as having dirty subobjects which will 
-			// destroy the subobject using the replicated state path of the owner.
 			FReplicationInfo& OwnerInfo = ReplicatedObjects[ObjectData.SubObjectRootIndex];
-			ObjectsWithDirtyChanges.SetBit(ObjectData.SubObjectRootIndex);
-			OwnerInfo.HasDirtySubObjects = 1U;
+			if (OwnerInfo.GetState() != EReplicatedObjectState::Invalid)
+			{
+				ObjectsWithDirtyChanges.SetBit(ObjectData.SubObjectRootIndex);
+				OwnerInfo.HasDirtySubObjects = 1U;
+			}
 		}
 	}
 }
@@ -942,12 +944,12 @@ void FReplicationWriter::HandleDeliveredRecord(const FReplicationRecord::FRecord
 		if (Info.LastAckedBaselineIndex != FDeltaCompressionBaselineManager::InvalidBaselineIndex)
 		{			
 			BaselineManager->DestroyBaseline(Parameters.ConnectionId, InternalIndex, Info.LastAckedBaselineIndex);
-			UE_LOG_REPLICATIONWRITER_CONN(TEXT("Destroyed old baseline %u for ( InternalObjectIndex: %u )"), Info.LastAckedBaselineIndex, InternalIndex);			
+			UE_LOG_REPLICATIONWRITER_CONN(TEXT("Destroyed old baseline %u for ( InternalIndex: %u )"), Info.LastAckedBaselineIndex, InternalIndex);			
 		}
 		Info.LastAckedBaselineIndex = RecordInfo.NewBaselineIndex;
 		Info.PendingBaselineIndex = FDeltaCompressionBaselineManager::InvalidBaselineIndex;
 
-		UE_LOG_REPLICATIONWRITER_CONN(TEXT("Acknowledged baseline %u for ( InternalObjectIndex: %u )"), RecordInfo.NewBaselineIndex, InternalIndex);
+		UE_LOG_REPLICATIONWRITER_CONN(TEXT("Acknowledged baseline %u for ( InternalIndex: %u )"), RecordInfo.NewBaselineIndex, InternalIndex);
 	}
 
 	// Update state
@@ -1412,7 +1414,7 @@ void FReplicationWriter::HandleDroppedRecord(const FReplicationRecord::FRecordIn
 	if (RecordInfo.NewBaselineIndex != FDeltaCompressionBaselineManager::InvalidBaselineIndex)
 	{
 		check(RecordInfo.NewBaselineIndex == Info.PendingBaselineIndex);
-		UE_LOG_REPLICATIONWRITER_CONN(TEXT("Lost baseline %u for ( InternalObjectIndex: %u )"), RecordInfo.NewBaselineIndex, InternalIndex);
+		UE_LOG_REPLICATIONWRITER_CONN(TEXT("Lost baseline %u for ( InternalIndex: %u )"), RecordInfo.NewBaselineIndex, InternalIndex);
 		
 		BaselineManager->LostBaseline(Parameters.ConnectionId, InternalIndex, RecordInfo.NewBaselineIndex);
 		Info.PendingBaselineIndex = FDeltaCompressionBaselineManager::InvalidBaselineIndex;
@@ -1876,7 +1878,7 @@ FReplicationWriter::EWriteObjectStatus FReplicationWriter::WriteObjectAndSubObje
 	{
 		if (CreatedBaselineIndex != FDeltaCompressionBaselineManager::InvalidBaselineIndex)
 		{
-			UE_LOG_REPLICATIONWRITER_CONN(TEXT("Destroy cancelled baseline %u for ( InternalObjectIndex: %u )"), CreatedBaselineIndex, InternalIndex);
+			UE_LOG_REPLICATIONWRITER_CONN(TEXT("Destroy cancelled baseline %u for ( InternalIndex: %u )"), CreatedBaselineIndex, InternalIndex);
 			BaselineManager->DestroyBaseline(Parameters.ConnectionId, InternalIndex, CreatedBaselineIndex);
 		}
 	};
@@ -1973,7 +1975,7 @@ FReplicationWriter::EWriteObjectStatus FReplicationWriter::WriteObjectAndSubObje
 							ApplyFilterToChangeMask(OutBatchInfo.ParentInternalIndex, InternalIndex, Info, ObjectData.Protocol, ReplicatedObjectStateBuffer, bIsInitialState);
 						}
 
-						UE_LOG_REPLICATIONWRITER_CONN(TEXT("Created new baseline %u for ( InternalObjectIndex: %u )"), CreatedBaselineIndex, InternalIndex);
+						UE_LOG_REPLICATIONWRITER_CONN(TEXT("Created new baseline %u for ( InternalIndex: %u )"), CreatedBaselineIndex, InternalIndex);
 					}
 				}
 			}
