Iris - Use both all ones and all zeros optimization in WriteSparseBitArray.

#rnx

diff --git a/Engine/Plugins/Runtime/ReplicationSystemTestPlugin/Source/Private/Tests/Serialization/TestCompressedChangeMaskData.h b/Engine/Plugins/Runtime/ReplicationSystemTestPlugin/Source/Private/Tests/Serialization/TestCompressedChangeMaskData.h
--- a/Engine/Plugins/Runtime/ReplicationSystemTestPlugin/Source/Private/Tests/Serialization/TestCompressedChangeMaskData.h
+++ b/Engine/Plugins/Runtime/ReplicationSystemTestPlugin/Source/Private/Tests/Serialization/TestCompressedChangeMaskData.h
@@ -6,6 +6,7 @@
 namespace UE::Net::Private
 {
 
+// Format is: BitCount Hint BitArrayData
 const FNetBitArrayView::StorageWordType VeryLargeChangeMaskTestData[] =
 {
 	1025U, 0U,  0x0,0x1,0x2,0x4,0x8,0x10,0x11,0x12, 0x14,0x18,0x0,0x1,0x2,0x4,0x8,0x10, 0x11,0x12,0x14,0x18,0x0,0x1,0x2,0x4, 0x8,0x10,0x11,0x12,0x14,0x18,0x0,0x1,0x1,
@@ -1162,6 +1163,7 @@ const FNetBitArrayView::StorageWordType FullTestData[] =
  91U,  0U, 0xA9,0x10000,0x0,
  91U,  1U, 0xFFF7DBBF,0xFFFFFFFF,0x7FFFFFF,
  9U,  1U, 0x1FF,
+ 646U, 1U, 0xFFFFFFFFU, 0xFFFFFFFFU, 0x0, 0x0, 0x0, 0x0, 0x0, 0x0, 0xFFFFFFFFU, 0xFFFFFFFFU, 0xFFFFFFFFU, 0xFFFFFFFFU, 0xFFFFFFFFU, 0xFFFFFFFFU, 0x1FU, 0x0, 0x0, 0x0, 0x0, 0x0, 0x20U,
 };
 
 }
diff --git a/Engine/Source/Runtime/Experimental/Iris/Core/Private/Iris/Serialization/NetBitStreamUtil.cpp b/Engine/Source/Runtime/Experimental/Iris/Core/Private/Iris/Serialization/NetBitStreamUtil.cpp
--- a/Engine/Source/Runtime/Experimental/Iris/Core/Private/Iris/Serialization/NetBitStreamUtil.cpp
+++ b/Engine/Source/Runtime/Experimental/Iris/Core/Private/Iris/Serialization/NetBitStreamUtil.cpp
@@ -584,7 +584,7 @@ static uint32 ReadSparseUint32UsingIndices(FNetBitStreamReader* Reader, uint32 B
 template<typename GetDataFunc, typename WriteSparseUint32Func>
 void WriteSparseBitArray(FNetBitStreamWriter* Writer, const uint32* Data, uint32 BitCount, GetDataFunc&& GetDataFunction, WriteSparseUint32Func&& WriteSparseUint32Function)
 {
-	checkSlow(BitCount <= SerializeSparseArrayMaxBitCount);
+	ensure(BitCount <= SerializeSparseArrayMaxBitCount);
 
 	using StorageWordType = FNetBitArrayBase::StorageWordType;
 	const uint32 WordBitCount = FNetBitArrayBase::WordBitCount;
@@ -592,18 +592,29 @@ void WriteSparseBitArray(FNetBitStreamWriter* Writer, const uint32* Data, uint32
 	const StorageWordType LastWordMask = ~StorageWordType(0) >> (uint32(-int32(BitCount)) & (WordBitCount - 1));
 	const uint32 LastWordIt = WordCount - 1;
 
-	// Build mask
-	uint32 DirtyWordMask = uint32(0);
+	// Build mask over words that are all zero or all ones.
+	uint32 NonZeroWordMask = uint32(0);
+	// Mask over words that are all ones.
+	uint32 InvertedWordMask = uint32(0);
 	{
-		for (uint32 WordIt = 0, CurrentBitMask = 1U; WordIt < WordCount; ++WordIt, CurrentBitMask <<= 1)
+		for (uint32 WordIt = 0, CurrentBitMask = 1U; WordIt < WordCount; ++WordIt, CurrentBitMask <<= 1U)
 		{
-			const StorageWordType CurrentWord = GetDataFunction(Data[WordIt]) & ((WordIt == LastWordIt) ? LastWordMask : ~0U);
-			DirtyWordMask |= CurrentWord ? CurrentBitMask : 0U;
+			const StorageWordType CurrentWord = Data[WordIt] & ((WordIt == LastWordIt) ? LastWordMask : ~0U);
+			// Set bits outside the valid range as ones when inverting the word to have the possibility of the zero word optimization.
+			const StorageWordType InvertedWord = ~(CurrentWord | (WordIt == LastWordIt ? ~LastWordMask : 0U));
+
+			NonZeroWordMask |= (CurrentWord == 0) or (InvertedWord == 0) ? 0U : CurrentBitMask;
+			InvertedWordMask |=  InvertedWord == 0 ? CurrentBitMask : 0U;
 		}
 	}
 
 	// Write Mask
-	Writer->WriteBits(DirtyWordMask, WordCount);
+	Writer->WriteBits(NonZeroWordMask, WordCount);
+	// Write all ones word mask.
+	if (Writer->WriteBool(InvertedWordMask != 0))
+	{
+		Writer->WriteBits(InvertedWordMask, WordCount);
+	}
 
 	// Encode dirty words
 	{
@@ -614,7 +625,7 @@ void WriteSparseBitArray(FNetBitStreamWriter* Writer, const uint32* Data, uint32
 
 		while (RemainingBits >= 32U)
 		{
-			if (DirtyWordMask & CurrentMaskBit)
+			if (NonZeroWordMask & CurrentMaskBit)
 			{
 				WriteSparseUint32Function(Writer, GetDataFunction(Data[WordIt]), 32U);
 			}
@@ -623,7 +634,7 @@ void WriteSparseBitArray(FNetBitStreamWriter* Writer, const uint32* Data, uint32
 			RemainingBits -= 32U;
 		}
 		// Last word
-		if (RemainingBits && (DirtyWordMask & CurrentMaskBit))
+		if (RemainingBits && (NonZeroWordMask & CurrentMaskBit))
 		{
 			const StorageWordType CurrentWord = GetDataFunction(Data[WordIt]) & LastWordMask;
 			WriteSparseUint32Function(Writer, CurrentWord, RemainingBits);
@@ -644,7 +655,12 @@ void ReadSparseBitArray(FNetBitStreamReader* Reader, uint32* OutData, uint32 Bit
 	const uint32 LastWordIt = WordCount - 1U;
 
 	// Read mask
-	uint32 DirtyWordMask = Reader->ReadBits(WordCount);
+	uint32 NonZeroWordMask = Reader->ReadBits(WordCount);
+	uint32 InvertedWordMask = uint32(0);
+	if (Reader->ReadBool())
+	{
+		InvertedWordMask = Reader->ReadBits(WordCount);
+	}
 
 	// Read and decode dirty words
 	uint32 CurrentMaskBit = 1U;
@@ -653,8 +669,17 @@ void ReadSparseBitArray(FNetBitStreamReader* Reader, uint32* OutData, uint32 Bit
 
 	while (RemainingBits >= 32U)
 	{
-		const StorageWordType ReadValue = DirtyWordMask & CurrentMaskBit ? ReadSparseUint32Function(Reader, 32U) : 0U;
-		OutData[WordIt] = GetDataFunction(ReadValue);
+		StorageWordType ReadValue = 0U;
+		if (NonZeroWordMask & CurrentMaskBit)
+		{
+			ReadValue = ReadSparseUint32Function(Reader, 32U);
+			ReadValue = GetDataFunction(ReadValue);
+		}
+		else if (InvertedWordMask & CurrentMaskBit)
+		{
+			ReadValue = ~0U;
+		}
+		OutData[WordIt] = ReadValue;
 
 		CurrentMaskBit += CurrentMaskBit;
 		RemainingBits -= 32U;
@@ -664,8 +689,17 @@ void ReadSparseBitArray(FNetBitStreamReader* Reader, uint32* OutData, uint32 Bit
 	// Last word, make sure we do not overwrite existing data
 	if (RemainingBits > 0U)
 	{
-		const StorageWordType ReadValue = DirtyWordMask & CurrentMaskBit ? ReadSparseUint32Function(Reader, RemainingBits) : 0U;
-		OutData[WordIt] = (OutData[WordIt] & ~LastWordMask) | GetDataFunction(ReadValue) & LastWordMask;
+		StorageWordType ReadValue = 0U;
+		if (NonZeroWordMask & CurrentMaskBit)
+		{
+			ReadValue = ReadSparseUint32Function(Reader, RemainingBits);
+			ReadValue = GetDataFunction(ReadValue) & LastWordMask;
+		}
+		else if (InvertedWordMask & CurrentMaskBit)
+		{
+			ReadValue = LastWordMask;
+		}
+		OutData[WordIt] = (OutData[WordIt] & ~LastWordMask) | ReadValue;
 	}
 }
 
