Fix cluster union remove operation regenerating the whole connectivity graph.

* Exposed the ability to make and break connections explicitly on the clustering manager so caller objects can have more control over the connection graph within.
* Converted union update method to use a flags object to avoid unamed bool parameters.
* Left add and update operations still rebuilding the graph for now until a later change.

diff --git a/Engine/Source/Runtime/Experimental/Chaos/Private/Chaos/ClusterUnionManager.cpp b/Engine/Source/Runtime/Experimental/Chaos/Private/Chaos/ClusterUnionManager.cpp
--- a/Engine/Source/Runtime/Experimental/Chaos/Private/Chaos/ClusterUnionManager.cpp
+++ b/Engine/Source/Runtime/Experimental/Chaos/Private/Chaos/ClusterUnionManager.cpp
@@ -333,7 +333,7 @@ namespace Chaos
 			// The anchored flag is taken care of in UpdateKinematicProperties so it must be set before that is called.
 			Cluster->InternalCluster->SetIsAnchored(true);
 		}
-		UpdateAllClusterUnionProperties(*Cluster, bIsNewCluster);
+		UpdateAllClusterUnionProperties(*Cluster, bIsNewCluster ? EUpdateClusterUnionPropertiesFlags::All : EUpdateClusterUnionPropertiesFlags::ForceGenerateConnectionGraph);
 
 		if (OldProxy)
 		{
@@ -364,7 +364,16 @@ namespace Chaos
 	void FClusterUnionManager::PostRemovalClusterUnionUpdate(FClusterUnion& Union)
 	{
 		SCOPE_CYCLE_COUNTER(STAT_PostRemovalClusterUnionUpdate);
-		UpdateAllClusterUnionProperties(Union, false);
+
+		if(TArray<FPBDRigidParticleHandle*>* Children = MClustering.GetChildrenMap().Find(Union.InternalCluster))
+		{
+			for(FPBDRigidParticleHandle* Child : *Children)
+			{
+				MClustering.RemoveNodeConnections(Child);
+			}
+		}
+
+		UpdateAllClusterUnionProperties(Union, EUpdateClusterUnionPropertiesFlags::None);
 
 		if (!Union.ChildParticles.IsEmpty())
 		{
@@ -486,12 +495,14 @@ namespace Chaos
 	}
 
 	DECLARE_CYCLE_STAT(TEXT("FClusterUnionManager::UpdateClusterUnionProperties"), STAT_UpdateClusterUnionProperties, STATGROUP_Chaos);
-	void FClusterUnionManager::UpdateAllClusterUnionProperties(FClusterUnion& ClusterUnion, bool bRecomputeMassOrientation)
+	void FClusterUnionManager::UpdateAllClusterUnionProperties(FClusterUnion& ClusterUnion, EUpdateClusterUnionPropertiesFlags Flags)
 	{
 		SCOPE_CYCLE_COUNTER(STAT_UpdateClusterUnionProperties);
 		// Update cluster properties.
 		FMatrix33 ClusterInertia(0);
 
+		const bool bRecomputeMassOrientation = EnumHasAnyFlags(Flags, EUpdateClusterUnionPropertiesFlags::RecomputeMassOrientation);
+
 		// TODO: These functions are generally just re-building the cluster from scratch. Need to figure out a way
 		// to get these functions to update the already existing cluster instead.
 		TSet<FPBDRigidParticleHandle*> FullChildrenSet(ClusterUnion.ChildParticles);
@@ -541,8 +552,11 @@ namespace Chaos
 			}
 		}
 
-		MClustering.ClearConnectionGraph(ClusterUnion.InternalCluster);
-		MClustering.GenerateConnectionGraph(ClusterUnion.InternalCluster, ClusterUnion.Parameters);
+		if(EnumHasAnyFlags(Flags, EUpdateClusterUnionPropertiesFlags::ForceGenerateConnectionGraph))
+		{
+			MClustering.ClearConnectionGraph(ClusterUnion.InternalCluster);
+			MClustering.GenerateConnectionGraph(ClusterUnion.InternalCluster, ClusterUnion.Parameters);
+		}
 	}
 
 	DECLARE_CYCLE_STAT(TEXT("FClusterUnionManager::GetOrCreateClusterUnionIndexFromExplicitIndex"), STAT_GetOrCreateClusterUnionIndexFromExplicitIndex, STATGROUP_Chaos);
@@ -666,7 +680,7 @@ namespace Chaos
 			}
 		}
 
-		UpdateAllClusterUnionProperties(*ClusterUnion, false);
+		UpdateAllClusterUnionProperties(*ClusterUnion, EUpdateClusterUnionPropertiesFlags::ForceGenerateConnectionGraph);
 
 		for (TPair<FPBDRigidClusteredParticleHandle*, bool>& Pair : DirtyParticleLockStates)
 		{
diff --git a/Engine/Source/Runtime/Experimental/Chaos/Private/Chaos/PBDRigidClustering.cpp b/Engine/Source/Runtime/Experimental/Chaos/Private/Chaos/PBDRigidClustering.cpp
--- a/Engine/Source/Runtime/Experimental/Chaos/Private/Chaos/PBDRigidClustering.cpp
+++ b/Engine/Source/Runtime/Experimental/Chaos/Private/Chaos/PBDRigidClustering.cpp
@@ -2104,6 +2104,11 @@ namespace Chaos
 		ConnectClusteredNodes(ClusteredChild1, ClusteredChild2, FromParticles, ToParticles);
 	}
 
+	void FRigidClustering::CreateNodeConnection(FPBDRigidClusteredParticleHandle* A, FPBDRigidClusteredParticleHandle* B)
+	{
+		ConnectClusteredNodes(A, B, nullptr, nullptr);
+	}
+
 	DECLARE_CYCLE_STAT(TEXT("TPBDRigidClustering<>::UpdateConnectivityGraphUsingPointImplicit"), STAT_UpdateConnectivityGraphUsingPointImplicit, STATGROUP_Chaos);
 	void 
 	FRigidClustering::UpdateConnectivityGraphUsingPointImplicit(
diff --git a/Engine/Source/Runtime/Experimental/Chaos/Public/Chaos/ClusterUnionManager.h b/Engine/Source/Runtime/Experimental/Chaos/Public/Chaos/ClusterUnionManager.h
--- a/Engine/Source/Runtime/Experimental/Chaos/Public/Chaos/ClusterUnionManager.h
+++ b/Engine/Source/Runtime/Experimental/Chaos/Public/Chaos/ClusterUnionManager.h
@@ -5,6 +5,7 @@
 #include "Chaos/ParticleHandleFwd.h"
 #include "Containers/Array.h"
 #include "Containers/Map.h"
+#include "Misc/EnumClassFlags.h"
 
 namespace Chaos
 {
@@ -32,6 +33,15 @@ namespace Chaos
 		Immediate
 	};
 
+	enum class EUpdateClusterUnionPropertiesFlags : int32
+	{
+		None = 0,
+		RecomputeMassOrientation = 1 << 0,
+		ForceGenerateConnectionGraph = 1 << 1,
+		All = RecomputeMassOrientation | ForceGenerateConnectionGraph
+	};
+	ENUM_CLASS_FLAGS(EUpdateClusterUnionPropertiesFlags);
+
 	struct CHAOS_API FClusterUnionCreationParameters
 	{
 		FClusterUnionExplicitIndex ExplicitIndex = INDEX_NONE;
@@ -140,7 +150,7 @@ namespace Chaos
 		void UpdateClusterUnionParticlesChildToParent(FClusterUnionIndex Index, const TArray<FPBDRigidParticleHandle*>& Particles, const TArray<FTransform>& ChildToParent, bool bLock);
 
 		// Update the cluster union's properties after its set of particle changes.
-		void UpdateAllClusterUnionProperties(FClusterUnion& ClusterUnion, bool bRecomputeMassOrientation);
+		void UpdateAllClusterUnionProperties(FClusterUnion& ClusterUnion, EUpdateClusterUnionPropertiesFlags Flags = EUpdateClusterUnionPropertiesFlags::All);
 
 		// Returns all cluster unions. Really meant only to be used for debugging.
 		const TMap<FClusterUnionIndex, FClusterUnion>& GetAllClusterUnions() const { return ClusterUnions; }
diff --git a/Engine/Source/Runtime/Experimental/Chaos/Public/Chaos/PBDRigidClustering.h b/Engine/Source/Runtime/Experimental/Chaos/Public/Chaos/PBDRigidClustering.h
--- a/Engine/Source/Runtime/Experimental/Chaos/Public/Chaos/PBDRigidClustering.h
+++ b/Engine/Source/Runtime/Experimental/Chaos/Public/Chaos/PBDRigidClustering.h
@@ -387,9 +387,27 @@ public:
 	FClusterUnionManager& GetClusterUnionManager() { return ClusterUnionManager; }
 	const FClusterUnionManager& GetClusterUnionManager() const { return ClusterUnionManager; }
 
-
 	const TSet<Chaos::FPBDRigidClusteredParticleHandle*>& GetTopLevelClusterParentsStrained() const { return TopLevelClusterParentsStrained; }
 
+	// Remove connectivity edges for specified particles
+	void RemoveNodeConnections(FPBDRigidParticleHandle* Child);
+	void RemoveNodeConnections(FPBDRigidClusteredParticleHandle* Child);
+
+	template<typename ParticleHandleTypeA, typename ParticleHandleTypeB>
+	void CreateNodeConnection(ParticleHandleTypeA* A, ParticleHandleTypeB* B)
+	{
+		if(A && B)
+		{
+			CreateNodeConnection(A->CastToClustered(), B->CastToClustered());
+		}
+		else
+		{
+			ensureMsgf(false, TEXT("CreateNodeConnection asked to connect a null particle, ignoring connection."));
+		}
+	}
+
+	void CreateNodeConnection(FPBDRigidClusteredParticleHandle* A, FPBDRigidClusteredParticleHandle* B);
+
  protected:
 
 	void ComputeStrainFromCollision(const FPBDCollisionConstraints& CollisionRule, const FReal Dt);
@@ -436,9 +454,6 @@ public:
 		const Chaos::FPBDRigidClusteredParticleHandle* Parent,
 		const FClusterCreationParameters& Parameters = FClusterCreationParameters());
 	
-	void RemoveNodeConnections(FPBDRigidParticleHandle* Child);
-	void RemoveNodeConnections(FPBDRigidClusteredParticleHandle* Child);
-
 	void RemoveChildFromParent(FPBDRigidParticleHandle* Child, const FPBDRigidClusteredParticleHandle* ClusteredParent);
 	void RemoveChildFromParentAndChildrenArray(FPBDRigidParticleHandle* Child, const FPBDRigidClusteredParticleHandle* ClusteredParent);
 
