Iris
- Fixed crash if the instance of a replicated objects was destroyed during dispatch

diff --git a/Engine/Source/Runtime/Experimental/Iris/Core/Private/Iris/ReplicationSystem/DequantizeAndApplyHelper.cpp b/Engine/Source/Runtime/Experimental/Iris/Core/Private/Iris/ReplicationSystem/DequantizeAndApplyHelper.cpp
--- a/Engine/Source/Runtime/Experimental/Iris/Core/Private/Iris/ReplicationSystem/DequantizeAndApplyHelper.cpp
+++ b/Engine/Source/Runtime/Experimental/Iris/Core/Private/Iris/ReplicationSystem/DequantizeAndApplyHelper.cpp
@@ -23,6 +23,7 @@ struct FDequantizeAndApplyHelper::FContext
 		FReplicationStateApplyContext::FStateBufferData StateBufferData;
 		const FReplicationStateDescriptor* Descriptor;
 		FReplicationFragment* Fragment;
+		EReplicationFragmentTraits FragmentTraits;
 		bool bHasUnresolvableReferences;
 		bool bMightHaveUnresolvableInitReferences;
 	};
@@ -98,6 +99,7 @@ FDequantizeAndApplyHelper::FContext* FDequantizeAndApplyHelper::Initialize(FNetS
 			FContext::FStateData& StateData = CachedStateData[CachedStateCount];
 			StateData.Descriptor = CurrentDescriptor;
 			StateData.Fragment = Fragments[StateIt];
+			StateData.FragmentTraits = StateData.Fragment->GetTraits();
 
 			// Dequantize state data
 			if (!EnumHasAnyFlags(StateData.Fragment->GetTraits(), EReplicationFragmentTraits::HasPersistentTargetStateBuffer))
@@ -175,7 +177,8 @@ void FDequantizeAndApplyHelper::Deinitialize(FDequantizeAndApplyHelper::FContext
 	{
 		for (const FContext::FStateData& StateData : MakeArrayView(Context->CachedStateData, Context->CachedStateCount))
 		{
-			if (!EnumHasAnyFlags(StateData.Fragment->GetTraits(), EReplicationFragmentTraits::HasPersistentTargetStateBuffer))
+			// Note: We do not use the fragment directly here as there are some code paths that destroys replicated instance & associated Fragments from PostNetReceive/RPC`s
+			if (!EnumHasAnyFlags(StateData.FragmentTraits, EReplicationFragmentTraits::HasPersistentTargetStateBuffer))
 			{
 				StateData.Descriptor->DestructReplicationState(StateData.StateBufferData.ExternalStateBuffer, StateData.Descriptor);
 			}
