Chaos Cloth - Do not dirty properties when setting to the same value.

#rnx

diff --git a/Engine/Source/Runtime/Experimental/Chaos/Public/Chaos/CollectionPropertyFacade.h b/Engine/Source/Runtime/Experimental/Chaos/Public/Chaos/CollectionPropertyFacade.h
--- a/Engine/Source/Runtime/Experimental/Chaos/Public/Chaos/CollectionPropertyFacade.h
+++ b/Engine/Source/Runtime/Experimental/Chaos/Public/Chaos/CollectionPropertyFacade.h
@@ -243,7 +243,7 @@ PRAGMA_ENABLE_DEPRECATION_WARNINGS
 		template<typename T, TEMPLATE_REQUIRES(TIsWeightedType<T>::Value)>
 		void SetValue(int32 KeyIndex, const T& Value) { SetWeightedValue(KeyIndex, Value, Value); }
 
-		void SetStringValue(int32 KeyIndex, const FString& Value) { GetStringValueArray()[KeyIndex] = Value; SetStringDirty(KeyIndex); }
+		void SetStringValue(int32 KeyIndex, const FString& Value) { if (GetStringValueArray()[KeyIndex] != Value) { GetStringValueArray()[KeyIndex] = Value; SetStringDirty(KeyIndex); } }
 
 		UE_DEPRECATED(5.3, "SetFlags is being phased out to promote correct dirtying operations.")
 		void SetFlags(int32 KeyIndex, uint8 Flags) { SetValue(KeyIndex, GetFlagsArray(), (ECollectionPropertyFlags)Flags); }
@@ -415,8 +415,11 @@ PRAGMA_ENABLE_DEPRECATION_WARNINGS
 	template<typename T>
 	inline void FCollectionPropertyFacade::SetValue(int32 KeyIndex, const TArrayView<T>& ValueArray, const T& Value)
 	{
-		ValueArray[KeyIndex] = Value;
-		SetDirty(KeyIndex);
+		if (ValueArray[KeyIndex] != Value)
+		{
+			ValueArray[KeyIndex] = Value;
+			SetDirty(KeyIndex);
+		}
 	}
 
 	template<typename T, typename TEnableIf<TIsWeightedType<T>::Value, int>::type>
