Iris - When FNames are encoded as strings make sure to store them case-insensitive. Fixes issues with mismatching protocols.

#rnx

diff --git a/Engine/Source/Runtime/Experimental/Iris/Core/Private/Iris/Serialization/StringNetSerializers.cpp b/Engine/Source/Runtime/Experimental/Iris/Core/Private/Iris/Serialization/StringNetSerializers.cpp
--- a/Engine/Source/Runtime/Experimental/Iris/Core/Private/Iris/Serialization/StringNetSerializers.cpp
+++ b/Engine/Source/Runtime/Experimental/Iris/Core/Private/Iris/Serialization/StringNetSerializers.cpp
@@ -180,6 +180,8 @@ void FNameNetSerializer::Quantize(FNetSerializationContext& Context, const FNetQ
 	const bool bIsString = (AsEName == nullptr || !ShouldReplicateAsInteger(*AsEName, SourceName));
 	if (bIsString)
 	{
+		const bool bRequireCaseInsensitivity = true;
+
 		const FNameEntry* DisplayNameEntry = SourceName.GetDisplayNameEntry();
 		TargetName.bIsString = 1;
 		TargetName.bIsEncoded = DisplayNameEntry->IsWide();
@@ -190,6 +192,10 @@ void FNameNetSerializer::Quantize(FNetSerializationContext& Context, const FNetQ
 		{
 			WIDECHAR TempWideBuffer[NAME_SIZE];
 			SourceName.GetPlainWIDEString(TempWideBuffer);
+			if (bRequireCaseInsensitivity)
+			{
+				FCString::Strupr(TempWideBuffer);
+			}
 
 			// Our codec uses up to 3 bytes per codepoint
 			const uint32 NameLength = DisplayNameEntry->GetNameLength() + 1U;
@@ -218,6 +224,11 @@ void FNameNetSerializer::Quantize(FNetSerializationContext& Context, const FNetQ
 			// At this time it's impossible to avoid the double copy unless we use a fixed excessive storage
 			ANSICHAR TempAnsiBuffer[NAME_SIZE];
 			SourceName.GetPlainANSIString(TempAnsiBuffer);
+			if (bRequireCaseInsensitivity)
+			{
+				FCStringAnsi::Strupr(TempAnsiBuffer);
+			}
+
 			FMemory::Memcpy(TargetName.ElementStorage, TempAnsiBuffer, NewElementCount*sizeof(ANSICHAR));
 		}
 	}
