I/O - support for resolving CDN's from distribution endpoint(s)

diff --git a/Engine/Source/Runtime/Experimental/IoStoreOnDemand/IoStoreOnDemand.Build.cs b/Engine/Source/Runtime/Experimental/IoStoreOnDemand/IoStoreOnDemand.Build.cs
--- a/Engine/Source/Runtime/Experimental/IoStoreOnDemand/IoStoreOnDemand.Build.cs
+++ b/Engine/Source/Runtime/Experimental/IoStoreOnDemand/IoStoreOnDemand.Build.cs
@@ -7,6 +7,7 @@ public class IoStoreOnDemand : ModuleRules
 	public IoStoreOnDemand(ReadOnlyTargetRules Target) : base(Target)
 	{
         PublicDependencyModuleNames.Add("Core");
+        PrivateDependencyModuleNames.AddRange(new string[] { "HTTP", "Json" });
 		bAllowConfidentialPlatformDefines = true;
     }
 }
diff --git a/Engine/Source/Runtime/Experimental/IoStoreOnDemand/Private/IoStoreOnDemand.cpp b/Engine/Source/Runtime/Experimental/IoStoreOnDemand/Private/IoStoreOnDemand.cpp
--- a/Engine/Source/Runtime/Experimental/IoStoreOnDemand/Private/IoStoreOnDemand.cpp
+++ b/Engine/Source/Runtime/Experimental/IoStoreOnDemand/Private/IoStoreOnDemand.cpp
@@ -267,25 +267,26 @@ uint64 FIoStoreOnDemandModule::ParseSizeParam(const TCHAR* Param)
 void FIoStoreOnDemandModule::StartupModule()
 {
 #if !WITH_EDITOR
-	FString Url, Toc;
+	UE::FOnDemandEndpoint Endpoint;
 	
 	FString UrlParam;
 	if (FParse::Value(FCommandLine::Get(), TEXT("IoStoreOnDemand="), UrlParam))
 	{
 		FStringView UrlView(UrlParam);
-		if (UrlView.StartsWith(TEXTVIEW("http")) && UrlView.EndsWith(TEXTVIEW(".iochunktoc")))
+		if (UrlView.StartsWith(TEXTVIEW("http://")) && UrlView.EndsWith(TEXTVIEW(".iochunktoc")))
 		{
 			int32 Delim = INDEX_NONE;
-			if (UrlView.FindLastChar(TEXT('/'), Delim))
+			if (UrlView.RightChop(7).FindChar(TEXT('/'), Delim))
 			{
-				Url = UrlView.Left(Delim);
-				Toc = UrlView.RightChop(Delim + 1);
+				Endpoint.ServiceUrl = UrlView.Left(7 +  Delim);
+				Endpoint.TocPath = UrlView.RightChop(Endpoint.ServiceUrl.Len() + 1);
 			}
 		}
 	}
 	
-	if (Url.IsEmpty() || Toc.IsEmpty())
+	if (!Endpoint.IsValid())
 	{
+		Endpoint = UE::FOnDemandEndpoint();
 		FString ConfigFileName = TEXT("IoStoreOnDemand.ini");
 		FString ConfigPath = FPaths::Combine(TEXT("Cloud"), ConfigFileName);
 		FString ConfigContent = FPlatformMisc::LoadTextFileFromPlatformPackage(ConfigPath);
@@ -295,12 +296,28 @@ void FIoStoreOnDemandModule::StartupModule()
 			FConfigFile Config;
 			Config.ProcessInputFileContents(ConfigContent, ConfigFileName);
 
-			Config.GetString(TEXT("Endpoint"), TEXT("Url"), Url);
-			Config.GetString(TEXT("Endpoint"), TEXT("Toc"), Toc);
+			Config.GetString(TEXT("Endpoint"), TEXT("DistributionUrl"), Endpoint.DistributionUrl);
+			Config.GetString(TEXT("Endpoint"), TEXT("ServiceUrl"), Endpoint.ServiceUrl);
+			Config.GetString(TEXT("Endpoint"), TEXT("TocPath"), Endpoint.TocPath);
+			
+			if (Endpoint.DistributionUrl.EndsWith(TEXT("/")))
+			{
+				Endpoint.DistributionUrl = Endpoint.DistributionUrl.Left(Endpoint.DistributionUrl.Len() - 1);
+			}
+			
+			if (Endpoint.ServiceUrl.EndsWith(TEXT("/")))
+			{
+				Endpoint.ServiceUrl = Endpoint.DistributionUrl.Left(Endpoint.ServiceUrl.Len() - 1);
+			}
+
+			if (Endpoint.TocPath.StartsWith(TEXT("/")))
+			{
+				Endpoint.TocPath.RightChopInline(1);
+			}
 		}
 	}
 
-	if (!Url.IsEmpty() && !Toc.IsEmpty())
+	if (Endpoint.IsValid())
 	{
 		TSharedPtr<IIoCache> Cache;
 
@@ -323,12 +340,8 @@ void FIoStoreOnDemandModule::StartupModule()
 		}
 		TSharedPtr<UE::IOnDemandIoDispatcherBackend> Backend = UE::MakeOnDemandIoDispatcherBackend(Cache);
 
-		UE::FOnDemandEndpoint Endpoint;
 		Endpoint.EndpointType = UE::EOnDemandEndpointType::CDN;
-		Endpoint.Url = MoveTemp(Url); 
-		Endpoint.TocPath = MoveTemp(Toc); 
 		Backend->Mount(Endpoint);
-
 		FIoDispatcher::Get().Mount(Backend.ToSharedRef(), -10);
 	}
 #endif // !WITH_EDITOR
@@ -337,4 +350,3 @@ void FIoStoreOnDemandModule::StartupModule()
 void FIoStoreOnDemandModule::ShutdownModule()
 {
 }
-
diff --git a/Engine/Source/Runtime/Experimental/IoStoreOnDemand/Private/OnDemandIoDispatcherBackend.cpp b/Engine/Source/Runtime/Experimental/IoStoreOnDemand/Private/OnDemandIoDispatcherBackend.cpp
--- a/Engine/Source/Runtime/Experimental/IoStoreOnDemand/Private/OnDemandIoDispatcherBackend.cpp
+++ b/Engine/Source/Runtime/Experimental/IoStoreOnDemand/Private/OnDemandIoDispatcherBackend.cpp
@@ -2,6 +2,7 @@
 
 #include "OnDemandIoDispatcherBackend.h"
 #include "Containers/StringView.h"
+#include "CoreHttp/Client.h"
 #include "EncryptionKeyManager.h"
 #include "HAL/Event.h"
 #include "HAL/Platform.h"
@@ -9,27 +10,28 @@
 #include "HAL/PreprocessorHelpers.h"
 #include "HAL/Runnable.h"
 #include "HAL/RunnableThread.h"
+#include "Http.h"
+#include "HttpManager.h"
 #include "IO/IoAllocators.h"
 #include "IO/IoCache.h"
 #include "IO/IoDispatcher.h"
 #include "IO/IoOffsetLength.h"
+#include "IO/IoStatus.h"
 #include "IO/IoStore.h"
 #include "IO/IoStoreOnDemand.h"
 #include "IO/IoChunkEncoding.h"
 #include "Misc/ScopeLock.h"
 #include "Misc/ScopeRWLock.h"
+#include "Modules/ModuleManager.h"
 #include "ProfilingDebugging/CountersTrace.h"
 #include "Serialization/CompactBinary.h"
 #include "Serialization/CompactBinarySerialization.h"
+#include "Serialization/JsonReader.h"
+#include "Serialization/JsonSerializer.h"
 #include "Tasks/Task.h"
 #include "Tasks/Pipe.h"
 #include <atomic>
 
-#define WITH_HTTP_CLIENT 1
-#if WITH_HTTP_CLIENT
-#include "CoreHttp/Client.h"
-#endif
-
 namespace UE::IO::Private
 {
 
@@ -56,6 +58,227 @@ FIoHash GetChunkKey(const FIoHash& ChunkHash, const FIoOffsetAndLength& Range)
 	return HashBuilder.Finalize();
 }
 
+///////////////////////////////////////////////////////////////////////////////
+class FDistributionEndpoints
+{
+public:
+	using FOnEndpointResolved = TFunction<void(const FString&, TConstArrayView<FString>)>;
+	
+	FDistributionEndpoints() = default;
+	~FDistributionEndpoints();
+
+	void Resolve(const FString& DistributionUrl, FOnEndpointResolved&& OnResolved);
+	void ResolveDeferredEndpoints();
+
+private:
+	struct FResolvedEndpoint
+	{
+		TArray<FString> ServiceUrls;
+	};
+
+	struct FResolveRequest
+	{
+		FString DistributionUrl;
+		FHttpRequestPtr HttpRequest;
+		TArray<FOnEndpointResolved> Callbacks;
+		int32 RetryCount = 0;
+	};
+
+	void IssueRequests();
+	void CancelRequests();
+	void CompleteRequest(FResolveRequest& ResolveRequest, FHttpResponsePtr HttpResponse);
+
+	TMap<FString, TUniquePtr<FResolvedEndpoint>> ResolvedEndpoints;
+	TMap<FString, TUniquePtr<FResolveRequest>> PendingRequests;
+	FRWLock Lock;
+	bool bInitialized = false;
+};
+
+FDistributionEndpoints::~FDistributionEndpoints()
+{
+	CancelRequests();
+}
+
+void FDistributionEndpoints::Resolve(const FString& DistributionUrl, FOnEndpointResolved&& OnResolved)
+{
+	const FResolvedEndpoint* Ep = nullptr;
+	{
+		FReadScopeLock _(Lock);
+		if (TUniquePtr<FResolvedEndpoint>* Entry = ResolvedEndpoints.Find(DistributionUrl))
+		{
+			Ep = Entry->Get();
+		}
+	}
+
+	if (Ep)
+	{
+		return OnResolved(DistributionUrl, Ep->ServiceUrls);
+	}
+
+	bool bIssueRequest = false;
+	{
+		FWriteScopeLock _(Lock);
+		TUniquePtr<FResolveRequest>& Request = PendingRequests.FindOrAdd(DistributionUrl);
+		if (!Request.IsValid())
+		{
+			Request.Reset(new FResolveRequest{DistributionUrl});
+			bIssueRequest = bInitialized;
+		}
+		Request->Callbacks.Add(MoveTemp(OnResolved));
+	}
+
+	if (bIssueRequest)
+	{
+		IssueRequests();
+	}
+}
+
+void FDistributionEndpoints::ResolveDeferredEndpoints()
+{
+	{
+		FWriteScopeLock _(Lock);
+		bInitialized = true;
+	}
+
+	IssueRequests();
+}
+
+void FDistributionEndpoints::IssueRequests()
+{
+	// Currenlty we need to use the HTTP module in order to resolve service endpoints due to HTTPS
+	FHttpModule& HttpModule = FModuleManager::LoadModuleChecked<FHttpModule>("HTTP");
+	const int32 MaxAttempts = 3;
+
+	TArray<FHttpRequestPtr, TInlineAllocator<2>> HttpRequests;
+	{
+		FWriteScopeLock _(Lock);
+		check(bInitialized);
+
+		for (auto& Kv : PendingRequests)
+		{
+			if (Kv.Value->HttpRequest.IsValid())
+			{
+				continue;
+			}
+
+			FResolveRequest& ResolveRequest = *Kv.Value.Get();
+			UE_LOG(LogIoStoreOnDemand, Log, TEXT("Resolving '%s' (#%d/%d)"), *ResolveRequest.DistributionUrl, ResolveRequest.RetryCount + 1, MaxAttempts);
+
+			FHttpRequestPtr HttpRequest = HttpModule.Get().CreateRequest();
+			HttpRequest->SetTimeout(3.0f);
+			HttpRequest->SetURL(Kv.Key);
+			HttpRequest->SetVerb(TEXT("GET"));
+			HttpRequest->SetHeader(TEXT("Accept"), TEXT("application/json"));
+			HttpRequest->OnProcessRequestComplete().BindLambda(
+				[this, &ResolveRequest, MaxAttempts]
+				(FHttpRequestPtr, FHttpResponsePtr Response, bool bOk)
+				{
+					FHttpRequestPtr Request = MoveTemp(ResolveRequest.HttpRequest);
+					if (Response->GetResponseCode() != 200)
+					{
+						if (++ResolveRequest.RetryCount < MaxAttempts)
+						{
+							Request->OnProcessRequestComplete().Unbind();
+							return IssueRequests();
+						}
+					}
+
+					CompleteRequest(ResolveRequest, Response);
+				});
+
+			ResolveRequest.HttpRequest = HttpRequest;
+			HttpRequests.Add(HttpRequest);
+		}
+	}
+
+	for (FHttpRequestPtr& Request : HttpRequests)
+	{
+		Request->ProcessRequest();
+	}
+}
+
+void FDistributionEndpoints::CancelRequests()
+{
+	TArray<FHttpRequestPtr, TInlineAllocator<2>> HttpRequests;
+	{
+		FWriteScopeLock _(Lock);
+		for (auto& Kv : PendingRequests)
+		{
+			if (Kv.Value->HttpRequest.IsValid())
+			{
+				HttpRequests.Add(Kv.Value->HttpRequest);
+			}
+		}
+	}
+
+	if (!HttpRequests.IsEmpty())
+	{
+		FHttpModule& HttpModule = FModuleManager::LoadModuleChecked<FHttpModule>("HTTP");
+		for (FHttpRequestPtr& Request : HttpRequests)
+		{
+			HttpModule.GetHttpManager().RemoveRequest(Request.ToSharedRef());
+			//TODO: Flush?
+		}
+	}
+}
+
+void FDistributionEndpoints::CompleteRequest(FResolveRequest& ResolveRequest, FHttpResponsePtr HttpResponse)
+{
+	TRACE_CPUPROFILER_EVENT_SCOPE(FOnDemandIoBackend::CompleteDistributionRequest);
+
+	using FJsonValuePtr = TSharedPtr<FJsonValue>;
+	using FJsonObjPtr = TSharedPtr<FJsonObject>;
+	using FJsonReader = TJsonReader<TCHAR>;
+	using FJsonReaderPtr = TSharedRef<FJsonReader>;
+
+	TArray<FString> ServiceUrls;
+	if (HttpResponse->GetResponseCode() == 200)
+	{
+		FString Json = HttpResponse->GetContentAsString();
+		FJsonReaderPtr JsonReader = TJsonReaderFactory<TCHAR>::Create(Json);
+
+		FJsonObjPtr JsonObj;
+		if (FJsonSerializer::Deserialize(JsonReader, JsonObj))
+		{
+			TArray<FJsonValuePtr> JsonValues = JsonObj->GetArrayField(TEXT("distributions"));
+			for (const FJsonValuePtr& JsonValue : JsonValues)
+			{
+				FString ServiceUrl = JsonValue->AsString();
+				if (ServiceUrl.EndsWith(TEXT("/")))
+				{
+					ServiceUrl.LeftInline(ServiceUrl.Len() - 1);
+				}
+				ServiceUrls.Add(MoveTemp(ServiceUrl));
+			}
+		}
+	}
+
+	const FResolvedEndpoint* ResolvedEndpoint = nullptr;
+	FString DistributionUrl;
+	TArray<FOnEndpointResolved> Callbacks;
+
+	{
+		FWriteScopeLock _(Lock);
+		if (!ServiceUrls.IsEmpty())
+		{
+			ResolvedEndpoint = ResolvedEndpoints.Emplace(
+				ResolveRequest.DistributionUrl,
+				new FResolvedEndpoint{MoveTemp(ServiceUrls)})
+			.Get();
+		}
+
+		Callbacks = MoveTemp(ResolveRequest.Callbacks);
+		DistributionUrl = MoveTemp(ResolveRequest.DistributionUrl);
+		PendingRequests.Remove(DistributionUrl);
+	}
+
+	TConstArrayView<FString> Urls = ResolvedEndpoint ? ResolvedEndpoint->ServiceUrls : TConstArrayView<FString>();
+	for (FOnEndpointResolved& Callback : Callbacks)
+	{
+		Callback(DistributionUrl, Urls);
+	}
+}
+
 ///////////////////////////////////////////////////////////////////////////////
 class FHttpClient
 {
@@ -66,10 +289,8 @@ public:
 	bool Tick();
 
 private:
-#if WITH_HTTP_CLIENT
 	void Issue(UE::HTTP::FRequest&& Request, FIoReadCallback&& Callback, FAnsiStringView Url = FAnsiStringView(), const TCHAR* DebugName = nullptr);
 	UE::HTTP::FEventLoop& EventLoop;
-#endif
 };
 
 FHttpClient::FHttpClient(UE::HTTP::FEventLoop& Loop)
@@ -79,24 +300,19 @@ FHttpClient::FHttpClient(UE::HTTP::FEventLoop& Loop)
 
 void FHttpClient::Get(FAnsiStringView Url, const FIoOffsetAndLength& Range, FIoReadCallback&& Callback, const TCHAR* DebugName)
 {
-#if WITH_HTTP_CLIENT
 	const uint64 RangeStart = Range.GetOffset();
 	const uint64 RangeEnd = Range.GetOffset() + Range.GetLength();
 	UE::HTTP::FRequest Request = EventLoop.Get(Url);
 	Request.Header(ANSITEXTVIEW("Range"), WriteToAnsiString<64>(ANSITEXTVIEW("bytes="), RangeStart, ANSITEXTVIEW("-"), RangeEnd));
 	
 	Issue(MoveTemp(Request), MoveTemp(Callback), Url, DebugName);
-#endif
 }
 
 void FHttpClient::Get(FAnsiStringView Url, FIoReadCallback&& Callback, const TCHAR* DebugName)
 {
-#if WITH_HTTP_CLIENT
 	Issue(EventLoop.Get(Url), MoveTemp(Callback), Url, DebugName);
-#endif
 }
 
-#if WITH_HTTP_CLIENT
 void FHttpClient::Issue(UE::HTTP::FRequest&& Request, FIoReadCallback&& Callback, FAnsiStringView DebugUrl, const TCHAR* DebugName)
 {
 	using namespace UE::HTTP;
@@ -159,15 +375,10 @@ void FHttpClient::Issue(UE::HTTP::FRequest&& Request, FIoReadCallback&& Callback
 
 	EventLoop.Send(MoveTemp(Request), MoveTemp(Sink));
 }
-#endif
 
 bool FHttpClient::Tick()
 {
-#if WITH_HTTP_CLIENT
 	return EventLoop.Tick() != 0;
-#else
-	return false;
-#endif
 }
 
 ///////////////////////////////////////////////////////////////////////////////
@@ -186,6 +397,7 @@ public:
 	struct FContainer
 	{
 		FAES::FAESKey EncryptionKey;
+		const FOnDemandEndpoint* Endpoint = nullptr;
 		FString Name;
 		FString EncryptionKeyGuid;
 		FString ChunksDirectory;
@@ -210,21 +422,28 @@ public:
 			return TConstArrayView<uint32>(Container->BlockSizes.GetData() + Entry->BlockOffset, Entry->BlockCount);
 		}
 	};
-	
+
 	FOnDemandIoStore();
 	~FOnDemandIoStore();
 
-	void AddToc(const FString& Endpoint, FOnDemandToc&& Toc);
+	void Initialize();
+	FIoStatus AddEndpoint(const FOnDemandEndpoint& Endpoint);
 	TIoStatusOr<uint64> GetChunkSize(const FIoChunkId& ChunkId);
 	FChunkInfo GetChunkInfo(const FIoChunkId& ChunkId);
 
 private:
+	TIoStatusOr<FOnDemandToc> GetToc(const FString& ServiceUrl, const FString& TocPath);
+	void AddToc(const FOnDemandEndpoint& Ep, FOnDemandToc&& Toc);
+	FIoStatus AddDeferredEndpoints(const FString& DistributionUrl, const TConstArrayView<FString>& ServiceUrls);
+	void AddDeferredContainers();
 	void OnEncryptionKeyAdded(const FGuid& Id, const FAES::FAESKey& Key);
-	void RegisterPendingContainers();
-
+	
+	FDistributionEndpoints DistributionEndpoints;
+	TChunkedArray<FOnDemandEndpoint> Endpoints;
 	TChunkedArray<FContainer> Containers;
 	TArray<FContainer*> RegisteredContainers;
-	TArray<FContainer*> PendingContainers;
+	TArray<FContainer*> DeferredContainers;
+	TArray<FOnDemandEndpoint> DeferredEndpoints;
 	mutable FRWLock Lock;
 };
 
@@ -238,20 +457,141 @@ FOnDemandIoStore::~FOnDemandIoStore()
 	FEncryptionKeyManager::Get().OnKeyAdded().RemoveAll(this);
 }
 
-void FOnDemandIoStore::AddToc(const FString& Endpoint, FOnDemandToc&& Toc)
+void FOnDemandIoStore::Initialize()
+{
+	DistributionEndpoints.ResolveDeferredEndpoints();
+}
+
+FIoStatus FOnDemandIoStore::AddEndpoint(const FOnDemandEndpoint& Endpoint)
+{
+	if ((Endpoint.DistributionUrl.IsEmpty() && Endpoint.ServiceUrl.IsEmpty()) || Endpoint.TocPath.IsEmpty())
+	{
+		return FIoStatus(EIoErrorCode::InvalidParameter, TEXT("Invalid endpoint parameters"));
+	}
+
+	if (Endpoint.ServiceUrl.IsEmpty())
+	{
+		DeferredEndpoints.Add(Endpoint);
+		DistributionEndpoints.Resolve(Endpoint.DistributionUrl, [this](const FString& DistributionUrl, TConstArrayView<FString> SerivceUrls)
+		{
+			if (FIoStatus Status = AddDeferredEndpoints(DistributionUrl, SerivceUrls); !Status.IsOk())
+			{
+				UE_LOG(LogIoStoreOnDemand, Error, TEXT("Failed to add on demand endpoint, reason '%s'"), *Status.ToString());
+			}
+		});
+
+		return FIoStatus::Unknown;
+	}
+	else
+	{
+		if (TIoStatusOr<FOnDemandToc> Toc = GetToc(Endpoint.ServiceUrl, Endpoint.TocPath); Toc.IsOk())
+		{
+			AddToc(Endpoint, Toc.ConsumeValueOrDie());
+			return FIoStatus::Ok;
+		}
+
+		return FIoStatus(EIoErrorCode::CorruptToc, TEXT("Failed to load TOC from endpoint"));
+	}
+}
+
+TIoStatusOr<uint64> FOnDemandIoStore::GetChunkSize(const FIoChunkId& ChunkId)
 {
+	if (FChunkInfo Info = GetChunkInfo(ChunkId))
+	{
+		return Info.Entry->RawSize;
+	}
+
+	return FIoStatus(EIoErrorCode::UnknownChunkID);
+}
+
+FOnDemandIoStore::FChunkInfo FOnDemandIoStore::GetChunkInfo(const FIoChunkId& ChunkId)
+{
+	FReadScopeLock _(Lock);
+
+	for (const FContainer* Container : RegisteredContainers)
+	{
+		if (const FTocEntry* Entry = Container->TocEntries.Find(ChunkId))
+		{
+			return FChunkInfo{Container, Entry};
+		}
+	}
+
+	return {};
+}
+
+TIoStatusOr<FOnDemandToc> FOnDemandIoStore::GetToc(const FString& ServiceUrl, const FString& TocPath)
+{
+	TRACE_CPUPROFILER_EVENT_SCOPE(FOnDemandIoBackend::GetToc);
+
+	HTTP::FEventLoop EventLoop;
+	FHttpClient Client(EventLoop);
+
+	TAnsiStringBuilder<256> Url;
+	Url << ServiceUrl << "/" << TocPath;
+
+	for (int32 Attempt = 0, MaxAttempts = 3; Attempt < MaxAttempts; ++Attempt)
+	{
+		UE_LOG(LogIoStoreOnDemand, Log, TEXT("Fetching TOC '%s/%s' (#%d/%d)"), *ServiceUrl, *TocPath, Attempt + 1, MaxAttempts);
+		
+		TIoStatusOr<FOnDemandToc> Toc;
+		Client.Get(Url.ToView(), [&Toc](TIoStatusOr<FIoBuffer> Response)
+		{
+			if (Response.IsOk())
+			{
+				FIoBuffer Buffer = Response.ConsumeValueOrDie();
+				FOnDemandToc NewToc;	
+				if (LoadFromCompactBinary(FCbFieldView(Buffer.GetData()), NewToc))
+				{
+					Toc = TIoStatusOr<FOnDemandToc>(MoveTemp(NewToc));
+				}
+				else
+				{
+					UE_LOG(LogIoStoreOnDemand, Error, TEXT("Failed loading on demand TOC from compact binary"));
+				}
+			}
+			else
+			{
+				UE_LOG(LogIoStoreOnDemand, Error, TEXT("Failed fetching TOC, reason '%s'"), *Response.Status().ToString());
+			}
+		});
+
+		while (Client.Tick());
+
+		if (Toc.IsOk())
+		{
+			return Toc;
+		}
+	}
+
+	return TIoStatusOr<FOnDemandToc>(FIoStatus(EIoErrorCode::NotFound));
+}
+
+void FOnDemandIoStore::AddToc(const FOnDemandEndpoint& Ep, FOnDemandToc&& Toc)
+{
+	check(Ep.IsValid());
+	UE_LOG(LogIoStoreOnDemand, Log, TEXT("Adding TOC '%s/%s'"), *Ep.ServiceUrl, *Ep.TocPath);
+
 	TRACE_CPUPROFILER_EVENT_SCOPE(FOnDemandIoBackend::AddToc);
 
+	FString Prefix;
+	int32 Idx = INDEX_NONE;
+	if (Ep.TocPath.FindLastChar(TCHAR('/'), Idx))
+	{
+		Prefix = Ep.TocPath.Left(Idx);
+	}
+
 	{
 		FWriteScopeLock _(Lock);
 
+		const FOnDemandEndpoint* Endpoint = new(Endpoints) FOnDemandEndpoint{Ep.EndpointType, Ep.DistributionUrl, Ep.ServiceUrl, Ep.TocPath};
 		const FOnDemandTocHeader& Header = Toc.Header;
 
 		for (FOnDemandTocContainerEntry& Container : Toc.Containers)
 		{
 			FContainer* NewContainer = new(Containers) FContainer();
+			NewContainer->Endpoint = Endpoint;
 			NewContainer->Name = Container.ContainerName;
-			NewContainer->ChunksDirectory = Header.ChunksDirectory.ToLower();
+			NewContainer->ChunksDirectory = (Prefix.IsEmpty() ? Header.ChunksDirectory : Prefix / Header.ChunksDirectory).ToLower();
 			NewContainer->CompressionFormat = FName(Header.CompressionFormat);
 			NewContainer->BlockSize = Header.BlockSize;
 			NewContainer->EncryptionKeyGuid = Container.EncryptionKeyGuid;
@@ -271,48 +611,62 @@ void FOnDemandIoStore::AddToc(const FString& Endpoint, FOnDemandToc&& Toc)
 
 			NewContainer->BlockSizes = MoveTemp(Container.BlockSizes);
 
-			PendingContainers.Add(NewContainer);
+			DeferredContainers.Add(NewContainer);
 		}
 	}
 
-	RegisterPendingContainers();
+	AddDeferredContainers();
 }
 
-TIoStatusOr<uint64> FOnDemandIoStore::GetChunkSize(const FIoChunkId& ChunkId)
+FIoStatus FOnDemandIoStore::AddDeferredEndpoints(const FString& DistributionUrl, const TConstArrayView<FString>& ServiceUrls)
 {
-	if (FChunkInfo Info = GetChunkInfo(ChunkId))
+	TArray<FOnDemandEndpoint, TInlineAllocator<4>> EndpointsToAdd;
 	{
-		return Info.Entry->RawSize;
+		FWriteScopeLock _(Lock);
+		for (auto It = DeferredEndpoints.CreateIterator(); It; ++It)
+		{
+			if (It->DistributionUrl.Compare(DistributionUrl, ESearchCase::IgnoreCase) == 0)
+			{
+				EndpointsToAdd.Add(*It);
+				It.RemoveCurrent();
+			}
+		}
 	}
 
-	return FIoStatus(EIoErrorCode::UnknownChunkID);
-}
-
-FOnDemandIoStore::FChunkInfo FOnDemandIoStore::GetChunkInfo(const FIoChunkId& ChunkId)
-{
-	FReadScopeLock _(Lock);
-
-	for (const FContainer* Container : RegisteredContainers)
+	for (const FOnDemandEndpoint& Ep : EndpointsToAdd)
 	{
-		if (const FTocEntry* Entry = Container->TocEntries.Find(ChunkId))
+		bool bOk = false;
+		for (const FString& SerivceUrl : ServiceUrls)
 		{
-			return FChunkInfo{Container, Entry};
+			// Currenlty we don't need use secure sockets to fetch on demand content
+			FString UnsecureUrl = SerivceUrl.Replace(TEXT("https"), TEXT("http"));
+			TIoStatusOr<FOnDemandToc> Toc = GetToc(UnsecureUrl, Ep.TocPath);
+			if (Toc.IsOk())
+			{
+				AddToc(Ep, Toc.ConsumeValueOrDie());
+				bOk = true;
+				break;
+			}
+			else
+			{
+				UE_LOG(LogIoStoreOnDemand, Log, TEXT("Failed fetch TOC '%s/%s'"), *UnsecureUrl, *Ep.TocPath);
+			}
+		}
+
+		if (!bOk)
+		{
+			return FIoStatus(EIoErrorCode::CorruptToc, TEXT("Failed to add deferred endpoint"));
 		}
 	}
 
-	return {};
+	return FIoStatus::Ok;
 }
 
-void FOnDemandIoStore::OnEncryptionKeyAdded(const FGuid& Id, const FAES::FAESKey& Key)
+void FOnDemandIoStore::AddDeferredContainers()
 {
-	RegisterPendingContainers();
-}
+	FWriteScopeLock _(Lock);
 
-void FOnDemandIoStore::RegisterPendingContainers()
-{
-	FReadScopeLock _(Lock);
-
-	for (auto It = PendingContainers.CreateIterator(); It; ++It)
+	for (auto It = DeferredContainers.CreateIterator(); It; ++It)
 	{
 		FContainer* Container = *It;
 		if (Container->EncryptionKeyGuid.IsEmpty())
@@ -340,6 +694,11 @@ void FOnDemandIoStore::RegisterPendingContainers()
 	}
 }
 
+void FOnDemandIoStore::OnEncryptionKeyAdded(const FGuid& Id, const FAES::FAESKey& Key)
+{
+	AddDeferredContainers();
+}
+
 ///////////////////////////////////////////////////////////////////////////////
 class FIoRequestQueue
 {
@@ -400,7 +759,10 @@ struct FChunkRequestParams
 	void GetUrl(FAnsiStringBuilderBase& Url) const
 	{
 		const FString HashString = LexToString(ChunkInfo.Entry->Hash);
-		Url << ChunkInfo.Container->ChunksDirectory << "/" << HashString.Left(2) << "/" << HashString << ANSITEXTVIEW(".iochunk");
+		Url << ChunkInfo.Container->Endpoint->ServiceUrl
+			<< "/" << ChunkInfo.Container->ChunksDirectory
+			<< "/" << HashString.Left(2)
+			<< "/" << HashString << ANSITEXTVIEW(".iochunk");
 	}
 
 	FIoChunkDecodingParams GetDecodingParams() const
@@ -652,10 +1014,6 @@ private:
 	TUniquePtr<FOnDemandIoStore> IoStore;
 	TArray<TUniquePtr<FHttpPipe>> HttpPipes;
 	FChunkRequests ChunkRequests;
-#if WITH_HTTP_CLIENT
-	UE::HTTP::FEventLoop HttpEventLoop;
-	TUniquePtr<FHttpClient> HttpClient;
-#endif
 	FIoRequestQueue CompletedRequests;
 	FOnDemandEndpoint CurrentEndpoint;
 	std::atomic_bool bStopRequested{false};
@@ -667,9 +1025,6 @@ FOnDemandIoBackend::FOnDemandIoBackend(TSharedPtr<IIoCache> InCache)
 	: Cache(InCache)
 {
 	IoStore = MakeUnique<FOnDemandIoStore>();
-#if WITH_HTTP_CLIENT
-	HttpClient = MakeUnique<FHttpClient>(HttpEventLoop);
-#endif
 
 	for (uint32 Idx = 0; Idx < HttpWorkerCount; ++Idx)
 	{
@@ -684,8 +1039,9 @@ FOnDemandIoBackend::~FOnDemandIoBackend()
 
 void FOnDemandIoBackend::Initialize(TSharedRef<const FIoDispatcherBackendContext> Context)
 {
-	UE_LOG(LogIoStoreOnDemand, Log, TEXT("Initializing HTTP I/O dispatcher backend"));
+	UE_LOG(LogIoStoreOnDemand, Log, TEXT("Initializing on demand I/O dispatcher backend"));
 	BackendContext = Context;
+	IoStore->Initialize();
 }
 
 void FOnDemandIoBackend::Shutdown()
@@ -695,7 +1051,7 @@ void FOnDemandIoBackend::Shutdown()
 		return;
 	}
 
-	UE_LOG(LogIoStoreOnDemand, Log, TEXT("Shutting down HTTP I/O dispatcher backend"));
+	UE_LOG(LogIoStoreOnDemand, Log, TEXT("Shutting down on demand I/O dispatcher backend"));
 
 	bStopRequested = true;
 	BackendContext.Reset();
@@ -820,9 +1176,7 @@ bool FOnDemandIoBackend::Resolve(FIoRequestImpl* Request)
 			}
 
 			TAnsiStringBuilder<256> Url;
-			Url << CurrentEndpoint.Url << "/";
 			ChunkRequest->Params.GetUrl(Url);
-
 			HttpPipe.Client.Get(Url.ToView(), ChunkRequest->Params.ChunkRange, [ChunkRequest](TIoStatusOr<FIoBuffer> Status)
 			{
 				if (Status.IsOk())
@@ -903,43 +1257,13 @@ void FOnDemandIoBackend::Mount(const FOnDemandEndpoint& Endpoint)
 {
 	TRACE_CPUPROFILER_EVENT_SCOPE(FOnDemandIoBackend::Mount);
 
-	CurrentEndpoint = Endpoint;
-
 	if (EnumHasAnyFlags(Endpoint.EndpointType, EOnDemandEndpointType::CDN))
 	{
-		UE_LOG(LogIoStoreOnDemand, Log, TEXT("Mounting CDN endpoint, Url='%s', Toc='%s'"), *Endpoint.Url, *Endpoint.TocPath);
-
-		TAnsiStringBuilder<256> Url;
-		Url << Endpoint.Url << "/" << Endpoint.TocPath;
-
-#if WITH_HTTP_CLIENT
-		HttpClient->Get(Url.ToView(), [this, Url = Endpoint.Url](TIoStatusOr<FIoBuffer> Status)
-		{
-			if (Status.IsOk())
-			{
-				FIoBuffer Buffer = Status.ConsumeValueOrDie();
-				FOnDemandToc Toc;
-				if (LoadFromCompactBinary(FCbFieldView(Buffer.GetData()), Toc))
-				{
-					IoStore->AddToc(Url, MoveTemp(Toc));
-				}
-				else
-				{
-					UE_LOG(LogIoStoreOnDemand, Error, TEXT("Failed loading ondemand TOC from compact binary"));
-				}
-			}
-			else
-			{
-				UE_LOG(LogIoStoreOnDemand, Error, TEXT("Mounting CDN endpoint failed, reason '%s'"), *Status.Status().ToString());
-			}
-		});
-
-		while (HttpClient->Tick());
-#endif
+		IoStore->AddEndpoint(Endpoint);
 	}
 	else
 	{
-		UE_LOG(LogIoStoreOnDemand, Log, TEXT("Mounting ZEN endpoint, Url='%s'"), *Endpoint.Url);
+		UE_LOG(LogIoStoreOnDemand, Log, TEXT("Mounting ZEN endpoint, Url='%s'"), *Endpoint.ServiceUrl);
 	}
 }
 
diff --git a/Engine/Source/Runtime/Experimental/IoStoreOnDemand/Private/OnDemandIoDispatcherBackend.h b/Engine/Source/Runtime/Experimental/IoStoreOnDemand/Private/OnDemandIoDispatcherBackend.h
--- a/Engine/Source/Runtime/Experimental/IoStoreOnDemand/Private/OnDemandIoDispatcherBackend.h
+++ b/Engine/Source/Runtime/Experimental/IoStoreOnDemand/Private/OnDemandIoDispatcherBackend.h
@@ -21,8 +21,11 @@ enum class EOnDemandEndpointType
 struct FOnDemandEndpoint
 {
 	EOnDemandEndpointType EndpointType;
-	FString Url;
+	FString DistributionUrl;
+	FString ServiceUrl;
 	FString TocPath;
+
+	bool IsValid() const { return (DistributionUrl.Len() > 0 || ServiceUrl.Len() > 0) && TocPath.Len() > 0; }
 };
 
 class IOnDemandIoDispatcherBackend
