EventRelay is created for break events.
It is allowing to iterate through all events from one function from a blueprint library.
We can also registered to the relay to be able to group or filter events out.
This is allowing to have a better control on the received events according to their component, tags, location, velocity...etc

diff --git a/Engine/Plugins/Experimental/ChaosCaching/Source/ChaosCaching/Private/Chaos/Adapters/GeometryCollectionComponentCacheAdapter.cpp b/Engine/Plugins/Experimental/ChaosCaching/Source/ChaosCaching/Private/Chaos/Adapters/GeometryCollectionComponentCacheAdapter.cpp
--- a/Engine/Plugins/Experimental/ChaosCaching/Source/ChaosCaching/Private/Chaos/Adapters/GeometryCollectionComponentCacheAdapter.cpp
+++ b/Engine/Plugins/Experimental/ChaosCaching/Source/ChaosCaching/Private/Chaos/Adapters/GeometryCollectionComponentCacheAdapter.cpp
@@ -375,6 +375,7 @@ namespace Chaos
 							CachedBreak.Mass = Event->Mass;
 							CachedBreak.BoundingBox = TAABB<FReal, 3>(Event->BoundingBoxMin, Event->BoundingBoxMax);
 							CachedBreak.BoundingBox = CachedBreak.BoundingBox.TransformedAABB(ComponentToWorld);
+							CachedBreak.bIsGlobal = Proxy->GetSimParameters().bDispatchGlobalBreakingData;
 
 							if (!SolverBreakingEventFilter->Enabled() || SolverBreakingEventFilter->Pass(CachedBreak))
 							{
diff --git a/Engine/Source/Runtime/Engine/Classes/Engine/World.h b/Engine/Source/Runtime/Engine/Classes/Engine/World.h
--- a/Engine/Source/Runtime/Engine/Classes/Engine/World.h
+++ b/Engine/Source/Runtime/Engine/Classes/Engine/World.h
@@ -60,6 +60,7 @@ class FTimerManager;
 class FWorldInGamePerformanceTrackers;
 class IInterface_PostProcessVolume;
 class UAISystemBase;
+class UChaosEventRelay;
 class UCanvas;
 class UDemoNetDriver;
 class UGameViewportClient;
@@ -1508,6 +1509,9 @@ public:
 	/** Default global physics scene. */
 	TSharedPtr<FPhysScene_Chaos> DefaultPhysicsScene_Chaos;
 
+	/** Access to the ChaosEventRelay to access all registered events. */
+	UChaosEventRelay* GetChaosEventRelay();
+
 	/** Physics Field component. */
 	UPROPERTY(Transient)
 	TObjectPtr<class UPhysicsFieldComponent> PhysicsField;
diff --git a/Engine/Source/Runtime/Engine/Classes/PhysicsEngine/ChaosBlueprintLibrary.h b/Engine/Source/Runtime/Engine/Classes/PhysicsEngine/ChaosBlueprintLibrary.h
new file mode 100644
--- /dev/null
+++ b/Engine/Source/Runtime/Engine/Classes/PhysicsEngine/ChaosBlueprintLibrary.h
@@ -0,0 +1,17 @@
+// Copyright Epic Games, Inc. All Rights Reserved.
+#pragma once
+
+#include "Kismet/BlueprintFunctionLibrary.h"
+
+#include "ChaosBlueprintLibrary.generated.h"
+
+class UChaosEventRelay;
+
+UCLASS()
+class UChaosBlueprintLibrary : public UBlueprintFunctionLibrary
+{
+	GENERATED_BODY()
+public:
+	UFUNCTION(BlueprintPure, Category = "Physics Object")
+	static const UChaosEventRelay* GetEventRelayFromContext(UObject* ContextObject);
+};
\ No newline at end of file
diff --git a/Engine/Source/Runtime/Engine/Private/PhysicsEngine/Experimental/ChaosBlueprintLibrary.cpp b/Engine/Source/Runtime/Engine/Private/PhysicsEngine/Experimental/ChaosBlueprintLibrary.cpp
new file mode 100644
--- /dev/null
+++ b/Engine/Source/Runtime/Engine/Private/PhysicsEngine/Experimental/ChaosBlueprintLibrary.cpp
@@ -0,0 +1,16 @@
+// Copyright Epic Games, Inc. All Rights Reserved.
+#include "PhysicsEngine/ChaosBlueprintLibrary.h"
+#include "Physics/Experimental/ChaosEventRelay.h"
+
+const UChaosEventRelay* UChaosBlueprintLibrary::GetEventRelayFromContext(UObject* ContextObject)
+{
+	if (ContextObject)
+	{
+		UWorld* World = GEngine->GetWorldFromContextObject(ContextObject, EGetWorldErrorMode::ReturnNull);
+		if (World)
+		{
+			return World->GetChaosEventRelay();
+		}
+	}
+	return nullptr;
+}
diff --git a/Engine/Source/Runtime/Engine/Private/PhysicsEngine/Experimental/ChaosEventRelay.cpp b/Engine/Source/Runtime/Engine/Private/PhysicsEngine/Experimental/ChaosEventRelay.cpp
new file mode 100644
--- /dev/null
+++ b/Engine/Source/Runtime/Engine/Private/PhysicsEngine/Experimental/ChaosEventRelay.cpp
@@ -0,0 +1,15 @@
+// Copyright Epic Games, Inc. All Rights Reserved.
+
+#include "Physics/Experimental/ChaosEventRelay.h"
+
+UChaosEventRelay::UChaosEventRelay()
+{
+}
+
+void UChaosEventRelay::DispatchPhysicsBreakEvents(const TArray<FBreakChaosEvent>& BreakEvents)
+{
+	if (OnBreakEvent.IsBound())
+	{
+		OnBreakEvent.Broadcast(BreakEvents);
+	}
+}
\ No newline at end of file
diff --git a/Engine/Source/Runtime/Engine/Private/PhysicsEngine/Experimental/ChaosEventType.cpp b/Engine/Source/Runtime/Engine/Private/PhysicsEngine/Experimental/ChaosEventType.cpp
new file mode 100644
--- /dev/null
+++ b/Engine/Source/Runtime/Engine/Private/PhysicsEngine/Experimental/ChaosEventType.cpp
@@ -0,0 +1,31 @@
+// Copyright Epic Games, Inc. All Rights Reserved.
+
+#include "Physics/Experimental/ChaosEventType.h"
+#include "Chaos/ExternalCollisionData.h"
+
+
+#include UE_INLINE_GENERATED_CPP_BY_NAME(ChaosEventType)
+
+FBreakChaosEvent::FBreakChaosEvent()
+	: Component(nullptr)
+	, Location(FVector::ZeroVector)
+	, Velocity(FVector::ZeroVector)
+	, AngularVelocity(FVector::ZeroVector)
+	, Extents(FVector::ZeroVector)
+	, Mass(0.0f)
+	, Index(INDEX_NONE)
+	, bFromCrumble(false)
+{
+}
+
+FBreakChaosEvent::FBreakChaosEvent(const Chaos::FBreakingData& BreakingData)
+	: Component(nullptr)
+	, Location(BreakingData.Location)
+	, Velocity(BreakingData.Velocity)
+	, AngularVelocity(BreakingData.AngularVelocity)
+	, Extents(BreakingData.BoundingBox.Extents())
+	, Mass(BreakingData.Mass)
+	, Index(BreakingData.TransformGroupIndex)
+	, bFromCrumble(BreakingData.bFromCrumble)
+{
+}
diff --git a/Engine/Source/Runtime/Engine/Private/PhysicsEngine/Experimental/PhysScene_Chaos.cpp b/Engine/Source/Runtime/Engine/Private/PhysicsEngine/Experimental/PhysScene_Chaos.cpp
--- a/Engine/Source/Runtime/Engine/Private/PhysicsEngine/Experimental/PhysScene_Chaos.cpp
+++ b/Engine/Source/Runtime/Engine/Private/PhysicsEngine/Experimental/PhysScene_Chaos.cpp
@@ -15,6 +15,7 @@
 #include "PhysicsEngine/ConstraintInstance.h"
 #include "PhysicsEngine/PhysicsCollisionHandler.h"
 #include "PhysicsEngine/PhysicsObjectExternalInterface.h"
+#include "Physics/Experimental/ChaosEventRelay.h"
 
 #include "ChaosSolversModule.h"
 
@@ -480,6 +481,7 @@ FPhysScene_Chaos::FPhysScene_Chaos(AActor* InSolverActor
 	, PhysicsReplication(nullptr)
 	, SolverActor(InSolverActor)
 	, LastEventDispatchTime(Chaos::FReal(-1))
+	, LastBreakEventDispatchTime(Chaos::FReal(-1))
 #if WITH_EDITOR
 	, SingleStepCounter(0)
 #endif
@@ -501,6 +503,8 @@ FPhysScene_Chaos::FPhysScene_Chaos(AActor* InSolverActor
 
 	Chaos::FEventManager* EventManager = SceneSolver->GetEventManager();
 	EventManager->RegisterHandler<Chaos::FCollisionEventData>(Chaos::EEventType::Collision, this, &FPhysScene_Chaos::HandleCollisionEvents);
+	EventManager->RegisterHandler<Chaos::FBreakingEventData>(Chaos::EEventType::Breaking, this, &FPhysScene_Chaos::HandleBreakingEvents);
+
 
 	//Initialize unique ptrs that are just here to allow forward declare. This should be reworked todo(ocohen)
 #if TODO_FIX_REFERENCES_TO_ADDARRAY
@@ -512,6 +516,8 @@ FPhysScene_Chaos::FPhysScene_Chaos(AActor* InSolverActor
 	PhysicsReplication = PhysicsReplicationFactory.IsValid() ? PhysicsReplicationFactory->Create(this) : new FPhysicsReplication(this);
 
 	FPhysicsDelegates::OnPhysSceneInit.Broadcast(this);
+
+	ChaosEventRelay = NewObject<UChaosEventRelay>();
 }
 
 FPhysScene_Chaos::~FPhysScene_Chaos()
@@ -755,6 +761,7 @@ void FPhysScene_Chaos::SetPhysicsReplication(FPhysicsReplication* InPhysicsRepli
 void FPhysScene_Chaos::AddReferencedObjects(FReferenceCollector& Collector)
 {
 	Super::AddReferencedObjects(Collector);
+	Collector.AddReferencedObject(ChaosEventRelay);
 #if WITH_EDITOR
 
 	for (auto& Pair : PhysicsProxyToComponentMap)
@@ -892,12 +899,14 @@ FORCEINLINE void FPhysScene_Chaos::HandleEachCollisionEvent(const TArray<int32>&
 			if (NotifyInfo.Info0.Component == nullptr)
 			{
 				NotifyInfo.Info0.Component = Comp0;
-				NotifyInfo.Info0.Actor = (NotifyInfo.Info0.Component != nullptr) ? NotifyInfo.Info0.Component->GetOwner() : nullptr;
+				check(Comp0);
+				NotifyInfo.Info0.Actor = Comp0->GetOwner();
 			}
 			if (NotifyInfo.Info1.Component == nullptr)
 			{
 				NotifyInfo.Info1.Component = Comp1;
-				NotifyInfo.Info1.Actor = (NotifyInfo.Info1.Component != nullptr) ? NotifyInfo.Info1.Component->GetOwner() : nullptr;
+				check(Comp1);
+				NotifyInfo.Info1.Actor = Comp1->GetOwner();
 			}
 		}
 	}
@@ -923,7 +932,7 @@ void FPhysScene_Chaos::HandleCollisionEvents(const Chaos::FCollisionEventData& E
 			for (TPair<IPhysicsProxyBase*, TArray<int32>> Pair : PhysicsProxyToCollisionIndicesMap)
 			{
 				IPhysicsProxyBase*  PhysicsProxy0 = Pair.Key;
-				TArray<int32> CollisionIndices = Pair.Value;
+				const TArray<int32>& CollisionIndices = Pair.Value;
 
 				UPrimitiveComponent* Comp0 = GetOwningComponent<UPrimitiveComponent>(PhysicsProxy0);
 
@@ -995,6 +1004,33 @@ void FPhysScene_Chaos::DispatchPendingCollisionNotifies()
 	}
 }
 
+void FPhysScene_Chaos::HandleBreakingEvents(const Chaos::FBreakingEventData& Event)
+{
+	QUICK_SCOPE_CYCLE_COUNTER(STAT_HandleBreakingEvents);
+
+	TArray<FBreakChaosEvent> PendingBreakEvents;
+	Chaos::FBreakingDataArray const& BreakingDataArray = Event.BreakingData.AllBreakingsArray;
+
+	if (Event.BreakingData.bHasGlobalEvent && LastBreakEventDispatchTime < Event.BreakingData.TimeCreated)
+	{
+		LastBreakEventDispatchTime = Event.BreakingData.TimeCreated;
+		for (const Chaos::FBreakingData& BreakingData: BreakingDataArray)
+		{
+			if (BreakingData.bIsGlobal && BreakingData.Proxy)
+			{
+				UPrimitiveComponent* Comp = GetOwningComponent<UPrimitiveComponent>(BreakingData.Proxy);
+				FBreakChaosEvent& BreakEvent = PendingBreakEvents.Emplace_GetRef(BreakingData);
+				BreakEvent.Component = Comp;
+			}
+		}
+	}
+
+	if (PendingBreakEvents.Num() > 0)
+	{
+		ChaosEventRelay->DispatchPhysicsBreakEvents(PendingBreakEvents);
+	}
+}
+
 #if CHAOS_WITH_PAUSABLE_SOLVER
 void FPhysScene_Chaos::OnUpdateWorldPause()
 {
diff --git a/Engine/Source/Runtime/Engine/Private/World.cpp b/Engine/Source/Runtime/Engine/Private/World.cpp
--- a/Engine/Source/Runtime/Engine/Private/World.cpp
+++ b/Engine/Source/Runtime/Engine/Private/World.cpp
@@ -66,6 +66,7 @@
 #include "Engine/NetworkObjectList.h"
 #include "GameFramework/GameStateBase.h"
 #include "PhysicsEngine/PhysicsConstraintComponent.h"
+#include "Physics/Experimental/ChaosEventRelay.h"
 #include "Components/BrushComponent.h"
 #include "Engine/Polys.h"
 #include "Components/ModelComponent.h"
@@ -1777,6 +1778,15 @@ void UWorld::RepairChaosActors()
 	DefaultPhysicsScene_Chaos = PhysicsScene_Chaos;
 }
 
+UChaosEventRelay* UWorld::GetChaosEventRelay()
+{
+	if (PhysicsScene)
+	{
+		return PhysicsScene->GetChaosEventRelay();
+	}
+	return nullptr;
+}
+
 void UWorld::RepairStreamingLevels()
 {
 	for (int32 Index = 0; Index < StreamingLevels.Num(); )
diff --git a/Engine/Source/Runtime/Engine/Public/Physics/Experimental/ChaosEventRelay.h b/Engine/Source/Runtime/Engine/Public/Physics/Experimental/ChaosEventRelay.h
new file mode 100644
--- /dev/null
+++ b/Engine/Source/Runtime/Engine/Public/Physics/Experimental/ChaosEventRelay.h
@@ -0,0 +1,32 @@
+// Copyright Epic Games, Inc. All Rights Reserved.
+
+#pragma once
+
+#include "Physics/Experimental/ChaosEventType.h"
+#include "UObject/Object.h"
+
+#include "ChaosEventRelay.generated.h"
+
+
+
+DECLARE_DYNAMIC_MULTICAST_DELEGATE_OneParam(FBreakEventSignature, const TArray<FBreakChaosEvent>&, BreakEvents);
+
+
+/**
+* An object managing events
+*/
+UCLASS()
+class ENGINE_API UChaosEventRelay : public UObject
+{
+	GENERATED_BODY()
+
+public:
+
+	UChaosEventRelay();
+
+	void DispatchPhysicsBreakEvents(const TArray<FBreakChaosEvent>& BreakEvents);
+
+	UPROPERTY(BlueprintAssignable, Category = "Physics Events")
+	FBreakEventSignature OnBreakEvent;
+
+};
diff --git a/Engine/Source/Runtime/Engine/Public/Physics/Experimental/ChaosEventType.h b/Engine/Source/Runtime/Engine/Public/Physics/Experimental/ChaosEventType.h
new file mode 100644
--- /dev/null
+++ b/Engine/Source/Runtime/Engine/Public/Physics/Experimental/ChaosEventType.h
@@ -0,0 +1,58 @@
+// Copyright Epic Games, Inc. All Rights Reserved.
+
+#pragma once
+
+#include "PhysicsPublic.h"
+#include "Math/MathFwd.h"
+
+#include "ChaosEventType.generated.h"
+
+
+namespace Chaos
+{
+	struct FBreakingData;
+}
+
+USTRUCT(BlueprintType)
+struct ENGINE_API FBreakChaosEvent
+{
+	GENERATED_BODY()
+
+public:
+
+	FBreakChaosEvent();
+	FBreakChaosEvent(const Chaos::FBreakingData& BreakingData);
+
+	/** primitive component involved in the break event */
+	UPROPERTY(BlueprintReadOnly, Category = "Break Event")
+	TObjectPtr<UPrimitiveComponent> Component = nullptr;
+
+	/** World location of the break */
+	UPROPERTY(BlueprintReadOnly, Category = "Break Event")
+	FVector Location;
+
+	/** Linear Velocity of the breaking particle  */
+	UPROPERTY(BlueprintReadOnly, Category = "Break Event")
+	FVector Velocity;
+
+	/** Angular Velocity of the breaking particle  */
+	UPROPERTY(BlueprintReadOnly, Category = "Break Event")
+	FVector AngularVelocity;
+
+	/** Extents of the bounding box */
+	UPROPERTY(BlueprintReadOnly, Category = "Break Event")
+	FVector Extents;
+
+	/** Mass of the breaking particle  */
+	UPROPERTY(BlueprintReadOnly, Category = "Break Event")
+	float Mass;
+
+	/** Index of the geometry collection bone if positive */
+	UPROPERTY(BlueprintReadOnly, Category = "Break Event")
+	int32 Index;
+
+	/** Whether the break event originated from a crumble event */
+	UPROPERTY(BlueprintReadOnly, Category = "Break Event")
+	bool bFromCrumble;
+};
+
diff --git a/Engine/Source/Runtime/Engine/Public/Physics/Experimental/PhysScene_Chaos.h b/Engine/Source/Runtime/Engine/Public/Physics/Experimental/PhysScene_Chaos.h
--- a/Engine/Source/Runtime/Engine/Public/Physics/Experimental/PhysScene_Chaos.h
+++ b/Engine/Source/Runtime/Engine/Public/Physics/Experimental/PhysScene_Chaos.h
@@ -6,6 +6,7 @@
 #include "Tickable.h"
 #include "EventsData.h"
 #include "Physics/PhysScene.h"
+#include "Physics/Experimental/ChaosEventType.h"
 #include "GameFramework/Actor.h"
 #include "PhysicsPublic.h"
 #include "PhysInterface_Chaos.h"
@@ -38,6 +39,7 @@ class FPerSolverFieldSystem;
 class IPhysicsProxyBase;
 
 class UWorld;
+class UChaosEventRelay;
 class AWorldSettings;
 class FPhysicsReplicationFactory;
 class FContactModifyCallbackFactory;
@@ -234,6 +236,7 @@ public:
 	void RemoveDeferredPhysicsStateCreation(UPrimitiveComponent* Component);
 	void ProcessDeferredCreatePhysicsState();
 
+	UChaosEventRelay* GetChaosEventRelay() const { return ChaosEventRelay.Get(); }
 
 	// Storage structure for replication data
 	// probably should just expose read/write API not the structure directly itself like this.
@@ -277,9 +280,10 @@ private:
 	// Chaos Event Handlers
 	void HandleEachCollisionEvent(const TArray<int32>& CollisionIndices, IPhysicsProxyBase* PhysicsProxy0, UPrimitiveComponent* const Comp0, Chaos::FCollisionDataArray const& CollisionData, Chaos::FReal MinDeltaVelocityThreshold);
 	void HandleCollisionEvents(const Chaos::FCollisionEventData& CollisionData);
-
 	void DispatchPendingCollisionNotifies();
 
+	void HandleBreakingEvents(const Chaos::FBreakingEventData& Event);
+
 	/** Replication manager that updates physics bodies towards replicated physics state */
 	FPhysicsReplication* PhysicsReplication;
 
@@ -364,7 +368,10 @@ private:
 	/** The SolverActor that spawned and owns this scene */
 	TWeakObjectPtr<AActor> SolverActor;
 
+	TObjectPtr<UChaosEventRelay> ChaosEventRelay;
+
 	Chaos::FReal LastEventDispatchTime;
+	Chaos::FReal LastBreakEventDispatchTime;
 	class FAsyncPhysicsTickCallback* AsyncPhysicsTickCallback = nullptr;
 
 #if WITH_EDITOR
diff --git a/Engine/Source/Runtime/Experimental/Chaos/Private/Chaos/PBDRigidClustering.cpp b/Engine/Source/Runtime/Experimental/Chaos/Private/Chaos/PBDRigidClustering.cpp
--- a/Engine/Source/Runtime/Experimental/Chaos/Private/Chaos/PBDRigidClustering.cpp
+++ b/Engine/Source/Runtime/Experimental/Chaos/Private/Chaos/PBDRigidClustering.cpp
@@ -569,6 +569,7 @@ namespace Chaos
 				}
 				ClusterBreak.TransformGroupIndex = ConcreteProxy->GetTransformGroupIndexFromHandle(ClusteredParticle);
 				ClusterBreak.bFromCrumble = bFromCrumble;
+				ClusterBreak.bIsGlobal = SimParams.bDispatchGlobalBreakingData;
 			}
 		}
 	}
diff --git a/Engine/Source/Runtime/Experimental/Chaos/Private/EventDefaults.cpp b/Engine/Source/Runtime/Experimental/Chaos/Private/EventDefaults.cpp
--- a/Engine/Source/Runtime/Experimental/Chaos/Private/EventDefaults.cpp
+++ b/Engine/Source/Runtime/Experimental/Chaos/Private/EventDefaults.cpp
@@ -327,6 +327,7 @@ namespace Chaos
 			{
 				FilteredBreakingDataArray.Reset();
 				FilteredBreakingIndicesByPhysicsProxy.Reset();
+				BreakingEventData.BreakingData.bHasGlobalEvent = false;
 			}
 			BreakingEventData.BreakingData.TimeCreated = Solver->MTime;
 
@@ -344,6 +345,7 @@ namespace Chaos
 					{
 						const int32 NewIndex = FilteredBreakingDataArray.Emplace(ClusterBreaking);
 						FilteredBreakingIndicesByPhysicsProxy.FindOrAdd(ClusterBreaking.Proxy).Add(FEventManager::EncodeCollisionIndex(NewIndex, false));
+						BreakingEventData.BreakingData.bHasGlobalEvent |= ClusterBreaking.bIsGlobal;
 					}
 				}
 			}
diff --git a/Engine/Source/Runtime/Experimental/Chaos/Private/PhysicsProxy/GeometryCollectionPhysicsProxy.cpp b/Engine/Source/Runtime/Experimental/Chaos/Private/PhysicsProxy/GeometryCollectionPhysicsProxy.cpp
--- a/Engine/Source/Runtime/Experimental/Chaos/Private/PhysicsProxy/GeometryCollectionPhysicsProxy.cpp
+++ b/Engine/Source/Runtime/Experimental/Chaos/Private/PhysicsProxy/GeometryCollectionPhysicsProxy.cpp
@@ -2463,6 +2463,13 @@ void FGeometryCollectionPhysicsProxy::SetNotifyBreakings_External(bool bNotify)
 	SetProxyDirty_External();
 }
 
+void FGeometryCollectionPhysicsProxy::SetNotifyGlobalBreakings_External(bool bNotify)
+{
+	check(IsInGameThread());
+	GameThreadPerFrameData.SetNotifyGlobalBreakings(bNotify);
+	SetProxyDirty_External();
+}
+
 void FGeometryCollectionPhysicsProxy::SetNotifyRemovals_External(bool bNotify)
 {
 	check(IsInGameThread());
@@ -2689,7 +2696,8 @@ void FGeometryCollectionPhysicsProxy::PushStateOnGameThread(Chaos::FPBDRigidsSol
 
 	if (GameThreadPerFrameData.GetIsNotificationDataDirty())
 	{
-		Parameters.bGenerateBreakingData = GameThreadPerFrameData.GetNotifyBreakings();
+		Parameters.bGenerateBreakingData = GameThreadPerFrameData.GetNotifyBreakings() || GameThreadPerFrameData.GetNotifyGlobalBreakings();
+		Parameters.bDispatchGlobalBreakingData = GameThreadPerFrameData.GetNotifyGlobalBreakings();
 		Parameters.bGenerateRemovalsData = GameThreadPerFrameData.GetNotifyRemovals();
 		Parameters.bGenerateCrumblingData = GameThreadPerFrameData.GetNotifyCrumblings();
 		Parameters.bGenerateCrumblingChildrenData = GameThreadPerFrameData.GetCrumblingEventIncludesChildren();
diff --git a/Engine/Source/Runtime/Experimental/Chaos/Public/Chaos/ExternalCollisionData.h b/Engine/Source/Runtime/Experimental/Chaos/Public/Chaos/ExternalCollisionData.h
--- a/Engine/Source/Runtime/Experimental/Chaos/Public/Chaos/ExternalCollisionData.h
+++ b/Engine/Source/Runtime/Experimental/Chaos/Public/Chaos/ExternalCollisionData.h
@@ -12,6 +12,13 @@ class UPhysicalMaterial;
 
 namespace Chaos
 {
+
+	struct FBaseEventFlag
+	{
+		FBaseEventFlag() : bIsGlobal(false) {}
+		bool bIsGlobal;
+	};
+
 	/**
 	 * Collision event data stored for use by other systems (e.g. Niagara, gameplay events)
 	 */
@@ -171,10 +178,11 @@ namespace Chaos
 	/*
 	BreakingData passed from the physics solver to subsystems
 	*/
-	struct FBreakingData
+	struct FBreakingData : public FBaseEventFlag
 	{
 		FBreakingData()
-			: Proxy(nullptr)
+			: FBaseEventFlag()
+			, Proxy(nullptr)
 			, Location(FVec3((FReal)0.0))
 			, Velocity(FVec3((FReal)0.0))
 			, AngularVelocity(FVec3((FReal)0.0))
diff --git a/Engine/Source/Runtime/Experimental/Chaos/Public/EventsData.h b/Engine/Source/Runtime/Experimental/Chaos/Public/EventsData.h
--- a/Engine/Source/Runtime/Experimental/Chaos/Public/EventsData.h
+++ b/Engine/Source/Runtime/Experimental/Chaos/Public/EventsData.h
@@ -75,7 +75,7 @@ namespace Chaos
 	 */
 	struct FAllBreakingData : public FTimeResource
 	{
-		FAllBreakingData() : AllBreakingsArray(FBreakingDataArray()) {}
+		FAllBreakingData() : AllBreakingsArray(FBreakingDataArray()), bHasGlobalEvent(false) {}
 
 		void Reset()
 		{
@@ -83,6 +83,7 @@ namespace Chaos
 		}
 
 		FBreakingDataArray AllBreakingsArray;
+		bool bHasGlobalEvent;
 	};
 
 	struct FBreakingEventData
diff --git a/Engine/Source/Runtime/Experimental/Chaos/Public/GeometryCollection/GeometryCollectionSimulationCoreTypes.h b/Engine/Source/Runtime/Experimental/Chaos/Public/GeometryCollection/GeometryCollectionSimulationCoreTypes.h
--- a/Engine/Source/Runtime/Experimental/Chaos/Public/GeometryCollection/GeometryCollectionSimulationCoreTypes.h
+++ b/Engine/Source/Runtime/Experimental/Chaos/Public/GeometryCollection/GeometryCollectionSimulationCoreTypes.h
@@ -205,6 +205,7 @@ struct FSimulationParameters
 		, bGenerateRemovalsData(false)
 		, bGenerateCrumblingData(false)
 		, bGenerateCrumblingChildrenData(false)
+		, bDispatchGlobalBreakingData(false)
 		, EnableGravity(true)
 		, GravityGroupIndex(0)
 		, UseInertiaConditioning(true)
@@ -259,6 +260,7 @@ struct FSimulationParameters
 		, bGenerateRemovalsData(Other.bGenerateRemovalsData)
 		, bGenerateCrumblingData(Other.bGenerateCrumblingData)
 		, bGenerateCrumblingChildrenData(Other.bGenerateCrumblingChildrenData)
+		, bDispatchGlobalBreakingData(Other.bDispatchGlobalBreakingData)
 		, Shared(Other.Shared)
 		, EnableGravity(Other.EnableGravity)
 		, GravityGroupIndex(Other.GravityGroupIndex)
@@ -342,6 +344,8 @@ struct FSimulationParameters
 	bool bGenerateCrumblingData;
 	bool bGenerateCrumblingChildrenData;
 
+	bool bDispatchGlobalBreakingData;
+
 	FSharedSimulationParameters Shared;
 
 	bool EnableGravity;
diff --git a/Engine/Source/Runtime/Experimental/Chaos/Public/GeometryCollectionProxyData.h b/Engine/Source/Runtime/Experimental/Chaos/Public/GeometryCollectionProxyData.h
--- a/Engine/Source/Runtime/Experimental/Chaos/Public/GeometryCollectionProxyData.h
+++ b/Engine/Source/Runtime/Experimental/Chaos/Public/GeometryCollectionProxyData.h
@@ -204,6 +204,13 @@ public:
 	}
 	bool GetNotifyBreakings() const { return bNotifyBreakings; }
 
+	void SetNotifyGlobalBreakings(bool bNotify)
+	{
+		bNotifyGlobalBreakings = bNotify;
+		bIsNotificationDataDirty = true;
+	}
+	bool GetNotifyGlobalBreakings() const { return bNotifyGlobalBreakings; }
+
 	void SetNotifyRemovals(bool bNotify)
 	{
 		bNotifyRemovals = bNotify;
@@ -247,6 +254,7 @@ private:
 	uint16 bNotifyRemovals : 1;
 	uint16 bNotifyCrumblings : 1;
 	uint16 bCrumblingEventIncludesChildren : 1;
+	uint16 bNotifyGlobalBreakings : 1;
 
 	/** updated when bDamageSettingsDataDirty is set */
 	uint16 bEnableStrainOnCollision : 1;
diff --git a/Engine/Source/Runtime/Experimental/Chaos/Public/PhysicsProxy/GeometryCollectionPhysicsProxy.h b/Engine/Source/Runtime/Experimental/Chaos/Public/PhysicsProxy/GeometryCollectionPhysicsProxy.h
--- a/Engine/Source/Runtime/Experimental/Chaos/Public/PhysicsProxy/GeometryCollectionPhysicsProxy.h
+++ b/Engine/Source/Runtime/Experimental/Chaos/Public/PhysicsProxy/GeometryCollectionPhysicsProxy.h
@@ -423,8 +423,9 @@ public:
 	void SetNotifyRemovals_External(bool bNotify);
 	void SetNotifyCrumblings_External(bool bNotify, bool bIncludeChildren);
 
-	float ComputeMaterialBasedDamageThreshold_Internal(Chaos::FPBDRigidClusteredParticleHandle& ClusteredParticle) const;
+	void SetNotifyGlobalBreakings_External(bool bNotify);
 
+	float ComputeMaterialBasedDamageThreshold_Internal(Chaos::FPBDRigidClusteredParticleHandle& ClusteredParticle) const;
 	FProxyInterpolationData& GetInterpolationData() { return InterpolationData; }
 	const FProxyInterpolationData& GetInterpolationData() const { return InterpolationData; }
 
diff --git a/Engine/Source/Runtime/Experimental/ChaosSolverEngine/Private/Chaos/ChaosGameplayEventDispatcher.cpp b/Engine/Source/Runtime/Experimental/ChaosSolverEngine/Private/Chaos/ChaosGameplayEventDispatcher.cpp
--- a/Engine/Source/Runtime/Experimental/ChaosSolverEngine/Private/Chaos/ChaosGameplayEventDispatcher.cpp
+++ b/Engine/Source/Runtime/Experimental/ChaosSolverEngine/Private/Chaos/ChaosGameplayEventDispatcher.cpp
@@ -18,29 +18,6 @@
 
 #include UE_INLINE_GENERATED_CPP_BY_NAME(ChaosGameplayEventDispatcher)
 
-FChaosBreakEvent::FChaosBreakEvent()
-	: Component(nullptr)
-	, Location(FVector::ZeroVector)
-	, Velocity(FVector::ZeroVector)
-	, AngularVelocity(FVector::ZeroVector)
-	, Extents(FVector::ZeroVector)
-	, Mass(0.0f)
-	, Index(INDEX_NONE)
-	, bFromCrumble(false)
-{
-}
-
-FChaosBreakEvent::FChaosBreakEvent(const Chaos::FBreakingData& BreakingData)
-	: Component(nullptr)
-	, Location(BreakingData.Location)
-	, Velocity(BreakingData.Velocity)
-	, AngularVelocity(BreakingData.AngularVelocity)
-	, Extents(BreakingData.BoundingBox.Extents())
-	, Mass(BreakingData.Mass)
-	, Index(BreakingData.TransformGroupIndex)
-	, bFromCrumble(BreakingData.bFromCrumble)
-{
-}
 
 FChaosRemovalEvent::FChaosRemovalEvent()
 	: Component(nullptr)
diff --git a/Engine/Source/Runtime/Experimental/ChaosSolverEngine/Public/Chaos/ChaosGameplayEventDispatcher.h b/Engine/Source/Runtime/Experimental/ChaosSolverEngine/Public/Chaos/ChaosGameplayEventDispatcher.h
--- a/Engine/Source/Runtime/Experimental/ChaosSolverEngine/Public/Chaos/ChaosGameplayEventDispatcher.h
+++ b/Engine/Source/Runtime/Experimental/ChaosSolverEngine/Public/Chaos/ChaosGameplayEventDispatcher.h
@@ -4,6 +4,7 @@
 
 #include "ChaosEventListenerComponent.h"
 #include "PhysicsPublic.h"
+#include "Physics/Experimental/ChaosEventType.h"
 #include "ChaosNotifyHandlerInterface.h"
 #include "ChaosGameplayEventDispatcher.generated.h"
 
@@ -19,47 +20,13 @@ namespace Chaos
 	struct FBreakingData;
 }
 
+// FChaosBreakEvent has been moved from ChaosSolverEngine to Engine, this structure need to be kept for backward compatibility.
 USTRUCT(BlueprintType)
-struct CHAOSSOLVERENGINE_API FChaosBreakEvent
+struct CHAOSSOLVERENGINE_API FChaosBreakEvent : public FBreakChaosEvent
 {
 	GENERATED_BODY()
-
-public:
-
-	FChaosBreakEvent();
-	FChaosBreakEvent(const Chaos::FBreakingData& BreakingData);
-
-	/** primitive component involved in the break event */
-	UPROPERTY(BlueprintReadOnly, Category = "Break Event")
-	TObjectPtr<UPrimitiveComponent> Component = nullptr;
-
-	/** World location of the break */
-	UPROPERTY(BlueprintReadOnly, Category = "Break Event")
-	FVector Location;
-
-	/** Linear Velocity of the breaking particle  */
-	UPROPERTY(BlueprintReadOnly, Category = "Break Event")
-	FVector Velocity;
-
-	/** Angular Velocity of the breaking particle  */
-	UPROPERTY(BlueprintReadOnly, Category = "Break Event")
-	FVector AngularVelocity;
-
-	/** Extents of the bounding box */
-	UPROPERTY(BlueprintReadOnly, Category = "Break Event")
-	FVector Extents;
-
-	/** Mass of the breaking particle  */
-	UPROPERTY(BlueprintReadOnly, Category = "Break Event")
-	float Mass;
-
-	/** Index of the geometry collection bone if positive */
-	UPROPERTY(BlueprintReadOnly, Category = "Break Event")
-	int32 Index;
-
-	/** Whether the break event originated from a crumble event */
-	UPROPERTY(BlueprintReadOnly, Category = "Break Event")
-	bool bFromCrumble;
+	FChaosBreakEvent() : FBreakChaosEvent() {}
+	FChaosBreakEvent(const Chaos::FBreakingData& BreakingData) : FBreakChaosEvent(BreakingData) {}
 };
 
 USTRUCT(BlueprintType)
diff --git a/Engine/Source/Runtime/Experimental/GeometryCollectionEngine/Private/GeometryCollection/GeometryCollectionComponent.cpp b/Engine/Source/Runtime/Experimental/GeometryCollectionEngine/Private/GeometryCollection/GeometryCollectionComponent.cpp
--- a/Engine/Source/Runtime/Experimental/GeometryCollectionEngine/Private/GeometryCollection/GeometryCollectionComponent.cpp
+++ b/Engine/Source/Runtime/Experimental/GeometryCollectionEngine/Private/GeometryCollection/GeometryCollectionComponent.cpp
@@ -344,6 +344,7 @@ UGeometryCollectionComponent::UGeometryCollectionComponent(const FObjectInitiali
 	, bNotifyCollisions(false)
 	, bNotifyRemovals(false)
 	, bNotifyCrumblings(false)
+	, bNotifyGlobalBreaks(false)
 	, bStoreVelocities(false)
 	, bShowBoneColors(false)
 	, bUseRootProxyForNavigation(false)
@@ -827,6 +828,18 @@ void UGeometryCollectionComponent::SetNotifyCrumblings(bool bNewNotifyCrumblings
 	}
 }
 
+void UGeometryCollectionComponent::SetNotifyGlobalBreaks(bool bNewNotifyGlobalBreaks)
+{
+	if (bNotifyGlobalBreaks != bNewNotifyGlobalBreaks)
+	{
+		if (PhysicsProxy)
+		{
+			PhysicsProxy->SetNotifyGlobalBreakings_External(bNewNotifyGlobalBreaks);
+		}
+		bNotifyGlobalBreaks = bNewNotifyGlobalBreaks;
+	}
+}
+
 FBodyInstance* UGeometryCollectionComponent::GetBodyInstance(FName BoneName /*= NAME_None*/, bool bGetWelded /*= true*/, int32 Index /*=INDEX_NONE*/) const
 {
 	return nullptr;// const_cast<FBodyInstance*>(&DummyBodyInstance);
@@ -1463,19 +1476,19 @@ void UGeometryCollectionComponent::DispatchChaosPhysicsCollisionBlueprintEvents(
 // call when first registering
 void UGeometryCollectionComponent::RegisterForEvents()
 {
-	if (EventDispatcher)
+	Chaos::FPhysicsSolver* Solver = GetWorld()->GetPhysicsScene()->GetSolver();
+	if (Solver)
 	{
-		// The new interested proxy system relies on us having having a proxy registered with the scene at the time we register Chaos events.
-		// Hence we need to do a double-take here to force the re-registration of chaos events. This is guaranteed to happen every time we
-		// re-create the physics proxy (e.g. in the case of the construction script in PIE). In that case, UChaosGameplayEventDispatcher::OnRegister
-		// will run *before* the new physics proxy is created.
-		EventDispatcher->UnregisterChaosEvents();
-		EventDispatcher->RegisterChaosEvents();
-
-		if (BodyInstance.bNotifyRigidBodyCollision || bNotifyBreaks || bNotifyCollisions || bNotifyRemovals || bNotifyCrumblings)
+		if (EventDispatcher)
 		{
-			Chaos::FPhysicsSolver* Solver = GetWorld()->GetPhysicsScene()->GetSolver();
-			if (Solver)
+			// The new interested proxy system relies on us having having a proxy registered with the scene at the time we register Chaos events.
+			// Hence we need to do a double-take here to force the re-registration of chaos events. This is guaranteed to happen every time we
+			// re-create the physics proxy (e.g. in the case of the construction script in PIE). In that case, UChaosGameplayEventDispatcher::OnRegister
+			// will run *before* the new physics proxy is created.
+			EventDispatcher->UnregisterChaosEvents();
+			EventDispatcher->RegisterChaosEvents();
+
+			if (BodyInstance.bNotifyRigidBodyCollision || bNotifyBreaks || bNotifyCollisions || bNotifyRemovals || bNotifyCrumblings)
 			{
 				if (bNotifyCollisions || BodyInstance.bNotifyRigidBodyCollision)
 				{
@@ -1490,12 +1503,10 @@ void UGeometryCollectionComponent::RegisterForEvents()
 				if (bNotifyBreaks)
 				{
 					EventDispatcher->RegisterForBreakEvents(this, &DispatchGeometryCollectionBreakEvent);
-
 					Solver->EnqueueCommandImmediate([Solver]()
 						{
 							Solver->SetGenerateBreakingData(true);
 						});
-
 				}
 
 				if (bNotifyRemovals)
@@ -1520,6 +1531,13 @@ void UGeometryCollectionComponent::RegisterForEvents()
 				}
 			}
 		}
+		if (bNotifyGlobalBreaks)
+		{
+			Solver->EnqueueCommandImmediate([Solver]()
+				{
+					Solver->SetGenerateBreakingData(true);
+				});
+		}
 	}
 }
 
@@ -2832,7 +2850,8 @@ void UGeometryCollectionComponent::RegisterAndInitializePhysicsProxy()
 		SimulationParameters.StartAwake = BodyInstance.bStartAwake;
 		SimulationParameters.CacheType = CacheParameters.CacheMode;
 		SimulationParameters.ReverseCacheBeginTime = CacheParameters.ReverseCacheBeginTime;
-		SimulationParameters.bGenerateBreakingData = bNotifyBreaks;
+		SimulationParameters.bGenerateBreakingData = bNotifyBreaks || bNotifyGlobalBreaks;
+		SimulationParameters.bDispatchGlobalBreakingData = bNotifyGlobalBreaks;
 		SimulationParameters.bGenerateCollisionData = bNotifyCollisions;
 		SimulationParameters.bGenerateTrailingData = bNotifyTrailing;
 		SimulationParameters.bGenerateRemovalsData = bNotifyRemovals;
diff --git a/Engine/Source/Runtime/Experimental/GeometryCollectionEngine/Public/GeometryCollection/GeometryCollectionComponent.h b/Engine/Source/Runtime/Experimental/GeometryCollectionEngine/Public/GeometryCollection/GeometryCollectionComponent.h
--- a/Engine/Source/Runtime/Experimental/GeometryCollectionEngine/Public/GeometryCollection/GeometryCollectionComponent.h
+++ b/Engine/Source/Runtime/Experimental/GeometryCollectionEngine/Public/GeometryCollection/GeometryCollectionComponent.h
@@ -988,6 +988,10 @@ public:
 	/** Changes whether or not this component will get future break notifications. */
 	UFUNCTION(BlueprintCallable, Category = "Physics")
 	void SetNotifyCrumblings(bool bNewNotifyCrumblings, bool bNewCrumblingEventIncludesChildren = false);
+
+	/** Changes whether or not this component will get future global break notifications. */
+	UFUNCTION(BlueprintCallable, Category = "Physics")
+	void SetNotifyGlobalBreaks(bool bNewNotifyGlobalBreaks);
 	
 	/** Overrideable native notification */
 	virtual void NotifyBreak(const FChaosBreakEvent& Event) {};
@@ -1099,6 +1103,10 @@ public:
 	UPROPERTY(EditAnywhere, BlueprintReadOnly, Category = "ChaosPhysics|Events")
 	bool bNotifyCrumblings;
 
+	/** If true, this component will generate breaking events that will be listened by the global event relay. */
+	UPROPERTY(EditAnywhere, BlueprintReadOnly, Category = "ChaosPhysics|Events|Global")
+	bool bNotifyGlobalBreaks;
+
 	/** If this and bNotifyCrumblings are true, the crumbling events will contain released children indices. */
 	UPROPERTY(EditAnywhere, BlueprintReadOnly, Category = "ChaosPhysics|Events", meta = (EditCondition = "bNotifyCrumblings"))
 	bool bCrumblingEventIncludesChildren;
