Geometry collection : update socket API to also look for sockets in the root proxies if they are available

diff --git a/Engine/Source/Runtime/Experimental/GeometryCollectionEngine/Private/GeometryCollection/GeometryCollectionComponent.cpp b/Engine/Source/Runtime/Experimental/GeometryCollectionEngine/Private/GeometryCollection/GeometryCollectionComponent.cpp
--- a/Engine/Source/Runtime/Experimental/GeometryCollectionEngine/Private/GeometryCollection/GeometryCollectionComponent.cpp
+++ b/Engine/Source/Runtime/Experimental/GeometryCollectionEngine/Private/GeometryCollection/GeometryCollectionComponent.cpp
@@ -13,6 +13,7 @@
 #include "Engine/CollisionProfile.h"
 #include "Engine/Engine.h"
 #include "Engine/InstancedStaticMesh.h"
+#include "Engine/StaticMeshSocket.h"
 #include "Field/FieldSystemComponent.h"
 #include "GeometryCollection/Facades/CollectionHierarchyFacade.h"
 #include "GeometryCollection/GeometryCollection.h"
@@ -2369,66 +2370,130 @@ void UGeometryCollectionComponent::OnUpdateTransform(EUpdateTransformFlags Updat
 
 bool UGeometryCollectionComponent::HasAnySockets() const
 {
-	if (RestCollection && RestCollection->GetGeometryCollection())
+	bool bHasAnySocket= false;
+	if (RestCollection)
 	{
-		return (RestCollection->GetGeometryCollection()->BoneName.Num() > 0);
+		if (RestCollection->GetGeometryCollection())
+		{
+			if (RestCollection->GetGeometryCollection()->BoneName.Num() > 0)
+			{
+				bHasAnySocket = true;
+			}
+		}
+		if (!bHasAnySocket)
+		{
+			for (const TObjectPtr<UStaticMesh>& ProxyMesh : RestCollection->RootProxyData.ProxyMeshes)
+			{
+				if (ProxyMesh && ProxyMesh->Sockets.Num() > 0)
+				{
+					bHasAnySocket = true;
+					break;
+				}
+			}
+		}
 	}
-	return false;
+	return bHasAnySocket;
 }
 
 bool UGeometryCollectionComponent::DoesSocketExist(FName InSocketName) const
 {
-	if (RestCollection && RestCollection->GetGeometryCollection())
+	bool bFoundSocket = false;
+	if (RestCollection)
 	{
-		return RestCollection->GetGeometryCollection()->BoneName.Contains(InSocketName.ToString());
+		if (RestCollection->GetGeometryCollection())
+		{
+			bFoundSocket |= RestCollection->GetGeometryCollection()->BoneName.Contains(InSocketName.ToString());
+		}
+		if (!bFoundSocket)
+		{
+			for (const TObjectPtr<UStaticMesh>& ProxyMesh : RestCollection->RootProxyData.ProxyMeshes)
+			{
+				if (ProxyMesh)
+				{
+					if (UStaticMeshSocket* MeshSocket = ProxyMesh->FindSocket(InSocketName))
+					{
+						bFoundSocket = true;
+						break;
+					}
+				}
+			}
+		}
 	}
-	return false;
+	return bFoundSocket;
 }
 
 FTransform UGeometryCollectionComponent::GetSocketTransform(FName InSocketName, ERelativeTransformSpace TransformSpace) const
 {
-	if (RestCollection && RestCollection->GetGeometryCollection())
+	if (RestCollection)
 	{
-		const TSharedPtr<FGeometryCollection, ESPMode::ThreadSafe> Collection = RestCollection->GetGeometryCollection();
-		const int32 TransformIndex = Collection->BoneName.Find(InSocketName.ToString());
-		if (TransformIndex != INDEX_NONE)
+		bool bFoundSocket = false;
+		FTransform SocketComponentSpaceTransform{ FTransform::Identity };
+		FTransform ParentComponentSpaceTransform{ FTransform::Identity };
+		if (RestCollection->GetGeometryCollection())
 		{
-			
-			if (ComponentSpaceTransforms.IsValidIndex(TransformIndex))
+			const TSharedPtr<FGeometryCollection, ESPMode::ThreadSafe> Collection = RestCollection->GetGeometryCollection();
+			if (Collection)
 			{
-				const FTransform BoneComponentSpaceTransform = ComponentSpaceTransforms[TransformIndex];
-				switch (TransformSpace)
+				const int32 TransformIndex = Collection->BoneName.Find(InSocketName.ToString());
+				if (TransformIndex != INDEX_NONE)
 				{
-				case RTS_World:
-					return BoneComponentSpaceTransform * GetComponentTransform();
-
-				case RTS_Actor:
-					{
-						if (const AActor* Actor = GetOwner())
-						{
-							const FTransform SocketWorldSpaceTransform = BoneComponentSpaceTransform * GetComponentTransform();
-							return SocketWorldSpaceTransform.GetRelativeTransform(Actor->GetTransform());
-						}
-						break;
-					}
-
-				case RTS_Component:
-					return BoneComponentSpaceTransform;
-					
-				case RTS_ParentBoneSpace:
+					if (ComponentSpaceTransforms.IsValidIndex(TransformIndex))
 					{
+						SocketComponentSpaceTransform = ComponentSpaceTransforms[TransformIndex];
 						const int32 ParentTransformIndex = Collection->Parent[TransformIndex];
-						FTransform ParentComponentSpaceTransform{ FTransform::Identity };
 						if (ComponentSpaceTransforms.IsValidIndex(ParentTransformIndex))
 						{
 							ParentComponentSpaceTransform = ComponentSpaceTransforms[ParentTransformIndex];
 						}
-						return BoneComponentSpaceTransform.GetRelativeTransform(ParentComponentSpaceTransform);
+						bFoundSocket = true;
 					}
+				}
+			}
+		}
+		if (!bFoundSocket)
+		{
+			for (const TObjectPtr<UStaticMesh>& ProxyMesh : RestCollection->RootProxyData.ProxyMeshes)
+			{
+				if (ProxyMesh)
+				{
+					if (UStaticMeshSocket* MeshSocket = ProxyMesh->FindSocket(InSocketName))
+					{
+						SocketComponentSpaceTransform = FTransform(MeshSocket->RelativeRotation, MeshSocket->RelativeLocation, MeshSocket->RelativeScale);
+						bFoundSocket = true;
+						break;
+					}
+				}
+			}
+		}
+		
+		if (bFoundSocket)
+		{
+			// convert to the right space 
+			switch (TransformSpace)
+			{
+			case RTS_World:
+				return SocketComponentSpaceTransform * GetComponentTransform();
 
-				default:
-					check(false);
+			case RTS_Actor:
+			{
+				if (const AActor* Actor = GetOwner())
+				{
+					const FTransform SocketWorldSpaceTransform = SocketComponentSpaceTransform * GetComponentTransform();
+					return SocketWorldSpaceTransform.GetRelativeTransform(Actor->GetTransform());
 				}
+				break;
+			}
+
+			case RTS_Component:
+				return SocketComponentSpaceTransform;
+
+			case RTS_ParentBoneSpace:
+			{
+				return SocketComponentSpaceTransform.GetRelativeTransform(ParentComponentSpaceTransform);
+			}
+
+			default:
+				check(false);
 			}
 		}
 	}
@@ -2437,13 +2502,32 @@ FTransform UGeometryCollectionComponent::GetSocketTransform(FName InSocketName,
 
 void UGeometryCollectionComponent::QuerySupportedSockets(TArray<FComponentSocketDescription>& OutSockets) const
 {
-	if (RestCollection && RestCollection->GetGeometryCollection())
+	if (RestCollection)
 	{
-		for (const FString& BoneName: RestCollection->GetGeometryCollection()->BoneName)
+		if (RestCollection->GetGeometryCollection())
+		{
+			for (const FString& BoneName : RestCollection->GetGeometryCollection()->BoneName)
+			{
+				FComponentSocketDescription& Desc = OutSockets.AddZeroed_GetRef();
+				Desc.Name = *BoneName;
+				Desc.Type = EComponentSocketType::Type::Bone;
+			}
+		}
+
+		for (const TObjectPtr<UStaticMesh>& ProxyMesh : RestCollection->RootProxyData.ProxyMeshes)
 		{
-			FComponentSocketDescription& Desc = OutSockets.AddZeroed_GetRef();
-			Desc.Name = *BoneName;
-			Desc.Type = EComponentSocketType::Type::Bone;
+			if (ProxyMesh)
+			{
+				for (const TObjectPtr<class UStaticMeshSocket>& MeshSocket : ProxyMesh->Sockets)
+				{
+					if (MeshSocket)
+					{
+						FComponentSocketDescription& Desc = OutSockets.AddZeroed_GetRef();
+						Desc.Name = MeshSocket->SocketName;
+						Desc.Type = EComponentSocketType::Type::Socket;
+					}
+				}
+			}
 		}
 	}
 }
