Iris - Fix most unsafe type casts in IrisCore.

#rnx

diff --git a/Engine/Source/Runtime/Experimental/Iris/Core/Private/Iris/Core/IrisDebugging.cpp b/Engine/Source/Runtime/Experimental/Iris/Core/Private/Iris/Core/IrisDebugging.cpp
--- a/Engine/Source/Runtime/Experimental/Iris/Core/Private/Iris/Core/IrisDebugging.cpp
+++ b/Engine/Source/Runtime/Experimental/Iris/Core/Private/Iris/Core/IrisDebugging.cpp
@@ -153,7 +153,7 @@ void SetIrisDebugObjectName(const ANSICHAR* NameBuffer)
 	}
 }
 
-void SetIrisDebugNetRefHandle(uint32 NetRefHandleId)
+void SetIrisDebugNetRefHandle(uint64 NetRefHandleId)
 {
 	IrisDebugHelperInternal::GIrisDebugNetRefHandle = Private::FNetRefHandleManager::MakeNetRefHandleFromId(NetRefHandleId);
 }
@@ -391,7 +391,7 @@ FNetReplicatedObjectDebugInfo DebugNetRefHandle(FNetRefHandle Handle)
 	return DebugNetRefHandleById(Handle.GetId(), Handle.GetReplicationSystemId());
 }
 
-FNetReplicatedObjectDebugInfo DebugNetRefHandleById(uint32 NetRefHandleId, uint32 ReplicationSystemId)
+FNetReplicatedObjectDebugInfo DebugNetRefHandleById(uint64 NetRefHandleId, uint32 ReplicationSystemId)
 {
 	using namespace UE::Net::Private;
 	FNetReplicatedObjectDebugInfo Info = {};
diff --git a/Engine/Source/Runtime/Experimental/Iris/Core/Private/Iris/DataStream/DataStreamManager.cpp b/Engine/Source/Runtime/Experimental/Iris/Core/Private/Iris/DataStream/DataStreamManager.cpp
--- a/Engine/Source/Runtime/Experimental/Iris/Core/Private/Iris/DataStream/DataStreamManager.cpp
+++ b/Engine/Source/Runtime/Experimental/Iris/Core/Private/Iris/DataStream/DataStreamManager.cpp
@@ -244,8 +244,8 @@ UDataStreamManager::EWriteResult UDataStreamManager::FImpl::WriteData(UE::Net::F
 	using namespace UE::Net;
 	using namespace UE::Net::Private;
 
-	const SIZE_T StreamCount = Streams.Num();
-	if (StreamCount == 0)
+	const int32 StreamCount = Streams.Num();
+	if (StreamCount <= 0)
 	{
 		return EWriteResult::NoData;
 	}
diff --git a/Engine/Source/Runtime/Experimental/Iris/Core/Private/Iris/ReplicationState/ReplicationStateDescriptorBuilder.cpp b/Engine/Source/Runtime/Experimental/Iris/Core/Private/Iris/ReplicationState/ReplicationStateDescriptorBuilder.cpp
--- a/Engine/Source/Runtime/Experimental/Iris/Core/Private/Iris/ReplicationState/ReplicationStateDescriptorBuilder.cpp
+++ b/Engine/Source/Runtime/Experimental/Iris/Core/Private/Iris/ReplicationState/ReplicationStateDescriptorBuilder.cpp
@@ -136,15 +136,20 @@ public:
 
 	struct FMemberProperty
 	{
-		const FProperty* Property;												// The property
-		const UFunction* PropertyRepNotifyFunction;								// RepNotifyFunction
-		const FPropertyNetSerializerInfo* SerializerInfo;						// SerializerInfo
-		CreateAndRegisterReplicationFragmentFunc CreateAndRegisterReplicationFragmentFunction; // Optional create method for fragment of properties that supports separate states (FastArrays)
-		EMemberPropertyTraits Traits;											// Traits for the property
-		ELifetimeCondition ReplicationCondition;								// Condition for when and where the property should be replicated
-		uint16 ChangeMaskBits;													// Bits needed to track dirtiness for this property. Init state properties use no bits.
-		FName ChangeMaskGroupName;												// We allow members to specify a changemask group to facilitate sharing of the same changemask bit
-		FMemoryLayoutUtil::FSizeAndAlignment ExternalSizeAndAlignment;			// If we do not have a property we need to pass this in explicitly
+		const FProperty* Property = nullptr;
+		const UFunction* PropertyRepNotifyFunction = nullptr;
+		const FPropertyNetSerializerInfo* SerializerInfo = nullptr;
+		// Optional create method for fragment of properties that supports separate states, such as FastArraySerializers and custom struct serializers with custom fragment.
+		CreateAndRegisterReplicationFragmentFunc CreateAndRegisterReplicationFragmentFunction = nullptr;
+		EMemberPropertyTraits Traits = EMemberPropertyTraits::None;
+		// Condition for when and where the property should be replicated
+		ELifetimeCondition ReplicationCondition = ELifetimeCondition::COND_None;
+		// Number of bits needed to track dirtiness for this property. Init state properties ignores this.
+		uint16 ChangeMaskBits = 0;
+		// We allow members to specify a changemask group to facilitate sharing of the same changemask bit
+		FName ChangeMaskGroupName;
+		// For members without a corresponding property this needs to be set explicitly
+		FMemoryLayoutUtil::FSizeAndAlignment ExternalSizeAndAlignment;
 	};
 
 	struct FMemberFunction
@@ -232,8 +237,9 @@ private:
 	struct FMemberTagCacheEntry
 	{
 		FRepTag Tag;
-		uint32 MemberIndex;
-		uint32 InDescriptor : 1; // Tags can be found in the member's ReplicationStateDescriptor
+		uint16 MemberIndex;
+		// If set then tags can be found in the member's ReplicationStateDescriptor
+		bool bInDescriptor;
 	};
 
 	struct FMemberTagCache
@@ -637,7 +643,7 @@ void FPropertyReplicationStateDescriptorBuilder::BuildMemberTagCache(FMemberTagC
 {
 	const FMemberProperty* MemberProperties = Members.GetData();
 
-	uint32 EntryIndex = 0;
+	uint16 EntryIndex = 0;
 	for (const FMemberCacheEntry& MemberCacheEntry : MemberCache)
 	{
 		const FMemberProperty& MemberProperty = MemberProperties[EntryIndex];
@@ -649,8 +655,8 @@ void FPropertyReplicationStateDescriptorBuilder::BuildMemberTagCache(FMemberTagC
 			{
 				FMemberTagCacheEntry TagEntry;
 				TagEntry.Tag = Tag;
-				TagEntry.InDescriptor = 0U;
 				TagEntry.MemberIndex = EntryIndex;
+				TagEntry.bInDescriptor = false;
 
 				++MemberTagCache.TagCount;
 				MemberTagCache.CacheEntries.Add(TagEntry);
@@ -667,8 +673,8 @@ void FPropertyReplicationStateDescriptorBuilder::BuildMemberTagCache(FMemberTagC
 			{
 				FMemberTagCacheEntry TagEntry;
 				TagEntry.Tag = 0;
-				TagEntry.InDescriptor = 1U;
 				TagEntry.MemberIndex = EntryIndex;
+				TagEntry.bInDescriptor = true;
 
 				MemberTagCache.TagCount += Descriptor->TagCount;
 				MemberTagCache.CacheEntries.Add(TagEntry);
@@ -789,10 +795,11 @@ bool FPropertyReplicationStateDescriptorBuilder::CalculateDefaultStateChecksum(c
 	const UPartialNetObjectAttachmentHandlerConfig* PartialNetObjectAttachmentHandlerConfig = GetDefault<UPartialNetObjectAttachmentHandlerConfig>();
 
 	constexpr uint32 SmallBufferSize = 1024;
-	const uint32 MaxBufferSize = PartialNetObjectAttachmentHandlerConfig->GetTotalMaxPayloadBitCount()/8U;
+	const uint64 MaxBufferSize = PartialNetObjectAttachmentHandlerConfig->GetTotalMaxPayloadBitCount()/8U;
+	check(MaxBufferSize/4 < std::numeric_limits<int32>::max());
 
 	// Setup small buffer first as it might be written to when the writer's destructor is called.
-	TArray<uint32, TInlineAllocator<SmallBufferSize/sizeof(uint32)> > TempBuffer;
+	TArray<uint32, TInlineAllocator<SmallBufferSize/sizeof(uint32)>> TempBuffer;
 	TempBuffer.SetNum(SmallBufferSize/sizeof(uint32));
 
 	// We setup a default context
@@ -808,8 +815,8 @@ bool FPropertyReplicationStateDescriptorBuilder::CalculateDefaultStateChecksum(c
 	// Try again with largest supported buffer
 	if (Writer.IsOverflown())
 	{
-		TempBuffer.SetNum(MaxBufferSize/sizeof(uint32));
-		Writer.InitBytes((uint8*)TempBuffer.GetData(), MaxBufferSize);
+		TempBuffer.SetNum(static_cast<int32>(static_cast<int64>(MaxBufferSize/sizeof(uint32))));
+		Writer.InitBytes((uint8*)TempBuffer.GetData(), static_cast<uint32>(MaxBufferSize));
 		FReplicationStateOperations::Serialize(Context, Descriptor->DefaultStateBuffer, Descriptor);
 	}
 
@@ -854,7 +861,7 @@ void FPropertyReplicationStateDescriptorBuilder::BuildMemberDescriptors(FReplica
 
 		ExternalBufferAlignment = FMath::Max(ExternalBufferAlignment, ExternalMemberAlignment);
 		ExternalOffset = Align(ExternalOffset, ExternalMemberAlignment);
-		CurrentMemberDescriptor->ExternalMemberOffset = ExternalOffset;
+		CurrentMemberDescriptor->ExternalMemberOffset = static_cast<uint32>(ExternalOffset);
 		ExternalOffset += ExternalMemberSize;
 
 		// Internal
@@ -864,7 +871,7 @@ void FPropertyReplicationStateDescriptorBuilder::BuildMemberDescriptors(FReplica
 
 		InternalBufferAlignment = FMath::Max(InternalBufferAlignment, InternalMemberAlignment);
 		InternalOffset = Align(InternalOffset, InternalMemberAlignment);
-		CurrentMemberDescriptor->InternalMemberOffset = InternalOffset;
+		CurrentMemberDescriptor->InternalMemberOffset = static_cast<uint32>(InternalOffset);
 		InternalOffset += InternalMemberSize;
 
 		++CurrentMemberDescriptor;
@@ -873,7 +880,7 @@ void FPropertyReplicationStateDescriptorBuilder::BuildMemberDescriptors(FReplica
 
 	// Fill in descriptor
 	Descriptor->MemberDescriptors = MemberDescriptors;
-	Descriptor->MemberCount = Members.Num();
+	Descriptor->MemberCount = static_cast<uint16>(Members.Num());
 
 	// Update Context
 	Context.External.Alignment = ExternalBufferAlignment;
@@ -922,7 +929,7 @@ void FPropertyReplicationStateDescriptorBuilder::BuildMemberDescriptorsForStruct
 
 		InternalBufferAlignment = FMath::Max(InternalBufferAlignment, InternalMemberAlignment);
 		InternalOffset = Align(InternalOffset, InternalMemberAlignment);
-		CurrentMemberDescriptor->InternalMemberOffset = InternalOffset;
+		CurrentMemberDescriptor->InternalMemberOffset = static_cast<uint32>(InternalOffset);
 		InternalOffset += InternalMemberSize;
 
 		++CurrentMemberDescriptor;
@@ -931,7 +938,7 @@ void FPropertyReplicationStateDescriptorBuilder::BuildMemberDescriptorsForStruct
 
 	// Fill in descriptor
 	Descriptor->MemberDescriptors = MemberDescriptors;
-	Descriptor->MemberCount = Members.Num();
+	Descriptor->MemberCount = static_cast<uint16>(Members.Num());
 
 	// Update Context
 	Context.External.Alignment = FMath::Max(Context.External.Alignment, ExternalBufferAlignment);
@@ -972,13 +979,13 @@ void FPropertyReplicationStateDescriptorBuilder::BuildMemberChangeMaskDescriptor
 
 	for (const FMemberProperty& Info : Members)
 	{
-		const uint32 BitCount = (bIsInitState ? 0U : Info.ChangeMaskBits);
+		const uint16 BitCount = (bIsInitState ? 0U : Info.ChangeMaskBits);
 		const FName ChangeMaskGroupName = Info.ChangeMaskGroupName;
 
 		if (ChangeMaskGroupName.IsNone())
 		{
 			CurrentMemberChangeMaskDescriptor->BitCount = BitCount;
-			CurrentMemberChangeMaskDescriptor->BitOffset = (bIsInitState ? 0U : CurrentBitOffset);
+			CurrentMemberChangeMaskDescriptor->BitOffset = static_cast<uint16>(bIsInitState ? 0U : CurrentBitOffset);
 			CurrentBitOffset += BitCount;
 		}
 		else if (FChangeMaskGroupInfo* SharedChangeMaskGroupInfo = ChangemaskGroups.FindByPredicate([ChangeMaskGroupName](const FChangeMaskGroupInfo& GroupInfo) { return GroupInfo.ChangeMaskGroupName == ChangeMaskGroupName; }))
@@ -988,7 +995,7 @@ void FPropertyReplicationStateDescriptorBuilder::BuildMemberChangeMaskDescriptor
 		else
 		{
 			CurrentMemberChangeMaskDescriptor->BitCount = BitCount;
-			CurrentMemberChangeMaskDescriptor->BitOffset = (bIsInitState ? 0U : CurrentBitOffset);
+			CurrentMemberChangeMaskDescriptor->BitOffset = static_cast<uint16>(bIsInitState ? 0U : CurrentBitOffset);
 			CurrentBitOffset += BitCount;
 
 			ChangemaskGroups.Add({ ChangeMaskGroupName, *CurrentMemberChangeMaskDescriptor });
@@ -1004,8 +1011,10 @@ void FPropertyReplicationStateDescriptorBuilder::BuildMemberChangeMaskDescriptor
 	// If state has custom conditionals we need an additional changemask.
 	const SIZE_T ConditionalChangeMaskMemberSize = (EnumHasAnyFlags(Descriptor->Traits, EReplicationStateTraits::HasLifetimeConditionals) ? ChangeMaskMemberSize : 0U);
 
-	Descriptor->ChangeMaskBitCount = CurrentBitOffset;
-	Descriptor->ChangeMasksExternalOffset = Align(Context.External.Size, ChangeMaskMemberAlignment);
+	checkSlow(CurrentBitOffset <= std::numeric_limits<uint16>::max());
+
+	Descriptor->ChangeMaskBitCount = static_cast<uint16>(CurrentBitOffset);
+	Descriptor->ChangeMasksExternalOffset = static_cast<uint32>(Align(Context.External.Size, ChangeMaskMemberAlignment));
 
 	// Update context
 	Context.External.Alignment = FMath::Max(Context.External.Alignment, ChangeMaskMemberAlignment);
@@ -1059,11 +1068,11 @@ void FPropertyReplicationStateDescriptorBuilder::BuildMemberTagDescriptors(FRepl
 	{
 		const FMemberCacheEntry& MemberEntry = MemberCacheEntries[TagEntry.MemberIndex];
 		const FReplicationStateMemberDescriptor& MemberDescriptor = MemberDescriptors[TagEntry.MemberIndex];
-		if (TagEntry.InDescriptor)
+		if (TagEntry.bInDescriptor)
 		{
 			// Copy tags from descriptor
 			const FReplicationStateDescriptor* MemberStateDescriptor = MemberEntry.Descriptor; 
-			for (uint32 TagIt = 0, TagEndIt = MemberStateDescriptor->TagCount; TagIt != TagEndIt; ++TagIt)
+			for (uint16 TagIt = 0, TagEndIt = MemberStateDescriptor->TagCount; TagIt != TagEndIt; ++TagIt)
 			{
 				CurrentMemberTagDescriptor->Tag = MemberStateDescriptor->MemberTagDescriptors[TagIt].Tag;
 				CurrentMemberTagDescriptor->MemberIndex = TagEntry.MemberIndex;
@@ -1081,6 +1090,7 @@ void FPropertyReplicationStateDescriptorBuilder::BuildMemberTagDescriptors(FRepl
 	}
 
 	// Fill in descriptor
+	ensureAlwaysMsgf(Context.MemberTagCache.TagCount < MAX_uint16 - 1U, TEXT("Tag count cannot exceed %u. Building descriptor %s."), MAX_uint16 - 1U, ToCStr(Descriptor->DebugName));
 	Descriptor->MemberTagDescriptors = MemberTagDescriptors;
 	Descriptor->TagCount = static_cast<uint16>(Context.MemberTagCache.TagCount);
 }
@@ -1103,7 +1113,7 @@ void FPropertyReplicationStateDescriptorBuilder::BuildMemberReferenceDescriptors
 			{
 				// Flatten out references from nested members with static storage
 				const FReplicationStateDescriptor* MemberStateDescriptor = MemberEntry.Descriptor; 
-				for (uint32 InnerIt = 0, InnerEndIt = MemberStateDescriptor->ObjectReferenceCount; InnerIt != InnerEndIt; ++InnerIt)
+				for (uint16 InnerIt = 0, InnerEndIt = MemberStateDescriptor->ObjectReferenceCount; InnerIt != InnerEndIt; ++InnerIt)
 				{
 					// Copy inner info
 					CurrentMemberReferenceDescriptor->Info = MemberStateDescriptor->MemberReferenceDescriptors[InnerIt].Info;
@@ -1111,7 +1121,7 @@ void FPropertyReplicationStateDescriptorBuilder::BuildMemberReferenceDescriptors
 					// Update offset to point to inner offset
 					CurrentMemberReferenceDescriptor->Offset = MemberDescriptors[Entry.MemberIndex].InternalMemberOffset + MemberStateDescriptor->MemberReferenceDescriptors[InnerIt].Offset;
 
-					CurrentMemberReferenceDescriptor->MemberIndex = Entry.MemberIndex;
+					CurrentMemberReferenceDescriptor->MemberIndex = static_cast<uint16>(Entry.MemberIndex);
 					CurrentMemberReferenceDescriptor->InnerReferenceIndex = InnerIt;
 
 					++CurrentMemberReferenceDescriptor;
@@ -1126,7 +1136,7 @@ void FPropertyReplicationStateDescriptorBuilder::BuildMemberReferenceDescriptors
 				// Offset to member containing data about the array
 				CurrentMemberReferenceDescriptor->Offset = MemberDescriptors[Entry.MemberIndex].InternalMemberOffset;
 	
-				CurrentMemberReferenceDescriptor->MemberIndex = Entry.MemberIndex;
+				CurrentMemberReferenceDescriptor->MemberIndex = static_cast<uint16>(Entry.MemberIndex);
 				CurrentMemberReferenceDescriptor->InnerReferenceIndex = MAX_uint16;
 
 				++CurrentMemberReferenceDescriptor;
@@ -1139,15 +1149,16 @@ void FPropertyReplicationStateDescriptorBuilder::BuildMemberReferenceDescriptors
 			// Offset to reference
 			CurrentMemberReferenceDescriptor->Offset = MemberDescriptors[Entry.MemberIndex].InternalMemberOffset;
 
-			CurrentMemberReferenceDescriptor->MemberIndex = Entry.MemberIndex;
+			CurrentMemberReferenceDescriptor->MemberIndex = static_cast<uint16>(Entry.MemberIndex);
 			CurrentMemberReferenceDescriptor->InnerReferenceIndex = MAX_uint16;
 			++CurrentMemberReferenceDescriptor;
 		}
 	}
 
 	// Fill in descriptor
+	ensureAlwaysMsgf(Context.ReferenceCache.ReferenceCount < std::numeric_limits<uint16>::max(), TEXT("Object reference count cannot exceed %u. Building descriptor %s."), std::numeric_limits<uint16>::max(), ToCStr(Descriptor->DebugName));
 	Descriptor->MemberReferenceDescriptors = MemberReferenceDescriptors;
-	Descriptor->ObjectReferenceCount = Context.ReferenceCache.ReferenceCount;
+	Descriptor->ObjectReferenceCount = static_cast<uint16>(Context.ReferenceCache.ReferenceCount);
 }
 
 void FPropertyReplicationStateDescriptorBuilder::BuildMemberFunctionDescriptors(FReplicationStateMemberFunctionDescriptor* MemberFunctionDescriptors, FReplicationStateDescriptor* Descriptor, const FBuilderContext& Context) const
@@ -1164,8 +1175,9 @@ void FPropertyReplicationStateDescriptorBuilder::BuildMemberFunctionDescriptors(
 		UE_LOG_DESCRIPTORBUILDER(Verbose, TEXT("FPropertyReplicationStateDescriptorBuilder::BuildMemberFunctionDescriptors AddingFunction %s"), ToCStr(MemberFunction.Function->GetName()));		
 	}
 
-	Descriptor->MemberFunctionDescriptors = MemberFunctionDescriptors;
+	ensureAlwaysMsgf(Context.MemberFunctionCache.Num() <= std::numeric_limits<uint16>::max(), TEXT("Function count cannot exceed %u. Building descriptor %s."), std::numeric_limits<uint16>::max(), ToCStr(Descriptor->DebugName));
 	Descriptor->FunctionCount = static_cast<uint16>(Context.MemberFunctionCache.Num());
+	Descriptor->MemberFunctionDescriptors = MemberFunctionDescriptors;
 }
 
 void FPropertyReplicationStateDescriptorBuilder::BuildMemberProperties(const FProperty** MemberProperties, FReplicationStateDescriptor* Descriptor) const
@@ -1235,7 +1247,7 @@ void FPropertyReplicationStateDescriptorBuilder::BuildMemberRepIndexToMemberDesc
 		const SIZE_T MemberIndex = &Info - Members.GetData();
 		const uint16 RepIndex = Info.Property->RepIndex;
 		checkSlow(RepIndex < RepIndexCount);
-		MemberRepIndexToMemberIndexDescriptors[RepIndex].MemberIndex = MemberIndex;
+		MemberRepIndexToMemberIndexDescriptors[RepIndex].MemberIndex = static_cast<uint16>(MemberIndex);
 	}
 
 	Descriptor->MemberRepIndexToMemberIndexDescriptors = MemberRepIndexToMemberIndexDescriptors;
@@ -1323,7 +1335,7 @@ void FPropertyReplicationStateDescriptorBuilder::BuildMemberSerializerConfigs(ui
 			{
 				FArrayPropertyNetSerializerConfig* SerializerConfig = new (SerializerConfigBuffer) FArrayPropertyNetSerializerConfig();
 				SerializerConfig->MaxElementCount = 65535U;
-				SerializerConfig->ElementCountBitCount = GetBitsNeeded(SerializerConfig->MaxElementCount);
+				SerializerConfig->ElementCountBitCount = static_cast<uint16>(GetBitsNeeded(SerializerConfig->MaxElementCount));
 				SerializerConfig->Property = TFieldPath<FArrayProperty>(const_cast<FArrayProperty*>(CastField<FArrayProperty>(Member.Property)));
 				SerializerConfig->StateDescriptor = MemberCacheEntry->Descriptor;
 				MemberSerializerDescriptor->SerializerConfig = SerializerConfig;
@@ -1354,7 +1366,7 @@ void FPropertyReplicationStateDescriptorBuilder::FixupNetRoleNetSerializerConfig
 	{
 		if (SerializerDescriptor.Serializer == &UE_NET_GET_SERIALIZER(FNetRoleNetSerializer))
 		{
-			const uint32 MemberIndex = &SerializerDescriptor - MemberSerializerDescriptors;
+			const uint32 MemberIndex = static_cast<uint32>(&SerializerDescriptor - MemberSerializerDescriptors);
 			if (FirstNetRoleMemberIndex == ~0U)
 			{
 				FirstNetRoleMemberIndex = MemberIndex;
@@ -1391,7 +1403,7 @@ void FPropertyReplicationStateDescriptorBuilder::FixupDescriptorForNativeFastArr
 {
 	// Native fast arrays quantized directly from the src state data which is the fastarray itself so the offset must be set to zero.
 	const_cast<FReplicationStateMemberDescriptor*>(&Descriptor->MemberDescriptors[0])->ExternalMemberOffset = 0U;
-	Descriptor->ChangeMasksExternalOffset = GetFastArrayChangeMaskOffset(Descriptor->MemberProperties[0]);
+	Descriptor->ChangeMasksExternalOffset = static_cast<uint32>(GetFastArrayChangeMaskOffset(Descriptor->MemberProperties[0]));
 }
 
 EReplicationStateTraits FPropertyReplicationStateDescriptorBuilder::BuildReplicationStateTraits(const FBuilderContext& Context) const
@@ -1495,11 +1507,19 @@ EReplicationStateTraits FPropertyReplicationStateDescriptorBuilder::BuildReplica
 
 void FPropertyReplicationStateDescriptorBuilder::FinalizeDescriptor(FReplicationStateDescriptor* Descriptor, FBuilderContext& Context) const
 {
-	Descriptor->ExternalAlignment = Context.External.Alignment;
-	Descriptor->ExternalSize = Context.External.Size;
+	const SIZE_T AlignedExternalSize = Align(Context.External.Size, Context.External.Alignment);
+	const SIZE_T AlignedInternalSize = Align(Context.Internal.Size, Context.Internal.Alignment);
+
+	checkSlow(Context.External.Alignment < std::numeric_limits<uint16>::max());
+	checkSlow(Context.Internal.Alignment < std::numeric_limits<uint16>::max());
+	checkSlow(AlignedExternalSize < std::numeric_limits<uint32>::max());
+	checkSlow(AlignedInternalSize < std::numeric_limits<uint32>::max());
 
-	Descriptor->InternalAlignment = Context.Internal.Alignment;
-	Descriptor->InternalSize = Context.Internal.Size;	
+	Descriptor->ExternalAlignment = static_cast<uint16>(Context.External.Alignment);
+	Descriptor->ExternalSize = static_cast<uint32>(AlignedExternalSize);
+
+	Descriptor->InternalAlignment = static_cast<uint16>(Context.Internal.Alignment);
+	Descriptor->InternalSize = static_cast<uint32>(AlignedInternalSize);
 }
 
 TRefCountPtr<const FReplicationStateDescriptor>
@@ -1509,6 +1529,9 @@ FPropertyReplicationStateDescriptorBuilder::Build(const FString& StateName, FRep
 	const SIZE_T MemberCount = Members.Num();
 	const SIZE_T FunctionCount = Functions.Num();
 
+	checkSlow(MemberCount < std::numeric_limits<uint16>::max());
+	checkSlow(FunctionCount < std::numeric_limits<uint16>::max());
+
 	FBuilderContext Context;
 	Context.BuildParams = BuildParams;
 	BuildMemberCache(Context, DescriptorRegistry);
@@ -1609,7 +1632,7 @@ FPropertyReplicationStateDescriptorBuilder::Build(const FString& StateName, FRep
 	AddSerializerConfigMemoryBlockToLayout(Layout, LayoutData.MemberSerializerConfigSizeAndOffset, Context.MemberCache);
 
 	// Allocate required memory for descriptor and space for all sections
-	uint8* Buffer = (uint8*)FMemory::Malloc(Layout.CurrentOffset, Layout.MaxAlignment);
+	uint8* Buffer = (uint8*)FMemory::Malloc(Layout.CurrentOffset, static_cast<uint32>(Layout.MaxAlignment));
 
 	if (Buffer)
 	{
@@ -1619,6 +1642,9 @@ FPropertyReplicationStateDescriptorBuilder::Build(const FString& StateName, FRep
 
 		check((void*)Descriptor == (void*)Buffer);
 		
+		// DebugName is handy for logging.
+		Descriptor->DebugName = CreatePersistentNetDebugName(StateName.GetCharArray().GetData());
+
 		// Setup Traits early on
 		Descriptor->Traits = Traits;
 
@@ -1712,7 +1738,6 @@ FPropertyReplicationStateDescriptorBuilder::Build(const FString& StateName, FRep
 			Descriptor->BaseStruct = BaseStruct;
 		}
 
-		Descriptor->DebugName = CreatePersistentNetDebugName(StateName.GetCharArray().GetData());
 		if (LayoutData.MemberDebugDescriptorsSizeAndOffset.Size > 0)
 		{
 			BuildMemberDebugDescriptors(reinterpret_cast<FReplicationStateMemberDebugDescriptor*>(Buffer + LayoutData.MemberDebugDescriptorsSizeAndOffset.Offset), Descriptor, Context);
@@ -1736,10 +1761,6 @@ FPropertyReplicationStateDescriptorBuilder::Build(const FString& StateName, FRep
 		// Finalize the descriptor and fill in remaining information
 		FinalizeDescriptor(Descriptor, Context);
 
-		// Pad up the size to the alignment of the state
-		Descriptor->ExternalSize = Align(Descriptor->ExternalSize, Descriptor->ExternalAlignment);
-		Descriptor->InternalSize = Align(Descriptor->InternalSize, Descriptor->InternalAlignment);
-
 		// Setup function pointers to required methods
 		Descriptor->ConstructReplicationState = ConstructPropertyReplicationState;
 		Descriptor->DestructReplicationState = DestructPropertyReplicationState;
@@ -1833,9 +1854,9 @@ EMemberPropertyTraits FPropertyReplicationStateDescriptorBuilder::GetConnectionF
 
 	static_assert(TConditionPairIsSorted<>(ConditionPairs), "ConditionPairs order needs to be updated.");
 	
-	const int32 ConditionPairIndex = Algo::BinarySearchBy(ConditionPairs, Condition, [](const FConditionPair& Pair) { return Pair.Condition; });
-
-	if (ConditionPairIndex >= 0)
+	const SIZE_T ConditionPairIndex = Algo::BinarySearchBy(ConditionPairs, Condition, [](const FConditionPair& Pair) { return Pair.Condition; });
+	// BinarySearchBy can return INDEX_NONE (-1) but compiler believes the auto return value should be SIZE_T.
+	if (ConditionPairIndex < UE_ARRAY_COUNT(ConditionPairs))
 	{
 		Traits |= ConditionPairs[ConditionPairIndex].PropertyTrait;
 	}
@@ -2711,7 +2732,7 @@ SIZE_T FReplicationStateDescriptorBuilder::CreateDescriptorsForClass(FResult& Cr
 			Classes[--ClassIndex] = CurrentClass;
 		}
 
-		for (const UClass* CurrentClass : MakeArrayView(Classes + ClassIndex, MaxClassHierarchyDepth - ClassIndex))
+		for (const UClass* CurrentClass : MakeArrayView(Classes + ClassIndex, static_cast<int32>(MaxClassHierarchyDepth - ClassIndex)))
 		{
 			for (const UField* NetField : MakeArrayView(CurrentClass->NetFields))
 			{
diff --git a/Engine/Source/Runtime/Experimental/Iris/Core/Private/Iris/ReplicationState/ReplicationStateStorage.cpp b/Engine/Source/Runtime/Experimental/Iris/Core/Private/Iris/ReplicationState/ReplicationStateStorage.cpp
--- a/Engine/Source/Runtime/Experimental/Iris/Core/Private/Iris/ReplicationState/ReplicationStateStorage.cpp
+++ b/Engine/Source/Runtime/Experimental/Iris/Core/Private/Iris/ReplicationState/ReplicationStateStorage.cpp
@@ -390,7 +390,7 @@ FReplicationStateStorage::ObjectInfoIndexType FReplicationStateStorage::AllocPer
 		FPerObjectInfo& ObjectInfo = ObjectInfos[InfoIndex];
 		ObjectInfo = FPerObjectInfo();
 
-		return InfoIndex;
+		return static_cast<ObjectInfoIndexType>(InfoIndex);
 	}
 
 	return InvalidObjectInfoIndex;
diff --git a/Engine/Source/Runtime/Experimental/Iris/Core/Private/Iris/ReplicationSystem/AttachmentReplication.cpp b/Engine/Source/Runtime/Experimental/Iris/Core/Private/Iris/ReplicationSystem/AttachmentReplication.cpp
--- a/Engine/Source/Runtime/Experimental/Iris/Core/Private/Iris/ReplicationSystem/AttachmentReplication.cpp
+++ b/Engine/Source/Runtime/Experimental/Iris/Core/Private/Iris/ReplicationSystem/AttachmentReplication.cpp
@@ -193,7 +193,7 @@ bool FNetObjectAttachmentSendQueue::Enqueue(TArrayView<const TRefCountPtr<FNetBl
 	}
 	else
 	{
-		const uint32 TotalCountNeeded = UnreliableQueue.Count() + Attachments.Num();
+		const SIZE_T TotalCountNeeded = UnreliableQueue.Count() + Attachments.Num();
 		if (TotalCountNeeded > MaxUnreliableCount)
 		{
 			UE_LOG(LogIris, Verbose, TEXT("Dropping old RPC due to too many unreliable Attachments: %d max: %u"), UnreliableQueue.Count(), MaxUnreliableCount);
@@ -240,13 +240,13 @@ void FNetObjectAttachmentSendQueue::SetUnreliableQueueCapacity(uint32 QueueCapac
 	// MaxUnreliableCount is what prevents the queue from growing too large.
 	MaxUnreliableCount = QueueCapacity;
 	
-	const uint32 UnreliableCount = UnreliableQueue.Count();
+	const SIZE_T UnreliableCount = UnreliableQueue.Count();
 	if (QueueCapacity >= UnreliableCount)
 	{
 		return;
 	}
 
-	const uint32 DropCount = UnreliableCount - QueueCapacity;
+	const SIZE_T DropCount = UnreliableCount - QueueCapacity;
 	UE_LOG(LogIris, Warning, TEXT("Dropping %u attachments due to change in unreliable queue capacity to %u"), DropCount, QueueCapacity);
 	UnreliableQueue.PopNoCheck(DropCount);
 }
@@ -359,7 +359,7 @@ uint32 FNetObjectAttachmentSendQueue::SerializeUnreliable(FNetSerializationConte
 	uint32 SerializedUnreliableCount = 0;
 	uint32 PrevHasMoreAttachmentsWritePos = 0;
 	const bool bSerializeWithObject = RefHandle.IsValid();
-	for (uint32 AttachmentIt = 0, AttachmentEndIt = UnreliableQueue.Count(); AttachmentIt != AttachmentEndIt; ++AttachmentIt)
+	for (SIZE_T AttachmentIt = 0, AttachmentEndIt = UnreliableQueue.Count(); AttachmentIt != AttachmentEndIt; ++AttachmentIt)
 	{
 		FNetBitStreamRollbackScope RollbackScope(*Writer);
 		FNetExportRollbackScope ExportScope(Context);
@@ -745,13 +745,13 @@ void FNetObjectAttachmentReceiveQueue::SetUnreliableQueueCapacity(uint32 QueueCa
 {
 	MaxUnreliableCount = QueueCapacity;
 	
-	const uint32 UnreliableCount = UnreliableQueue.Count();
+	const SIZE_T UnreliableCount = UnreliableQueue.Count();
 	if (QueueCapacity >= UnreliableCount)
 	{
 		return;
 	}
 
-	const uint32 DropCount = UnreliableCount - QueueCapacity;
+	const SIZE_T DropCount = UnreliableCount - QueueCapacity;
 	UE_LOG(LogIris, Warning, TEXT("Dropping %u attachments to due to change in unreliable queue capacity to %u"), DropCount, QueueCapacity);
 	UnreliableQueue.Pop(DropCount);
 }
diff --git a/Engine/Source/Runtime/Experimental/Iris/Core/Private/Iris/ReplicationSystem/Conditionals/ReplicationConditionals.cpp b/Engine/Source/Runtime/Experimental/Iris/Core/Private/Iris/ReplicationSystem/Conditionals/ReplicationConditionals.cpp
--- a/Engine/Source/Runtime/Experimental/Iris/Core/Private/Iris/ReplicationSystem/Conditionals/ReplicationConditionals.cpp
+++ b/Engine/Source/Runtime/Experimental/Iris/Core/Private/Iris/ReplicationSystem/Conditionals/ReplicationConditionals.cpp
@@ -131,10 +131,10 @@ void FReplicationConditionals::InitPropertyCustomConditions(FInternalNetRefIndex
 	FReplicationFragment* const * Fragments = ReplicatedObjectData.InstanceProtocol->Fragments;
 	const FReplicationInstanceProtocol* InstanceProtocol = ReplicatedObjectData.InstanceProtocol;
 	const uint32 FragmentCount = InstanceProtocol->FragmentCount;
-	const SIZE_T FirstRelevantStateIndex = Protocol->FirstLifetimeConditionalsStateIndex;
+	const uint32 FirstRelevantStateIndex = Protocol->FirstLifetimeConditionalsStateIndex;
 
 	const UObject* LastOwner = nullptr;
-	for (const FReplicationStateDescriptor*& StateDescriptor : MakeArrayView(Protocol->ReplicationStateDescriptors + FirstRelevantStateIndex, Protocol->ReplicationStateCount - FirstRelevantStateIndex))
+	for (const FReplicationStateDescriptor*& StateDescriptor : MakeArrayView(Protocol->ReplicationStateDescriptors + FirstRelevantStateIndex, static_cast<int32>(Protocol->ReplicationStateCount - FirstRelevantStateIndex)))
 	{
 		if (EnumHasAnyFlags(StateDescriptor->Traits, EReplicationStateTraits::HasLifetimeConditionals))
 		{
@@ -241,8 +241,8 @@ bool FReplicationConditionals::SetPropertyCustomCondition(FInternalNetRefIndex O
 
 		const FReplicationInstanceProtocol* InstanceProtocol = ReplicatedObjectData.InstanceProtocol;
 
-		const SIZE_T FirstRelevantStateIndex = Protocol->FirstLifetimeConditionalsStateIndex;
-		for (const FReplicationStateDescriptor*& StateDescriptor : MakeArrayView(Protocol->ReplicationStateDescriptors + FirstRelevantStateIndex, Protocol->ReplicationStateCount - FirstRelevantStateIndex))
+		const uint32 FirstRelevantStateIndex = Protocol->FirstLifetimeConditionalsStateIndex;
+		for (const FReplicationStateDescriptor*& StateDescriptor : MakeArrayView(Protocol->ReplicationStateDescriptors + FirstRelevantStateIndex, static_cast<int32>(Protocol->ReplicationStateCount - FirstRelevantStateIndex)))
 		{
 			if (EnumHasAnyFlags(StateDescriptor->Traits, EReplicationStateTraits::HasLifetimeConditionals))
 			{
@@ -372,8 +372,8 @@ bool FReplicationConditionals::SetPropertyDynamicCondition(FInternalNetRefIndex
 
 		const FReplicationInstanceProtocol* InstanceProtocol = ReplicatedObjectData.InstanceProtocol;
 
-		const SIZE_T FirstRelevantStateIndex = Protocol->FirstLifetimeConditionalsStateIndex;
-		for (const FReplicationStateDescriptor*& StateDescriptor : MakeArrayView(Protocol->ReplicationStateDescriptors + FirstRelevantStateIndex, Protocol->ReplicationStateCount - FirstRelevantStateIndex))
+		const uint32 FirstRelevantStateIndex = Protocol->FirstLifetimeConditionalsStateIndex;
+		for (const FReplicationStateDescriptor*& StateDescriptor : MakeArrayView(Protocol->ReplicationStateDescriptors + FirstRelevantStateIndex, static_cast<int32>(Protocol->ReplicationStateCount - FirstRelevantStateIndex)))
 		{
 			if (EnumHasAnyFlags(StateDescriptor->Traits, EReplicationStateTraits::HasLifetimeConditionals))
 			{
diff --git a/Engine/Source/Runtime/Experimental/Iris/Core/Private/Iris/ReplicationSystem/DeltaCompression/DeltaCompressionBaselineManager.cpp b/Engine/Source/Runtime/Experimental/Iris/Core/Private/Iris/ReplicationSystem/DeltaCompression/DeltaCompressionBaselineManager.cpp
--- a/Engine/Source/Runtime/Experimental/Iris/Core/Private/Iris/ReplicationSystem/DeltaCompression/DeltaCompressionBaselineManager.cpp
+++ b/Engine/Source/Runtime/Experimental/Iris/Core/Private/Iris/ReplicationSystem/DeltaCompression/DeltaCompressionBaselineManager.cpp
@@ -84,7 +84,7 @@ void FDeltaCompressionBaselineManager::Init(FDeltaCompressionBaselineManagerInit
 		const uint32 BytesPerConnection = sizeof(FObjectBaselineInfo);
 		// The PerObjectInfo already contains the connection specific data for one connection,
 		// but we need to accommodate for MaxConnectionCount.
-		BytesPerObjectInfo = Align(sizeof(FPerObjectInfo) + BytesPerConnection*(FPlatformMath::Max(MaxConnectionCount, 1U) - 1U), alignof(FPerObjectInfo));
+		BytesPerObjectInfo = static_cast<uint32>(Align(sizeof(FPerObjectInfo) + BytesPerConnection*(FPlatformMath::Max(MaxConnectionCount, 1U) - 1U), alignof(FPerObjectInfo)));
 	}
 
 	BaselineCounts.SetNumZeroed(MaxDeltaCompressedObjectCount + 1U);
@@ -257,7 +257,7 @@ void FDeltaCompressionBaselineManager::UpdateDirtyStateMasks(const FChangeMaskCa
 
 		ChangeMaskStorageType*const ConnectionChangeMaskStoragePtr = ObjectInfo->ChangeMasksForConnections;
 		ChangeMaskStorageType const*const UpdatedChanges = DirtyChangeMaskStoragePtr + Entry.StorageOffset;
-		const SIZE_T ChangeMaskStridePerConnection = ObjectInfo->ChangeMaskStride;
+		const uint32 ChangeMaskStridePerConnection = ObjectInfo->ChangeMaskStride;
 		for (uint32 ConnectionId : ValidConnections)
 		{
 			const SIZE_T ChangeMaskOffset = ChangeMaskStridePerConnection*ConnectionId;
@@ -284,13 +284,14 @@ FDeltaCompressionBaselineManager::FPerObjectInfo* FDeltaCompressionBaselineManag
 		const FNetRefHandleManager::FReplicatedObjectData& ObjectData = NetRefHandleManager->GetReplicatedObjectDataNoCheck(ObjectIndex);
 		const FReplicationProtocol* Protocol = ObjectData.Protocol;
 		const uint32 ChangeMaskStride = Align(Protocol->ChangeMaskBitCount, sizeof(ChangeMaskStorageType)*8U)/(sizeof(ChangeMaskStorageType)*8U);
+		check(ChangeMaskStride < 65536U);
 		if (ChangeMaskStride > 0)
 		{
 			const uint32 AllocSize = ChangeMaskStride*3*MaxConnectionCount*sizeof(ChangeMaskStorageType);
 			ChangeMaskStorageType* ChangeMasksForConnections = static_cast<ChangeMaskStorageType*>(ChangeMaskAllocator.Alloc(AllocSize, alignof(ChangeMaskStorageType)));
 			FPlatformMemory::Memzero(ChangeMasksForConnections, AllocSize);
 
-			ObjectInfo->ChangeMaskStride = ChangeMaskStride;
+			ObjectInfo->ChangeMaskStride = static_cast<uint16>(ChangeMaskStride);
 			ObjectInfo->ChangeMasksForConnections = ChangeMasksForConnections;
 		}
 
@@ -372,7 +373,7 @@ void FDeltaCompressionBaselineManager::FreeAllPerObjectInfos()
 {
 	auto DestructObjectInfo = [this](uint32 ObjectInfoIndex)
 	{
-		FPerObjectInfo* ObjectInfo = GetPerObjectInfo(ObjectInfoIndex);
+		FPerObjectInfo* ObjectInfo = GetPerObjectInfo(static_cast<ObjectInfoIndexType>(ObjectInfoIndex));
 		DestructPerObjectInfo(ObjectInfo);
 	};
 
@@ -512,7 +513,8 @@ void FDeltaCompressionBaselineManager::AdjustBaselineCount(const FPerObjectInfo*
 		return;
 	}
 
-	const uint32 NewBaselineCount = BaselineCounts[ObjectInfoIndex] + Adjustment;
+	const uint16 NewBaselineCount = static_cast<uint16>(BaselineCounts[ObjectInfoIndex] + Adjustment);
+	UE_NET_DC_BASELINE_CHECK(NewBaselineCount <= BaselineCounts[ObjectInfoIndex]);
 	BaselineCounts[ObjectInfoIndex] = NewBaselineCount;
 
 	if (NewBaselineCount == 0)
@@ -529,8 +531,9 @@ void FDeltaCompressionBaselineManager::AdjustBaselineCount(const FPerObjectInfo*
 
 void FDeltaCompressionBaselineManager::ReleaseBaselinesForConnection(uint32 ConnId)
 {
-	auto ReleaseObjectBaselines = [this, ConnId](uint32 ObjectInfoIndex)
+	auto ReleaseObjectBaselines = [this, ConnId](uint32 InObjectInfoIndex)
 	{
+		const uint16 ObjectInfoIndex = static_cast<uint16>(InObjectInfoIndex);
 		FPerObjectInfo* ObjectInfo = GetPerObjectInfo(ObjectInfoIndex);
 		FObjectBaselineInfo& BaselineInfo = ObjectInfo->BaselinesForConnections[ConnId];
 
diff --git a/Engine/Source/Runtime/Experimental/Iris/Core/Private/Iris/ReplicationSystem/Filtering/NetObjectGridFilter.cpp b/Engine/Source/Runtime/Experimental/Iris/Core/Private/Iris/ReplicationSystem/Filtering/NetObjectGridFilter.cpp
--- a/Engine/Source/Runtime/Experimental/Iris/Core/Private/Iris/ReplicationSystem/Filtering/NetObjectGridFilter.cpp
+++ b/Engine/Source/Runtime/Experimental/Iris/Core/Private/Iris/ReplicationSystem/Filtering/NetObjectGridFilter.cpp
@@ -114,7 +114,7 @@ void UNetObjectGridFilter::Filter(FNetObjectFilteringParams& Params)
 		{
 			if ((PrevCell.Cell.X == CellAndTimestamp.Cell.X) & (PrevCell.Cell.Y == CellAndTimestamp.Cell.Y))
 			{
-				PrevCells.RemoveAtSwap(&PrevCell - PrevCells.GetData());
+				PrevCells.RemoveAtSwap(static_cast<int32>(&PrevCell - PrevCells.GetData()));
 				break;
 			}
 		}
@@ -368,11 +368,19 @@ void UNetObjectGridFilter::CalculateCellBox(const UNetObjectGridFilter::FPerObje
 	MinPosition = MinPosition.ComponentMax(Config->MinPos);
 	MaxPosition = MaxPosition.ComponentMin(Config->MaxPos);
 
+	const int64 MinX = FPlatformMath::FloorToInt(MinPosition.X/Config->CellSizeX);
+	const int64 MinY = FPlatformMath::FloorToInt(MinPosition.Y/Config->CellSizeY);
+	const int64 MaxX = FPlatformMath::FloorToInt(MaxPosition.X/Config->CellSizeX);
+	const int64 MaxY = FPlatformMath::FloorToInt(MaxPosition.Y/Config->CellSizeY);
+
 	FCellBox CellBox;
-	CellBox.MinX = FPlatformMath::FloorToInt(MinPosition.X/Config->CellSizeX);
-	CellBox.MinY = FPlatformMath::FloorToInt(MinPosition.Y/Config->CellSizeY);
-	CellBox.MaxX = FPlatformMath::FloorToInt(MaxPosition.X/Config->CellSizeX);
-	CellBox.MaxY = FPlatformMath::FloorToInt(MaxPosition.Y/Config->CellSizeY);
+	CellBox.MinX = static_cast<int32>(MinX);
+	CellBox.MinY = static_cast<int32>(MinY);
+	CellBox.MaxX = static_cast<int32>(MaxX);
+	CellBox.MaxY = static_cast<int32>(MaxY);
+
+	// Current large world max of 8796093022208.0 requires a cell size of at least around 4500 to not overflow an int32. 
+	checkSlow(MinX == CellBox.MinX && MinY == CellBox.MinY && MaxX == CellBox.MaxX && MaxY == CellBox.MaxY);
 
 	OutCellBox = CellBox;
 }
@@ -380,8 +388,8 @@ void UNetObjectGridFilter::CalculateCellBox(const UNetObjectGridFilter::FPerObje
 void UNetObjectGridFilter::CalculateCellCoord(UNetObjectGridFilter::FCellCoord& OutCoord, const FVector& Pos)
 {
 	FCellCoord Coord;
-	Coord.X = FPlatformMath::FloorToInt(Pos.X/Config->CellSizeX);
-	Coord.Y = FPlatformMath::FloorToInt(Pos.Y/Config->CellSizeY);
+	Coord.X = static_cast<int32>(FPlatformMath::FloorToInt(Pos.X/Config->CellSizeX));
+	Coord.Y = static_cast<int32>(FPlatformMath::FloorToInt(Pos.Y/Config->CellSizeY));
 
 	OutCoord = Coord;
 }
@@ -445,13 +453,17 @@ void UNetObjectGridFragmentLocFilter::OnInit(FNetObjectFilterInitParams& InitPar
 bool UNetObjectGridFragmentLocFilter::BuildObjectInfo(uint32 ObjectIndex, FNetObjectFilterAddObjectParams& Params)
 {
 	UE::Net::FRepTagFindInfo WorldLocationTagInfo;
-
 	if (UE::Net::FindRepTag(Params.Protocol, UE::Net::RepTag_WorldLocation, WorldLocationTagInfo))
 	{
-		FObjectLocationInfo& ObjectLocationInfo = static_cast<FObjectLocationInfo&>(Params.OutInfo);
-		ObjectLocationInfo.SetLocationStateOffset(WorldLocationTagInfo.ExternalStateOffset);
-		ObjectLocationInfo.SetLocationStateIndex(WorldLocationTagInfo.StateIndex);
+		// Want to keep the memory footprint minimal so don't allow adding objects whose values would not fit.
+		if (WorldLocationTagInfo.ExternalStateOffset >= MAX_uint16 || WorldLocationTagInfo.StateIndex >= MAX_uint16)
+		{
+			return false;
+		}
 
+		FObjectLocationInfo& ObjectLocationInfo = static_cast<FObjectLocationInfo&>(Params.OutInfo);
+		ObjectLocationInfo.SetLocationStateOffset(static_cast<uint16>(WorldLocationTagInfo.ExternalStateOffset));
+		ObjectLocationInfo.SetLocationStateIndex(static_cast<uint16>(WorldLocationTagInfo.StateIndex));
 
 		// NetCullDistanceSqr is optional. 
 		UE::Net::FRepTagFindInfo NetCullDistanceSqrTagInfo;
@@ -460,8 +472,8 @@ bool UNetObjectGridFragmentLocFilter::BuildObjectInfo(uint32 ObjectIndex, FNetOb
 			if ((NetCullDistanceSqrTagInfo.ExternalStateOffset < MAX_uint16) && (NetCullDistanceSqrTagInfo.StateIndex < MAX_uint16))
 			{
 				FCullDistanceFragmentInfo FragmentInfo;
-				FragmentInfo.CullDistanceSqrStateIndex = NetCullDistanceSqrTagInfo.StateIndex;
-				FragmentInfo.CullDistanceSqrStateOffset = NetCullDistanceSqrTagInfo.ExternalStateOffset;
+				FragmentInfo.CullDistanceSqrStateIndex = static_cast<uint16>(NetCullDistanceSqrTagInfo.StateIndex);
+				FragmentInfo.CullDistanceSqrStateOffset = static_cast<uint16>(NetCullDistanceSqrTagInfo.ExternalStateOffset);
 
 				CullDistanceFragments.Add(ObjectIndex, MoveTemp(FragmentInfo));
 			}
diff --git a/Engine/Source/Runtime/Experimental/Iris/Core/Private/Iris/ReplicationSystem/Filtering/ReplicationFiltering.cpp b/Engine/Source/Runtime/Experimental/Iris/Core/Private/Iris/ReplicationSystem/Filtering/ReplicationFiltering.cpp
--- a/Engine/Source/Runtime/Experimental/Iris/Core/Private/Iris/ReplicationSystem/Filtering/ReplicationFiltering.cpp
+++ b/Engine/Source/Runtime/Experimental/Iris/Core/Private/Iris/ReplicationSystem/Filtering/ReplicationFiltering.cpp
@@ -225,7 +225,7 @@ FReplicationFiltering::FReplicationFiltering()
 
 void FReplicationFiltering::StaticChecks()
 {
-	static_assert(TNumericLimits<PerObjectInfoIndexType>::Max() % (UsedPerObjectInfoStorageGrowSize*32U) == UsedPerObjectInfoStorageGrowSize*32U - 1, "Bit array grow code expects not to be able to return an out of bound index.");
+	static_assert(std::numeric_limits<PerObjectInfoIndexType>::max() % (UsedPerObjectInfoStorageGrowSize*32U) == UsedPerObjectInfoStorageGrowSize*32U - 1, "Bit array grow code expects not to be able to return an out of bound index.");
 	static_assert(sizeof(FNetBitArrayBase::StorageWordType) == sizeof(uint32), "Expected FNetBitArrayBase::StorageWordType to be four bytes in size."); 
 }
 
@@ -391,7 +391,7 @@ void FReplicationFiltering::SetOwningConnection(FInternalNetRefIndex ObjectIndex
 	}
 
 	const uint16 OldConnectionId = ObjectIndexToOwningConnection[ObjectIndex];
-	ObjectIndexToOwningConnection[ObjectIndex] = ConnectionId;
+	ObjectIndexToOwningConnection[ObjectIndex] = static_cast<uint16>(ConnectionId);
 	if (ConnectionId != OldConnectionId)
 	{
 		bHasDirtyOwner = 1;
@@ -459,7 +459,7 @@ bool FReplicationFiltering::SetFilter(FInternalNetRefIndex ObjectIndex, FNetObje
 		{
 			++FilterInfo.ObjectCount;
 			FilterInfo.FilteredObjects.SetBit(ObjIndex);
-			this->ObjectIndexToDynamicFilterIndex[ObjIndex] = FilterIndex;
+			this->ObjectIndexToDynamicFilterIndex[ObjIndex] = static_cast<uint8>(FilterIndex);
 			this->DynamicFilterEnabledObjects.SetBit(ObjIndex);
 			return true;
 		}
@@ -568,7 +568,7 @@ FNetObjectFilterHandle FReplicationFiltering::GetFilterHandle(const FName Filter
 	{
 		if (Info.Name == FilterName)
 		{
-			return FNetObjectFilterHandleUtil::MakeDynamicFilterHandle(&Info - DynamicFilterInfos.GetData());
+			return FNetObjectFilterHandleUtil::MakeDynamicFilterHandle(static_cast<uint32>(&Info - DynamicFilterInfos.GetData()));
 		}
 	}
 
@@ -640,7 +640,7 @@ void FReplicationFiltering::RemoveConnection(uint32 ConnectionId)
 	// Reset SubObject filter for removed connection
 	SubObjectFilterGroups.ForAllSetBits([this, ConnectionId](uint32 GroupIndex)
 	{
-		FNetObjectGroupHandle GroupHandle(GroupIndex);
+		FNetObjectGroupHandle GroupHandle(static_cast<uint16>(GroupIndex));
 		if (!IsReservedNetObjectGroupHandle(GroupHandle))
 		{
 			SetConnectionFilterStatus(*GetPerObjectInfo(GroupInfos[GroupHandle].ConnectionStateIndex), ConnectionId, ENetFilterStatus::Disallow);
@@ -689,7 +689,7 @@ void FReplicationFiltering::InitNewConnections()
 		{
 			auto InitGroupFilterForConnection = [this, &ConnectionInfo, ConnectionId](uint32 GroupIndex)
 			{
-				const FNetObjectGroupHandle GroupHandle = GroupIndex;
+				const FNetObjectGroupHandle GroupHandle = static_cast<uint16>(GroupIndex);
 				const FPerObjectInfo* ConnectionState = GetPerObjectInfo(this->GroupInfos[GroupIndex].ConnectionStateIndex);
 
 				// Setup filter for connection
@@ -988,7 +988,7 @@ void FReplicationFiltering::UpdateOwnerAndConnectionFiltering()
 				const FNetRefHandleManager::FReplicatedObjectData& ObjectData = NetRefHandleManager->GetReplicatedObjectDataNoCheck(SubObjectIndex);
 				if (ObjectData.IsSubObject())
 				{
-					ObjectIndexToOwningConnectionStorage[SubObjectIndex] = OwningConnectionId;
+					ObjectIndexToOwningConnectionStorage[SubObjectIndex] = static_cast<uint16>(OwningConnectionId);
 				}
 			}
 		};
@@ -1075,7 +1075,7 @@ void FReplicationFiltering::UpdateGroupFiltering()
 			IRIS_PROFILER_SCOPE(FReplicationFiltering_UpdateGroupFiltering);
 
 			const FPerObjectInfo* ConnectionStateInfo = GetPerObjectInfo(this->GroupInfos[GroupIndex].ConnectionStateIndex);
-			const FNetObjectGroup* Group = this->Groups->GetGroup(GroupIndex);
+			const FNetObjectGroup* Group = this->Groups->GetGroup(FNetObjectGroupHandle(GroupIndex));
 			FPerConnectionInfo* LocalConnectionInfos = this->ConnectionInfos.GetData();
 
 			auto UpdateGroupFilterForConnection = [this, ConnectionStateInfo, Group, LocalConnectionInfos](uint32 ConnectionId)
@@ -1225,7 +1225,7 @@ void FReplicationFiltering::UpdateDynamicFiltering(ENetFilterType FilterType)
 			const uint32* DynamicFilterEnabledObjectsData = DynamicFilterEnabledObjects.GetData();
 			FNetBitArrayView DynamicFilteredOutObjects = MakeNetBitArrayView(ConnectionInfo.DynamicFilteredOutObjects);
 			uint32* DynamicFilteredOutObjectsData = ConnectionInfo.DynamicFilteredOutObjects.GetData();
-			for (SIZE_T WordIt = 0, WordEndIt = WordCountForObjectBitArrays; WordIt != WordEndIt; ++WordIt)
+			for (uint32 WordIt = 0, WordEndIt = WordCountForObjectBitArrays; WordIt != WordEndIt; ++WordIt)
 			{
 				const uint32 SubObjects = SubObjectsData[WordIt];
 				const uint32 DependentObjects = DependentObjectsData[WordIt];
@@ -1385,7 +1385,7 @@ void FReplicationFiltering::BatchNotifyFiltersOfDirtyObjects(FUpdateDirtyObjects
 		UpdateParameters.ObjectCount = PerFilterInfo.ObjectCount;
 		UpdateParameters.InstanceProtocols = PerFilterInfo.InstanceProtocols;
 
-		const SIZE_T FilterIndex = &PerFilterInfo - BatchHelper.PerFilterInfos.GetData();
+		const int32 FilterIndex = static_cast<int32>(&PerFilterInfo - BatchHelper.PerFilterInfos.GetData());
 		UNetObjectFilter* Filter = DynamicFilterInfos[FilterIndex].Filter.Get();
 		Filter->UpdateObjects(UpdateParameters);
 	}
@@ -1413,7 +1413,7 @@ FReplicationFiltering::PerObjectInfoIndexType FReplicationFiltering::AllocPerObj
 	// Grow used indices bit array if needed
 	if (FreeIndex == FNetBitArrayView::InvalidIndex)
 	{
-		checkf(UsedPerObjectInfos.GetNumBits() < TNumericLimits<PerObjectInfoIndexType>::Max(), TEXT("Filtering per object info storage exhausted. Contact the UE Network team."));
+		checkf(UsedPerObjectInfos.GetNumBits() < std::numeric_limits<PerObjectInfoIndexType>::max(), TEXT("Filtering per object info storage exhausted. Contact the UE Network team."));
 		FreeIndex = UsedPerObjectInfos.GetNumBits();
 		UsedPerObjectInfoStorage.AddZeroed(UsedPerObjectInfoStorageGrowSize/UsedPerObjectInfoBitsPerWord);
 		FNetBitArrayView NewUsedPerObjectInfos = MakeNetBitArrayView(UsedPerObjectInfoStorage.GetData(), UsedPerObjectInfoStorage.Num()*UsedPerObjectInfoBitsPerWord);
@@ -1431,7 +1431,7 @@ FReplicationFiltering::PerObjectInfoIndexType FReplicationFiltering::AllocPerObj
 		UsedPerObjectInfos.SetBit(FreeIndex);
 	}
 	
-	return FreeIndex;
+	return static_cast<PerObjectInfoIndexType>(FreeIndex);
 }
 
 void FReplicationFiltering::FreePerObjectInfo(PerObjectInfoIndexType Index)
@@ -1527,7 +1527,7 @@ void FReplicationFiltering::UpdateSubObjectFilters()
 	// We want to remove all groups that have no members and no enabled connections
 	auto UpdateSubObjectFilterGroup = [this](uint32 GroupIndex)
 	{
-		FNetObjectGroupHandle GroupHandle(GroupIndex);
+		FNetObjectGroupHandle GroupHandle(static_cast<uint16>(GroupIndex));
 		if (const FNetObjectGroup* Group = Groups->GetGroup(FNetObjectGroupHandle(GroupIndex)))
 		{
 			if (Group->Members.Num() == 0U && !IsConnectionFilterStatusAllowedForAnyConnection(*GetPerObjectInfo(GroupInfos[GroupHandle].ConnectionStateIndex)))
@@ -2008,6 +2008,10 @@ void FReplicationFiltering::SetPerObjectInfoFilterStatus(FPerObjectInfo& ObjectI
 void FReplicationFiltering::InitFilters()
 {
 	const UNetObjectFilterDefinitions* FilterDefinitions = GetDefault<UNetObjectFilterDefinitions>();
+
+	// We store a uint8 per object to filter.
+	check(FilterDefinitions->GetFilterDefinitions().Num() <= 256);
+
 	for (const FNetObjectFilterDefinition& FilterDefinition : FilterDefinitions->GetFilterDefinitions())
 	{
 		constexpr UObject* ClassOuter = nullptr;
diff --git a/Engine/Source/Runtime/Experimental/Iris/Core/Private/Iris/ReplicationSystem/NetBlob/NetBlobHandlerManager.cpp b/Engine/Source/Runtime/Experimental/Iris/Core/Private/Iris/ReplicationSystem/NetBlob/NetBlobHandlerManager.cpp
--- a/Engine/Source/Runtime/Experimental/Iris/Core/Private/Iris/ReplicationSystem/NetBlob/NetBlobHandlerManager.cpp
+++ b/Engine/Source/Runtime/Experimental/Iris/Core/Private/Iris/ReplicationSystem/NetBlob/NetBlobHandlerManager.cpp
@@ -39,7 +39,7 @@ bool FNetBlobHandlerManager::RegisterHandler(UNetBlobHandler* Handler)
 	{
 		if (Definition.ClassName.ToString() == ClassName)
 		{
-			const uint32 Index = &Definition - BlobHandlerDefinitions->NetBlobHandlerDefinitions.GetData();
+			const uint32 Index = static_cast<uint32>(&Definition - BlobHandlerDefinitions->NetBlobHandlerDefinitions.GetData());
 			Handler->NetBlobType = Index;
 			Handlers[Index] = Handler;
 			return true;
diff --git a/Engine/Source/Runtime/Experimental/Iris/Core/Private/Iris/ReplicationSystem/NetBlob/NetRPC.cpp b/Engine/Source/Runtime/Experimental/Iris/Core/Private/Iris/ReplicationSystem/NetBlob/NetRPC.cpp
--- a/Engine/Source/Runtime/Experimental/Iris/Core/Private/Iris/ReplicationSystem/NetBlob/NetRPC.cpp
+++ b/Engine/Source/Runtime/Experimental/Iris/Core/Private/Iris/ReplicationSystem/NetBlob/NetRPC.cpp
@@ -274,8 +274,8 @@ void FNetRPC::DeserializeFunctionLocator(FNetSerializationContext& Context)
 	UE_NET_TRACE_SCOPE(FunctionLocator, *Reader, Context.GetTraceCollector(), ENetTraceVerbosity::Trace);
 
 	const uint32 NibbleCount = Reader->ReadBits(2U) + 1U;
-	FunctionLocator.DescriptorIndex = Reader->ReadBits(NibbleCount*4U);
-	FunctionLocator.FunctionIndex = Reader->ReadBits(NibbleCount*4U);
+	FunctionLocator.DescriptorIndex = static_cast<uint16>(Reader->ReadBits(NibbleCount*4U));
+	FunctionLocator.FunctionIndex = static_cast<uint16>(Reader->ReadBits(NibbleCount*4U));
 }
 
 void FNetRPC::InternalSerializeObjectReference(FNetSerializationContext& Context) const
@@ -557,8 +557,8 @@ bool FNetRPC::IsServerAllowedToExecuteRPC(FNetSerializationContext& Context) con
 			{
 				if (FunctionDescriptor.Function == Function)
 				{
-					OutFunctionLocator.DescriptorIndex = &Descriptor - Protocol->ReplicationStateDescriptors;
-					OutFunctionLocator.FunctionIndex = &FunctionDescriptor - Descriptor->MemberFunctionDescriptors;
+					OutFunctionLocator.DescriptorIndex = static_cast<uint16>(&Descriptor - Protocol->ReplicationStateDescriptors);
+					OutFunctionLocator.FunctionIndex = static_cast<uint16>(&FunctionDescriptor - Descriptor->MemberFunctionDescriptors);
 					OutFunctionDescriptor = &FunctionDescriptor;
 					return true;
 				}
diff --git a/Engine/Source/Runtime/Experimental/Iris/Core/Private/Iris/ReplicationSystem/NetBlob/PartialNetBlob.cpp b/Engine/Source/Runtime/Experimental/Iris/Core/Private/Iris/ReplicationSystem/NetBlob/PartialNetBlob.cpp
--- a/Engine/Source/Runtime/Experimental/Iris/Core/Private/Iris/ReplicationSystem/NetBlob/PartialNetBlob.cpp
+++ b/Engine/Source/Runtime/Experimental/Iris/Core/Private/Iris/ReplicationSystem/NetBlob/PartialNetBlob.cpp
@@ -262,8 +262,8 @@ void FPartialNetBlob::SplitPayload(const FPartialNetBlob::FPayloadSplitParams& S
 		FPartialNetBlob* PartialBlob = new FPartialNetBlob(SplitParams.CreationInfo);
 		PartialBlob->SetDebugName(SplitParams.DebugName);
 		PartialBlob->OriginalCreationInfo = SplitParams.OriginalCreationInfo;
-		PartialBlob->PartIndex = PartIt;
-		PartialBlob->PartCount = PartialBlobCount;
+		PartialBlob->PartIndex = static_cast<uint16>(PartIt);
+		PartialBlob->PartCount = static_cast<uint16>(PartialBlobCount);
 		if (bIsFirstPart && EnumHasAnyFlags(SplitParams.OriginalCreationInfo.Flags, ENetBlobFlags::HasExports))
 		{
 			PartialBlob->OriginalBlob = SplitParams.OriginalBlob;
@@ -274,7 +274,7 @@ void FPartialNetBlob::SplitPayload(const FPartialNetBlob::FPayloadSplitParams& S
 		const uint32 PartialBlobBitCount = FPlatformMath::Min(SplitParams.PayloadBitCount - PayloadBitOffset, SplitParams.PartBitCount);
 		const uint32 PartialBlobWordCount = (PartialBlobBitCount + 31U)/32U;
 		const uint32 PayloadWordOffset = PayloadBitOffset/32U;
-		PartialBlob->PayloadBitCount = PartialBlobBitCount;
+		PartialBlob->PayloadBitCount = static_cast<uint16>(PartialBlobBitCount);
 		PartialBlob->Payload.SetNumUninitialized(PartialBlobWordCount);
 		FPlatformMemory::Memcpy(PartialBlob->Payload.GetData(), SplitParams.Payload + PayloadWordOffset, PartialBlobWordCount*4U);
 
diff --git a/Engine/Source/Runtime/Experimental/Iris/Core/Private/Iris/ReplicationSystem/NetExports.cpp b/Engine/Source/Runtime/Experimental/Iris/Core/Private/Iris/ReplicationSystem/NetExports.cpp
--- a/Engine/Source/Runtime/Experimental/Iris/Core/Private/Iris/ReplicationSystem/NetExports.cpp
+++ b/Engine/Source/Runtime/Experimental/Iris/Core/Private/Iris/ReplicationSystem/NetExports.cpp
@@ -8,8 +8,8 @@ namespace UE::Net::Private
 
 void FNetExports::InitExportRecordForPacket()
 {
-	CurrentExportInfo.NetHandleExportCount = NetHandleExports.Count();
-	CurrentExportInfo.NetTokenExportCount = NetTokenExports.Count();
+	CurrentExportInfo.NetHandleExportCount = static_cast<uint32>(NetHandleExports.Count());
+	CurrentExportInfo.NetTokenExportCount = static_cast<uint32>(NetTokenExports.Count());
 }
 
 void FNetExports::CommitExportsToRecord(FExportScope& ExportScope)
@@ -28,8 +28,8 @@ void FNetExports::CommitExportsToRecord(FExportScope& ExportScope)
 void FNetExports::PushExportRecordForPacket()
 {
 	FExportInfo Info;
-	Info.NetHandleExportCount = NetHandleExports.Count() - CurrentExportInfo.NetHandleExportCount;
-	Info.NetTokenExportCount = NetTokenExports.Count() - CurrentExportInfo.NetTokenExportCount;
+	Info.NetHandleExportCount = static_cast<uint32>(NetHandleExports.Count()) - CurrentExportInfo.NetHandleExportCount;
+	Info.NetTokenExportCount = static_cast<uint32>(NetTokenExports.Count()) - CurrentExportInfo.NetTokenExportCount;
 	ExportRecord.Enqueue(Info);
 }
 
diff --git a/Engine/Source/Runtime/Experimental/Iris/Core/Private/Iris/ReplicationSystem/ObjectReferenceCache.cpp b/Engine/Source/Runtime/Experimental/Iris/Core/Private/Iris/ReplicationSystem/ObjectReferenceCache.cpp
--- a/Engine/Source/Runtime/Experimental/Iris/Core/Private/Iris/ReplicationSystem/ObjectReferenceCache.cpp
+++ b/Engine/Source/Runtime/Experimental/Iris/Core/Private/Iris/ReplicationSystem/ObjectReferenceCache.cpp
@@ -1191,7 +1191,7 @@ void FObjectReferenceCache::ReadFullReferenceInternal(FNetSerializationContext&
 	if (bIsExported)
 	{
 		// Can we load this object?
-		bNoLoad = Reader->ReadBits(1U);
+		bNoLoad = Reader->ReadBits(1U) & 1;
 		if (Reader->IsOverflown())
 		{
 			check(false);
diff --git a/Engine/Source/Runtime/Experimental/Iris/Core/Private/Iris/ReplicationSystem/Prioritization/LocationBasedNetObjectPrioritizer.cpp b/Engine/Source/Runtime/Experimental/Iris/Core/Private/Iris/ReplicationSystem/Prioritization/LocationBasedNetObjectPrioritizer.cpp
--- a/Engine/Source/Runtime/Experimental/Iris/Core/Private/Iris/ReplicationSystem/Prioritization/LocationBasedNetObjectPrioritizer.cpp
+++ b/Engine/Source/Runtime/Experimental/Iris/Core/Private/Iris/ReplicationSystem/Prioritization/LocationBasedNetObjectPrioritizer.cpp
@@ -43,8 +43,8 @@ bool ULocationBasedNetObjectPrioritizer::AddObject(uint32 ObjectIndex, FNetObjec
 	}
 
 	FObjectLocationInfo& ObjectInfo = static_cast<FObjectLocationInfo&>(Params.OutInfo);
-	ObjectInfo.SetLocationStateOffset(TagInfo.ExternalStateOffset);
-	ObjectInfo.SetLocationStateIndex(TagInfo.StateIndex);
+	ObjectInfo.SetLocationStateOffset(static_cast<uint16>(TagInfo.ExternalStateOffset));
+	ObjectInfo.SetLocationStateIndex(static_cast<uint16>(TagInfo.StateIndex));
 	const uint32 LocationIndex = AllocLocation();
 	ObjectInfo.SetLocationIndex(LocationIndex);
 
diff --git a/Engine/Source/Runtime/Experimental/Iris/Core/Private/Iris/ReplicationSystem/Prioritization/NetObjectCountLimiter.cpp b/Engine/Source/Runtime/Experimental/Iris/Core/Private/Iris/ReplicationSystem/Prioritization/NetObjectCountLimiter.cpp
--- a/Engine/Source/Runtime/Experimental/Iris/Core/Private/Iris/ReplicationSystem/Prioritization/NetObjectCountLimiter.cpp
+++ b/Engine/Source/Runtime/Experimental/Iris/Core/Private/Iris/ReplicationSystem/Prioritization/NetObjectCountLimiter.cpp
@@ -122,7 +122,7 @@ uint16 UNetObjectCountLimiter::AllocInternalIndex()
 	}
 
 	InternalObjectIndices.SetBit(Index);
-	return Index;
+	return static_cast<uint16>(Index);
 }
 
 void UNetObjectCountLimiter::FreeInternalIndex(uint16 Index)
@@ -158,7 +158,7 @@ void UNetObjectCountLimiter::PrePrioritizeForRoundRobin()
 
 	if (ObjectCount)
 	{
-		RoundRobinState.NextIndexToConsider = Indices[ObjectCount - 1U] + 1U;
+		RoundRobinState.NextIndexToConsider = static_cast<uint16>(Indices[ObjectCount - 1U] + 1U);
 		for (uint32 Index : MakeArrayView(Indices, ObjectCount))
 		{
 			RoundRobinState.InternalObjectIndices.SetBit(Index);
diff --git a/Engine/Source/Runtime/Experimental/Iris/Core/Private/Iris/ReplicationSystem/Prioritization/ReplicationPrioritization.cpp b/Engine/Source/Runtime/Experimental/Iris/Core/Private/Iris/ReplicationSystem/Prioritization/ReplicationPrioritization.cpp
--- a/Engine/Source/Runtime/Experimental/Iris/Core/Private/Iris/ReplicationSystem/Prioritization/ReplicationPrioritization.cpp
+++ b/Engine/Source/Runtime/Experimental/Iris/Core/Private/Iris/ReplicationSystem/Prioritization/ReplicationPrioritization.cpp
@@ -390,7 +390,7 @@ bool FReplicationPrioritization::SetPrioritizer(uint32 ObjectIndex, FNetObjectPr
 		FPrioritizerInfo& PrioritizerInfo = PrioritizerInfos[NewPrioritizer];
 		if (PrioritizerInfo.Prioritizer->AddObject(ObjectIndex, AddParams))
 		{
-			Prioritizer = NewPrioritizer;
+			Prioritizer = static_cast<uint8>(NewPrioritizer);
 			++PrioritizerInfo.ObjectCount;
 			return true;
 		}
@@ -699,7 +699,7 @@ void FReplicationPrioritization::PrioritizeForConnection(uint32 ConnId, FPriorit
 					PrioParameters.ObjectIndices = PerPrioritizerInfo.ObjectIndices.GetFirstChunkData();
 					PrioParameters.ObjectCount = PerPrioritizerInfo.ObjectIndices.GetFirstChunkNum();
 
-					const SIZE_T PrioritizerIndex = &PerPrioritizerInfo - BatchHelper.PerPrioritizerInfos.GetData();
+					const int32 PrioritizerIndex = static_cast<int32>(&PerPrioritizerInfo - BatchHelper.PerPrioritizerInfos.GetData());
 					UNetObjectPrioritizer* Prioritizer = PrioritizerInfos[PrioritizerIndex].Prioritizer.Get();
 					Prioritizer->Prioritize(PrioParameters);
 
@@ -720,7 +720,7 @@ void FReplicationPrioritization::PrioritizeForConnection(uint32 ConnId, FPriorit
 				PrioParameters.ObjectIndices = PerPrioritizerInfo.ObjectIndices.GetFirstChunkData();
 				PrioParameters.ObjectCount = PerPrioritizerInfo.ObjectIndices.GetFirstChunkNum();
 
-				const SIZE_T PrioritizerIndex = &PerPrioritizerInfo - BatchHelper.PerPrioritizerInfos.GetData();
+				const int32 PrioritizerIndex = static_cast<int32>(&PerPrioritizerInfo - BatchHelper.PerPrioritizerInfos.GetData());
 				UNetObjectPrioritizer* Prioritizer = PrioritizerInfos[PrioritizerIndex].Prioritizer.Get();
 				Prioritizer->Prioritize(PrioParameters);
 
@@ -778,7 +778,7 @@ void FReplicationPrioritization::BatchNotifyPrioritizersOfDirtyObjects(FUpdateDi
 		UpdateParameters.ObjectCount = PerPrioritizerInfo.ObjectCount;
 		UpdateParameters.InstanceProtocols = PerPrioritizerInfo.InstanceProtocols;
 
-		const SIZE_T PrioritizerIndex = &PerPrioritizerInfo - BatchHelper.PerPrioritizerInfos.GetData();
+		const int32 PrioritizerIndex = static_cast<int32>(&PerPrioritizerInfo - BatchHelper.PerPrioritizerInfos.GetData());
 		UNetObjectPrioritizer* Prioritizer = PrioritizerInfos[PrioritizerIndex].Prioritizer.Get();
 		Prioritizer->UpdateObjects(UpdateParameters);
 	}
@@ -799,6 +799,9 @@ void FReplicationPrioritization::InitPrioritizers()
 	TArray<FNetObjectPrioritizerDefinition> Definitions;
 	PrioritizerDefinitions->GetValidDefinitions(Definitions);
 
+	// We store a uint8 per object to prioritizer.
+	check(Definitions.Num() <= 256);
+
 	PrioritizerInfos.Reserve(Definitions.Num());
 	for (FNetObjectPrioritizerDefinition& Definition : Definitions)
 	{
diff --git a/Engine/Source/Runtime/Experimental/Iris/Core/Private/Iris/ReplicationSystem/Prioritization/SphereWithOwnerBoostNetObjectPrioritizer.cpp b/Engine/Source/Runtime/Experimental/Iris/Core/Private/Iris/ReplicationSystem/Prioritization/SphereWithOwnerBoostNetObjectPrioritizer.cpp
--- a/Engine/Source/Runtime/Experimental/Iris/Core/Private/Iris/ReplicationSystem/Prioritization/SphereWithOwnerBoostNetObjectPrioritizer.cpp
+++ b/Engine/Source/Runtime/Experimental/Iris/Core/Private/Iris/ReplicationSystem/Prioritization/SphereWithOwnerBoostNetObjectPrioritizer.cpp
@@ -58,7 +58,7 @@ void USphereWithOwnerBoostNetObjectPrioritizer::UpdateObjects(FNetObjectPrioriti
 	{
 		FObjectLocationInfo& ObjectInfo = static_cast<FObjectLocationInfo&>(Params.PrioritizationInfos[ObjectIndex]);
 		const uint32 OwningConnection = ReplicationFiltering.GetOwningConnection(ObjectIndex);
-		OwningConnections[ObjectInfo.GetLocationIndex()] = OwningConnection;
+		OwningConnections[ObjectInfo.GetLocationIndex()] = static_cast<uint16>(OwningConnection);
 	}
 }
 
diff --git a/Engine/Source/Runtime/Experimental/Iris/Core/Private/Iris/ReplicationSystem/RepTag.cpp b/Engine/Source/Runtime/Experimental/Iris/Core/Private/Iris/ReplicationSystem/RepTag.cpp
--- a/Engine/Source/Runtime/Experimental/Iris/Core/Private/Iris/ReplicationSystem/RepTag.cpp
+++ b/Engine/Source/Runtime/Experimental/Iris/Core/Private/Iris/ReplicationSystem/RepTag.cpp
@@ -92,7 +92,7 @@ bool FindRepTag(const FReplicationProtocol* Protocol, FRepTag RepTag, FRepTagFin
 		
 		if (FindRepTag(StateDescriptor, RepTag, OutRepTagFindInfo))
 		{
-			OutRepTagFindInfo.StateIndex = &StateDescriptor - Protocol->ReplicationStateDescriptors;
+			OutRepTagFindInfo.StateIndex = static_cast<uint32>(&StateDescriptor - Protocol->ReplicationStateDescriptors);
 			OutRepTagFindInfo.InternalStateAbsoluteOffset += InternalOffset;
 			return true;
 		}
diff --git a/Engine/Source/Runtime/Experimental/Iris/Core/Private/Iris/ReplicationSystem/ReplicationProtocolManager.cpp b/Engine/Source/Runtime/Experimental/Iris/Core/Private/Iris/ReplicationSystem/ReplicationProtocolManager.cpp
--- a/Engine/Source/Runtime/Experimental/Iris/Core/Private/Iris/ReplicationSystem/ReplicationProtocolManager.cpp
+++ b/Engine/Source/Runtime/Experimental/Iris/Core/Private/Iris/ReplicationSystem/ReplicationProtocolManager.cpp
@@ -35,6 +35,10 @@ static FAutoConsoleVariableRef CVarIrisLogReplicationProtocols(
 FReplicationInstanceProtocol* FReplicationProtocolManager::CreateInstanceProtocol(const FReplicationFragments& Fragments)
 {
 	const uint32 FragmentCount = Fragments.Num();
+	if (!ensure(FragmentCount < 65536))
+	{
+		return nullptr;
+	}
 
 	// We want to keep this as a single allocation so we first build a layout and allocate enough space for the InstanceProtocol and its data
 	struct FReplicationInstanceProtocolLayoutData
@@ -51,14 +55,13 @@ FReplicationInstanceProtocol* FReplicationProtocolManager::CreateInstanceProtoco
 	FMemoryLayoutUtil::AddToLayout<FReplicationFragment*>(Layout, LayoutData.FragmentsSizeAndOffset, FragmentCount);
 
 	// Allocate memory for the instance protocol
-	uint8* Buffer = (uint8*)FMemory::Malloc(Layout.CurrentOffset, Layout.MaxAlignment);
-	FMemory::Memset(Buffer, 0, Layout.CurrentOffset);
+	uint8* Buffer = static_cast<uint8*>(FMemory::MallocZeroed(Layout.CurrentOffset, static_cast<uint32>(Layout.MaxAlignment)));
 
 	// Init FReplicationInstanceProtocol
 	FReplicationInstanceProtocol* InstanceProtocol = new (Buffer) FReplicationInstanceProtocol;
 	InstanceProtocol->FragmentData = reinterpret_cast<FReplicationInstanceProtocol::FFragmentData*>(Buffer + LayoutData.FragmentDataSizeAndOffset.Offset);
 	InstanceProtocol->Fragments = reinterpret_cast<FReplicationFragment* const *>(Buffer + LayoutData.FragmentsSizeAndOffset.Offset);
-	InstanceProtocol->FragmentCount = FragmentCount;
+	InstanceProtocol->FragmentCount = static_cast<uint16>(FragmentCount);
 
 	// Fill in fragment data and fragment pointer
 	uint32 FragmentIt = 0u;
@@ -230,6 +233,10 @@ const FReplicationProtocol* FReplicationProtocolManager::CreateReplicationProtoc
 
 	// Create the protocol
 	const uint32 FragmentCount = Fragments.Num();
+	if (!ensure(FragmentCount <= 65536))
+	{
+		return nullptr;
+	}
 
 	// We want to keep this as a single allocation so we first build a layout and allocate enough space for the InstanceProtocol and its data
 	struct FReplicationProtocolLayoutData
@@ -245,8 +252,7 @@ const FReplicationProtocol* FReplicationProtocolManager::CreateReplicationProtoc
 
 	// Allocate memory for the protocol, the replication protocol must be refcounted by all NetObjects
 	// We could also choose to explicitly control the lifetime 
-	uint8* Buffer = (uint8*)FMemory::Malloc(Layout.CurrentOffset, Layout.MaxAlignment);
-	FMemory::Memset(Buffer, 0, Layout.CurrentOffset);
+	uint8* Buffer = static_cast<uint8*>(FMemory::MallocZeroed(Layout.CurrentOffset, static_cast<uint32>(Layout.MaxAlignment)));
 
 	// Init FReplicationInstanceProtocol
 	FReplicationProtocol* Protocol = new (Buffer) FReplicationProtocol;
@@ -335,8 +341,8 @@ const FReplicationProtocol* FReplicationProtocolManager::CreateReplicationProtoc
 
 		Protocol->InternalChangeMasksOffset = InternalSize;
 
-		Protocol->FirstLifetimeConditionalsStateIndex = FirstLifetimeConditionalsStateIndex;
-		Protocol->LifetimeConditionalsStateCount = LifetimeConditionalsStateCount;
+		Protocol->FirstLifetimeConditionalsStateIndex = static_cast<uint16>(FirstLifetimeConditionalsStateIndex);
+		Protocol->LifetimeConditionalsStateCount = static_cast<uint16>(LifetimeConditionalsStateCount);
 		Protocol->FirstLifetimeConditionalsChangeMaskOffset = FirstLifetimeConditionalsChangeMaskOffset;
 
 		InternalSize += FNetBitArrayView::CalculateRequiredWordCount(ChangeMaskBitCount)*sizeof(FNetBitArrayView::StorageWordType);
diff --git a/Engine/Source/Runtime/Experimental/Iris/Core/Private/Iris/ReplicationSystem/ReplicationReader.cpp b/Engine/Source/Runtime/Experimental/Iris/Core/Private/Iris/ReplicationSystem/ReplicationReader.cpp
--- a/Engine/Source/Runtime/Experimental/Iris/Core/Private/Iris/ReplicationSystem/ReplicationReader.cpp
+++ b/Engine/Source/Runtime/Experimental/Iris/Core/Private/Iris/ReplicationSystem/ReplicationReader.cpp
@@ -261,14 +261,14 @@ FNetRefHandle FReplicationReader::ReadNetRefHandleId(FNetSerializationContext& C
 	return RefHandle;
 }
 	
-uint16 FReplicationReader::ReadObjectsPendingDestroy(FNetSerializationContext& Context)
+uint32 FReplicationReader::ReadObjectsPendingDestroy(FNetSerializationContext& Context)
 {
 	FNetBitStreamReader& Reader = *Context.GetBitStreamReader();
 
 	UE_NET_TRACE_SCOPE(ObjectsPendingDestroy, Reader, Context.GetTraceCollector(), ENetTraceVerbosity::Trace);
 
 	// Read how many destroyed objects we have
-	const uint16 ObjectsToRead = Reader.ReadBits(16);
+	const uint32 ObjectsToRead = Reader.ReadBits(16);
 	
 	if (!Context.HasErrorOrOverflow())
 	{
@@ -710,8 +710,8 @@ void FReplicationReader::ReadObjectInBatch(FNetSerializationContext& Context, FN
 		IncompleteHandle = ReadNetRefHandleId(Context, Reader);
 	}
 	
-	// Read replicated destroy header if necessary
-	const bool bReadReplicatedDestroyHeader = !IsObjectIndexForOOBAttachment(IncompleteHandle.GetId());
+	// Read replicated destroy header if necessary. We don't know the internal index yet so can't do the more appropriate check IsObjectIndexForOOBAttachment.
+	const bool bReadReplicatedDestroyHeader = IncompleteHandle.IsValid();
 	const uint32 ReplicatedDestroyHeaderFlags = bReadReplicatedDestroyHeader ? Reader.ReadBits(ReplicatedDestroyHeaderFlags_BitCount) : ReplicatedDestroyHeaderFlags_None;
 
 	const bool bHasState = Reader.ReadBool();
@@ -766,7 +766,7 @@ void FReplicationReader::ReadObjectInBatch(FNetSerializationContext& Context, FN
 		}
 		
 		// We got a read error
-		if (Reader.IsOverflown() || IsObjectIndexForOOBAttachment(IncompleteHandle.GetId()))
+		if (Reader.IsOverflown() || !IncompleteHandle.IsValid())
 		{
 			UE_LOG_REPLICATIONREADER_ERROR(TEXT("FReplicationReader::ReadObject Bitstream corrupted."));
 			const FName& NetError = (Reader.IsOverflown() ? GNetError_BitStreamOverflow : GNetError_BitStreamError);
@@ -798,7 +798,7 @@ void FReplicationReader::ReadObjectInBatch(FNetSerializationContext& Context, FN
 	else
 	{
 		bHasErrors = bHasErrors || Context.HasErrorOrOverflow();
-		if (bHasErrors || IsObjectIndexForOOBAttachment(IncompleteHandle.GetId()))
+		if (bHasErrors || !IncompleteHandle.IsValid())
 		{
 			InternalIndex = ObjectIndexForOOBAttachment;
 		}
diff --git a/Engine/Source/Runtime/Experimental/Iris/Core/Private/Iris/ReplicationSystem/ReplicationReader.h b/Engine/Source/Runtime/Experimental/Iris/Core/Private/Iris/ReplicationSystem/ReplicationReader.h
--- a/Engine/Source/Runtime/Experimental/Iris/Core/Private/Iris/ReplicationSystem/ReplicationReader.h
+++ b/Engine/Source/Runtime/Experimental/Iris/Core/Private/Iris/ReplicationSystem/ReplicationReader.h
@@ -171,7 +171,7 @@ private:
 	uint32 ReadObjectsInBatch(FNetSerializationContext& Context, FNetRefHandle InCompleteHandle, bool bHasBatchOwnerData, uint32 BatchEndBitPosition);
 
 	// Read all objects pending destroy
-	uint16 ReadObjectsPendingDestroy(FNetSerializationContext& Context);
+	uint32 ReadObjectsPendingDestroy(FNetSerializationContext& Context);
 
 	// Read state data for all incoming objects
 	void ReadObjects(FNetSerializationContext& Context, uint32 ObjectCountToRead);
diff --git a/Engine/Source/Runtime/Experimental/Iris/Core/Private/Iris/ReplicationSystem/ReplicationRecord.h b/Engine/Source/Runtime/Experimental/Iris/Core/Private/Iris/ReplicationSystem/ReplicationRecord.h
--- a/Engine/Source/Runtime/Experimental/Iris/Core/Private/Iris/ReplicationSystem/ReplicationRecord.h
+++ b/Engine/Source/Runtime/Experimental/Iris/Core/Private/Iris/ReplicationSystem/ReplicationRecord.h
@@ -75,7 +75,7 @@ public:
 	// PeekInfo . If the info indicates there's an attachment record one should call DequeueAttachmentRecord() as well.
 	inline const FRecordInfo& PeekInfo() const { return RecordInfos.Peek(); }
 	inline const FRecordInfo& PeekInfoAtOffset(uint32 Offset) const { return RecordInfos.PeekAtOffset(Offset); }
-	inline const uint16 GetInfoCount() const { return RecordInfos.Count(); }
+	inline const uint32 GetInfoCount() const { return static_cast<uint32>(RecordInfos.Count()); }
 
 	// If the info from PeekInfo() indicates there's an attachment record one need to call this function as well.
 	uint64 DequeueAttachmentRecord();
@@ -92,7 +92,7 @@ public:
 	inline const uint32 PeekRecordAtOffset(uint32 Offset) const { return Record.PeekAtOffset(Offset); }
 
 	// Num Records
-	inline const uint32 GetRecordCount() const { return Record.Count(); }
+	inline const uint32 GetRecordCount() const { return static_cast<uint32>(Record.Count()); }
 	
 private:
 
diff --git a/Engine/Source/Runtime/Experimental/Iris/Core/Private/Iris/ReplicationSystem/ReplicationWriter.cpp b/Engine/Source/Runtime/Experimental/Iris/Core/Private/Iris/ReplicationSystem/ReplicationWriter.cpp
--- a/Engine/Source/Runtime/Experimental/Iris/Core/Private/Iris/ReplicationSystem/ReplicationWriter.cpp
+++ b/Engine/Source/Runtime/Experimental/Iris/Core/Private/Iris/ReplicationSystem/ReplicationWriter.cpp
@@ -263,7 +263,7 @@ static bool s_ValidateReplicationRecord(const FReplicationRecord* ReplicationRec
 	}
 	
 	// Verify last / first record
-	const uint16 RecordInfoCount = bVerifyFirstRecord ? ReplicationRecord->PeekRecordAtOffset(0) : (ReplicationRecord->PeekRecordAtOffset(ReplicationRecord->GetRecordCount() - 1));
+	const uint32 RecordInfoCount = bVerifyFirstRecord ? ReplicationRecord->PeekRecordAtOffset(0) : ReplicationRecord->PeekRecordAtOffset(ReplicationRecord->GetRecordCount() - 1);
 
 	// check for duplicates
 	{
@@ -1209,7 +1209,7 @@ void FReplicationWriter::HandleDroppedRecord<FReplicationWriter::EReplicatedObje
 		bool bNeedToResendAttachments = RecordInfo.HasAttachments;
 		bool bNeedToResendState = false;
 
-		FNetBitArrayView LostChangeMask = FChangeMaskUtil::MakeChangeMask(RecordInfo.ChangeMaskOrPtr, (RecordInfo.HasChangeMask ? Info.ChangeMaskBitCount : 1U));
+		FNetBitArrayView LostChangeMask = FChangeMaskUtil::MakeChangeMask(RecordInfo.ChangeMaskOrPtr, static_cast<uint32>(RecordInfo.HasChangeMask ? Info.ChangeMaskBitCount : 1U));
 		if (RecordInfo.HasChangeMask)
 		{
 			// Iterate over all data in flight for this object and mask away any already re-transmitted changes
@@ -2305,7 +2305,7 @@ int FReplicationWriter::PrepareAndSendHugeObjectPayload(FNetSerializationContext
 	const uint32 BitsPerStorageWord = sizeof(HugeObjectStorageType) * 8;
 
 	TArray<HugeObjectStorageType> HugeObjectPayload;
-	HugeObjectPayload.AddUninitialized((PartialNetObjectAttachmentHandler->GetConfig()->GetTotalMaxPayloadBitCount() + (BitsPerStorageWord - 1U))/BitsPerStorageWord);
+	HugeObjectPayload.AddUninitialized(static_cast<int32>(PartialNetObjectAttachmentHandler->GetConfig()->GetTotalMaxPayloadBitCount() + (BitsPerStorageWord - 1U))/BitsPerStorageWord);
 
 	// Setup a special context for the huge object serialization.
 	FNetBitStreamWriter HugeObjectWriter;
@@ -3063,7 +3063,7 @@ UDataStream::EWriteResult FReplicationWriter::Write(FNetSerializationContext& Co
 	FNetBitStreamRollbackScope Rollback(Writer);
 	const uint32 HeaderPos = Writer.GetPosBits();
 
-	uint16 WrittenObjectCount = 0;
+	uint32 WrittenObjectCount = 0;
 	const uint32 OldReplicationInfoCount = ReplicationRecord.GetInfoCount();
 
 	// Written batch count
@@ -3072,7 +3072,7 @@ UDataStream::EWriteResult FReplicationWriter::Write(FNetSerializationContext& Co
 	// Write timestamps etc? Or do we do this in header.
 	// WriteReplicationFrameData();
 
-	const uint16 WrittenObjectsPendingDestroyCount = WriteObjectsPendingDestroy(Context);
+	const uint32 WrittenObjectsPendingDestroyCount = WriteObjectsPendingDestroy(Context);
 	WrittenObjectCount += WrittenObjectsPendingDestroyCount;
 
 	// Only reason for overflow here is if we did not fit header
@@ -3107,14 +3107,14 @@ UDataStream::EWriteResult FReplicationWriter::Write(FNetSerializationContext& Co
 		{
 			// Seek back to HeaderPos and update the header
 			FNetBitStreamWriteScope WriteScope(Writer, HeaderPos);
-			const uint16 TotalWrittenBatchCount = WriteContext.WrittenBatchCount + WrittenObjectsPendingDestroyCount;
+			const uint32 TotalWrittenBatchCount = WriteContext.WrittenBatchCount + WrittenObjectsPendingDestroyCount;
 			Writer.WriteBits(TotalWrittenBatchCount, 16);
 		}
 
 		//UE_LOG_REPLICATIONWRITER(TEXT("FReplicationWriter::Write() Wrote %u Objects for ConnectionId:%u, ReplicationSystemId: %u."), WrittenObjectCount, Parameters.ConnectionId, Parameters.ReplicationSystem->GetId());	
 	
 		// Push record
-		const uint32 ReplicationInfoCount = ReplicationRecord.GetInfoCount() - OldReplicationInfoCount;
+		const uint16 ReplicationInfoCount = static_cast<uint16>(ReplicationRecord.GetInfoCount() - OldReplicationInfoCount);
 		ReplicationRecord.PushRecord(ReplicationInfoCount);
 
 #if UE_NET_VALIDATE_REPLICATION_RECORD
@@ -3339,7 +3339,7 @@ void FReplicationWriter::StopAllReplication()
 		FChangeMaskStorageOrPointer::Free(Info.ChangeMaskOrPtr, Info.ChangeMaskBitCount, s_DefaultChangeMaskAllocator);
 
 		// Release object reference
-		const uint32 InternalIndex = &Info - FirstInfo;
+		const FInternalNetRefIndex InternalIndex = static_cast<FInternalNetRefIndex>(&Info - FirstInfo);
 		NetRefHandleManager->ReleaseNetObjectRef(InternalIndex);
 	}
 }
diff --git a/Engine/Source/Runtime/Experimental/Iris/Core/Private/Iris/Serialization/ArrayPropertyNetSerializer.cpp b/Engine/Source/Runtime/Experimental/Iris/Core/Private/Iris/Serialization/ArrayPropertyNetSerializer.cpp
--- a/Engine/Source/Runtime/Experimental/Iris/Core/Private/Iris/Serialization/ArrayPropertyNetSerializer.cpp
+++ b/Engine/Source/Runtime/Experimental/Iris/Core/Private/Iris/Serialization/ArrayPropertyNetSerializer.cpp
@@ -170,8 +170,8 @@ void FArrayPropertyNetSerializer::CollectNetReferences(FNetSerializationContext&
 		// We currently use a simple modulo scheme for bits in the changemask, if we have more elements in the array then is covered by the changemask
 		// several entries in the array will be considered dirty and be serialized
 		// As the first bit is used by the owning property we need to offset by one and deduct one from the usable bits
-		const uint16 ChangeMaskBitOffset = Args.ChangeMaskInfo.BitOffset + 1U;
-		const uint16 ChangeMaskBitCount = Args.ChangeMaskInfo.BitCount - 1U;
+		const uint32 ChangeMaskBitOffset = Args.ChangeMaskInfo.BitOffset + 1U;
+		const uint32 ChangeMaskBitCount = Args.ChangeMaskInfo.BitCount - 1U;
 
 		// This is the info that will be stored in the entries
 		for (uint32 ElementIt = 0, ElementEndIt = Array.ElementCount; ElementIt < ElementEndIt; ++ElementIt)
@@ -219,7 +219,7 @@ void FArrayPropertyNetSerializer::Deserialize(FNetSerializationContext& Context,
 	const FReplicationStateMemberSerializerDescriptor& ElementSerializerDescriptor = ElementStateDescriptor->MemberSerializerDescriptors[0];
 
 	// Memory management
-	AdjustArraySize(Context, Array, Config, NewElementCount);
+	AdjustArraySize(Context, Array, Config, static_cast<uint16>(NewElementCount));
 
 	// Deserialize
 	FNetDeserializeArgs ElementArgs;
@@ -395,7 +395,7 @@ void FArrayPropertyNetSerializer::DeserializeDelta(FNetSerializationContext& Con
 	const FNetSerializer* ElementSerializer = ElementSerializerDescriptor.Serializer;
 	const uint32 ElementSize = ElementStateDescriptor->InternalSize;
 
-	AdjustArraySize(Context, Array, Config, ElementCount);
+	AdjustArraySize(Context, Array, Config, static_cast<uint16>(ElementCount));
 
 	// Elements in the current array up to the previous size can use delta serialization.
 	if (PrevElementCount)
@@ -469,9 +469,15 @@ void FArrayPropertyNetSerializer::Quantize(FNetSerializationContext& Context, co
 
 	FScriptArrayHelper ScriptArrayHelper(Config->Property.Get(), reinterpret_cast<const void*>(Args.Source));
 	const uint32 ElementCount = static_cast<uint32>(ScriptArrayHelper.Num());
-	QuantizedType& TargetArray = *reinterpret_cast<QuantizedType*>(Args.Target);
 
-	AdjustArraySize(Context, TargetArray, Config, ElementCount);
+	if (ElementCount > Config->MaxElementCount)
+	{
+		Context.SetError(GNetError_ArraySizeTooLarge);
+		return;
+	}
+
+	QuantizedType& TargetArray = *reinterpret_cast<QuantizedType*>(Args.Target);
+	AdjustArraySize(Context, TargetArray, Config, static_cast<uint16>(ElementCount));
 
 	const FReplicationStateDescriptor* ElementStateDescriptor = Config->StateDescriptor;
 	const FReplicationStateMemberSerializerDescriptor& ElementSerializerDescriptor = ElementStateDescriptor->MemberSerializerDescriptors[0];
@@ -663,21 +669,20 @@ void FArrayPropertyNetSerializer::CloneDynamicState(FNetSerializationContext& Co
 
 	const SIZE_T ElementSize = ElementStateDescriptor->InternalSize;
 	const SIZE_T ElementAlignment = ElementStateDescriptor->InternalAlignment;
-	const uint32 ElementCount = SourceArray.ElementCount;
 
 	// Clone storage
 	void* ElementStorage = nullptr;
-	if (ElementCount > 0)
+	if (SourceArray.ElementCount > 0)
 	{
-		ElementStorage = Context.GetInternalContext()->Alloc(ElementSize*ElementCount, ElementAlignment);
-		FMemory::Memcpy(ElementStorage, SourceArray.ElementStorage, ElementSize*ElementCount);
+		ElementStorage = Context.GetInternalContext()->Alloc(ElementSize*SourceArray.ElementCount, ElementAlignment);
+		FMemory::Memcpy(ElementStorage, SourceArray.ElementStorage, ElementSize*SourceArray.ElementCount);
 	}
-	TargetArray.ElementCapacityCount = ElementCount;
-	TargetArray.ElementCount = ElementCount;
+	TargetArray.ElementCapacityCount = SourceArray.ElementCount;
+	TargetArray.ElementCount = SourceArray.ElementCount;
 	TargetArray.ElementStorage = ElementStorage;
 
 	// If no member has dynamic state then there's nothing more to do.
-	if (!EnumHasAnyFlags(ElementStateDescriptor->Traits, EReplicationStateTraits::HasDynamicState) || ElementCount == 0)
+	if (!EnumHasAnyFlags(ElementStateDescriptor->Traits, EReplicationStateTraits::HasDynamicState) || SourceArray.ElementCount == 0)
 	{
 		return;
 	}
@@ -691,7 +696,7 @@ void FArrayPropertyNetSerializer::CloneDynamicState(FNetSerializationContext& Co
 	ElementArgs.Target = NetSerializerValuePointer(TargetArray.ElementStorage);
 
 	const FNetSerializer* ElementSerializer = ElementSerializerDescriptor.Serializer;
-	for (uint32 ElementIt = 0, ElementEndIt = ElementCount; ElementIt < ElementEndIt; ++ElementIt)
+	for (uint32 ElementIt = 0, ElementEndIt = SourceArray.ElementCount; ElementIt < ElementEndIt; ++ElementIt)
 	{
 		ElementSerializer->CloneDynamicState(Context, ElementArgs);
 		ElementArgs.Source += ElementSize;
diff --git a/Engine/Source/Runtime/Experimental/Iris/Core/Private/Iris/Serialization/BitfieldNetSerializer.cpp b/Engine/Source/Runtime/Experimental/Iris/Core/Private/Iris/Serialization/BitfieldNetSerializer.cpp
--- a/Engine/Source/Runtime/Experimental/Iris/Core/Private/Iris/Serialization/BitfieldNetSerializer.cpp
+++ b/Engine/Source/Runtime/Experimental/Iris/Core/Private/Iris/Serialization/BitfieldNetSerializer.cpp
@@ -38,7 +38,7 @@ void FBitfieldNetSerializer::Serialize(FNetSerializationContext& Context, const
 void FBitfieldNetSerializer::Deserialize(FNetSerializationContext& Context, const FNetDeserializeArgs& Args)
 {
 	FNetBitStreamReader* Reader = Context.GetBitStreamReader();
-	const QuantizedType Value = Reader->ReadBits(1U);
+	const QuantizedType Value = static_cast<QuantizedType>(Reader->ReadBits(1U));
 
 	QuantizedType* Target = reinterpret_cast<QuantizedType*>(Args.Target);
 	*Target = Value;
diff --git a/Engine/Source/Runtime/Experimental/Iris/Core/Private/Iris/Serialization/BoolNetSerializer.cpp b/Engine/Source/Runtime/Experimental/Iris/Core/Private/Iris/Serialization/BoolNetSerializer.cpp
--- a/Engine/Source/Runtime/Experimental/Iris/Core/Private/Iris/Serialization/BoolNetSerializer.cpp
+++ b/Engine/Source/Runtime/Experimental/Iris/Core/Private/Iris/Serialization/BoolNetSerializer.cpp
@@ -41,7 +41,7 @@ void FBoolNetSerializer::Serialize(FNetSerializationContext& Context, const FNet
 void FBoolNetSerializer::Deserialize(FNetSerializationContext& Context, const FNetDeserializeArgs& Args)
 {
 	FNetBitStreamReader* Reader = Context.GetBitStreamReader();
-	const SourceType Value = Reader->ReadBits(1U);
+	const SourceType Value = static_cast<SourceType>(Reader->ReadBits(1U));
 
 	SourceType* Target = reinterpret_cast<SourceType*>(Args.Target);
 	*Target = Value;
diff --git a/Engine/Source/Runtime/Experimental/Iris/Core/Private/Iris/Serialization/IntNetSerializerBase.h b/Engine/Source/Runtime/Experimental/Iris/Core/Private/Iris/Serialization/IntNetSerializerBase.h
--- a/Engine/Source/Runtime/Experimental/Iris/Core/Private/Iris/Serialization/IntNetSerializerBase.h
+++ b/Engine/Source/Runtime/Experimental/Iris/Core/Private/Iris/Serialization/IntNetSerializerBase.h
@@ -7,8 +7,7 @@
 #include "Iris/Serialization/NetBitStreamReader.h"
 #include "Iris/Serialization/NetBitStreamWriter.h"
 #include "Iris/Core/BitTwiddling.h"
-#include "Traits/IntType.h"
-#include "Templates/IsSigned.h"
+#include <type_traits>
 
 namespace UE::Net::Private
 {
@@ -19,9 +18,9 @@ struct FIntNetSerializerBase
 {
 	static const uint32 Version = 0;
 
-	typedef InSourceType SourceType;
+	using SourceType = InSourceType;
 	// For convenience we'll use an unsigned type as the quantized type
-	typedef typename TUnsignedIntType<sizeof(SourceType)>::Type QuantizedType;
+	using QuantizedType = std::make_unsigned_t<SourceType>;
 	typedef InConfigType ConfigType;
 
 	static void Serialize(FNetSerializationContext& Context, const FNetSerializeArgs& Args)
@@ -86,11 +85,11 @@ struct FIntNetSerializerBase
 		}
 		else
 		{
-			Value = Reader->ReadBits(BitCount);
+			Value = static_cast<QuantizedType>(Reader->ReadBits(BitCount));
 		}
 
 		// Sign-extend the value if needed
-		if constexpr (TIsSigned<SourceType>::Value)
+		if constexpr (std::is_signed_v<SourceType>)
 		{
 			const QuantizedType SignMask = QuantizedType(1) << (BitCount - 1U);
 			Value = (Value ^ SignMask) - SignMask;
@@ -118,7 +117,7 @@ struct FIntNetSerializerBase
 		const uint32 BitCountTableEntryCount = DeltaBitCountTableEntryCount[IndexToBitCountEntries];
 
 		FNetBitStreamWriter* Writer = Context.GetBitStreamWriter();
-		if constexpr (TIsSigned<SourceType>::Value)
+		if constexpr (std::is_signed_v<SourceType>)
 		{
 			SerializeIntDelta(*Writer, Value, PrevValue, BitCountTable, BitCountTableEntryCount, MaxBitCount);
 		}
@@ -147,7 +146,7 @@ struct FIntNetSerializerBase
 		const uint32 BitCountTableEntryCount = DeltaBitCountTableEntryCount[IndexToBitCountEntries];
 
 		FNetBitStreamReader* Reader = Context.GetBitStreamReader();
-		if constexpr (TIsSigned<SourceType>::Value)
+		if constexpr (std::is_signed_v<SourceType>)
 		{
 			DeserializeIntDelta(*Reader, Target, PrevValue, BitCountTable, BitCountTableEntryCount, MaxBitCount);
 		}
@@ -165,7 +164,7 @@ struct FIntNetSerializerBase
 		const uint32 BitCount = Config->BitCount;
 
 		QuantizedType Value;
-		if constexpr (TIsSigned<SourceType>::Value)
+		if constexpr (std::is_signed_v<SourceType>)
 		{
 			const QuantizedType SignMask = (QuantizedType(1) << (BitCount - 1U));
 			const QuantizedType ValueMask = SignMask | (SignMask - QuantizedType(1));
diff --git a/Engine/Source/Runtime/Experimental/Iris/Core/Private/Iris/Serialization/InternalEnumNetSerializers.cpp b/Engine/Source/Runtime/Experimental/Iris/Core/Private/Iris/Serialization/InternalEnumNetSerializers.cpp
--- a/Engine/Source/Runtime/Experimental/Iris/Core/Private/Iris/Serialization/InternalEnumNetSerializers.cpp
+++ b/Engine/Source/Runtime/Experimental/Iris/Core/Private/Iris/Serialization/InternalEnumNetSerializers.cpp
@@ -24,7 +24,7 @@ public:
 			ensure(EnumValueCount > 0);
 			OutConfig.LowerBound = 0;
 			OutConfig.UpperBound = 0;
-			OutConfig.BitCount = GetBitsNeededForRange(OutConfig.LowerBound, OutConfig.UpperBound);
+			OutConfig.BitCount = static_cast<uint8>(GetBitsNeededForRange(OutConfig.LowerBound, OutConfig.UpperBound));
 			OutConfig.Enum = Enum;
 		}
 		else
@@ -41,7 +41,7 @@ public:
 
 			OutConfig.LowerBound = static_cast<SourceType>(SmallestValue);
 			OutConfig.UpperBound = static_cast<SourceType>(LargestValue);
-			OutConfig.BitCount = GetBitsNeededForRange(OutConfig.LowerBound, OutConfig.UpperBound);
+			OutConfig.BitCount = static_cast<uint8>(GetBitsNeededForRange(OutConfig.LowerBound, OutConfig.UpperBound));
 			OutConfig.Enum = Enum;
 		}
 
diff --git a/Engine/Source/Runtime/Experimental/Iris/Core/Private/Iris/Serialization/InternalNetSerializationContext.cpp b/Engine/Source/Runtime/Experimental/Iris/Core/Private/Iris/Serialization/InternalNetSerializationContext.cpp
--- a/Engine/Source/Runtime/Experimental/Iris/Core/Private/Iris/Serialization/InternalNetSerializationContext.cpp
+++ b/Engine/Source/Runtime/Experimental/Iris/Core/Private/Iris/Serialization/InternalNetSerializationContext.cpp
@@ -18,7 +18,7 @@ FInternalNetSerializationContext::FInternalNetSerializationContext(UReplicationS
 
 void* FInternalNetSerializationContext::Alloc(SIZE_T Size, SIZE_T Alignment)
 {
-	return GMalloc->Malloc(Size, Alignment);
+	return GMalloc->Malloc(Size, static_cast<uint32>(Alignment));
 }
 
 void FInternalNetSerializationContext::Free(void* Ptr)
diff --git a/Engine/Source/Runtime/Experimental/Iris/Core/Private/Iris/Serialization/IrisObjectReferencePackageMap.cpp b/Engine/Source/Runtime/Experimental/Iris/Core/Private/Iris/Serialization/IrisObjectReferencePackageMap.cpp
--- a/Engine/Source/Runtime/Experimental/Iris/Core/Private/Iris/Serialization/IrisObjectReferencePackageMap.cpp
+++ b/Engine/Source/Runtime/Experimental/Iris/Core/Private/Iris/Serialization/IrisObjectReferencePackageMap.cpp
@@ -16,9 +16,10 @@ bool UIrisObjectReferencePackageMap::SerializeObject(FArchive& Ar, UClass* InCla
 		{
 			Index = References->Add(Obj);
 		}
-		uint8 IndexByte = Index;
-		if (ensureAlwaysMsgf(IndexByte < References->Num(), TEXT("UIrisObjectReferencePackageMap::SerializeObject, failed to serialize object reference. A Maximum of 256 references are currently supported by this PackageMap")))
+		constexpr int32 MaxNumReferences = 256;
+		if (ensureAlwaysMsgf(Index > 0 && Index < MaxNumReferences, TEXT("UIrisObjectReferencePackageMap::SerializeObject, failed to serialize object reference. A Maximum of 256 references are currently supported by this PackageMap")))
 		{
+			uint8 IndexByte = static_cast<uint8>(Index);
 			Ar << IndexByte;
 		}
 		else
diff --git a/Engine/Source/Runtime/Experimental/Iris/Core/Private/Iris/Serialization/LastResortPropertyNetSerializer.cpp b/Engine/Source/Runtime/Experimental/Iris/Core/Private/Iris/Serialization/LastResortPropertyNetSerializer.cpp
--- a/Engine/Source/Runtime/Experimental/Iris/Core/Private/Iris/Serialization/LastResortPropertyNetSerializer.cpp
+++ b/Engine/Source/Runtime/Experimental/Iris/Core/Private/Iris/Serialization/LastResortPropertyNetSerializer.cpp
@@ -46,6 +46,8 @@ struct FLastResortPropertyNetSerializer
 	static void FreeDynamicState(FNetSerializationContext&, const FNetFreeDynamicStateArgs&);
 
 private:
+	static constexpr uint32 AllocationAlignment = 4U;
+
 	static void FreeDynamicStateInternal(FNetSerializationContext&, QuantizedType& Value);
 	static void GrowDynamicStateInternal(FNetSerializationContext&, QuantizedType& Value, uint16 NewBitCount);
 	static void ShrinkDynamicStateInternal(FNetSerializationContext&, QuantizedType& Value, uint16 NewBitCount);
@@ -74,11 +76,11 @@ void FLastResortPropertyNetSerializer::Deserialize(FNetSerializationContext& Con
 	const uint32 NewBitCount = ReadPackedUint32(Reader);
 	if (NewBitCount > 65535U)
 	{
-		Reader->DoOverflow();
+		Context.SetError(GNetError_ArraySizeTooLarge);
 		return;
 	}
 
-	AdjustStorageSize(Context, Value, NewBitCount);
+	AdjustStorageSize(Context, Value, static_cast<uint16>(NewBitCount));
 
 	Reader->ReadBitStream(static_cast<uint32*>(Value.Storage), NewBitCount);
 }
@@ -91,15 +93,15 @@ void FLastResortPropertyNetSerializer::Quantize(FNetSerializationContext& Contex
 	FNetBitWriter Archive(8192);
 	Property->NetSerializeItem(Archive, static_cast<UPackageMap*>(nullptr), reinterpret_cast<void*>(Args.Source));
 
-	const uint32 BitCount = Archive.GetNumBits();
-	if (BitCount > 65535)
+	const uint64 BitCount = Archive.GetNumBits();
+	if (BitCount > 65535U)
 	{
-		Context.SetError(GNetError_BitStreamOverflow);
+		Context.SetError(GNetError_ArraySizeTooLarge);
 		return;
 	}
 
 	QuantizedType& Value = *reinterpret_cast<QuantizedType*>(Args.Target);
-	AdjustStorageSize(Context, Value, BitCount);
+	AdjustStorageSize(Context, Value, static_cast<uint16>(BitCount));
 	if (BitCount > 0)
 	{
 		FMemory::Memcpy(Value.Storage, Archive.GetData(), (BitCount + 7U)/8U);
@@ -126,7 +128,7 @@ bool FLastResortPropertyNetSerializer::IsEqual(FNetSerializationContext& Context
 		{
 			return false;
 		}
-		const bool bIsEqual = FMemory::Memcmp(Value0.Storage, Value1.Storage, Align((Value0.BitCount + 7U)/8U, 4U)) == 0;
+		const bool bIsEqual = FMemory::Memcmp(Value0.Storage, Value1.Storage, Align((Value0.BitCount + 7U)/8U, AllocationAlignment)) == 0;
 		return bIsEqual;
 	}
 	else
@@ -146,13 +148,12 @@ void FLastResortPropertyNetSerializer::CloneDynamicState(FNetSerializationContex
 	QuantizedType& Target = *reinterpret_cast<QuantizedType*>(Args.Target);
 	const QuantizedType& Source = *reinterpret_cast<const QuantizedType*>(Args.Source);
 
-	constexpr SIZE_T Alignment = 4U;
-	const uint32 ByteCount = Align((Source.BitCount + 7U)/8U, 4U);
+	const uint16 ByteCount = static_cast<uint16>(Align((Source.BitCount + 7U)/8U, AllocationAlignment));
 
 	void* Storage = nullptr;
 	if (ByteCount > 0)
 	{
-		Storage = Context.GetInternalContext()->Alloc(ByteCount, Alignment);
+		Storage = Context.GetInternalContext()->Alloc(ByteCount, AllocationAlignment);
 		FMemory::Memcpy(Storage, Source.Storage, ByteCount);
 	}
 	Target.ByteCapacity = ByteCount;
@@ -177,13 +178,12 @@ void FLastResortPropertyNetSerializer::GrowDynamicStateInternal(FNetSerializatio
 {
 	checkSlow(NewBitCount > Value.BitCount);
 
-	constexpr SIZE_T Alignment = 4U;
-	const uint32 ByteCount = Align((NewBitCount + 7U)/8U, 4U);
+	const uint16 ByteCount = static_cast<uint16>(Align((NewBitCount + 7U)/8U, AllocationAlignment));
 
 	// We don't support delta compression for the unknown contents of the bits so we don't need to copy the old data.
 	Context.GetInternalContext()->Free(Value.Storage);
 
-	void* Storage = Context.GetInternalContext()->Alloc(ByteCount, Alignment);
+	void* Storage = Context.GetInternalContext()->Alloc(ByteCount, AllocationAlignment);
 
 	// Clear the last word to support IsEqual Memcmp optimization.
 	const uint32 LastWordIndex = ByteCount/4U - 1U;
@@ -196,7 +196,7 @@ void FLastResortPropertyNetSerializer::GrowDynamicStateInternal(FNetSerializatio
 
 void FLastResortPropertyNetSerializer::AdjustStorageSize(FNetSerializationContext& Context, QuantizedType& Value, uint16 NewBitCount)
 {
-	const uint32 NewByteCapacity = Align((NewBitCount + 7U)/8U, 4U);
+	const uint32 NewByteCapacity = Align((NewBitCount + 7U)/8U, AllocationAlignment);
 	if (NewByteCapacity == 0)
 	{
 		// Free everything
diff --git a/Engine/Source/Runtime/Experimental/Iris/Core/Private/Iris/Serialization/NetBitStreamUtil.cpp b/Engine/Source/Runtime/Experimental/Iris/Core/Private/Iris/Serialization/NetBitStreamUtil.cpp
--- a/Engine/Source/Runtime/Experimental/Iris/Core/Private/Iris/Serialization/NetBitStreamUtil.cpp
+++ b/Engine/Source/Runtime/Experimental/Iris/Core/Private/Iris/Serialization/NetBitStreamUtil.cpp
@@ -148,7 +148,7 @@ uint16 ReadPackedUint16(FNetBitStreamReader* Reader)
 	const uint32 ByteCountToRead = Reader->ReadBits(1U) + 1U;
 	const uint32 BitCountToRead = ByteCountToRead * 8U;
 
-	const uint16 Value = Reader->ReadBits(BitCountToRead);
+	const uint16 Value = static_cast<uint16>(Reader->ReadBits(BitCountToRead));
 	return Value;
 }
 
diff --git a/Engine/Source/Runtime/Experimental/Iris/Core/Private/Iris/Serialization/NetRoleNetSerializer.cpp b/Engine/Source/Runtime/Experimental/Iris/Core/Private/Iris/Serialization/NetRoleNetSerializer.cpp
--- a/Engine/Source/Runtime/Experimental/Iris/Core/Private/Iris/Serialization/NetRoleNetSerializer.cpp
+++ b/Engine/Source/Runtime/Experimental/Iris/Core/Private/Iris/Serialization/NetRoleNetSerializer.cpp
@@ -110,15 +110,15 @@ bool PartialInitNetRoleSerializerConfig(FNetRoleNetSerializerConfig& OutConfig,
 		check(EnumConfig.LowerBound == 0);
 
 		const int64 AutonomousProxyValue = Enum->GetValueByNameString(TEXT("ROLE_AutonomousProxy"), EGetByNameFlags::CaseSensitive);
-		check(AutonomousProxyValue != INDEX_NONE);
+		check(AutonomousProxyValue > 0 && AutonomousProxyValue < 256);
 		const int64 SimulatedProxyValue = Enum->GetValueByNameString(TEXT("ROLE_SimulatedProxy"), EGetByNameFlags::CaseSensitive);
-		check(SimulatedProxyValue != INDEX_NONE);
+		check(SimulatedProxyValue > 0 && SimulatedProxyValue < 256);
 		
 		FNetRoleNetSerializer::CachedConfig.LowerBound = EnumConfig.LowerBound;
 		FNetRoleNetSerializer::CachedConfig.UpperBound = EnumConfig.UpperBound;
 		FNetRoleNetSerializer::CachedConfig.BitCount = EnumConfig.BitCount;
-		FNetRoleNetSerializer::CachedConfig.AutonomousProxyValue = AutonomousProxyValue;
-		FNetRoleNetSerializer::CachedConfig.SimulatedProxyValue = SimulatedProxyValue;
+		FNetRoleNetSerializer::CachedConfig.AutonomousProxyValue = static_cast<uint8>(AutonomousProxyValue);
+		FNetRoleNetSerializer::CachedConfig.SimulatedProxyValue = static_cast<uint8>(SimulatedProxyValue);
 		FNetRoleNetSerializer::CachedConfig.Enum = Enum;
 	}
 
diff --git a/Engine/Source/Runtime/Experimental/Iris/Core/Private/Iris/Serialization/PackedVectorNetSerializers.cpp b/Engine/Source/Runtime/Experimental/Iris/Core/Private/Iris/Serialization/PackedVectorNetSerializers.cpp
--- a/Engine/Source/Runtime/Experimental/Iris/Core/Private/Iris/Serialization/PackedVectorNetSerializers.cpp
+++ b/Engine/Source/Runtime/Experimental/Iris/Core/Private/Iris/Serialization/PackedVectorNetSerializers.cpp
@@ -392,15 +392,15 @@ void FPackedVectorNetSerializerBase::SerializeDelta(uint32 ScaleBitCount, FNetSe
 				{
 					if (XDiffers)
 					{
-						Writer->WriteBits(Value.X, 32U);
+						Writer->WriteBits(static_cast<uint32>(Value.X), 32U);
 					}
 					if (YDiffers)
 					{
-						Writer->WriteBits(Value.Y, 32U);
+						Writer->WriteBits(static_cast<uint32>(Value.Y), 32U);
 					}
 					if (ZDiffers)
 					{
-						Writer->WriteBits(Value.Z, 32U);
+						Writer->WriteBits(static_cast<uint32>(Value.Z), 32U);
 					}
 				}
 			}
@@ -822,9 +822,9 @@ void FVectorNetQuantizeNormalNetSerializer::Quantize(FNetSerializationContext&,
 	QuantizedType& Target = *reinterpret_cast<QuantizedType*>(Args.Target);
 
 	QuantizedType TempValue = {};
-	TempValue.X = QuantizeSignedUnitFloat(Source.X, PrecisionBitCount);
-	TempValue.Y = QuantizeSignedUnitFloat(Source.Y, PrecisionBitCount);
-	TempValue.Z = QuantizeSignedUnitFloat(Source.Z, PrecisionBitCount);
+	TempValue.X = QuantizeSignedUnitFloat(static_cast<float>(Source.X), PrecisionBitCount);
+	TempValue.Y = QuantizeSignedUnitFloat(static_cast<float>(Source.Y), PrecisionBitCount);
+	TempValue.Z = QuantizeSignedUnitFloat(static_cast<float>(Source.Z), PrecisionBitCount);
 
 	Target = TempValue;
 }
@@ -857,17 +857,17 @@ bool FVectorNetQuantizeNormalNetSerializer::IsEqual(FNetSerializationContext& Co
 		const SourceType& SourceValue0 = *reinterpret_cast<const SourceType*>(Args.Source0);
 		const SourceType& SourceValue1 = *reinterpret_cast<const SourceType*>(Args.Source1);
 
-		if (QuantizeSignedUnitFloat(SourceValue0.X, PrecisionBitCount) != QuantizeSignedUnitFloat(SourceValue1.X, PrecisionBitCount))
+		if (QuantizeSignedUnitFloat(static_cast<float>(SourceValue0.X), PrecisionBitCount) != QuantizeSignedUnitFloat(static_cast<float>(SourceValue1.X), PrecisionBitCount))
 		{
 			return false;
 		}
 
-		if (QuantizeSignedUnitFloat(SourceValue0.Y, PrecisionBitCount) != QuantizeSignedUnitFloat(SourceValue1.Y, PrecisionBitCount))
+		if (QuantizeSignedUnitFloat(static_cast<float>(SourceValue0.Y), PrecisionBitCount) != QuantizeSignedUnitFloat(static_cast<float>(SourceValue1.Y), PrecisionBitCount))
 		{
 			return false;
 		}
 
-		if (QuantizeSignedUnitFloat(SourceValue0.Z, PrecisionBitCount) != QuantizeSignedUnitFloat(SourceValue1.Z, PrecisionBitCount))
+		if (QuantizeSignedUnitFloat(static_cast<float>(SourceValue0.Z), PrecisionBitCount) != QuantizeSignedUnitFloat(static_cast<float>(SourceValue1.Z), PrecisionBitCount))
 		{
 			return false;
 		}
diff --git a/Engine/Source/Runtime/Experimental/Iris/Core/Private/Iris/Serialization/RotatorNetSerializers.cpp b/Engine/Source/Runtime/Experimental/Iris/Core/Private/Iris/Serialization/RotatorNetSerializers.cpp
--- a/Engine/Source/Runtime/Experimental/Iris/Core/Private/Iris/Serialization/RotatorNetSerializers.cpp
+++ b/Engine/Source/Runtime/Experimental/Iris/Core/Private/Iris/Serialization/RotatorNetSerializers.cpp
@@ -142,11 +142,11 @@ void FRotatorAsShortNetSerializer::Deserialize(FNetSerializationContext& Context
 	FNetBitStreamReader* Reader = Context.GetBitStreamReader();
 
 	QuantizedType Value;
-	const uint32 XYZIsNotZero = Reader->ReadBits(3U);
+	const uint16 XYZIsNotZero = static_cast<uint16>(Reader->ReadBits(3U));
 	Value.XYZIsNotZero = XYZIsNotZero;
-	Value.X = XYZIsNotZero & XDiffersMask ? Reader->ReadBits(16U) : 0U;
-	Value.Y = XYZIsNotZero & YDiffersMask ? Reader->ReadBits(16U) : 0U;
-	Value.Z = XYZIsNotZero & ZDiffersMask ? Reader->ReadBits(16U) : 0U;
+	Value.X = static_cast<uint16>(XYZIsNotZero & XDiffersMask ? Reader->ReadBits(16U) : 0U);
+	Value.Y = static_cast<uint16>(XYZIsNotZero & YDiffersMask ? Reader->ReadBits(16U) : 0U);
+	Value.Z = static_cast<uint16>(XYZIsNotZero & ZDiffersMask ? Reader->ReadBits(16U) : 0U);
 
 	QuantizedType& Target = *reinterpret_cast<QuantizedType*>(Args.Target);
 	Target = Value;
@@ -196,22 +196,22 @@ void FRotatorAsShortNetSerializer::DeserializeDelta(FNetSerializationContext& Co
 	const uint32 XYZDiffers = Reader->ReadBits(3U);
 	if (XYZDiffers & XDiffersMask)
 	{
-		const uint16 DX = Reader->ReadBits(16U);
+		const uint16 DX = static_cast<uint16>(Reader->ReadBits(16U));
 		TempValue.X += DX;
 	}
 	if (XYZDiffers & YDiffersMask)
 	{
-		const uint16 DY = Reader->ReadBits(16U);
+		const uint16 DY = static_cast<uint16>(Reader->ReadBits(16U));
 		TempValue.Y += DY;
 	}
 	if (XYZDiffers & ZDiffersMask)
 	{
-		const uint16 DZ = Reader->ReadBits(16U);
+		const uint16 DZ = static_cast<uint16>(Reader->ReadBits(16U));
 		TempValue.Z += DZ;
 	}
 
 	// Reconstruct flags
-	uint32 XYZIsNotZero = 0U;
+	uint16 XYZIsNotZero = 0U;
 	XYZIsNotZero |= (TempValue.X != 0 ? XDiffersMask : 0U);
 	XYZIsNotZero |= (TempValue.Y != 0 ? YDiffersMask : 0U);
 	XYZIsNotZero |= (TempValue.Z != 0 ? ZDiffersMask : 0U);
@@ -308,11 +308,11 @@ void FRotatorAsByteNetSerializer::Deserialize(FNetSerializationContext& Context,
 	FNetBitStreamReader* Reader = Context.GetBitStreamReader();
 
 	QuantizedType Value;
-	const uint32 XYZIsNotZero = Reader->ReadBits(3U);
+	const uint8 XYZIsNotZero = static_cast<uint8>(Reader->ReadBits(3U));
 	Value.XYZIsNotZero = XYZIsNotZero;
-	Value.X = XYZIsNotZero & XDiffersMask ? Reader->ReadBits(8U) : 0U;
-	Value.Y = XYZIsNotZero & YDiffersMask ? Reader->ReadBits(8U) : 0U;
-	Value.Z = XYZIsNotZero & ZDiffersMask ? Reader->ReadBits(8U) : 0U;
+	Value.X = static_cast<uint8>(XYZIsNotZero & XDiffersMask ? Reader->ReadBits(8U) : 0U);
+	Value.Y = static_cast<uint8>(XYZIsNotZero & YDiffersMask ? Reader->ReadBits(8U) : 0U);
+	Value.Z = static_cast<uint8>(XYZIsNotZero & ZDiffersMask ? Reader->ReadBits(8U) : 0U);
 
 	QuantizedType& Target = *reinterpret_cast<QuantizedType*>(Args.Target);
 	Target = Value;
@@ -354,22 +354,22 @@ void FRotatorAsByteNetSerializer::DeserializeDelta(FNetSerializationContext& Con
 	const uint32 XYZDiffers = Reader->ReadBits(3U);
 	if (XYZDiffers & XDiffersMask)
 	{
-		TempValue.X = Reader->ReadBits(8U);
+		TempValue.X = static_cast<uint8>(Reader->ReadBits(8U));
 	}
 	if (XYZDiffers & YDiffersMask)
 	{
-		TempValue.Y = Reader->ReadBits(8U);
+		TempValue.Y = static_cast<uint8>(Reader->ReadBits(8U));
 	}
 	if (XYZDiffers & ZDiffersMask)
 	{
-		TempValue.Z = Reader->ReadBits(8U);
+		TempValue.Z = static_cast<uint8>(Reader->ReadBits(8U));
 	}
 
 	// Reconstruct flags
-	uint32 XYZIsNotZero = 0U;
-	XYZIsNotZero |= (TempValue.X != 0 ? XDiffersMask : 0U);
-	XYZIsNotZero |= (TempValue.Y != 0 ? YDiffersMask : 0U);
-	XYZIsNotZero |= (TempValue.Z != 0 ? ZDiffersMask : 0U);
+	uint8 XYZIsNotZero = 0U;
+	XYZIsNotZero |= static_cast<uint8>(TempValue.X != 0 ? XDiffersMask : 0U);
+	XYZIsNotZero |= static_cast<uint8>(TempValue.Y != 0 ? YDiffersMask : 0U);
+	XYZIsNotZero |= static_cast<uint8>(TempValue.Z != 0 ? ZDiffersMask : 0U);
 	TempValue.XYZIsNotZero = XYZIsNotZero;
 
 	QuantizedType& Target = *reinterpret_cast<QuantizedType*>(Args.Target);
diff --git a/Engine/Source/Runtime/Experimental/Iris/Core/Private/Iris/Serialization/StringNetSerializerUtils.cpp b/Engine/Source/Runtime/Experimental/Iris/Core/Private/Iris/Serialization/StringNetSerializerUtils.cpp
--- a/Engine/Source/Runtime/Experimental/Iris/Core/Private/Iris/Serialization/StringNetSerializerUtils.cpp
+++ b/Engine/Source/Runtime/Experimental/Iris/Core/Private/Iris/Serialization/StringNetSerializerUtils.cpp
@@ -46,7 +46,7 @@ void FStringNetSerializerBase::Deserialize(FNetSerializationContext& Context, co
 	}
 
 	Target.bIsEncoded = bIsEncoded;
-	FStringNetSerializerUtils::AdjustArraySize<QuantizedType, uint8>(Context, Target, (NewElementCount > 0 ? NewElementCount + 1 : 0));
+	FStringNetSerializerUtils::AdjustArraySize<QuantizedType, uint8>(Context, Target, static_cast<uint16>(NewElementCount > 0 ? NewElementCount + 1 : 0));
 	if (NewElementCount > 0)
 	{
 		Reader->ReadBitStream(static_cast<uint32*>(Target.ElementStorage), NewElementCount*8U);
@@ -81,17 +81,28 @@ void FStringNetSerializerBase::Quantize(FNetSerializationContext& Context, const
 	if (Target.bIsEncoded)
 	{
 		const uint32 StringLength = Source.Len() + 1U;
-		FStringNetSerializerUtils::AdjustArraySize<QuantizedType, uint8>(Context, Target, 3U*StringLength);
+		constexpr uint32 MaxStringLength = 65535U/3U;
+		if (StringLength > MaxStringLength)
+		{
+			Context.SetError(GNetError_ArraySizeTooLarge);
+			return;
+		}
+		FStringNetSerializerUtils::AdjustArraySize<QuantizedType, uint8>(Context, Target, static_cast<uint16>(3U*StringLength));
 
 		uint32 OutDestLen = 0;
 		uint8* EncodingBuffer = static_cast<uint8*>(Target.ElementStorage);
 		const bool bEncodingSuccess = FStringNetSerializerUtils::TStringCodec<FString::ElementType>::Encode(EncodingBuffer, 3U*StringLength, GetData(Source), StringLength, OutDestLen);
-		Target.ElementCount = OutDestLen;
+		Target.ElementCount = static_cast<uint16>(OutDestLen);
 	}
 	else
 	{
 		const uint32 StringLength = Source.Len();
-		FStringNetSerializerUtils::AdjustArraySize<QuantizedType, uint8>(Context, Target, StringLength > 0 ? StringLength + 1 : 0);
+		if (StringLength + 1U > 65535U)
+		{
+			Context.SetError(GNetError_ArraySizeTooLarge);
+			return;
+		}
+		FStringNetSerializerUtils::AdjustArraySize<QuantizedType, uint8>(Context, Target, static_cast<uint16>(StringLength > 0 ? StringLength + 1 : 0));
 		if (StringLength > 0)
 		{
 			const TCHAR* SourceString = GetData(Source);
diff --git a/Engine/Source/Runtime/Experimental/Iris/Core/Private/Iris/Serialization/StringNetSerializerUtils.h b/Engine/Source/Runtime/Experimental/Iris/Core/Private/Iris/Serialization/StringNetSerializerUtils.h
--- a/Engine/Source/Runtime/Experimental/Iris/Core/Private/Iris/Serialization/StringNetSerializerUtils.h
+++ b/Engine/Source/Runtime/Experimental/Iris/Core/Private/Iris/Serialization/StringNetSerializerUtils.h
@@ -412,16 +412,15 @@ public:
 
 		constexpr SIZE_T ElementSize = sizeof(ElementType);
 		constexpr SIZE_T ElementAlignment = Align(alignof(ElementType), 4U);
-		const uint32 ElementCount = Source.ElementCount;
 
 		void* ElementStorage = nullptr;
-		if (ElementCount > 0)
+		if (Source.ElementCount > 0)
 		{
-			ElementStorage = Context.GetInternalContext()->Alloc(Align(ElementSize*ElementCount, 4U), ElementAlignment);
-			FMemory::Memcpy(ElementStorage, Source.ElementStorage, ElementSize*ElementCount);
+			ElementStorage = Context.GetInternalContext()->Alloc(Align(ElementSize*Source.ElementCount, 4U), ElementAlignment);
+			FMemory::Memcpy(ElementStorage, Source.ElementStorage, ElementSize*Source.ElementCount);
 		}
-		Target.ElementCapacityCount = ElementCount;
-		Target.ElementCount = ElementCount;
+		Target.ElementCapacityCount = Source.ElementCount;
+		Target.ElementCount = Source.ElementCount;
 		Target.ElementStorage = ElementStorage;
 	}
 
diff --git a/Engine/Source/Runtime/Experimental/Iris/Core/Private/Iris/Serialization/StringNetSerializers.cpp b/Engine/Source/Runtime/Experimental/Iris/Core/Private/Iris/Serialization/StringNetSerializers.cpp
--- a/Engine/Source/Runtime/Experimental/Iris/Core/Private/Iris/Serialization/StringNetSerializers.cpp
+++ b/Engine/Source/Runtime/Experimental/Iris/Core/Private/Iris/Serialization/StringNetSerializers.cpp
@@ -141,7 +141,7 @@ void FNameNetSerializer::Deserialize(FNetSerializationContext& Context, const FN
 		Target.bEncodeNumberFromIntMax = bEncodeNumberFromIntMax;
 		Target.bIsEncoded = bIsEncoded;
 		Target.ENameOrNumber = Number;
-		FStringNetSerializerUtils::AdjustArraySize<QuantizedType, uint8>(Context, Target, NewElementCount);
+		FStringNetSerializerUtils::AdjustArraySize<QuantizedType, uint8>(Context, Target, static_cast<uint16>(NewElementCount));
 		Reader->ReadBitStream(static_cast<uint32*>(Target.ElementStorage), (NewElementCount - 1U)*8U);
 		static_cast<uint8*>(Target.ElementStorage)[NewElementCount - 1] = 0;
 		if (bIsEncoded && !FStringNetSerializerUtils::TStringCodec<WIDECHAR>::IsValidEncoding(static_cast<uint8*>(Target.ElementStorage), NewElementCount - 1U))
@@ -199,14 +199,20 @@ void FNameNetSerializer::Quantize(FNetSerializationContext& Context, const FNetQ
 
 			// Our codec uses up to 3 bytes per codepoint
 			const uint32 NameLength = DisplayNameEntry->GetNameLength() + 1U;
-			FStringNetSerializerUtils::AdjustArraySize<QuantizedType, uint8>(Context, TargetName, 3U*NameLength);
+			constexpr uint32 MaxArrayCount = 65536U/3U;
+			if (NameLength > MaxArrayCount)
+			{
+				Context.SetError(GNetError_ArraySizeTooLarge);
+				return;
+			}
+			FStringNetSerializerUtils::AdjustArraySize<QuantizedType, uint8>(Context, TargetName, static_cast<uint16>(3U*NameLength));
 
 			uint32 OutDestLen = 0;
 			uint8* EncodingBuffer = static_cast<uint8*>(TargetName.ElementStorage);
 			const bool bEncodingSuccess = FStringNetSerializerUtils::TStringCodec<WIDECHAR>::Encode(EncodingBuffer, 3U*NameLength, TempWideBuffer, NameLength, OutDestLen);
 			if (bEncodingSuccess)
 			{
-				TargetName.ElementCount = OutDestLen;
+				TargetName.ElementCount = static_cast<uint16>(OutDestLen);
 			}
 			else
 			{
@@ -219,7 +225,13 @@ void FNameNetSerializer::Quantize(FNetSerializationContext& Context, const FNetQ
 		{
 			// For debugging purposes we store null terminated strings.
 			const uint32 NewElementCount = DisplayNameEntry->GetNameLength() + 1U;
-			FStringNetSerializerUtils::AdjustArraySize<QuantizedType, uint8>(Context, TargetName, NewElementCount);
+			constexpr uint32 MaxArrayCount = 65536U;
+			if (NewElementCount > MaxArrayCount)
+			{
+				Context.SetError(GNetError_ArraySizeTooLarge);
+				return;
+			}
+			FStringNetSerializerUtils::AdjustArraySize<QuantizedType, uint8>(Context, TargetName, static_cast<uint16>(NewElementCount));
 
 			// At this time it's impossible to avoid the double copy unless we use a fixed excessive storage
 			ANSICHAR TempAnsiBuffer[NAME_SIZE];
@@ -234,7 +246,7 @@ void FNameNetSerializer::Quantize(FNetSerializationContext& Context, const FNetQ
 	}
 	else
 	{
-		constexpr uint32 NewElementCount = 0U;
+		constexpr uint16 NewElementCount = 0U;
 		FStringNetSerializerUtils::AdjustArraySize<QuantizedType, uint8>(Context, TargetName, NewElementCount);
 
 		TargetName.bIsString = 0;
diff --git a/Engine/Source/Runtime/Experimental/Iris/Core/Public/Iris/Core/IrisDebugging.h b/Engine/Source/Runtime/Experimental/Iris/Core/Public/Iris/Core/IrisDebugging.h
--- a/Engine/Source/Runtime/Experimental/Iris/Core/Public/Iris/Core/IrisDebugging.h
+++ b/Engine/Source/Runtime/Experimental/Iris/Core/Public/Iris/Core/IrisDebugging.h
@@ -91,7 +91,7 @@ extern "C" IRISCORE_API FNetReplicatedObjectDebugInfo DebugNetObjectById(UObject
 extern "C" IRISCORE_API FNetReplicatedObjectDebugInfo DebugNetRefHandle(FNetRefHandle Handle);
 
 /** Look up replicated handle specified by handle id and replicationsystem id and return debug information */
-extern "C" IRISCORE_API FNetReplicatedObjectDebugInfo DebugNetRefHandleById(uint32 NetRefHandleId, uint32 ReplicationSystemId);
+extern "C" IRISCORE_API FNetReplicatedObjectDebugInfo DebugNetRefHandleById(uint64 NetRefHandleId, uint32 ReplicationSystemId);
 
 /** Look up replicated handle specified by an internal index and replicationsystem id and return debug information */
 extern "C" IRISCORE_API FNetReplicatedObjectDebugInfo DebugNetInternalIndex(UE::Net::Private::FInternalNetRefIndex InternalIndex, uint32 ReplicationSystemId);
@@ -105,7 +105,7 @@ extern "C" IRISCORE_API const TCHAR* DebugNetObjectProtocolReferencesToString(ui
 extern "C" IRISCORE_API void SetIrisDebugObjectName(const ANSICHAR* NameBuffer);
 
 /** Set the NetHandle to break on */
-extern "C" IRISCORE_API void SetIrisDebugNetRefHandle(uint32 NetHandleId);
+extern "C" IRISCORE_API void SetIrisDebugNetRefHandle(uint64 NetHandleId);
 
 /** Set the InternalIndex to break on */
 extern "C" IRISCORE_API void SetIrisDebugInternalIndex(UE::Net::Private::FInternalNetRefIndex InternalIndex);
diff --git a/Engine/Source/Runtime/Experimental/Iris/Core/Public/Iris/Serialization/NetSerializer.h b/Engine/Source/Runtime/Experimental/Iris/Core/Public/Iris/Serialization/NetSerializer.h
--- a/Engine/Source/Runtime/Experimental/Iris/Core/Public/Iris/Serialization/NetSerializer.h
+++ b/Engine/Source/Runtime/Experimental/Iris/Core/Public/Iris/Serialization/NetSerializer.h
@@ -5,7 +5,7 @@
 #include "UObject/ObjectMacros.h"
 #include "Iris/Serialization/NetSerializationContext.h"
 #include "Iris/Serialization/NetSerializerConfig.h"
-
+#include <limits>
 
 /**
  * A FNetSerializer is needed for replication support of a certain type. Most types that can be 
@@ -139,9 +139,9 @@ typedef const FNetSerializerConfig* NetSerializerConfigParam;
 struct FNetSerializerChangeMaskParam
 {
 	/** Offset in the change mask where we store the bits for the member. */
-	uint16 BitOffset = 0;
+	uint32 BitOffset = 0;
 	/** Number of bits used for the member. This is typically 1, except for special cases like arrays. */
-	uint16 BitCount = 0;
+	uint32 BitCount = 0;
 };
 
 /**
@@ -416,11 +416,15 @@ public:
 
 		Serializer.DefaultConfig = Builder.GetDefaultConfig();
 
-		Serializer.QuantizedTypeSize = Builder.GetQuantizedTypeSize();
-		Serializer.QuantizedTypeAlignment = Builder.GetQuantizedTypeAlignment();
+		static_assert(Builder.GetQuantizedTypeSize() <= std::numeric_limits<uint16>::max(), "");
+		Serializer.QuantizedTypeSize = static_cast<uint16>(Builder.GetQuantizedTypeSize());
+		static_assert(Builder.GetQuantizedTypeAlignment() <= std::numeric_limits<uint16>::max(), "");
+		Serializer.QuantizedTypeAlignment = static_cast<uint16>(Builder.GetQuantizedTypeAlignment());
 
-		Serializer.ConfigTypeSize = Builder.GetConfigTypeSize();
-		Serializer.ConfigTypeAlignment = Builder.GetConfigTypeAlignment();
+		static_assert(Builder.GetConfigTypeSize() <= std::numeric_limits<uint16>::max(), "");
+		Serializer.ConfigTypeSize = static_cast<uint16>(Builder.GetConfigTypeSize());
+		static_assert(Builder.GetConfigTypeAlignment() <= std::numeric_limits<uint16>::max(), "");
+		Serializer.ConfigTypeAlignment = static_cast<uint16>(Builder.GetConfigTypeAlignment());
 
 		Serializer.Name = Name;
 		return Serializer;
