Geometry collection : make sure we don't construct FName at runtime for the removal related facade classes

diff --git a/Engine/Source/Runtime/Experimental/GeometryCollectionEngine/Private/GeometryCollection/GeometryCollectionEngineRemoval.cpp b/Engine/Source/Runtime/Experimental/GeometryCollectionEngine/Private/GeometryCollection/GeometryCollectionEngineRemoval.cpp
--- a/Engine/Source/Runtime/Experimental/GeometryCollectionEngine/Private/GeometryCollection/GeometryCollectionEngineRemoval.cpp
+++ b/Engine/Source/Runtime/Experimental/GeometryCollectionEngine/Private/GeometryCollection/GeometryCollectionEngineRemoval.cpp
@@ -6,11 +6,25 @@
 
 ///////////////////////////////////////////////////////////////////////////////////////////////////////////////////////////
 
+namespace
+{
+	const FName BreakTimerAttributeName = "BreakTimer";
+	const FName PostBreakDurationAttributeName = "PostBreakDuration";
+	const FName BreakRemovalDurationAttributeName = "BreakRemovalDuration";
+
+	const FName SleepTimerAttributeName = "SleepTimer";
+	const FName MaxSleepTimeAttributeName = "MaxSleepTime";
+	const FName SleepRemovalDurationAttributeName = "SleepRemovalDuration";
+	const FName LastPositionAttributeName = "LastPosition";
+
+	const FName DecayAttributeName = "Decay";
+}
+
 FGeometryCollectionRemoveOnBreakDynamicFacade::FGeometryCollectionRemoveOnBreakDynamicFacade(FManagedArrayCollection& InCollection)
-	: BreakTimerAttribute(InCollection, "BreakTimer", FGeometryCollection::TransformGroup)
-	, PostBreakDurationAttribute(InCollection, "PostBreakDuration", FGeometryCollection::TransformGroup)
-	, BreakRemovalDurationAttribute(InCollection, "BreakRemovalDuration", FGeometryCollection::TransformGroup)
-	, ChildrenAttribute(InCollection, "Children", FTransformCollection::TransformGroup)
+	: BreakTimerAttribute(InCollection, BreakTimerAttributeName, FGeometryCollection::TransformGroup)
+	, PostBreakDurationAttribute(InCollection, PostBreakDurationAttributeName, FGeometryCollection::TransformGroup)
+	, BreakRemovalDurationAttribute(InCollection, BreakRemovalDurationAttributeName, FGeometryCollection::TransformGroup)
+	, ChildrenAttribute(InCollection, FTransformCollection::ChildrenAttribute, FTransformCollection::TransformGroup)
 
 {
 }
@@ -109,10 +123,10 @@ float FGeometryCollectionRemoveOnBreakDynamicFacade::UpdateBreakTimerAndComputeD
 ///////////////////////////////////////////////////////////////////////////////////////////////////////////////////////////
 
 FGeometryCollectionRemoveOnSleepDynamicFacade::FGeometryCollectionRemoveOnSleepDynamicFacade(FManagedArrayCollection& InCollection)
-	: SleepTimerAttribute(InCollection, "SleepTimer", FGeometryCollection::TransformGroup)
-	, MaxSleepTimeAttribute(InCollection, "MaxSleepTime", FGeometryCollection::TransformGroup)
-	, SleepRemovalDurationAttribute(InCollection, "SleepRemovalDuration", FGeometryCollection::TransformGroup)
-	, LastPositionAttribute(InCollection, "LastPosition", FGeometryCollection::TransformGroup)
+	: SleepTimerAttribute(InCollection, SleepTimerAttributeName, FGeometryCollection::TransformGroup)
+	, MaxSleepTimeAttribute(InCollection, MaxSleepTimeAttributeName, FGeometryCollection::TransformGroup)
+	, SleepRemovalDurationAttribute(InCollection, SleepRemovalDurationAttributeName, FGeometryCollection::TransformGroup)
+	, LastPositionAttribute(InCollection, LastPositionAttributeName, FGeometryCollection::TransformGroup)
 {
 }
 
@@ -202,7 +216,7 @@ float FGeometryCollectionRemoveOnSleepDynamicFacade::ComputeDecay(int32 Transfor
 ///////////////////////////////////////////////////////////////////////////////////////////////////////////////////////////
 
 FGeometryCollectionDecayDynamicFacade::FGeometryCollectionDecayDynamicFacade(FManagedArrayCollection& InCollection)
-	: DecayAttribute(InCollection, "Decay", FGeometryCollection::TransformGroup)
+	: DecayAttribute(InCollection, DecayAttributeName, FGeometryCollection::TransformGroup)
 {
 }
 
