Check if a deleted proxy is being added back, trying to catch the root cause of a crash

diff --git a/Engine/Source/Runtime/Experimental/Chaos/Public/Chaos/Framework/PhysicsSolverBase.h b/Engine/Source/Runtime/Experimental/Chaos/Public/Chaos/Framework/PhysicsSolverBase.h
--- a/Engine/Source/Runtime/Experimental/Chaos/Public/Chaos/Framework/PhysicsSolverBase.h
+++ b/Engine/Source/Runtime/Experimental/Chaos/Public/Chaos/Framework/PhysicsSolverBase.h
@@ -282,6 +282,7 @@ namespace Chaos
 
 		void AddDirtyProxy(IPhysicsProxyBase * ProxyBaseIn)
 		{
+			check(ProxyBaseIn->GetMarkedDeleted() == false);
 			MarshallingManager.GetProducerData_External()->DirtyProxiesDataBuffer.Add(ProxyBaseIn);
 		}
 		void RemoveDirtyProxy(IPhysicsProxyBase * ProxyBaseIn)
