geometry collection cluster particles should only have a cluster group index set on them if they are a root particle (does not have another parent)

diff --git a/Engine/Source/Runtime/Experimental/Chaos/Private/PhysicsProxy/GeometryCollectionPhysicsProxy.cpp b/Engine/Source/Runtime/Experimental/Chaos/Private/PhysicsProxy/GeometryCollectionPhysicsProxy.cpp
--- a/Engine/Source/Runtime/Experimental/Chaos/Private/PhysicsProxy/GeometryCollectionPhysicsProxy.cpp
+++ b/Engine/Source/Runtime/Experimental/Chaos/Private/PhysicsProxy/GeometryCollectionPhysicsProxy.cpp
@@ -1663,7 +1663,7 @@ FGeometryCollectionPhysicsProxy::BuildClusters_Internal(
 	// Construct an active cluster particle, disable children, derive M and I from children:
 	Chaos::TPBDRigidClusteredParticleHandle<Chaos::FReal, 3>* Parent =
 		static_cast<Chaos::FPBDRigidsSolver*>(Solver)->GetEvolution()->GetRigidClustering().CreateClusterParticle(
-			Parameters.ClusterGroupIndex, 
+			(ParentIndex[CollectionClusterIndex] == INDEX_NONE) ? Parameters.ClusterGroupIndex : 0, 
 			MoveTemp(ChildHandlesCopy),
 			ClusterCreationParameters,
 			Implicits[CollectionClusterIndex], // union from children if null
