Iris
* Localized the Polling code into a dedicated class

#rnx

diff --git a/Engine/Source/Runtime/Experimental/Iris/Core/Private/Iris/ReplicationSystem/ObjectReplicationBridge.cpp b/Engine/Source/Runtime/Experimental/Iris/Core/Private/Iris/ReplicationSystem/ObjectReplicationBridge.cpp
--- a/Engine/Source/Runtime/Experimental/Iris/Core/Private/Iris/ReplicationSystem/ObjectReplicationBridge.cpp
+++ b/Engine/Source/Runtime/Experimental/Iris/Core/Private/Iris/ReplicationSystem/ObjectReplicationBridge.cpp
@@ -29,6 +29,7 @@
 #include "Iris/ReplicationSystem/ReplicationOperations.h"
 #include "Iris/ReplicationSystem/ReplicationOperationsInternal.h"
 #include "Iris/ReplicationSystem/RepTag.h"
+#include "Iris/ReplicationSystem/Polling/ObjectPoller.h"
 #include "Iris/Serialization/NetSerializationContext.h"
 #include "Iris/Serialization/InternalNetSerializationContext.h"
 #include "Iris/Serialization/NetBitStreamUtil.h"
@@ -702,172 +703,22 @@ void UObjectReplicationBridge::PreUpdateAndPollImpl(FNetRefHandle Handle)
 	using namespace UE::Net;
 	using namespace UE::Net::Private;
 
-	FReplicationSystemInternal* ReplicationSystemInternal = GetReplicationSystem()->GetReplicationSystemInternal();
-	const FNetRefHandleManager& LocalNetRefHandleManager = ReplicationSystemInternal->GetNetRefHandleManager();
-	const TArray<UObject*>& ReplicatedInstances = LocalNetRefHandleManager.GetReplicatedInstances();
-	const bool bIsUsingPushModel = IsIrisPushModelEnabled();
-
-	FDirtyObjectsAccessor DirtyObjectsAccessor(ReplicationSystemInternal->GetDirtyNetObjectTracker());
-
-	FNetBitArrayView DirtyObjectsThisFrame = DirtyObjectsAccessor.GetDirtyNetObjects();
-
-	struct FPreUpdateAndPollStats
-	{
-		uint32 PreUpdatedObjectCount = 0;
-		uint32 PolledObjectCount = 0;
-		uint32 PolledReferencesObjectCount = 0;
-	};
-
-	FPreUpdateAndPollStats Stats;
-
-	auto UpdateAndPollFunction = [&ReplicatedInstances, this, &LocalNetRefHandleManager, &DirtyObjectsThisFrame, &Stats](uint32 InternalObjectIndex)
-	{
-		const FNetRefHandleManager::FReplicatedObjectData& ObjectData = LocalNetRefHandleManager.GetReplicatedObjectDataNoCheck(InternalObjectIndex);
-		if (ObjectData.InstanceProtocol)
-		{
-			IRIS_PROFILER_PROTOCOL_NAME(ObjectData.Protocol->DebugName->Name);
-
-			// Call per-instance PreUpdate function
-			if (PreUpdateInstanceFunction && EnumHasAnyFlags(ObjectData.InstanceProtocol->InstanceTraits, EReplicationInstanceProtocolTraits::NeedsPreSendUpdate))
-			{
-				IRIS_PROFILER_SCOPE_VERBOSE(PreReplicationUpdate);
-				(*PreUpdateInstanceFunction)(ObjectData.RefHandle, ReplicatedInstances[InternalObjectIndex], this);
-				++Stats.PreUpdatedObjectCount;
-			}
-
-			// Poll properties if the instance protocol requires it
-			if (EnumHasAnyFlags(ObjectData.InstanceProtocol->InstanceTraits, EReplicationInstanceProtocolTraits::NeedsPoll))
-			{
-				IRIS_PROFILER_SCOPE_VERBOSE(Poll);
-
-				const bool bIsGCAffectedObject = GarbageCollectionAffectedObjects.GetBit(InternalObjectIndex);
-				GarbageCollectionAffectedObjects.ClearBit(InternalObjectIndex);
-
-				// If this object has been around for a garbage collect and it has object references we must make sure that we update all cached object references
-				EReplicationFragmentPollFlags PollOptions = EReplicationFragmentPollFlags::PollAllState;
-				PollOptions |= bIsGCAffectedObject ? EReplicationFragmentPollFlags::ForceRefreshCachedObjectReferencesAfterGC : EReplicationFragmentPollFlags::None;
-
-				const bool bWasMarkedDirty = FReplicationInstanceOperations::PollAndRefreshCachedPropertyData(ObjectData.InstanceProtocol, PollOptions);
-				if (bWasMarkedDirty)
-				{
-					DirtyObjectsThisFrame.SetBit(InternalObjectIndex);
-				}
-				++Stats.PolledObjectCount;
-			}
-		}
-	};
-
-	const FNetBitArray& PrevScopableObjects = LocalNetRefHandleManager.GetPrevFrameScopableInternalIndices();
-
-	auto PushModelUpdateAndPollFunction = [&ReplicatedInstances, this, &LocalNetRefHandleManager, &DirtyObjectsThisFrame, &PrevScopableObjects, &Stats](uint32 InternalObjectIndex)
-	{
-		const FNetRefHandleManager::FReplicatedObjectData& ObjectData = LocalNetRefHandleManager.GetReplicatedObjectDataNoCheck(InternalObjectIndex);
-		if (const FReplicationInstanceProtocol* InstanceProtocol = ObjectData.InstanceProtocol)
-		{
-			IRIS_PROFILER_PROTOCOL_NAME(ObjectData.Protocol->DebugName->Name);
-
-			const EReplicationInstanceProtocolTraits InstanceTraits = InstanceProtocol->InstanceTraits;
-			const bool bNeedsPoll = EnumHasAnyFlags(InstanceTraits, EReplicationInstanceProtocolTraits::NeedsPoll);
-			bool bIsDirtyObject = DirtyObjectsThisFrame.GetBit(InternalObjectIndex);
-
-			// Call per-instance PreUpdate function
-			if (PreUpdateInstanceFunction && EnumHasAnyFlags(InstanceTraits, EReplicationInstanceProtocolTraits::NeedsPreSendUpdate))
-			{
-				IRIS_PROFILER_SCOPE_VERBOSE(PreReplicationUpdate);
-
-				(*PreUpdateInstanceFunction)(ObjectData.RefHandle, ReplicatedInstances[InternalObjectIndex], this);
-				++Stats.PreUpdatedObjectCount;
-
-				// Pre update may dirty push based properties. Detect it.
-				if ((bNeedsPoll & !bIsDirtyObject))
-				{
-					bIsDirtyObject = FGlobalDirtyNetObjectTracker::IsNetObjectStateDirty(ObjectData.NetHandle);
-					if (bIsDirtyObject)
-					{
-						DirtyObjectsThisFrame.SetBit(bIsDirtyObject);
-					}
-				}
-			}
-
-			const bool bIsNewInScope = !PrevScopableObjects.GetBit(InternalObjectIndex);
-			const bool bIsGCAffectedObject = GarbageCollectionAffectedObjects.GetBit(InternalObjectIndex);			
-			GarbageCollectionAffectedObjects.ClearBit(InternalObjectIndex);
-
-			// Early out if the instance does not require polling
-			if (!bNeedsPoll)
-			{
-				return;
-			}
-
-			IRIS_PROFILER_SCOPE_VERBOSE(PollPushBased);
-
-			// If the object is fully push model we only need to poll it if it's dirty, unless it's a new object or was garbage collected.
-			bool bWasMarkedDirty = false;
-			if (EnumHasAnyFlags(InstanceTraits, EReplicationInstanceProtocolTraits::HasFullPushBasedDirtiness))
-			{
-				if (bIsDirtyObject | bIsNewInScope)
-				{
-					// We need to do a poll if object is marked as dirty
-					EReplicationFragmentPollFlags PollOptions = EReplicationFragmentPollFlags::PollAllState;
-					PollOptions |= bIsGCAffectedObject ? EReplicationFragmentPollFlags::ForceRefreshCachedObjectReferencesAfterGC : EReplicationFragmentPollFlags::None;
-					bWasMarkedDirty = FReplicationInstanceOperations::PollAndRefreshCachedPropertyData(InstanceProtocol, EReplicationFragmentTraits::None, PollOptions);
-					++Stats.PolledObjectCount;
-				}
-				else if (bIsGCAffectedObject)
-				{
-					// If this object might have been affected by GC, only refresh cached references
-					const EReplicationFragmentTraits RequiredTraits = EReplicationFragmentTraits::HasPushBasedDirtiness;
-					bWasMarkedDirty = FReplicationInstanceOperations::PollAndRefreshCachedObjectReferences(InstanceProtocol, RequiredTraits);
-					++Stats.PolledReferencesObjectCount;
-				}
-			}
-			else
-			{
-				// If the object has pushed based fragments, and not is marked dirty and object is affected by GC we need to make sure that we refresh cached references for all push based fragments
-				const bool bIsPushBasedObject = EnumHasAnyFlags(InstanceTraits, EReplicationInstanceProtocolTraits::HasPartialPushBasedDirtiness | EReplicationInstanceProtocolTraits::HasFullPushBasedDirtiness);
-				const bool bHasObjectReferences = EnumHasAnyFlags(InstanceTraits, EReplicationInstanceProtocolTraits::HasObjectReference);
-				const bool bNeedsRefreshOfCachedObjectReferences = ((!(bIsNewInScope | bIsDirtyObject)) & bIsGCAffectedObject & bIsPushBasedObject & bHasObjectReferences);
-				if (bNeedsRefreshOfCachedObjectReferences)
-				{
-					// Only states which has push based dirtiness need to be updated as the other states will be polled in full anyway.
-					const EReplicationFragmentTraits RequiredTraits = EReplicationFragmentTraits::HasPushBasedDirtiness;
-					bWasMarkedDirty = FReplicationInstanceOperations::PollAndRefreshCachedObjectReferences(InstanceProtocol, RequiredTraits);
-					++Stats.PolledReferencesObjectCount;
-				}
-
-				// If this object has been around for a garbage collect and it has object references we must make sure that we update all cached object references 
-				EReplicationFragmentPollFlags PollOptions = EReplicationFragmentPollFlags::PollAllState;
-				PollOptions |= bIsGCAffectedObject ? EReplicationFragmentPollFlags::ForceRefreshCachedObjectReferencesAfterGC : EReplicationFragmentPollFlags::None;
-
-				// If the object is not new or dirty at this point we only need to poll non-push based fragments as we know that pushed based states have not been modified
-				const EReplicationFragmentTraits ExcludeTraits = (bIsDirtyObject | bIsNewInScope) ? EReplicationFragmentTraits::None : EReplicationFragmentTraits::HasPushBasedDirtiness;
-				bWasMarkedDirty |= FReplicationInstanceOperations::PollAndRefreshCachedPropertyData(InstanceProtocol, ExcludeTraits, PollOptions);
-				++Stats.PolledObjectCount;
-			}
-
-			if (bWasMarkedDirty)
-			{
-				DirtyObjectsThisFrame.SetBit(InternalObjectIndex);
-			}
-		}
-	};
-
 	// If we are asked to poll a single object
 	if (Handle.IsValid())
 	{
-		if (uint32 InternalObjectIndex = LocalNetRefHandleManager.GetInternalIndex(Handle))
-		{
-			// From here we call into user code via PreUpdateInstanceFunction, so allow external code to set dirty flags since DirtyObjects is not read anymore.
-			ReplicationSystemInternal->GetDirtyNetObjectTracker().AllowExternalAccess();
+		FObjectPoller::FInitParams PollerInitParams;
+		PollerInitParams.ObjectReplicationBridge = this;
+		PollerInitParams.ReplicationSystemInternal = GetReplicationSystem()->GetReplicationSystemInternal();
+		FObjectPoller Poller(PollerInitParams);
 
-			UpdateAndPollFunction(InternalObjectIndex);
-		}
+		Poller.PollSingleObject(Handle);
 
 		return;
 	}
 	
 	// Update every relevant objects from here
-
+	FReplicationSystemInternal* ReplicationSystemInternal = GetReplicationSystem()->GetReplicationSystemInternal();
+	const FNetRefHandleManager& LocalNetRefHandleManager = ReplicationSystemInternal->GetNetRefHandleManager();
 	const FNetBitArrayView RelevantObjects = LocalNetRefHandleManager.GetRelevantObjectsInternalIndices();
 	const FNetBitArrayView WantToBeDormantObjects = MakeNetBitArrayView(LocalNetRefHandleManager.GetWantToBeDormantInternalIndices());
 
@@ -944,7 +795,12 @@ void UObjectReplicationBridge::PreUpdateAndPollImpl(FNetRefHandle Handle)
 		FNetBitArray DirtyAndRelevantObjects(RelevantObjects.GetNumBits());
 		FNetBitArrayView DirtyAndRelevantObjectsView = MakeNetBitArrayView(DirtyAndRelevantObjects);
 
-		DirtyAndRelevantObjectsView.Set(RelevantObjects, FNetBitArray::AndOp, DirtyObjectsThisFrame);
+		{
+			FDirtyObjectsAccessor DirtyObjectsAccessor(ReplicationSystemInternal->GetDirtyNetObjectTracker());	
+			FNetBitArrayView DirtyObjectsThisFrame = DirtyObjectsAccessor.GetDirtyNetObjects();
+
+			DirtyAndRelevantObjectsView.Set(RelevantObjects, FNetBitArray::AndOp, DirtyObjectsThisFrame);
+		}
 
 		// Update subobjects' owner first and owners' subobjects second. It's the only way to properly mark all groups of objects in two passes.
 		const FNetBitArrayView SubObjects = MakeNetBitArrayView(LocalNetRefHandleManager.GetSubObjectInternalIndices());
@@ -968,21 +824,17 @@ void UObjectReplicationBridge::PreUpdateAndPollImpl(FNetRefHandle Handle)
 		}
 	}
 
-	{
-		IRIS_PROFILER_SCOPE(PreUpdateAndPollImpl_Poll);
+	
+	IRIS_PROFILER_SCOPE(PreUpdateAndPollImpl_Poll);
 
-		// From here we call into user code via PreUpdateInstanceFunction, so allow external code to set dirty flags since DirtyObjects is not read anymore.
-		ReplicationSystemInternal->GetDirtyNetObjectTracker().AllowExternalAccess();
+	FObjectPoller::FInitParams PollerInitParams;
+	PollerInitParams.ObjectReplicationBridge = this;
+	PollerInitParams.ReplicationSystemInternal = GetReplicationSystem()->GetReplicationSystemInternal();
 
-		if (IsIrisPushModelEnabled())
-		{
-			ObjectsConsideredForPolling.ForAllSetBits(PushModelUpdateAndPollFunction);
-		}
-		else
-		{
-			ObjectsConsideredForPolling.ForAllSetBits(UpdateAndPollFunction);
-		}
-	}
+	FObjectPoller Poller(PollerInitParams);
+	Poller.PollObjects(ObjectsConsideredForPollingView);
+	
+	FObjectPoller::FPreUpdateAndPollStats Stats = Poller.GetPollStats();
 
 	// Report stats
 	UE_NET_TRACE_FRAME_STATSCOUNTER(ReplicationSystem->GetId(), ReplicationSystem.PreUpdatedObjectCount, Stats.PreUpdatedObjectCount, ENetTraceVerbosity::Trace);
diff --git a/Engine/Source/Runtime/Experimental/Iris/Core/Private/Iris/ReplicationSystem/Polling/ObjectPoller.cpp b/Engine/Source/Runtime/Experimental/Iris/Core/Private/Iris/ReplicationSystem/Polling/ObjectPoller.cpp
new file mode 100644
--- /dev/null
+++ b/Engine/Source/Runtime/Experimental/Iris/Core/Private/Iris/ReplicationSystem/Polling/ObjectPoller.cpp
@@ -0,0 +1,216 @@
+// Copyright Epic Games, Inc. All Rights Reserved.
+
+#include "ObjectPoller.h"
+
+#include "Iris/ReplicationSystem/ObjectReplicationBridge.h"
+#include "Iris/ReplicationSystem/ReplicationOperations.h"
+#include "Iris/ReplicationSystem/ReplicationSystem.h"
+#include "Iris/ReplicationSystem/ReplicationSystemInternal.h"
+#include "Iris/ReplicationSystem/NetRefHandleManager.h"
+#include "Iris/ReplicationSystem/LegacyPushModel.h"
+#include "Iris/ReplicationSystem/NetRefHandle.h"
+
+#include "Iris/Core/IrisProfiler.h"
+
+#include "Net/Core/Trace/NetDebugName.h"
+
+namespace UE::Net::Private
+{
+
+FObjectPoller::FObjectPoller(const FInitParams& InitParams)
+	: ObjectReplicationBridge(InitParams.ObjectReplicationBridge)
+	, ReplicationSystemInternal(InitParams.ReplicationSystemInternal)
+	, LocalNetRefHandleManager(ReplicationSystemInternal->GetNetRefHandleManager())
+	, ReplicatedInstances(LocalNetRefHandleManager.GetReplicatedInstances())
+	, AccumulatedDirtyObjects(ReplicationSystemInternal->GetDirtyNetObjectTracker().GetAccumulatedDirtyNetObjects())
+	, PrevFrameGlobalScopeObjects(MakeNetBitArrayView(LocalNetRefHandleManager.GetPrevFrameScopableInternalIndices()))
+{
+	GarbageCollectionAffectedObjects = MakeNetBitArrayView(ObjectReplicationBridge->GarbageCollectionAffectedObjects);
+
+	// DirtyObjectsThisFrame is acquired only during polling 
+}
+
+void FObjectPoller::PollObjects(const FNetBitArrayView& ObjectsConsideredForPolling)
+{
+	FDirtyObjectsAccessor DirtyObjectsAccessor(ReplicationSystemInternal->GetDirtyNetObjectTracker());
+	DirtyObjectsThisFrame = DirtyObjectsAccessor.GetDirtyNetObjects();
+
+	// From here we call into user code via PreUpdateInstanceFunction, so allow external code to set dirty flags since DirtyObjects is not read anymore.
+	ReplicationSystemInternal->GetDirtyNetObjectTracker().AllowExternalAccess();
+
+	if (IsIrisPushModelEnabled())
+	{
+		ObjectsConsideredForPolling.ForAllSetBits([this](FInternalNetRefIndex Objectindex)
+		{
+			PushModelPollObject(Objectindex);
+		});
+	}
+	else
+	{
+		ObjectsConsideredForPolling.ForAllSetBits([this](FInternalNetRefIndex Objectindex)
+		{
+			ForcePollObject(Objectindex);
+		});
+	}
+
+	// Clear ref to locked dirty bit array
+	DirtyObjectsThisFrame = FNetBitArrayView();
+}
+
+void FObjectPoller::PollSingleObject(FNetRefHandle Handle)
+{
+	if (uint32 InternalObjectIndex = LocalNetRefHandleManager.GetInternalIndex(Handle))
+	{
+		FDirtyObjectsAccessor DirtyObjectsAccessor(ReplicationSystemInternal->GetDirtyNetObjectTracker());
+		DirtyObjectsThisFrame = DirtyObjectsAccessor.GetDirtyNetObjects();
+
+		// From here we call into user code via PreUpdateInstanceFunction, so allow external code to set dirty flags since DirtyObjects is not read anymore.
+		ReplicationSystemInternal->GetDirtyNetObjectTracker().AllowExternalAccess();
+
+		ForcePollObject(InternalObjectIndex);
+
+		// Clear ref to locked dirty bit array
+		DirtyObjectsThisFrame = FNetBitArrayView();
+	}
+}
+
+void FObjectPoller::ForcePollObject(FInternalNetRefIndex ObjectIndex)
+{
+	const FNetRefHandleManager::FReplicatedObjectData& ObjectData = LocalNetRefHandleManager.GetReplicatedObjectDataNoCheck(ObjectIndex);
+	if (ObjectData.InstanceProtocol == nullptr)
+	{
+		return;
+	}
+
+	IRIS_PROFILER_PROTOCOL_NAME(ObjectData.Protocol->DebugName->Name);
+
+	// Call per-instance PreUpdate function
+	if (ObjectReplicationBridge->PreUpdateInstanceFunction && EnumHasAnyFlags(ObjectData.InstanceProtocol->InstanceTraits, EReplicationInstanceProtocolTraits::NeedsPreSendUpdate))
+	{
+		IRIS_PROFILER_SCOPE_VERBOSE(PreReplicationUpdate);
+		(*ObjectReplicationBridge->PreUpdateInstanceFunction)(ObjectData.RefHandle, ReplicatedInstances[ObjectIndex], ObjectReplicationBridge);
+		++PollStats.PreUpdatedObjectCount;
+	}
+
+	// Poll properties if the instance protocol requires it
+	if (EnumHasAnyFlags(ObjectData.InstanceProtocol->InstanceTraits, EReplicationInstanceProtocolTraits::NeedsPoll))
+	{
+		IRIS_PROFILER_SCOPE_VERBOSE(Poll);
+
+		const bool bIsGCAffectedObject = GarbageCollectionAffectedObjects.GetBit(ObjectIndex);
+		GarbageCollectionAffectedObjects.ClearBit(ObjectIndex);
+
+		// If this object has been around for a garbage collect and it has object references we must make sure that we update all cached object references
+		EReplicationFragmentPollFlags PollOptions = EReplicationFragmentPollFlags::PollAllState;
+		PollOptions |= bIsGCAffectedObject ? EReplicationFragmentPollFlags::ForceRefreshCachedObjectReferencesAfterGC : EReplicationFragmentPollFlags::None;
+
+		const bool bWasMarkedDirty = FReplicationInstanceOperations::PollAndRefreshCachedPropertyData(ObjectData.InstanceProtocol, PollOptions);
+		if (bWasMarkedDirty)
+		{
+			DirtyObjectsThisFrame.SetBit(ObjectIndex);
+		}
+		++PollStats.PolledObjectCount;
+	}
+}
+
+void FObjectPoller::PushModelPollObject(FInternalNetRefIndex ObjectIndex)
+{
+	const FNetRefHandleManager::FReplicatedObjectData& ObjectData = LocalNetRefHandleManager.GetReplicatedObjectDataNoCheck(ObjectIndex);
+	if (ObjectData.InstanceProtocol == nullptr)
+	{
+		return;
+	}
+
+	IRIS_PROFILER_PROTOCOL_NAME(ObjectData.Protocol->DebugName->Name);
+
+	const FReplicationInstanceProtocol* InstanceProtocol = ObjectData.InstanceProtocol;
+
+	const EReplicationInstanceProtocolTraits InstanceTraits = InstanceProtocol->InstanceTraits;
+	const bool bNeedsPoll = EnumHasAnyFlags(InstanceTraits, EReplicationInstanceProtocolTraits::NeedsPoll);
+
+	//$IRIS TODO: should use the Accumulated array in case we were dirty while culled.  Also move this test to just after PreUpdate in case it gets dirty inside the call.
+	bool bIsDirtyObject = DirtyObjectsThisFrame.GetBit(ObjectIndex);
+
+	// Call per-instance PreUpdate function
+	if (ObjectReplicationBridge->PreUpdateInstanceFunction && EnumHasAnyFlags(InstanceTraits, EReplicationInstanceProtocolTraits::NeedsPreSendUpdate))
+	{
+		IRIS_PROFILER_SCOPE_VERBOSE(PreReplicationUpdate);
+
+		(*ObjectReplicationBridge->PreUpdateInstanceFunction)(ObjectData.RefHandle, ReplicatedInstances[ObjectIndex], ObjectReplicationBridge);
+		++PollStats.PreUpdatedObjectCount;
+
+		// Pre update may dirty push based properties. Detect it.
+		if ((bNeedsPoll & !bIsDirtyObject))
+		{
+			bIsDirtyObject = FGlobalDirtyNetObjectTracker::IsNetObjectStateDirty(ObjectData.NetHandle);
+			if (bIsDirtyObject)
+			{
+				DirtyObjectsThisFrame.SetBit(bIsDirtyObject);
+			}
+		}
+	}
+
+	const bool bIsNewInScope = !PrevFrameGlobalScopeObjects.GetBit(ObjectIndex);
+	const bool bIsGCAffectedObject = GarbageCollectionAffectedObjects.GetBit(ObjectIndex);
+	GarbageCollectionAffectedObjects.ClearBit(ObjectIndex);
+
+	// Early out if the instance does not require polling
+	if (!bNeedsPoll)
+	{
+		return;
+	}
+
+	IRIS_PROFILER_SCOPE_VERBOSE(PollPushBased);
+
+	// If the object is fully push model we only need to poll it if it's dirty, unless it's a new object or was garbage collected.
+	bool bWasMarkedDirty = false;
+	if (EnumHasAnyFlags(InstanceTraits, EReplicationInstanceProtocolTraits::HasFullPushBasedDirtiness))
+	{
+		if (bIsDirtyObject | bIsNewInScope)
+		{
+			// We need to do a poll if object is marked as dirty
+			EReplicationFragmentPollFlags PollOptions = EReplicationFragmentPollFlags::PollAllState;
+			PollOptions |= bIsGCAffectedObject ? EReplicationFragmentPollFlags::ForceRefreshCachedObjectReferencesAfterGC : EReplicationFragmentPollFlags::None;
+			bWasMarkedDirty = FReplicationInstanceOperations::PollAndRefreshCachedPropertyData(InstanceProtocol, EReplicationFragmentTraits::None, PollOptions);
+			++PollStats.PolledObjectCount;
+		}
+		else if (bIsGCAffectedObject)
+		{
+			// If this object might have been affected by GC, only refresh cached references
+			const EReplicationFragmentTraits RequiredTraits = EReplicationFragmentTraits::HasPushBasedDirtiness;
+			bWasMarkedDirty = FReplicationInstanceOperations::PollAndRefreshCachedObjectReferences(InstanceProtocol, RequiredTraits);
+			++PollStats.PolledReferencesObjectCount;
+		}
+	}
+	else
+	{
+		// If the object has pushed based fragments, and not is marked dirty and object is affected by GC we need to make sure that we refresh cached references for all push based fragments
+		const bool bIsPushBasedObject = EnumHasAnyFlags(InstanceTraits, EReplicationInstanceProtocolTraits::HasPartialPushBasedDirtiness | EReplicationInstanceProtocolTraits::HasFullPushBasedDirtiness);
+		const bool bHasObjectReferences = EnumHasAnyFlags(InstanceTraits, EReplicationInstanceProtocolTraits::HasObjectReference);
+		const bool bNeedsRefreshOfCachedObjectReferences = ((!(bIsNewInScope | bIsDirtyObject)) & bIsGCAffectedObject & bIsPushBasedObject & bHasObjectReferences);
+		if (bNeedsRefreshOfCachedObjectReferences)
+		{
+			// Only states which has push based dirtiness need to be updated as the other states will be polled in full anyway.
+			const EReplicationFragmentTraits RequiredTraits = EReplicationFragmentTraits::HasPushBasedDirtiness;
+			bWasMarkedDirty = FReplicationInstanceOperations::PollAndRefreshCachedObjectReferences(InstanceProtocol, RequiredTraits);
+			++PollStats.PolledReferencesObjectCount;
+		}
+
+		// If this object has been around for a garbage collect and it has object references we must make sure that we update all cached object references 
+		EReplicationFragmentPollFlags PollOptions = EReplicationFragmentPollFlags::PollAllState;
+		PollOptions |= bIsGCAffectedObject ? EReplicationFragmentPollFlags::ForceRefreshCachedObjectReferencesAfterGC : EReplicationFragmentPollFlags::None;
+
+		// If the object is not new or dirty at this point we only need to poll non-push based fragments as we know that pushed based states have not been modified
+		const EReplicationFragmentTraits ExcludeTraits = (bIsDirtyObject | bIsNewInScope) ? EReplicationFragmentTraits::None : EReplicationFragmentTraits::HasPushBasedDirtiness;
+		bWasMarkedDirty |= FReplicationInstanceOperations::PollAndRefreshCachedPropertyData(InstanceProtocol, ExcludeTraits, PollOptions);
+		++PollStats.PolledObjectCount;
+	}
+
+	if (bWasMarkedDirty)
+	{
+		DirtyObjectsThisFrame.SetBit(ObjectIndex);
+	}
+
+}
+
+} // end namespace UE::Net::Private
\ No newline at end of file
diff --git a/Engine/Source/Runtime/Experimental/Iris/Core/Private/Iris/ReplicationSystem/Polling/ObjectPoller.h b/Engine/Source/Runtime/Experimental/Iris/Core/Private/Iris/ReplicationSystem/Polling/ObjectPoller.h
new file mode 100644
--- /dev/null
+++ b/Engine/Source/Runtime/Experimental/Iris/Core/Private/Iris/ReplicationSystem/Polling/ObjectPoller.h
@@ -0,0 +1,85 @@
+// Copyright Epic Games, Inc. All Rights Reserved.
+
+#pragma once
+
+#include "Net/Core/NetBitArray.h"
+
+// Forward declarations
+class UObjectReplicationBridge;
+
+namespace UE::Net
+{
+	class FNetRefHandle;
+
+	namespace Private
+	{
+		typedef uint32 FInternalNetRefIndex;
+
+		class FReplicationSystemInternal;
+		class FNetRefHandleManager;
+	}
+}
+
+namespace UE::Net::Private
+{
+
+/** Class that holds the required information needed to execute the poll phase on one or multiple replicated objects. */
+class FObjectPoller
+{
+public:
+
+	/** Holds statistics on how the polling went */
+	struct FPreUpdateAndPollStats
+	{
+		uint32 PreUpdatedObjectCount = 0;
+		uint32 PolledObjectCount = 0;
+		uint32 PolledReferencesObjectCount = 0;
+	};
+
+	struct FInitParams
+	{
+		FReplicationSystemInternal* ReplicationSystemInternal = nullptr;
+		UObjectReplicationBridge* ObjectReplicationBridge = nullptr;
+	};
+
+public:
+
+	FObjectPoller(const FInitParams& InitParams);
+
+	const FPreUpdateAndPollStats& GetPollStats() const { return PollStats; }
+
+	/** Poll all the objects whose bit index is set in the array */
+	void PollObjects(const FNetBitArrayView& ObjectsConsideredForPolling);
+
+	/** Poll a single replicated object */
+	void PollSingleObject(FNetRefHandle Handle);
+
+private:
+
+	/** Polls an object in every circumstance */
+	void ForcePollObject(FInternalNetRefIndex ObjectIndex);
+
+	/** Polls an object only if it is required or considered dirty */
+	void PushModelPollObject(FInternalNetRefIndex ObjectIndex);
+
+private:
+
+	UObjectReplicationBridge* ObjectReplicationBridge;
+	FReplicationSystemInternal* ReplicationSystemInternal;
+
+	const FNetRefHandleManager& LocalNetRefHandleManager;
+	const TArray<UObject*>& ReplicatedInstances;
+
+	const FNetBitArrayView AccumulatedDirtyObjects;
+	const FNetBitArrayView PrevFrameGlobalScopeObjects;
+
+	FNetBitArrayView DirtyObjectsThisFrame;
+
+	FNetBitArrayView GarbageCollectionAffectedObjects;
+
+	FPreUpdateAndPollStats PollStats;
+
+};
+
+} // end namespace UE::Net::Private
+
diff --git a/Engine/Source/Runtime/Experimental/Iris/Core/Public/Iris/ReplicationSystem/ObjectReplicationBridge.h b/Engine/Source/Runtime/Experimental/Iris/Core/Public/Iris/ReplicationSystem/ObjectReplicationBridge.h
--- a/Engine/Source/Runtime/Experimental/Iris/Core/Public/Iris/ReplicationSystem/ObjectReplicationBridge.h
+++ b/Engine/Source/Runtime/Experimental/Iris/Core/Public/Iris/ReplicationSystem/ObjectReplicationBridge.h
@@ -22,6 +22,7 @@ namespace UE::Net
 	{
 		typedef uint32 FInternalNetRefIndex;
 		class FObjectPollFrequencyLimiter;
+		class FObjectPoller;
 	}
 }
 
@@ -259,6 +260,8 @@ private:
 
 private:
 
+	friend UE::Net::Private::FObjectPoller;
+
 	TMap<const UClass*, FName> ConfigClassPathNameCache;
 
 	// Prioritization
