UGeometryCollectionComponent::SetPerParticleCollisionProfileName can now adapt to cases where the number of particles in the rest collection changes at runtime

diff --git a/Engine/Source/Runtime/Experimental/GeometryCollectionEngine/Private/GeometryCollection/GeometryCollectionComponent.cpp b/Engine/Source/Runtime/Experimental/GeometryCollectionEngine/Private/GeometryCollection/GeometryCollectionComponent.cpp
--- a/Engine/Source/Runtime/Experimental/GeometryCollectionEngine/Private/GeometryCollection/GeometryCollectionComponent.cpp
+++ b/Engine/Source/Runtime/Experimental/GeometryCollectionEngine/Private/GeometryCollection/GeometryCollectionComponent.cpp
@@ -3241,9 +3241,10 @@ void UGeometryCollectionComponent::SetPerParticleCollisionProfileName(const TArr
 		return;
 	}
 
-	if (CollisionProfilePerParticle.IsEmpty())
+	const int32 NumTransforms = RestCollection->GetGeometryCollection()->Transform.Num();
+	if (CollisionProfilePerParticle.Num() != NumTransforms)
 	{
-		CollisionProfilePerParticle.SetNumZeroed(RestCollection->GetGeometryCollection()->Transform.Num());
+		CollisionProfilePerParticle.SetNumZeroed(NumTransforms);
 	}
 
 	for (int32 Id : BoneIds)
