fix crash on fracture due to out of sync geometry collection components: make FGeometryCollectionEdit propagate update info to all components w/ the same rest collection by default
+ add an extra guard so we will ensure instead of crash if we do get into the bad/mismatched state that this bug was causing

diff --git a/Engine/Source/Runtime/Experimental/GeometryCollectionEngine/Private/GeometryCollection/GeometryCollectionComponent.cpp b/Engine/Source/Runtime/Experimental/GeometryCollectionEngine/Private/GeometryCollection/GeometryCollectionComponent.cpp
--- a/Engine/Source/Runtime/Experimental/GeometryCollectionEngine/Private/GeometryCollection/GeometryCollectionComponent.cpp
+++ b/Engine/Source/Runtime/Experimental/GeometryCollectionEngine/Private/GeometryCollection/GeometryCollectionComponent.cpp
@@ -653,7 +653,11 @@ FBox UGeometryCollectionComponent::ComputeBounds(const FTransform& LocalToWorldW
 
 			const TManagedArray<FTransform>& Transforms = GetTransformArray();
 			const TManagedArray<int32>& ParentIndices = GetParentArray();
-
+			if (!ensure(Transforms.Num() == NumElements))
+			{
+				return FBox(ForceInitToZero);
+			}
+			
 			TArray<FTransform> TmpComponentSpaceTransforms;
 			GeometryCollectionAlgo::GlobalMatrices(Transforms, ParentIndices, TmpComponentSpaceTransforms);
 			if (TmpComponentSpaceTransforms.Num() == 0)
@@ -2743,6 +2747,10 @@ void UGeometryCollectionComponent::ResetDynamicCollection()
 		MarkRenderDynamicDataDirty();
 		SetRenderStateDirty();
 	}
+	else
+	{
+		DynamicCollection.Reset();
+	}
 
 	// make sure we have the RestTransforms up to date, other wise, otherwise there may be case where they do not match the Restcollection ones
 	// can happen if the RestCollection asset has been changed without the component knowing about it 
@@ -3552,21 +3560,46 @@ FString UGeometryCollectionComponent::GetDebugInfo()
 	return DebugInfo;
 }
 
-FGeometryCollectionEdit::FGeometryCollectionEdit(UGeometryCollectionComponent* InComponent, GeometryCollection::EEditUpdate InEditUpdate, bool bShapeIsUnchanged)
+FGeometryCollectionEdit::FGeometryCollectionEdit(UGeometryCollectionComponent* InComponent, GeometryCollection::EEditUpdate InEditUpdate, bool bShapeIsUnchanged, bool bPropagateToAllMatchingComponents)
 	: Component(InComponent)
 	, EditUpdate(InEditUpdate)
 	, bShapeIsUnchanged(bShapeIsUnchanged)
+	, bPropagateToAllMatchingComponents(bPropagateToAllMatchingComponents)
 {
-	bHadPhysicsState = Component->HasValidPhysicsState();
-	if (EnumHasAnyFlags(EditUpdate, GeometryCollection::EEditUpdate::Physics) && bHadPhysicsState)
+	if (!Component->RestCollection)
 	{
-		Component->DestroyPhysicsState();
+		return;
 	}
 
-	if (EnumHasAnyFlags(EditUpdate, GeometryCollection::EEditUpdate::Rest) && GetRestCollection())
+	auto ProcessComponent = [this](UGeometryCollectionComponent* ToProcess)
+	{
+		bool bHadPhysicsState = ToProcess->HasValidPhysicsState();
+		if (EnumHasAnyFlags(EditUpdate, GeometryCollection::EEditUpdate::Physics) && bHadPhysicsState)
+		{
+			HadPhysicsState.Add(ToProcess);
+			ToProcess->DestroyPhysicsState();
+		}
+
+		if (EnumHasAnyFlags(EditUpdate, GeometryCollection::EEditUpdate::Rest) && GetRestCollection())
+		{
+			ToProcess->Modify();
+			GetRestCollection()->Modify(); // Note: If multiple components share the rest collection, its Modify() will be called multiple times, but this should be ok
+		}
+	};
+
+	if (bPropagateToAllMatchingComponents)
 	{
-		Component->Modify();
-		GetRestCollection()->Modify();
+		for (TObjectIterator<UGeometryCollectionComponent> It(RF_ClassDefaultObject, false, EInternalObjectFlags::Garbage); It; ++It)
+		{
+			if (It->RestCollection == Component->RestCollection)
+			{
+				ProcessComponent(*It);
+			}
+		}
+	}
+	else
+	{
+		ProcessComponent(Component);
 	}
 }
 
@@ -3575,11 +3608,6 @@ FGeometryCollectionEdit::~FGeometryCollectionEdit()
 #if WITH_EDITOR
 	if (!!EditUpdate)
 	{
-		if (EnumHasAnyFlags(EditUpdate, GeometryCollection::EEditUpdate::Dynamic))
-		{
-			Component->ResetDynamicCollection();
-		}
-
 		if (EnumHasAnyFlags(EditUpdate, GeometryCollection::EEditUpdate::Rest) && GetRestCollection())
 		{
 			if (!bShapeIsUnchanged)
@@ -3589,9 +3617,32 @@ FGeometryCollectionEdit::~FGeometryCollectionEdit()
 			GetRestCollection()->InvalidateCollection();
 		}
 
-		if (EnumHasAnyFlags(EditUpdate, GeometryCollection::EEditUpdate::Physics) && bHadPhysicsState)
+		auto ProcessComponent = [this](UGeometryCollectionComponent* ToProcess)
+		{
+			if (EnumHasAnyFlags(EditUpdate, GeometryCollection::EEditUpdate::Dynamic))
+			{
+				ToProcess->ResetDynamicCollection();
+			}
+
+			if (EnumHasAnyFlags(EditUpdate, GeometryCollection::EEditUpdate::Physics) && HadPhysicsState.Contains(ToProcess))
+			{
+				ToProcess->RecreatePhysicsState();
+			}
+		};
+
+		if (bPropagateToAllMatchingComponents)
+		{
+			for (TObjectIterator<UGeometryCollectionComponent> It(RF_ClassDefaultObject, false, EInternalObjectFlags::Garbage); It; ++It)
+			{
+				if (It->RestCollection == Component->RestCollection)
+				{
+					ProcessComponent(*It);
+				}
+			}
+		}
+		else
 		{
-			Component->RecreatePhysicsState();
+			ProcessComponent(Component);
 		}
 	}
 #endif
diff --git a/Engine/Source/Runtime/Experimental/GeometryCollectionEngine/Public/GeometryCollection/GeometryCollectionComponent.h b/Engine/Source/Runtime/Experimental/GeometryCollectionEngine/Public/GeometryCollection/GeometryCollectionComponent.h
--- a/Engine/Source/Runtime/Experimental/GeometryCollectionEngine/Public/GeometryCollection/GeometryCollectionComponent.h
+++ b/Engine/Source/Runtime/Experimental/GeometryCollectionEngine/Public/GeometryCollection/GeometryCollectionComponent.h
@@ -202,11 +202,12 @@ class GEOMETRYCOLLECTIONENGINE_API FGeometryCollectionEdit
 {
 public:
 	/**
-	 * @param InComponent				The component to edit
-	 * @param EditUpdate				What parts of the geometry collection to update
-	 * @param bShapeIsUnchanged			Override indicating the overall shape of the geometry and clusters is unchanged, even if the rest collection changed.  Useful to e.g., not re-compute convex hulls when we don't need to.
+	 * @param InComponent					The component to edit
+	 * @param EditUpdate					What parts of the geometry collection to update
+	 * @param bShapeIsUnchanged				Override indicating the overall shape of the geometry and clusters is unchanged, even if the rest collection changed.  Useful to e.g., not re-compute convex hulls when we don't need to.
+	 * @param bPropagateAcrossComponents	Propagate updates to all components with the same underlying Rest Collection
 	 */
-	FGeometryCollectionEdit(UGeometryCollectionComponent* InComponent, GeometryCollection::EEditUpdate EditUpdate = GeometryCollection::EEditUpdate::RestPhysics, bool bShapeIsUnchanged = false);
+	FGeometryCollectionEdit(UGeometryCollectionComponent* InComponent, GeometryCollection::EEditUpdate EditUpdate = GeometryCollection::EEditUpdate::RestPhysics, bool bShapeIsUnchanged = false, bool bPropagateToAllMatchingComponents = true);
 	~FGeometryCollectionEdit();
 
 	UGeometryCollection* GetRestCollection();
@@ -214,8 +215,9 @@ public:
 private:
 	UGeometryCollectionComponent* Component;
 	const GeometryCollection::EEditUpdate EditUpdate;
-	bool bHadPhysicsState;
+	TSet<UGeometryCollectionComponent*> HadPhysicsState;
 	bool bShapeIsUnchanged;
+	bool bPropagateToAllMatchingComponents;
 };
 
 #if WITH_EDITOR
