Networked Physics: Render interpolate the offset from resimulation via linear offset decay.

Along with another optional interpolation layer that takes incoming velocity into account to make the correction more smooth / organic.

The optional "error velocity smoothing" is guarded behind a define (RENDERINTERP_ERRORVELOCITYSMOOTHING) and takes a little more CPU and memory to run.

There are CVars for how long time the error correction should be done over and for rendering draw debug.

The error offset gets cached in PT and needs to be marshalled to GT where it's used.

diff --git a/Engine/Source/Programs/HeadlessChaos/Private/HeadlessChaosTestRewind.cpp b/Engine/Source/Programs/HeadlessChaos/Private/HeadlessChaosTestRewind.cpp
--- a/Engine/Source/Programs/HeadlessChaos/Private/HeadlessChaosTestRewind.cpp
+++ b/Engine/Source/Programs/HeadlessChaos/Private/HeadlessChaosTestRewind.cpp
@@ -2108,6 +2108,7 @@ namespace ChaosTest {
 				FReal Time = 0;
 				FReal ZStart = 0;
 				const FReal ZVel = -1;
+				FReal LastCorrectionStep = 0;
 
 				Solver->SetRewindCallback(MoveTemp(UniqueRewindCallback));
 
@@ -2154,17 +2155,32 @@ namespace ChaosTest {
 						const int32 NextSimStep = FMath::CeilToInt(InterpolatedTime / SimDt);
 						const FReal NextSimStepTime = NextSimStep * SimDt;
 						const FReal ExpectedValue = ZStart + ZVel * InterpolatedTime;
-						const FReal TargetValue = ZStart + ZVel * NextSimStepTime;
+					//	const FReal TargetValue = ZStart + ZVel * NextSimStepTime;
 
-						if(!Proxy->GetInterpolationData().IsResimSmoothing())
+						if (Proxy->GetInterpolationData().IsErrorSmoothing())
 						{
-							//no resim interpolation, just simple value interpolation
-							EXPECT_NEAR(Particle.X()[2], ExpectedValue, 1e-2);
+#if !RENDERINTERP_ERRORVELOCITYSMOOTHING
+							const FReal CorrectionStep = Particle.X()[2] - PrevZ;
+							if (LastCorrectionStep != 0)
+							{
+								// Make sure we have a linear correction
+								EXPECT_NEAR(LastCorrectionStep, CorrectionStep, 1e-2);
+							}
+
+							LastCorrectionStep = CorrectionStep;
+#endif
 						}
 						else
 						{
-							//exponential decay from current state to target
-							EXPECT_NEAR(Particle.X()[2], FMath::Lerp(PrevZ, ExpectedValue, Chaos::ResimInterpStrength), 1e-2);
+							if (!!LastCorrectionStep)
+							{
+								// Correction is now done, check if the object arrived at the current position by the correction or if it was snapped into place via not doing correction anymore
+								EXPECT_NEAR(PrevZ, Particle.X()[2] - LastCorrectionStep, 1e-2);
+								LastCorrectionStep = 0;
+							}
+
+							//no resim interpolation, just simple value interpolation
+							EXPECT_NEAR(Particle.X()[2], ExpectedValue, 1e-2);
 						}
 					}
 				}
diff --git a/Engine/Source/Runtime/Experimental/Chaos/Private/Chaos/Framework/ChaosResultsManager.cpp b/Engine/Source/Runtime/Experimental/Chaos/Private/Chaos/Framework/ChaosResultsManager.cpp
--- a/Engine/Source/Runtime/Experimental/Chaos/Private/Chaos/Framework/ChaosResultsManager.cpp
+++ b/Engine/Source/Runtime/Experimental/Chaos/Private/Chaos/Framework/ChaosResultsManager.cpp
@@ -85,7 +85,7 @@ namespace Chaos
 		{
 			if (auto Proxy = Data.GetProxy())
 			{
-				FProxyInterpolationData& InterpolationData = Proxy->GetInterpolationData(); 
+				FProxyInterpolationBase& InterpolationData = Proxy->GetInterpolationData();
 				
 				//If proxy is not associated with this channel, do nothing
 				if (InterpolationData.GetInterpChannel_External() != ChannelIdx)
@@ -293,8 +293,8 @@ namespace Chaos
 		{
 			FSingleParticlePhysicsProxy* Proxy = Itr.Key;
 
-			FProxyInterpolationData& InterpolationData = Proxy->GetInterpolationData(); 
-			if(InterpolationData.IsResimSmoothing())
+			FProxyInterpolationBase& InterpolationData = Proxy->GetInterpolationData();
+			if (InterpolationData.IsErrorSmoothing())
 			{
 				if (InterpolationData.GetPullDataInterpIdx_External() == INDEX_NONE)	//not in results array
 				{
@@ -399,12 +399,11 @@ namespace Chaos
 		{
 			if (FSingleParticlePhysicsProxy* ResimProxy = ResimDirty.GetProxy())
 			{
-				FProxyInterpolationData& InterpolationData = ResimProxy->GetInterpolationData(); 
+				FProxyInterpolationBase& InterpolationData = ResimProxy->GetInterpolationData(); 
 				//Mark as resim only if proxy is owned by this channel
 				if(InterpolationData.GetInterpChannel_External() == ChannelIdx)
 				{
 					ParticleToResimTarget.FindOrAdd(ResimProxy) = ResimDirty;
-					InterpolationData.SetResimSmoothing(true);
 				}
 			}
 		}
diff --git a/Engine/Source/Runtime/Experimental/Chaos/Private/Chaos/PullPhysicsDataImp.cpp b/Engine/Source/Runtime/Experimental/Chaos/Private/Chaos/PullPhysicsDataImp.cpp
--- a/Engine/Source/Runtime/Experimental/Chaos/Private/Chaos/PullPhysicsDataImp.cpp
+++ b/Engine/Source/Runtime/Experimental/Chaos/Private/Chaos/PullPhysicsDataImp.cpp
@@ -8,6 +8,7 @@ namespace Chaos
 void FPullPhysicsData::Reset()
 {
 	DirtyRigids.Reset();
+	DirtyRigidErrors.Reset();
 	DirtyGeometryCollections.Reset();
 	DirtyClusterUnions.Reset();
 	DirtyJointConstraints.Reset();
diff --git a/Engine/Source/Runtime/Experimental/Chaos/Private/PBDRigidsSolver.cpp b/Engine/Source/Runtime/Experimental/Chaos/Private/PBDRigidsSolver.cpp
--- a/Engine/Source/Runtime/Experimental/Chaos/Private/PBDRigidsSolver.cpp
+++ b/Engine/Source/Runtime/Experimental/Chaos/Private/PBDRigidsSolver.cpp
@@ -1661,10 +1661,10 @@ namespace Chaos
 			const int32 ResimStep = MRewindCallback->TriggerRewindIfNeeded_Internal(LastStep);
 			const int32 NumResimSteps = LastStep - ResimStep + 1;
 
-			const bool bEnableNetworlPredictionDebug = IsNetworkPhysicsPredictionEnabled() && CanDebugNetworkPhysicsPrediction();
+			const bool bEnableNetworkPredictionDebug = IsNetworkPhysicsPredictionEnabled() && CanDebugNetworkPhysicsPrediction();
 			
 #if !(UE_BUILD_SHIPPING || UE_BUILD_TEST)
-			if (bEnableNetworlPredictionDebug)
+			if (bEnableNetworkPredictionDebug)
 			{
 				UE_LOG(LogChaos, Log, TEXT("	-> Trying to Rewind Frame = %d | At Time = %f | Num Steps = %d | Resim Step = %d | Last Step = %d"), CurrentFrame,  MTime, NumResimSteps, ResimStep, LastStep);
 				for(auto& Handle : GetParticles().GetNonDisabledDynamicView())
@@ -1713,7 +1713,7 @@ namespace Chaos
 								
 							GetEvolution()->SetReset(bFirst);
 #if !(UE_BUILD_SHIPPING || UE_BUILD_TEST)
-							if(bEnableNetworlPredictionDebug)
+							if(bEnableNetworkPredictionDebug)
 							{
 								UE_LOG(LogChaos, Log, TEXT("		-> Re-simulating Frame = %d "), Step);
 								for(auto& Handle : GetParticles().GetNonDisabledDynamicView())
@@ -1865,6 +1865,12 @@ namespace Chaos
 			}
 		}
 
+		// Move correction error data from RewindData to FPullPhysicsData for marshalling
+		if (FRewindData* RewindData = GetRewindData())
+		{
+			RewindData->BufferPhysicsResults(PullData->DirtyRigidErrors);
+		}
+
 		{
 			ensure(PullData->DirtyGeometryCollections.Num() == 0);	//we only fill this once per frame
 			PullData->DirtyGeometryCollections.Reserve(ActiveGC.Num());
diff --git a/Engine/Source/Runtime/Experimental/Chaos/Private/PhysicsProxy/SingleParticlePhysicsProxy.cpp b/Engine/Source/Runtime/Experimental/Chaos/Private/PhysicsProxy/SingleParticlePhysicsProxy.cpp
--- a/Engine/Source/Runtime/Experimental/Chaos/Private/PhysicsProxy/SingleParticlePhysicsProxy.cpp
+++ b/Engine/Source/Runtime/Experimental/Chaos/Private/PhysicsProxy/SingleParticlePhysicsProxy.cpp
@@ -16,6 +16,7 @@
 #include "Chaos/ChaosMarshallingManager.h"
 #include "Chaos/PullPhysicsDataImp.h"
 #include "RewindData.h"
+#include "Chaos/DebugDrawQueue.h"
 
 namespace Chaos
 {
@@ -251,19 +252,17 @@ void FSingleParticlePhysicsProxy::BufferPhysicsResults_External(Chaos::FDirtyRig
 	}
 }
 
-FRealSingle ResimInterpStrength = 0.2f;
-FAutoConsoleVariableRef CVarResimInterpStrength(TEXT("p.ResimInterpStrength"), ResimInterpStrength, TEXT("How strong the resim interp leash is. 1 means immediately snap to new target, 0 means do not interpolate at all"));
 
-FRealSingle ResimInterpStrength2 = 0.05f;
-FAutoConsoleVariableRef CVarResimInterpStrength2(TEXT("p.ResimInterpStrength2"), ResimInterpStrength2, TEXT("How strong the resim interp leash is for object in channel 2. 1 means immediately snap to new target, 0 means do not interpolate at all"));
+float RenderInterpErrorCorrectionDuration = 0.5f;
+FAutoConsoleVariableRef CVarRenderInterpErrorCorrectionDuration(TEXT("p.RenderInterp.ErrorCorrectionDuration"), RenderInterpErrorCorrectionDuration, TEXT("How long in seconds to apply error correction over."));
 
-FRealSingle MinLinError2ForResimInterp = 1.f;
-FAutoConsoleVariableRef CVarMinLinError2ForResimInterp(TEXT("p.MinLinError2ForResimInterp"), MinLinError2ForResimInterp, TEXT("The minimum squared error needed to continue interpolation during a resim"));
+float RenderInterpErrorVelocitySmoothingDuration = 0.5f;
+FAutoConsoleVariableRef CVarRenderInterpErrorVelocitySmoothingDuration(TEXT("p.RenderInterp.ErrorVelocitySmoothingDuration"), RenderInterpErrorVelocitySmoothingDuration, TEXT("How long in seconds to apply error velocity smoothing correction over, should be smaller than or equal to p.RenderInterp.ErrorCorrectionDuration. RENDERINTERPOLATION_VELOCITYSMOOTHING needs to be defined."));
 
-FRealSingle MinRotErrorForResimInterp = 0.1f;
-FAutoConsoleVariableRef CVarMinRotErrorForResimInterp(TEXT("p.MinRotErrorForResimInterp"), MinRotErrorForResimInterp, TEXT("The minimum rotation error needed to continue interpolation during a resim"));
+int32 RenderInterpDebugDraw = 0;
+FAutoConsoleVariableRef CVarRenderInterpDebugDraw(TEXT("p.RenderInterp.DebugDraw"), RenderInterpDebugDraw, TEXT("Draw debug lines for physics render interpolation, also needs p.Chaos.DebugDraw.Enabled set"));
 
-bool FSingleParticlePhysicsProxy::PullFromPhysicsState(const Chaos::FDirtyRigidParticleData& PullData,int32 SolverSyncTimestamp, const Chaos::FDirtyRigidParticleData* NextPullData, const Chaos::FRealSingle* Alpha)
+bool FSingleParticlePhysicsProxy::PullFromPhysicsState(const Chaos::FDirtyRigidParticleData& PullData,int32 SolverSyncTimestamp, const Chaos::FDirtyRigidParticleData* NextPullData, const Chaos::FRealSingle* Alpha, const FDirtyRigidParticleReplicationErrorData* Error, const Chaos::FReal AsyncFixedTimeStep)
 {
 	using namespace Chaos;
 	// Move buffered data into the TPBDRigidParticle without triggering invalidation of the physics state.
@@ -289,7 +288,20 @@ bool FSingleParticlePhysicsProxy::PullFromPhysicsState(const Chaos::FDirtyRigidP
 
 		const FSingleParticleProxyTimestamp* ProxyTimestamp = PullData.GetTimestamp();
 		
-		if(NextPullData)
+#if RENDERINTERP_ERRORVELOCITYSMOOTHING
+		const int32 RenderInterpErrorVelocitySmoothingDurationTicks = FMath::CeilToInt32(RenderInterpErrorVelocitySmoothingDuration / AsyncFixedTimeStep); // Convert duration from seconds to simulation ticks
+#endif
+
+		if (Error)
+		{
+			const int32 RenderInterpErrorCorrectionDurationTicks = FMath::CeilToInt32(RenderInterpErrorCorrectionDuration / AsyncFixedTimeStep); // Convert duration from seconds to simulation ticks
+			InterpolationData.AccumlateErrorXR(Error->ErrorX, Error->ErrorR, SolverSyncTimestamp, RenderInterpErrorCorrectionDurationTicks);
+#if RENDERINTERP_ERRORVELOCITYSMOOTHING
+			InterpolationData.SetVelocitySmoothing(Rigid->V(), Rigid->X(), RenderInterpErrorVelocitySmoothingDurationTicks);
+#endif
+		}
+
+		if (NextPullData)
 		{
 			auto LerpHelper = [SolverSyncTimestamp](const auto& Prev, const auto& OverwriteProperty) -> const auto*
 			{
@@ -303,41 +315,78 @@ bool FSingleParticlePhysicsProxy::PullFromPhysicsState(const Chaos::FDirtyRigidP
 
 			if (bUpdatePositionFromSimulation)
 			{
-				const float UseResimInterpStrength = InterpolationData.GetInterpChannel_External() == 0 ? ResimInterpStrength : ResimInterpStrength2;
+				const bool bIsReplicationErrorSmoothing = InterpolationData.IsErrorSmoothing();
+#if RENDERINTERP_ERRORVELOCITYSMOOTHING
+				const bool bIsErrorVelocitySmoothing = InterpolationData.IsErrorVelocitySmoothing();
+#endif
+
+				InterpolationData.UpdateError(SolverSyncTimestamp, AsyncFixedTimeStep);
 
-				bool bKeepSmoothing = false;
 				if (const FVec3* Prev = LerpHelper(PullData.X, ProxyTimestamp->OverWriteX))
 				{
 					FVec3 Target = FMath::Lerp(*Prev, NextPullData->X, *Alpha);
-					if (InterpolationData.IsResimSmoothing())
+					if (bIsReplicationErrorSmoothing)
 					{
-						const FVec3 SmoothedTarget = FMath::Lerp(Rigid->X(), Target, UseResimInterpStrength);
-						if((SmoothedTarget - Target).SizeSquared() > MinLinError2ForResimInterp)
+						Target += InterpolationData.GetErrorX(*Alpha);
+
+#if RENDERINTERP_ERRORVELOCITYSMOOTHING
+						if (bIsErrorVelocitySmoothing)
 						{
-							bKeepSmoothing = true;
-							Target = SmoothedTarget;
+#if CHAOS_DEBUG_DRAW
+							if (!!RenderInterpDebugDraw)
+							{
+								Chaos::FDebugDrawQueue::GetInstance().DrawDebugDirectionalArrow((Target - InterpolationData.GetErrorX(*Alpha)), Target, 1, FColor::Blue, false, 5.0f, 0, 0.5f);
+								Chaos::FDebugDrawQueue::GetInstance().DrawDebugBox(Target, FVector(2, 1, 1), Rigid->R(), FColor::Cyan, false, 5.f, 0, 0.25f);
+							}
+#endif // CHAOS_DEBUG_DRAW
+
+							Target = FMath::Lerp(Target, InterpolationData.GetErrorVelocitySmoothingX(*Alpha), InterpolationData.GetErrorVelocitySmoothingAlpha(RenderInterpErrorVelocitySmoothingDurationTicks));
 						}
-
+#endif // RENDERINTERP_ERRORVELOCITYSMOOTHING
 					}
-					Rigid->SetX(Target, false);
+
+					Rigid->SetX(Target, false);					
 				}
 
 				if (const FQuat* Prev = LerpHelper(PullData.R, ProxyTimestamp->OverWriteR))
 				{
 					FQuat Target = FMath::Lerp(*Prev, NextPullData->R, *Alpha);
-					if (InterpolationData.IsResimSmoothing())
+					if (bIsReplicationErrorSmoothing)
 					{
-						const FQuat SmoothedTarget = FMath::Lerp<FQuat>(Rigid->R(), Target, UseResimInterpStrength);
-						if(FQuat::ErrorAutoNormalize(SmoothedTarget, Target) > MinRotErrorForResimInterp)
-						{
-							bKeepSmoothing = true;
-							Target = SmoothedTarget;
-						}
+						Target = InterpolationData.GetErrorR(*Alpha) * Target;
 					}
 					Rigid->SetR(Target, false);
 				}
+				
+#if CHAOS_DEBUG_DRAW
+				if (!!RenderInterpDebugDraw)
+				{
+					Chaos::FDebugDrawQueue::GetInstance().DrawDebugBox(NextPullData->X, FVector(2, 1, 1), NextPullData->R, FColor::Yellow, false, 5.f, 0, 0.5f);
+					Chaos::FDebugDrawQueue::GetInstance().DrawDebugDirectionalArrow(PullData.X, NextPullData->X, 0.5f, FColor::Yellow, false, 5.0f, 0, 0.5f);
+					Chaos::FDebugDrawQueue::GetInstance().DrawDebugBox(Rigid->X(), FVector(2, 1, 1), Rigid->R(), FColor::Green, false, 5.f, 0, 0.5f);
 
-				InterpolationData.SetResimSmoothing(bKeepSmoothing);
+					if (bIsReplicationErrorSmoothing)
+					{
+						if (Error)
+						{
+							Chaos::FDebugDrawQueue::GetInstance().DrawDebugBox(PullData.X, FVector(4, 2, 2), PullData.R, FColor::Red, false, 5.f, 0, 0.5f);
+							Chaos::FDebugDrawQueue::GetInstance().DrawDebugDirectionalArrow(PullData.X, (PullData.X + InterpolationData.GetErrorX(0)), 1, FColor::Red, false, 5.0f, 0, 0.5f);
+						}
+
+#if RENDERINTERP_ERRORVELOCITYSMOOTHING
+						if (bIsErrorVelocitySmoothing)
+						{
+							Chaos::FDebugDrawQueue::GetInstance().DrawDebugBox(InterpolationData.GetErrorVelocitySmoothingX(*Alpha), FVector(2, 2, 2), Rigid->R(), FColor::Purple, false, 5.f, 0, 0.5f);
+						}
+						else
+#endif // RENDERINTERP_ERRORVELOCITYSMOOTHING
+						{
+							Chaos::FDebugDrawQueue::GetInstance().DrawDebugDirectionalArrow((Rigid->X() - InterpolationData.GetErrorX(*Alpha)), Rigid->X(), 1, FColor::Blue, false, 5.0f, 0, 0.5f);
+						}
+					}
+
+				}
+#endif // CHAOS_DEBUG_DRAW
 			}
 
 			if (const FVec3* Prev = LerpHelper(PullData.V, ProxyTimestamp->OverWriteV))
diff --git a/Engine/Source/Runtime/Experimental/Chaos/Private/RewindData.cpp b/Engine/Source/Runtime/Experimental/Chaos/Private/RewindData.cpp
--- a/Engine/Source/Runtime/Experimental/Chaos/Private/RewindData.cpp
+++ b/Engine/Source/Runtime/Experimental/Chaos/Private/RewindData.cpp
@@ -271,6 +271,8 @@ bool FRewindData::RewindToFrame(int32 Frame)
 		
 		FGeometryParticleStateBase& History = DirtyParticleInfo.AddFrame(CurFrame);	//non-const in case we need to record what's at head for a rewind (CurFrame has already been increased to the next frame)
 
+		History.CachePreCorrectionState(*PTParticle);
+
 		const bool bResimAsFollower = DirtyParticleInfo.bResimAsFollower;
 
 		bool bAnyChange = RewindHelper(PTParticle, bResimAsFollower, History.ParticlePositionRotation, [](auto Particle, const auto& Data) {Particle->SetXR(Data); });
@@ -341,6 +343,40 @@ void FRewindData::DesyncIfNecessary(TDirtyInfo& Info, const FFrameAndPhase Frame
 	}
 }
 
+template<>
+void FRewindData::AccumulateErrorIfNecessary(FGeometryParticleHandle& Obj, const FFrameAndPhase FrameAndPhase)
+{
+	const FDirtyParticleInfo* DirtyInfo = DirtyParticles.Find(&Obj);
+	if (!DirtyInfo)
+	{
+		return;
+	}
+
+	// Get the error offset after a correction
+	const FVec3 ErrorX = DirtyInfo->GetHistory().PreCorrectionXR.X() - Obj.X();
+	FQuat ErrorR = DirtyInfo->GetHistory().PreCorrectionXR.R() * Obj.R().Inverse();
+	ErrorR.EnforceShortestArcWith(FQuat::Identity);
+	ErrorR.Normalize();
+
+	// Check if error is large enough to hide behind render interpolation
+	if (!ErrorX.IsNearlyZero(0.1) || !ErrorR.IsIdentity(0.02))
+	{
+		// Find or add FDirtyParticleErrorInfo for the particle that has an error
+		FDirtyParticleErrorInfo& ErrorInfo = [&]() ->FDirtyParticleErrorInfo&
+		{
+			if (FDirtyParticleErrorInfo* Found = DirtyParticleErrors.Find(&Obj))
+			{
+				return *Found;
+			}
+
+			return DirtyParticleErrors.Add(&Obj, FDirtyParticleErrorInfo(Obj));
+		}();
+
+		// Cache error for particle 
+		ErrorInfo.AccumulateError(ErrorX, ErrorR);
+	}
+}
+
 void FRewindData::FinishFrame()
 {
 	QUICK_SCOPE_CYCLE_COUNTER(RewindDataFinishFrame);
@@ -365,6 +401,9 @@ void FRewindData::FinishFrame()
 				{
 					if (IsFinalResim())
 					{
+						// Cache the correction offset after a resimulation
+						AccumulateErrorIfNecessary(Handle, FutureFrame);
+
 						//Last resim so mark as in sync
 						Handle.SetSyncState(ESyncState::InSync);
 						Handle.SetEnabledDuringResim(false);
@@ -650,6 +689,22 @@ void FRewindData::SpawnProxyIfNeeded(FSingleParticlePhysicsProxy& Proxy)
 	}
 }
 
+// Hand over error data per particle from RewindData to a solver collection instead which gets marshalled to GT
+void FRewindData::BufferPhysicsResults(TMap<const FSingleParticlePhysicsProxy*, struct FDirtyRigidParticleReplicationErrorData>& DirtyRigidErrors)
+{
+	DirtyRigidErrors.Reserve(DirtyParticleErrors.Num());
+
+	for (const FDirtyParticleErrorInfo& ErrorInfo : DirtyParticleErrors)
+	{
+		FDirtyRigidParticleReplicationErrorData ErrorData;
+		ErrorData.ErrorX = ErrorInfo.GetErrorX();
+		ErrorData.ErrorR = ErrorInfo.GetErrorR();
+		DirtyRigidErrors.Add(static_cast<const FSingleParticlePhysicsProxy*>(ErrorInfo.GetObjectPtr()->PhysicsProxy()), ErrorData);
+	}
+
+	DirtyParticleErrors.Reset();
+}
+
 void FRewindData::MarkDirtyFromPT(FGeometryParticleHandle& Handle)
 {
 	FDirtyParticleInfo& Info = FindOrAddDirtyObj(Handle);
diff --git a/Engine/Source/Runtime/Experimental/Chaos/Public/Chaos/Framework/PhysicsProxyBase.h b/Engine/Source/Runtime/Experimental/Chaos/Public/Chaos/Framework/PhysicsProxyBase.h
--- a/Engine/Source/Runtime/Experimental/Chaos/Public/Chaos/Framework/PhysicsProxyBase.h
+++ b/Engine/Source/Runtime/Experimental/Chaos/Public/Chaos/Framework/PhysicsProxyBase.h
@@ -155,27 +155,153 @@ struct PhysicsProxyWrapper
 	EPhysicsProxyType Type;
 };
 
-// Data ise used by interpolation code
-struct FProxyInterpolationData
+
+struct FProxyInterpolationBase
 {
-public:
-	FProxyInterpolationData()
+	FProxyInterpolationBase()
 		: PullDataInterpIdx_External(INDEX_NONE)
 		, InterpChannel_External(0)
-		, bResimSmoothing(false)
 	{}
 
 	int32 GetPullDataInterpIdx_External() const { return PullDataInterpIdx_External; }
-	void SetPullDataInterpIdx_External(const int32 Idx) { PullDataInterpIdx_External = Idx;	}
+	void SetPullDataInterpIdx_External(const int32 Idx) { PullDataInterpIdx_External = Idx; }
 
 	int32 GetInterpChannel_External() const { return InterpChannel_External; }
-	void SetInterpChannel_External(int32 Channel) { InterpChannel_External = Channel; }
+	void SetInterpChannel_External(const int32 Channel) { InterpChannel_External = Channel; }
 
-	void SetResimSmoothing(bool ResimSmoothing) { bResimSmoothing = ResimSmoothing; }
-	bool IsResimSmoothing() const { return bResimSmoothing; }
-	
-private:
+	virtual bool IsErrorSmoothing() { return false; }
+
+protected:
 	int32 PullDataInterpIdx_External;
 	int32 InterpChannel_External;
-	bool bResimSmoothing;
+};
+
+// Render interpolation that can correct errors from resimulation / repositions through a linear decay over N simulation tick.
+struct FProxyInterpolationError : FProxyInterpolationBase
+{
+	using Super = FProxyInterpolationBase;
+
+	FProxyInterpolationError() : Super()
+	{}
+
+	virtual bool IsErrorSmoothing() override { return ErrorSmoothingCount > 0; }
+
+	Chaos::FVec3 GetErrorX(const Chaos::FRealSingle Alpha) { return FMath::Lerp(ErrorXPrev, ErrorX, Alpha); } // Get the ErrorX based on current Alpha between GT and PT
+	FQuat GetErrorR(const Chaos::FRealSingle Alpha) { return FMath::Lerp(ErrorRPrev, ErrorR, Alpha); } // Get the ErrorR based on current Alpha between GT and PT
+
+	void AccumlateErrorXR(const Chaos::FVec3 X, const FQuat R, const int32 CurrentSimTick, const int32 ErrorSmoothDuration)
+	{
+		ErrorX += X;
+		ErrorXPrev = ErrorX;
+		ErrorR *= R;
+		ErrorRPrev = ErrorR;
+		ErrorSmoothingCount = ErrorSmoothDuration; // How many simulation ticks to correct error over
+		LastSimTick = CurrentSimTick - 1; // Error is from the previous simulation tick, not the current
+	}
+
+	virtual bool UpdateError(const int32 CurrentSimTick, const Chaos::FReal AsyncFixedTimeStep)
+	{
+		// Cache how many simulation ticks have passed since last call
+		SimTicks = CurrentSimTick - LastSimTick;
+		LastSimTick = CurrentSimTick;
+
+		if (IsErrorSmoothing() && SimTicks > 0)
+		{
+			DecayError();
+			return true;
+		}
+		return false;
+	}
+
+
+protected:
+	void DecayError()
+	{
+		// Linear decay
+		// Example: If we want to decay an error of 100 over 10ticks (i.e. 10% each tick)
+		// First step:  9/10 = 0.9   |  100 * 0.9  = 90 error
+		// Second step: 8/9  = 0.888 |  90 * 0.888 = 80 error
+		// Third step: 7/8  = 0.875 |  80 * 0.875 = 70 error
+		// etc.
+		Chaos::FRealSingle Alpha = Chaos::FRealSingle(ErrorSmoothingCount - SimTicks) / Chaos::FRealSingle(ErrorSmoothingCount);
+		Alpha = FMath::Clamp(Alpha, 0.f, 1.f);
+
+		ErrorXPrev = ErrorX;
+		ErrorX *= Alpha;
+
+		ErrorRPrev = ErrorR;
+		ErrorR = FMath::Lerp(FQuat::Identity, ErrorR, Alpha);
+
+		ErrorSmoothingCount = FMath::Max(ErrorSmoothingCount - SimTicks, 0);
+	}
+
+
+protected:
+	int32 LastSimTick = 0;
+	int32 SimTicks = 0;
+
+	Chaos::FVec3 ErrorX = { 0,0,0 };
+	Chaos::FVec3 ErrorXPrev = { 0,0,0 };
+	FQuat ErrorR = FQuat::Identity;
+	FQuat ErrorRPrev = FQuat::Identity;
+	int32 ErrorSmoothingCount = 0;
+
+};
+
+
+// Take incoming velocity into consideration when performing render interpolation, the correction will be more organic but might result in clipping and it's heavier for memory and CPU.
+#ifndef RENDERINTERP_ERRORVELOCITYSMOOTHING
+#define RENDERINTERP_ERRORVELOCITYSMOOTHING 0
+#endif
+// Render Interpolation that both perform the linear error correction from FProxyInterpolationError and takes incoming velocity into account to make a smoother and more organic correction of the error.
+struct FProxyInterpolationErrorVelocity : FProxyInterpolationError
+{
+	using Super = FProxyInterpolationError;
+
+	FProxyInterpolationErrorVelocity() : Super()
+	{}
+
+	bool IsErrorVelocitySmoothing() { return ErrorVelocitySmoothingCount > 0; }
+
+	// Returns the Alpha of how much to take previous velocity into account, used to lerp from linear extrapolation to the predicted position based on previous velocity.
+	Chaos::FRealSingle GetErrorVelocitySmoothingAlpha(const int32 ErrorVelocitySmoothDuration) { return Chaos::FRealSingle(ErrorVelocitySmoothingCount) / Chaos::FRealSingle(ErrorVelocitySmoothDuration); }
+	Chaos::FVec3 GetErrorVelocitySmoothingX(const Chaos::FRealSingle Alpha) { return FMath::Lerp(ErrorVelocitySmoothingXPrev, ErrorVelocitySmoothingX, Alpha); } // Get the VelocityErrorX based on current Alpha between GT and PT
+
+	bool UpdateError(const int32 CurrentSimTick, const Chaos::FReal AsyncFixedTimeStep) override
+	{		
+		if (Super::UpdateError(CurrentSimTick, AsyncFixedTimeStep))
+		{
+			StepErrorVelocitySmoothingData(AsyncFixedTimeStep);
+			return true;
+		}
+		return false;
+	}
+
+	void SetVelocitySmoothing(const Chaos::FVec3 CurrV, const Chaos::FVec3 CurrX, const int32 ErrorVelocitySmoothDuration)
+	{
+		// Cache pre error velocity and position to be used when smoothing out error correction
+		ErrorVelocitySmoothingV = CurrV;
+		ErrorVelocitySmoothingX = CurrX;
+		ErrorVelocitySmoothingXPrev = ErrorVelocitySmoothingX;
+		ErrorVelocitySmoothingCount = ErrorVelocitySmoothDuration;
+	}
+
+protected:
+	void StepErrorVelocitySmoothingData(const Chaos::FReal AsyncFixedTimeStep)
+	{
+		// Step the error velocity smoothing position forward along the previous velocity to have a new position to base smoothing on each tick
+		if (IsErrorVelocitySmoothing())
+		{
+			const Chaos::FReal Time = AsyncFixedTimeStep * SimTicks;
+			ErrorVelocitySmoothingXPrev = ErrorVelocitySmoothingX;
+			ErrorVelocitySmoothingX += ErrorVelocitySmoothingV * Time;
+
+			ErrorVelocitySmoothingCount = FMath::Max(ErrorVelocitySmoothingCount - SimTicks, 0);
+		}
+	}
+
+	Chaos::FVec3 ErrorVelocitySmoothingV = { 0,0,0 };
+	Chaos::FVec3 ErrorVelocitySmoothingX = { 0,0,0 };
+	Chaos::FVec3 ErrorVelocitySmoothingXPrev = { 0,0,0 };
+	int32 ErrorVelocitySmoothingCount = 0;
 };
\ No newline at end of file
diff --git a/Engine/Source/Runtime/Experimental/Chaos/Public/Chaos/PhysicsSolverBaseImpl.h b/Engine/Source/Runtime/Experimental/Chaos/Public/Chaos/PhysicsSolverBaseImpl.h
--- a/Engine/Source/Runtime/Experimental/Chaos/Public/Chaos/PhysicsSolverBaseImpl.h
+++ b/Engine/Source/Runtime/Experimental/Chaos/Public/Chaos/PhysicsSolverBaseImpl.h
@@ -81,10 +81,12 @@ namespace Chaos
 					{
 						if(FSingleParticlePhysicsProxy* Proxy = RigidInterp.Prev.GetProxy())
 						{
-							if(Proxy->PullFromPhysicsState(RigidInterp.Prev, SolverTimestamp, &RigidInterp.Next, &Results.Alpha))
+							const FDirtyRigidParticleReplicationErrorData* ErrorData = LatestData->DirtyRigidErrors.Find(Proxy);
+							if(Proxy->PullFromPhysicsState(RigidInterp.Prev, SolverTimestamp, &RigidInterp.Next, &Results.Alpha, ErrorData, GetAsyncDeltaTime()))
 							{
 								Private::TPullPhysicsStateDispatchHelper<TDispatcher>::Apply(Dispatcher, Proxy);
 							}
+							LatestData->DirtyRigidErrors.Remove(Proxy);
 						}
 					}
 				}
diff --git a/Engine/Source/Runtime/Experimental/Chaos/Public/Chaos/PullPhysicsDataImp.h b/Engine/Source/Runtime/Experimental/Chaos/Public/Chaos/PullPhysicsDataImp.h
--- a/Engine/Source/Runtime/Experimental/Chaos/Public/Chaos/PullPhysicsDataImp.h
+++ b/Engine/Source/Runtime/Experimental/Chaos/Public/Chaos/PullPhysicsDataImp.h
@@ -58,6 +58,12 @@ struct FDirtyRigidParticleData : public TBasePullData<FSingleParticlePhysicsProx
 	EObjectStateType ObjectState;
 };
 
+struct FDirtyRigidParticleReplicationErrorData : public TBasePullData<FSingleParticlePhysicsProxy, FSingleParticleProxyTimestamp>
+{
+	FVec3 ErrorX;
+	FQuat ErrorR;
+};
+
 struct FDirtyGeometryCollectionData : public TBasePullData<FGeometryCollectionPhysicsProxy, FProxyTimestampBase>
 {
 	FGeometryCollectionResults Results;
@@ -110,6 +116,7 @@ class FPullPhysicsData
 {
 public:
 	TArray<FDirtyRigidParticleData> DirtyRigids;
+	TMap<const FSingleParticlePhysicsProxy*, FDirtyRigidParticleReplicationErrorData> DirtyRigidErrors;
 	TArray<FDirtyGeometryCollectionData> DirtyGeometryCollections;
 	TArray<FDirtyClusterUnionData> DirtyClusterUnions;
 	TArray<FDirtyJointConstraintData> DirtyJointConstraints;
diff --git a/Engine/Source/Runtime/Experimental/Chaos/Public/PhysicsProxy/ClusterUnionPhysicsProxy.h b/Engine/Source/Runtime/Experimental/Chaos/Public/PhysicsProxy/ClusterUnionPhysicsProxy.h
--- a/Engine/Source/Runtime/Experimental/Chaos/Public/PhysicsProxy/ClusterUnionPhysicsProxy.h
+++ b/Engine/Source/Runtime/Experimental/Chaos/Public/PhysicsProxy/ClusterUnionPhysicsProxy.h
@@ -110,8 +110,8 @@ namespace Chaos
 		void SyncRemoteData(FDirtyPropertiesManager& Manager, int32 DataIdx, FDirtyChaosProperties& RemoteData) const;
 		void ClearAccumulatedData();
 
-		FProxyInterpolationData& GetInterpolationData() { return InterpolationData; }
-		const FProxyInterpolationData& GetInterpolationData() const { return InterpolationData; }
+		FProxyInterpolationBase& GetInterpolationData() { return InterpolationData; }
+		const FProxyInterpolationBase& GetInterpolationData() const { return InterpolationData; }
 
 		FClusterUnionIndex GetClusterUnionIndex() const { return ClusterUnionIndex; }
 
@@ -127,7 +127,7 @@ namespace Chaos
 		FInternalParticle* Particle_Internal = nullptr;
 		FClusterUnionIndex ClusterUnionIndex = INDEX_NONE;
 
-		FProxyInterpolationData InterpolationData;
+		FProxyInterpolationBase InterpolationData;
 
 		//~ Begin TPhysicsProxy Interface
 	public:
diff --git a/Engine/Source/Runtime/Experimental/Chaos/Public/PhysicsProxy/GeometryCollectionPhysicsProxy.h b/Engine/Source/Runtime/Experimental/Chaos/Public/PhysicsProxy/GeometryCollectionPhysicsProxy.h
--- a/Engine/Source/Runtime/Experimental/Chaos/Public/PhysicsProxy/GeometryCollectionPhysicsProxy.h
+++ b/Engine/Source/Runtime/Experimental/Chaos/Public/PhysicsProxy/GeometryCollectionPhysicsProxy.h
@@ -426,8 +426,9 @@ public:
 	void SetNotifyGlobalBreakings_External(bool bNotify);
 
 	float ComputeMaterialBasedDamageThreshold_Internal(Chaos::FPBDRigidClusteredParticleHandle& ClusteredParticle) const;
-	FProxyInterpolationData& GetInterpolationData() { return InterpolationData; }
-	const FProxyInterpolationData& GetInterpolationData() const { return InterpolationData; }
+
+	FProxyInterpolationBase& GetInterpolationData() { return InterpolationData; }
+	const FProxyInterpolationBase& GetInterpolationData() const { return InterpolationData; }
 
 	enum class EReplicationMode: uint8
 	{
@@ -610,7 +611,7 @@ private:
 	// paradigm, at least for this component of the handshake.
 	Chaos::FGuardedTripleBuffer<FGeometryCollectionResults> PhysToGameInterchange;
 
-	FProxyInterpolationData InterpolationData;
+	FProxyInterpolationBase InterpolationData;
 
 	// this is used as a unique ID when collecting data from runtime
 	FGuid CollectorGuid;
diff --git a/Engine/Source/Runtime/Experimental/Chaos/Public/PhysicsProxy/SingleParticlePhysicsProxy.h b/Engine/Source/Runtime/Experimental/Chaos/Public/PhysicsProxy/SingleParticlePhysicsProxy.h
--- a/Engine/Source/Runtime/Experimental/Chaos/Public/PhysicsProxy/SingleParticlePhysicsProxy.h
+++ b/Engine/Source/Runtime/Experimental/Chaos/Public/PhysicsProxy/SingleParticlePhysicsProxy.h
@@ -19,7 +19,6 @@
 #include "RewindData.h"
 
 
-
 namespace Chaos
 {
 
@@ -55,8 +54,6 @@ private:
 class FRigidBodyHandle_External;
 class FRigidBodyHandle_Internal;
 
-extern CHAOS_API FRealSingle ResimInterpStrength;
-
 class CHAOS_API FSingleParticlePhysicsProxy : public IPhysicsProxyBase
 {
 public:
@@ -70,8 +67,8 @@ public:
 	FSingleParticlePhysicsProxy(FSingleParticlePhysicsProxy&&) = delete;
 	virtual ~FSingleParticlePhysicsProxy();
 
-	const FProxyInterpolationData& GetInterpolationData() const { return InterpolationData; }
-	FProxyInterpolationData& GetInterpolationData() { return InterpolationData; }
+	const FProxyInterpolationBase& GetInterpolationData() const { return InterpolationData; }
+	FProxyInterpolationBase& GetInterpolationData() { return InterpolationData; }
 
 	FORCEINLINE FRigidBodyHandle_External& GetGameThreadAPI()
 	{
@@ -131,7 +128,7 @@ public:
 	void BufferPhysicsResults_External(FDirtyRigidParticleData&);
 
 	/**/
-	bool PullFromPhysicsState(const FDirtyRigidParticleData& PullData, int32 SolverSyncTimestamp, const FDirtyRigidParticleData* NextPullData = nullptr, const FRealSingle* Alpha = nullptr);
+	bool PullFromPhysicsState(const FDirtyRigidParticleData& PullData, int32 SolverSyncTimestamp, const FDirtyRigidParticleData* NextPullData = nullptr, const FRealSingle* Alpha = nullptr, const FDirtyRigidParticleReplicationErrorData* Error = nullptr, const Chaos::FReal AsyncFixedTimeStep = 0);
 
 	/**/
 	bool IsDirty();
@@ -181,7 +178,11 @@ protected:
 	FPhysicsObjectUniquePtr Reference;
 
 private:
-	FProxyInterpolationData InterpolationData;
+#if RENDERINTERP_ERRORVELOCITYSMOOTHING
+	FProxyInterpolationErrorVelocity InterpolationData;
+#else
+	FProxyInterpolationError InterpolationData;
+#endif
 
 	//use static Create
 	FSingleParticlePhysicsProxy(TUniquePtr<PARTICLE_TYPE>&& InParticle, FParticleHandle* InHandle, UObject* InOwner = nullptr);
@@ -191,7 +192,7 @@ private:
 template <bool bExternal>
 class TThreadedSingleParticlePhysicsProxyBase : protected FSingleParticlePhysicsProxy
 {
-	TThreadedSingleParticlePhysicsProxyBase() = delete;	//You should only ever new FSingleParticlePhysicsProxy, derrived types are simply there for API constraining, no new data
+	TThreadedSingleParticlePhysicsProxyBase() = delete;	//You should only ever new FSingleParticlePhysicsProxy, derived types are simply there for API constraining, no new data
 public:
 
 	FSingleParticlePhysicsProxy* GetProxy() { return static_cast<FSingleParticlePhysicsProxy*>(this); }
diff --git a/Engine/Source/Runtime/Experimental/Chaos/Public/RewindData.h b/Engine/Source/Runtime/Experimental/Chaos/Public/RewindData.h
--- a/Engine/Source/Runtime/Experimental/Chaos/Public/RewindData.h
+++ b/Engine/Source/Runtime/Experimental/Chaos/Public/RewindData.h
@@ -782,6 +782,13 @@ struct FGeometryParticleStateBase
 	void SyncSimWritablePropsFromSim(FDirtyPropData Manager,const TPBDRigidParticleHandle<FReal,3>& Rigid);
 	void SyncDirtyDynamics(FDirtyPropData& DestManager,const FDirtyChaosProperties& Dirty,const FConstDirtyPropData& SrcManager);
 	
+	template<typename TParticle>
+	void CachePreCorrectionState(const TParticle& Particle)
+	{
+		PreCorrectionXR.SetX(Particle.X());
+		PreCorrectionXR.SetR(Particle.R());
+	}
+
 	TParticlePropertyBuffer<FParticlePositionRotation,EChaosProperty::XR> ParticlePositionRotation;
 	TParticlePropertyBuffer<FParticleNonFrequentData,EChaosProperty::NonFrequentData> NonFrequentData;
 	TParticlePropertyBuffer<FParticleVelocities,EChaosProperty::Velocities> Velocities;
@@ -794,6 +801,8 @@ struct FGeometryParticleStateBase
 	TParticlePropertyBuffer<FParticleVelocities, EChaosProperty::Velocities> TargetVelocities;
 
 	FShapesArrayStateBase ShapesArrayState;
+
+	FParticlePositionRotation PreCorrectionXR;
 };
 
 class FGeometryParticleState
@@ -1079,6 +1088,12 @@ public:
 			KeyToIdx.Remove(Key);
 		}
 	}
+	
+	void Reset()
+	{
+		DenseVals.Reset();
+		KeyToIdx.Reset();
+	}
 
 	int32 Num() const { return DenseVals.Num(); }
 
@@ -1277,6 +1292,8 @@ public:
 	/** Rewind the states from all the history states datas */
 	void RewindStates(const int32 RewindFrame, const bool bResetSolver);
 
+	void BufferPhysicsResults(TMap<const FSingleParticlePhysicsProxy*, struct FDirtyRigidParticleReplicationErrorData>& DirtyRigidErrors);
+
 	/** Return the rewind data solver */
 	const FPBDRigidsSolver* GetSolver() const { return Solver; }
 
@@ -1368,6 +1385,28 @@ private:
 	using FDirtyParticleInfo = TDirtyObjectInfo<FGeometryParticleStateBase, FGeometryParticleHandle>;
 	using FDirtyJointInfo = TDirtyObjectInfo<FJointStateBase, FPBDJointConstraintHandle>;
 
+	struct FDirtyParticleErrorInfo
+	{
+	private:
+		FGeometryParticleHandle* HandlePtr;
+		FVec3 ErrorX = { 0,0,0 };
+		FQuat ErrorR = FQuat::Identity;
+
+	public:
+		FDirtyParticleErrorInfo(FGeometryParticleHandle& InHandle) : HandlePtr(&InHandle)
+		{ }
+
+		void AccumulateError(FVec3 NewErrorX, FQuat NewErrorR)
+		{
+			ErrorX += NewErrorX;
+			ErrorR *= NewErrorR;
+		}
+
+		FGeometryParticleHandle* GetObjectPtr() const { return HandlePtr; }
+		FVec3 GetErrorX() const { return ErrorX; }
+		FQuat GetErrorR() const { return ErrorR; }
+	};
+
 	template <typename TDirtyObjs, typename TObj>
 	auto& FindOrAddDirtyObjImp(TDirtyObjs& DirtyObjs, TObj& Handle, const int32 InitializedOnFrame = INDEX_NONE)
 	{
@@ -1420,6 +1459,7 @@ private:
 
 	TDirtyObjects<FDirtyParticleInfo> DirtyParticles;
 	TDirtyObjects<FDirtyJointInfo> DirtyJoints;
+	TDirtyObjects<FDirtyParticleErrorInfo> DirtyParticleErrors;
 
 	TArray<TWeakPtr<FBaseInputsHistory>> InputsHistories;
 	TArray<TWeakPtr<FBaseStatesHistory>> StatesHistories;
@@ -1437,6 +1477,9 @@ private:
 
 	template <bool bSkipDynamics, typename TDirtyInfo>
 	void DesyncIfNecessary(TDirtyInfo& Info, const FFrameAndPhase FrameAndPhase);
+
+	template<typename TObj>
+	void AccumulateErrorIfNecessary(TObj& Handle, const FFrameAndPhase FrameAndPhase) { }
 };
 
 struct FResimDebugInfo
