Iris: Fix issue with serializing UIrisObjectReferencePackageMap at index 0 and causing mismatched deserialization on the other end

#rnx

diff --git a/Engine/Source/Runtime/Experimental/Iris/Core/Private/Iris/Serialization/IrisObjectReferencePackageMap.cpp b/Engine/Source/Runtime/Experimental/Iris/Core/Private/Iris/Serialization/IrisObjectReferencePackageMap.cpp
--- a/Engine/Source/Runtime/Experimental/Iris/Core/Private/Iris/Serialization/IrisObjectReferencePackageMap.cpp
+++ b/Engine/Source/Runtime/Experimental/Iris/Core/Private/Iris/Serialization/IrisObjectReferencePackageMap.cpp
@@ -9,34 +9,41 @@ bool UIrisObjectReferencePackageMap::SerializeObject(FArchive& Ar, UClass* InCla
 		return false;
 	}
 
+	constexpr uint8 MaxNumReferences = 255U;
+
 	if (Ar.IsSaving())
 	{
-		int32 Index;
-		if (!References->Find(Obj, Index))
+		int32 Index = MaxNumReferences;
+		if (!References->Find(Obj, Index) && References->Num() < MaxNumReferences)
 		{
 			Index = References->Add(Obj);
 		}
-		constexpr int32 MaxNumReferences = 256;
-		if (ensureAlwaysMsgf(Index > 0 && Index < MaxNumReferences, TEXT("UIrisObjectReferencePackageMap::SerializeObject, failed to serialize object reference. A Maximum of 256 references are currently supported by this PackageMap")))
+
+		if (References->IsValidIndex(Index))
 		{
 			uint8 IndexByte = static_cast<uint8>(Index);
 			Ar << IndexByte;
 		}
 		else
 		{
+			ensureMsgf(false, TEXT("UIrisObjectReferencePackageMap::SerializeObject, failed to serialize object reference with Index %u (%s). A Maximum of %u references are currently supported by this PackageMap"),
+				Index, *GetNameSafe(Obj), MaxNumReferences);
+			uint8 IndexByte = MaxNumReferences;
+			Ar << IndexByte;
 			return false;
 		}
 	}
 	else
 	{
-		uint8 IndexByte = 255U;
+		uint8 IndexByte = MaxNumReferences;
 		Ar << IndexByte;
-		if (ensureAlwaysMsgf(IndexByte < References->Num(), TEXT("UIrisObjectReferencePackageMap::SerializeObject, failed to read object reference index %u is out of bounds"), IndexByte))
+		if (References->IsValidIndex(IndexByte) && IndexByte < MaxNumReferences)
 		{
 			Obj = (*References)[IndexByte];
 		}
 		else
 		{
+			ensureMsgf(false, TEXT("UIrisObjectReferencePackageMap::SerializeObject, failed to read object reference index %u is out of bounds. Current ObjectReference num: %u"), IndexByte, References->Num());
 			return false;
 		}
 	}
