Fixed static analysis warning in CoreHttp

diff --git a/Engine/Source/Runtime/Experimental/IoStoreOnDemand/Private/CoreHttp/Http.cpp b/Engine/Source/Runtime/Experimental/IoStoreOnDemand/Private/CoreHttp/Http.cpp
--- a/Engine/Source/Runtime/Experimental/IoStoreOnDemand/Private/CoreHttp/Http.cpp
+++ b/Engine/Source/Runtime/Experimental/IoStoreOnDemand/Private/CoreHttp/Http.cpp
@@ -2005,7 +2005,12 @@ static void MiscTest()
 	check(ParseMessage("HTTP/1.1 100 _Message with a \r in it" CRLF, MsgOut) == -1);
 
 	bool AllIsWell = true;
-	auto NotExpectedToBeCalled = [&] (auto, auto) { return AllIsWell = false; };
+	auto NotExpectedToBeCalled = [&AllIsWell] (auto, auto)
+	{ 
+		AllIsWell = false;
+		return false;
+	};
+
 	EnumerateHeaders("",		NotExpectedToBeCalled); check(AllIsWell);
 	EnumerateHeaders(CRLF,		NotExpectedToBeCalled); check(AllIsWell);
 	EnumerateHeaders("foo",		NotExpectedToBeCalled); check(AllIsWell);
