Chaos : optimize SpatialAccelerationCollection RemoveFrom by avoiding allocating an array

diff --git a/Engine/Source/Runtime/Experimental/Chaos/Public/Chaos/SpatialAccelerationCollection.h b/Engine/Source/Runtime/Experimental/Chaos/Public/Chaos/SpatialAccelerationCollection.h
--- a/Engine/Source/Runtime/Experimental/Chaos/Public/Chaos/SpatialAccelerationCollection.h
+++ b/Engine/Source/Runtime/Experimental/Chaos/Public/Chaos/SpatialAccelerationCollection.h
@@ -383,7 +383,9 @@ public:
 	using FirstAccelerationType = typename std::tuple_element<0, std::tuple< TSpatialAccelerationTypes...>>::type;
 	using TPayloadType = typename FirstAccelerationType::PayloadType;
 	using T = typename FirstAccelerationType::TType;
-	static constexpr int d = FirstAccelerationType::D;	
+	static constexpr int d = FirstAccelerationType::D;
+	using BucketEntryType = TSpatialAccelerationBucketEntry<TPayloadType, T, d>;
+	using BucketType = TSpatialCollectionBucket<BucketEntryType>;
 
 	TSpatialAccelerationCollection()
 	{
@@ -396,7 +398,7 @@ public:
 		Result.Bucket = BucketIdx;
 		
 		ISpatialAcceleration<TPayloadType, T, d>* AccelPtr = Substructure.Get();
-		TSpatialAccelerationBucketEntry<TPayloadType, T, d> BucketEntry;
+		BucketEntryType BucketEntry;
 		BucketEntry.Acceleration = MoveTemp(Substructure);
 		
 		const int32 TypeIdx = GetTypeIdx(AccelPtr);
@@ -430,7 +432,7 @@ public:
 	{
 		TUniquePtr <ISpatialAcceleration<TPayloadType, T, d>> Removed;
 		{
-			auto& BucketEntry = Buckets[Idx.Bucket].Objects[Idx.InnerIdx];
+			BucketEntryType& BucketEntry = Buckets[Idx.Bucket].Objects[Idx.InnerIdx];
 			const int32 TypeIdx = GetTypeIdx(BucketEntry.Acceleration.Get());
 			switch (TypeIdx)
 			{
@@ -528,10 +530,10 @@ public:
 	template <typename SQVisitor>
 	typename std::enable_if_t<std::is_same_v<SQVisitor, typename Private::FSimOverlapVisitor>, void> Overlap(const TAABB<T, 3>& QueryBounds, SQVisitor& Visitor) const
 	{		
-		const TSpatialCollectionBucket<TSpatialAccelerationBucketEntry<TPayloadType, T, d>>& Bucket = Buckets[0];
+		const BucketType& Bucket = Buckets[0];
 		{
 			int BucketInnerIdx = 0;
-			for (const TSpatialAccelerationBucketEntry<TPayloadType, T, d>& BucketEntry : Bucket.Objects)
+			for (const BucketEntryType& BucketEntry : Bucket.Objects)
 			{
 				if (BucketInnerIdx != ESpatialAccelerationCollectionBucketInnerIdx::DefaultQueryOnly && BucketInnerIdx != ESpatialAccelerationCollectionBucketInnerIdx::DynamicQueryOnly)
 				{
@@ -550,23 +552,36 @@ public:
 		return ObjList;
 	}
 
-	virtual TArray<FSpatialAccelerationIdx> GetAllSpatialIndices() const override
+	// visitor signature is : void Visitor(FSpatialAccelerationIdx SpatialIdx)
+	template <typename TVisitor>
+	void VisitAllSpatialIndices(TVisitor Visitor) const 
 	{
 		uint16 BucketIdx = 0;
-		TArray<FSpatialAccelerationIdx> Indices;
-		for (const auto& Bucket : Buckets)
+		for (const BucketType& Bucket : Buckets)
 		{
 			uint16 InnerIdx = 0;
-			for (const auto& Entry : Bucket.Objects)
+			for (const BucketEntryType& Entry : Bucket.Objects)
 			{
 				if (Entry.Acceleration)
 				{
-					Indices.Add(FSpatialAccelerationIdx{ BucketIdx, InnerIdx });
+					Visitor(FSpatialAccelerationIdx{ BucketIdx, InnerIdx });
 				}
 				++InnerIdx;
 			}
 			++BucketIdx;
 		}
+	}
+
+	virtual TArray<FSpatialAccelerationIdx> GetAllSpatialIndices() const override
+	{
+		uint16 BucketIdx = 0;
+		TArray<FSpatialAccelerationIdx> Indices;
+		VisitAllSpatialIndices(
+			[&Indices](FSpatialAccelerationIdx Idx)
+			{
+				Indices.Add(Idx);
+			}
+		);
 		return Indices;
 	}
 
@@ -574,19 +589,22 @@ public:
 	virtual bool RemoveElementFrom(const TPayloadType& Payload, FSpatialAccelerationIdx SpatialIdx) override
 	{
 		const uint16 UseBucket = ((1 << SpatialIdx.Bucket) & this->ActiveBucketsMask) ? SpatialIdx.Bucket : 0;
-		bool Success = Buckets[UseBucket].Objects[SpatialIdx.InnerIdx].Acceleration->RemoveElement(Payload);
+		bool bSuccess = Buckets[UseBucket].Objects[SpatialIdx.InnerIdx].Acceleration->RemoveElement(Payload);
 		//ensure(Success); // Debug check to see if SpatialIdx was what the caller expected it to be
 		// Make sure that we remove this Payload even if the SpatialIdx is wrong
-		for (FSpatialAccelerationIdx Idx : GetAllSpatialIndices())
-		{
-			if (!(Idx == SpatialIdx))
+		VisitAllSpatialIndices(
+			[this, &Payload, &SpatialIdx, &bSuccess](FSpatialAccelerationIdx Idx)
 			{
-				const uint16 Buckt = ((1 << Idx.Bucket) & this->ActiveBucketsMask) ? Idx.Bucket : 0;
-				const bool Removed = Buckets[Buckt].Objects[Idx.InnerIdx].Acceleration->RemoveElement(Payload);
-				Success = Success || Removed;
+				if (!(Idx == SpatialIdx))
+				{
+					const uint16 Buckt = ((1 << Idx.Bucket) & this->ActiveBucketsMask) ? Idx.Bucket : 0;
+					const bool bRemoved = Buckets[Buckt].Objects[Idx.InnerIdx].Acceleration->RemoveElement(Payload);
+					check(!bSuccess || !bRemoved); // make sure we only remove from one acceleration structure only 
+					bSuccess |= bRemoved;
+				}
 			}
-		}
-		return Success;
+		);
+		return bSuccess;
 	}
 
 	// Returns true if the element was in fact updated or if it was moved from one SpatialIdx to another during the update or false if not found at any SpatialIdx
@@ -623,9 +641,9 @@ public:
 #if !UE_BUILD_SHIPPING
 	virtual void DebugDraw(ISpacialDebugDrawInterface<T>* InInterface) const override
 	{
-		for (const auto& Bucket : Buckets)
+		for (const BucketType& Bucket : Buckets)
 		{
-			for (const auto& Entry : Bucket.Objects)
+			for (const BucketEntryType& Entry : Bucket.Objects)
 			{
 				if (Entry.Acceleration)
 				{
@@ -647,7 +665,7 @@ public:
 
 			if (Ar.IsLoading())
 			{
-				for (const TSpatialAccelerationBucketEntry<TPayloadType, T, d>& Entry : Buckets[BucketIdx].Objects)
+				for (const BucketEntryType& Entry : Buckets[BucketIdx].Objects)
 				{
 					ISpatialAcceleration<TPayloadType, T, d>* RawPtr = Entry.Acceleration.Get();
 					const int32 TypeIdx = GetTypeIdx(RawPtr);
@@ -679,7 +697,7 @@ private:
 		{
 			Buckets[BucketIdx].CopyFrom(Other.Buckets[BucketIdx]);
 			bool bFirst = true;
-			for (const TSpatialAccelerationBucketEntry<TPayloadType, T, d>& Entry : Buckets[BucketIdx].Objects)
+			for (const BucketEntryType& Entry : Buckets[BucketIdx].Objects)
 			{
 				ISpatialAcceleration<TPayloadType, T, d>* RawPtr = Entry.Acceleration.Get();
 				const int32 TypeIdx = GetTypeIdx(RawPtr);
@@ -756,7 +774,7 @@ private:
 	}
 
 	static constexpr uint16 MaxBuckets = 8;
-	TSpatialCollectionBucket<TSpatialAccelerationBucketEntry<TPayloadType, T, d>> Buckets[MaxBuckets];
+	BucketType Buckets[MaxBuckets];
 	TSpatialTypeTuple< TSpatialAccelerationTypes...> Types; // Have buckets of acceleration structure pointers
 	static constexpr uint32 NumTypes = sizeof...(TSpatialAccelerationTypes);
 };
