cluster union now builds its connectivity graph incrementally using the implicit object visit overlapping function to find particles to connect to

diff --git a/Engine/Source/Runtime/Experimental/Chaos/Private/Chaos/ClusterUnionManager.cpp b/Engine/Source/Runtime/Experimental/Chaos/Private/Chaos/ClusterUnionManager.cpp
--- a/Engine/Source/Runtime/Experimental/Chaos/Private/Chaos/ClusterUnionManager.cpp
+++ b/Engine/Source/Runtime/Experimental/Chaos/Private/Chaos/ClusterUnionManager.cpp
@@ -1,6 +1,9 @@
 // Copyright Epic Games, Inc. All Rights Reserved.
 
 #include "Chaos/ClusterUnionManager.h"
+
+#include "Chaos/CastingUtilities.h"
+#include "Chaos/GeometryQueries.h"
 #include "Chaos/ImplicitObjectUnion.h"
 #include "Chaos/PBDRigidClustering.h"
 #include "Chaos/PBDRigidClusteringAlgo.h"
@@ -10,6 +13,26 @@
 
 namespace Chaos
 {
+	namespace
+	{
+		FRigidTransform3 GetParticleRigidFrameInClusterUnion(FPBDRigidParticleHandle* Child, const FRigidTransform3& ClusterWorldTM)
+		{
+			FRigidTransform3 Frame = FRigidTransform3::Identity;
+
+			if (FPBDRigidClusteredParticleHandle* ClusterChild = Child->CastToClustered(); ClusterChild && ClusterChild->IsChildToParentLocked())
+			{
+				Frame = ClusterChild->ChildToParent();
+			}
+			else
+			{
+				const FRigidTransform3 ChildWorldTM(Child->X(), Child->R());
+				Frame = ChildWorldTM.GetRelativeTransform(ClusterWorldTM);
+			}
+
+			return Frame;
+		}
+	}
+
 	FClusterUnionManager::FClusterUnionManager(FRigidClustering& InClustering, FPBDRigidsEvolutionGBF& InEvolution)
 		: MClustering(InClustering)
 		, MEvolution(InEvolution)
@@ -100,18 +123,7 @@ namespace Chaos
 
 		for (FPBDRigidParticleHandle* Child : Union.ChildParticles)
 		{
-			FRigidTransform3 Frame = FRigidTransform3::Identity;
-			
-			if (FPBDRigidClusteredParticleHandle* ClusterChild = Child->CastToClustered(); ClusterChild && ClusterChild->IsChildToParentLocked())
-			{
-				Frame = ClusterChild->ChildToParent();
-			}
-			else
-			{
-				const FRigidTransform3 ChildWorldTM(Child->X(), Child->R());
-				Frame = ChildWorldTM.GetRelativeTransform(ClusterWorldTM);
-			}
-
+			const FRigidTransform3 Frame = GetParticleRigidFrameInClusterUnion(Child, ClusterWorldTM);
 			if (Child->Geometry())
 			{
 				Objects.Add(TUniquePtr<FImplicitObject>(CreateTransformGeometryForClusterUnion<EThreadContext::Internal>(Child, Frame)));
@@ -310,6 +322,10 @@ namespace Chaos
 		for (FPBDRigidParticleHandle* Particle : FinalParticlesToAdd)
 		{
 			ParticleToClusterUnionIndex.Add(Particle, ClusterIndex);
+			if (!bIsNewCluster)
+			{
+				Cluster->PendingConnectivityOperations.Add({Particle, EClusterUnionConnectivityOperation ::Add});
+			}
 		}
 
 		MClustering.AddParticlesToCluster(Cluster->InternalCluster, FinalParticlesToAdd, ChildToParentMap);
@@ -333,7 +349,8 @@ namespace Chaos
 			// The anchored flag is taken care of in UpdateKinematicProperties so it must be set before that is called.
 			Cluster->InternalCluster->SetIsAnchored(true);
 		}
-		UpdateAllClusterUnionProperties(*Cluster, bIsNewCluster ? EUpdateClusterUnionPropertiesFlags::All : EUpdateClusterUnionPropertiesFlags::ForceGenerateConnectionGraph);
+
+		UpdateAllClusterUnionProperties(*Cluster, bIsNewCluster ? EUpdateClusterUnionPropertiesFlags::All : EUpdateClusterUnionPropertiesFlags::IncrementalGenerateConnectionGraph);
 
 		if (OldProxy)
 		{
@@ -360,20 +377,12 @@ namespace Chaos
 		MEvolution.GetParticles().MarkTransientDirtyParticle(Cluster->InternalCluster);
 	}
 
-	DECLARE_CYCLE_STAT(TEXT("FClusterUnionManager::PostRemovalClusterUnionUpdate"), STAT_PostRemovalClusterUnionUpdate, STATGROUP_Chaos);
-	void FClusterUnionManager::PostRemovalClusterUnionUpdate(FClusterUnion& Union)
+	DECLARE_CYCLE_STAT(TEXT("FClusterUnionManager::DeferredClusterUnionUpdate"), STAT_DeferredClusterUnionUpdate, STATGROUP_Chaos);
+	void FClusterUnionManager::DeferredClusterUnionUpdate(FClusterUnion& Union)
 	{
-		SCOPE_CYCLE_COUNTER(STAT_PostRemovalClusterUnionUpdate);
-
-		if(TArray<FPBDRigidParticleHandle*>* Children = MClustering.GetChildrenMap().Find(Union.InternalCluster))
-		{
-			for(FPBDRigidParticleHandle* Child : *Children)
-			{
-				MClustering.RemoveNodeConnections(Child);
-			}
-		}
+		SCOPE_CYCLE_COUNTER(STAT_DeferredClusterUnionUpdate);
 
-		UpdateAllClusterUnionProperties(Union, EUpdateClusterUnionPropertiesFlags::None);
+		UpdateAllClusterUnionProperties(Union, EUpdateClusterUnionPropertiesFlags::IncrementalGenerateConnectionGraph);
 
 		if (!Union.ChildParticles.IsEmpty())
 		{
@@ -401,7 +410,7 @@ namespace Chaos
 		{
 			if (FClusterUnion* Union = FindClusterUnion(Index))
 			{
-				PostRemovalClusterUnionUpdate(*Union);
+				DeferredClusterUnionUpdate(*Union);
 			}
 		}
 
@@ -447,6 +456,8 @@ namespace Chaos
 				{
 					Proxy->SetParentProxy(nullptr);
 				}
+
+				Cluster->PendingConnectivityOperations.Add({ Handle, EClusterUnionConnectivityOperation::Remove });
 			}
 		}
 
@@ -467,7 +478,7 @@ namespace Chaos
 		switch (UpdateClusterPropertiesTiming)
 		{
 		case EClusterUnionOperationTiming::Immediate:
-			PostRemovalClusterUnionUpdate(*Cluster);
+			DeferredClusterUnionUpdate(*Cluster);
 			break;
 		case EClusterUnionOperationTiming::Defer:
 			DeferredClusterUnionsForUpdateProperties.Add(ClusterIndex);
@@ -556,7 +567,105 @@ namespace Chaos
 		{
 			MClustering.ClearConnectionGraph(ClusterUnion.InternalCluster);
 			MClustering.GenerateConnectionGraph(ClusterUnion.InternalCluster, ClusterUnion.Parameters);
+			ClusterUnion.PendingConnectivityOperations.Empty();
+		}
+		else if (EnumHasAnyFlags(Flags, EUpdateClusterUnionPropertiesFlags::IncrementalGenerateConnectionGraph))
+		{
+			FlushIncrementalConnectivityGraphOperations(ClusterUnion);
+		}
+	}
+
+	DECLARE_CYCLE_STAT(TEXT("FClusterUnionManager::FlushIncrementalConnectivityGraphOperations"), STAT_FlushIncrementalConnectivityGraphOperations, STATGROUP_Chaos);
+	void FClusterUnionManager::FlushIncrementalConnectivityGraphOperations(FClusterUnion& ClusterUnion)
+	{
+		constexpr FReal kAddThickness = 5.0;
+		SCOPE_CYCLE_COUNTER(STAT_FlushIncrementalConnectivityGraphOperations);
+		const FRigidTransform3 ClusterWorldTM(ClusterUnion.InternalCluster->X(), ClusterUnion.InternalCluster->R());
+		for (const TPair<FPBDRigidParticleHandle*, EClusterUnionConnectivityOperation>& Op : ClusterUnion.PendingConnectivityOperations)
+		{
+			if (Op.Value == EClusterUnionConnectivityOperation::Add)
+			{
+				// Use the acceleration structure of the cluster union itself to make finding overlaps easy.
+				const FImplicitObjectUnion& ShapeUnion = ClusterUnion.SharedGeometry->GetObjectChecked<FImplicitObjectUnion>();
+				const FRigidTransform3 FromTransform = GetParticleRigidFrameInClusterUnion(Op.Key, ClusterWorldTM);
+				FAABB3 ParticleLocalBounds = Op.Key->SharedGeometry()->BoundingBox().TransformedAABB(FromTransform);
+				ParticleLocalBounds.Thicken(kAddThickness);
+
+				ShapeUnion.VisitOverlappingLeafObjects(
+					ParticleLocalBounds,
+					[this, &ClusterUnion, &FromTransform, &Op](const FImplicitObject* ToGeom, const FRigidTransform3& ToTransform, const int32 RootObjectIndex, const int32 ObjectIndex, const int32 LeafObjectIndex)
+					{
+						check(ToGeom != nullptr);
+						// RootObjectIndex is the index in the cluster union.
+						// NOTE: This will need to be re-thought if we ever decide to make the mapping between child shapes and child particles not 1-to-1 (e.g. if we ever attempt to simplify the cluster union shape).
+
+						// Ignore intersections against the same particle since we just added the particle (potentially) into the geometry.
+						if (Op.Key == ClusterUnion.ChildParticles[RootObjectIndex])
+						{
+							return;
+						}
+
+						FAABB3 ToAABB = ToGeom->CalculateTransformedBounds(ToTransform);
+						ToAABB.Thicken(kAddThickness);
+
+						const bool bOverlap = Utilities::CastHelper(*ToGeom, ToTransform, 
+							[&Op, &FromTransform, &ToAABB](const auto& ToGeomDowncast, const FRigidTransform3& FinalToTransform)
+							{
+								const FShapesArray& AllFromShapes = Op.Key->ShapesArray();
+								for (int32 FromIndex = 0; FromIndex < AllFromShapes.Num(); ++FromIndex)
+								{
+									if (!AllFromShapes[FromIndex])
+									{
+										continue;
+									}
+
+									const FPerShapeData& FromShape = *AllFromShapes[FromIndex];
+									const FImplicitObject* FromGeom = FromShape.GetGeometry().Get();
+
+									if (!FromGeom)
+									{
+										continue;
+									}
+
+									FAABB3 FromAABB = FromGeom->CalculateTransformedBounds(FromTransform);
+									FromAABB.Thicken(kAddThickness);
+
+									// First sanity check to see if the two shape AABB's intersect.
+									if (!FromAABB.Intersects(ToAABB))
+									{
+										continue;
+									}
+
+									// Now do a more accurate overlap check between the two shapes.
+									// Note that passing MTD here is critical as it gets us a more accurate check for some reason...
+									FMTDInfo MTDInfo;
+									if (OverlapQuery(*FromGeom, FromTransform, ToGeomDowncast, FinalToTransform, kAddThickness, &MTDInfo))
+									{
+										return true;
+									}
+									else
+									{
+										continue;
+									}
+								}
+
+								return false;
+							});
+
+						if (bOverlap)
+						{
+							MClustering.CreateNodeConnection(Op.Key, ClusterUnion.ChildParticles[RootObjectIndex]);
+						}
+					}
+				);
+			}
+			else if (Op.Value == EClusterUnionConnectivityOperation::Remove)
+			{
+				MClustering.RemoveNodeConnections(Op.Key);
+			}
 		}
+
+		ClusterUnion.PendingConnectivityOperations.Empty();
 	}
 
 	DECLARE_CYCLE_STAT(TEXT("FClusterUnionManager::GetOrCreateClusterUnionIndexFromExplicitIndex"), STAT_GetOrCreateClusterUnionIndexFromExplicitIndex, STATGROUP_Chaos);
@@ -675,12 +784,17 @@ namespace Chaos
 						}
 						ChildHandle->SetChildToParentLocked(true);
 						PendingChildToParentUpdates.Remove(ChildHandle);
+
+						// A child to parent update needs to remove *and* add to the connectivity graph (in that order) since
+						// the child to parent update might move the node so far away as to make the old connectivity edges incorrect.
+						ClusterUnion->PendingConnectivityOperations.Add({ Particle, EClusterUnionConnectivityOperation::Remove });
+						ClusterUnion->PendingConnectivityOperations.Add({ Particle, EClusterUnionConnectivityOperation::Add });
 					}
 				}
 			}
 		}
 
-		UpdateAllClusterUnionProperties(*ClusterUnion, EUpdateClusterUnionPropertiesFlags::ForceGenerateConnectionGraph);
+		UpdateAllClusterUnionProperties(*ClusterUnion, EUpdateClusterUnionPropertiesFlags::IncrementalGenerateConnectionGraph);
 
 		for (TPair<FPBDRigidClusteredParticleHandle*, bool>& Pair : DirtyParticleLockStates)
 		{
diff --git a/Engine/Source/Runtime/Experimental/Chaos/Public/Chaos/ClusterUnionManager.h b/Engine/Source/Runtime/Experimental/Chaos/Public/Chaos/ClusterUnionManager.h
--- a/Engine/Source/Runtime/Experimental/Chaos/Public/Chaos/ClusterUnionManager.h
+++ b/Engine/Source/Runtime/Experimental/Chaos/Public/Chaos/ClusterUnionManager.h
@@ -38,10 +38,17 @@ namespace Chaos
 		None = 0,
 		RecomputeMassOrientation = 1 << 0,
 		ForceGenerateConnectionGraph = 1 << 1,
+		IncrementalGenerateConnectionGraph = 1 << 2,
 		All = RecomputeMassOrientation | ForceGenerateConnectionGraph
 	};
 	ENUM_CLASS_FLAGS(EUpdateClusterUnionPropertiesFlags);
 
+	enum class EClusterUnionConnectivityOperation : int8
+	{
+		Add,
+		Remove
+	};
+
 	struct CHAOS_API FClusterUnionCreationParameters
 	{
 		FClusterUnionExplicitIndex ExplicitIndex = INDEX_NONE;
@@ -75,6 +82,9 @@ namespace Chaos
 
 		// Whether or not the position/rotation needs to be computed the first time a particle is added.
 		bool bNeedsXRInitialization = true;
+
+		// Pending particles that need to be added into the connectivity graph.
+		TArray<TPair<FPBDRigidParticleHandle*, EClusterUnionConnectivityOperation>> PendingConnectivityOperations;
 	};
 
 	struct FClusterUnionChildToParentUpdate
@@ -206,8 +216,11 @@ namespace Chaos
 		// Forcefully recreate the shared geometry on a cluster. Potentially expensive so ideally should be used rarely.
 		TSharedPtr<FImplicitObject, ESPMode::ThreadSafe> ForceRecreateClusterUnionSharedGeometry(const FClusterUnion& Union);
 
-		// Handles updating the cluster union after it has particles removed from it. Could be called from either immediate or deferred removal.
-		void PostRemovalClusterUnionUpdate(FClusterUnion& Union);
+		// Handles updating the cluster union.
+		void DeferredClusterUnionUpdate(FClusterUnion& Union);
+
+		// Flush the cluster union's incremental connectivity operations
+		void FlushIncrementalConnectivityGraphOperations(FClusterUnion& ClusterUnion);
 	};
 
 }
\ No newline at end of file
