Geometry collection : Expose Damage model property to blueprints

diff --git a/Engine/Source/Runtime/Experimental/GeometryCollectionEngine/Public/GeometryCollection/GeometryCollectionObject.h b/Engine/Source/Runtime/Experimental/GeometryCollectionEngine/Public/GeometryCollection/GeometryCollectionObject.h
--- a/Engine/Source/Runtime/Experimental/GeometryCollectionEngine/Public/GeometryCollection/GeometryCollectionObject.h
+++ b/Engine/Source/Runtime/Experimental/GeometryCollectionEngine/Public/GeometryCollection/GeometryCollectionObject.h
@@ -512,7 +512,7 @@ public:
 	int32 MaxClusterLevel;
 
 	/** Damage model to use for evaluating destruction. */
-	UPROPERTY(EditAnywhere, Category = "Damage")
+	UPROPERTY(EditAnywhere, BlueprintReadWrite, Category = "Damage")
 	EDamageModelTypeEnum DamageModel;
 
 	/** Damage threshold for clusters at different levels. */
