fix compile error on certain platforms due to kAddThickness not being captured in lambda

diff --git a/Engine/Source/Runtime/Experimental/Chaos/Private/Chaos/ClusterUnionManager.cpp b/Engine/Source/Runtime/Experimental/Chaos/Private/Chaos/ClusterUnionManager.cpp
--- a/Engine/Source/Runtime/Experimental/Chaos/Private/Chaos/ClusterUnionManager.cpp
+++ b/Engine/Source/Runtime/Experimental/Chaos/Private/Chaos/ClusterUnionManager.cpp
@@ -593,7 +593,7 @@ namespace Chaos
 
 				ShapeUnion.VisitOverlappingLeafObjects(
 					ParticleLocalBounds,
-					[this, &ClusterUnion, &FromTransform, &Op](const FImplicitObject* ToGeom, const FRigidTransform3& ToTransform, const int32 RootObjectIndex, const int32 ObjectIndex, const int32 LeafObjectIndex)
+					[this, &ClusterUnion, &FromTransform, &Op, kAddThickness](const FImplicitObject* ToGeom, const FRigidTransform3& ToTransform, const int32 RootObjectIndex, const int32 ObjectIndex, const int32 LeafObjectIndex)
 					{
 						check(ToGeom != nullptr);
 						// RootObjectIndex is the index in the cluster union.
@@ -609,7 +609,7 @@ namespace Chaos
 						ToAABB.Thicken(kAddThickness);
 
 						const bool bOverlap = Utilities::CastHelper(*ToGeom, ToTransform, 
-							[&Op, &FromTransform, &ToAABB](const auto& ToGeomDowncast, const FRigidTransform3& FinalToTransform)
+							[&Op, &FromTransform, &ToAABB, kAddThickness](const auto& ToGeomDowncast, const FRigidTransform3& FinalToTransform)
 							{
 								const FShapesArray& AllFromShapes = Op.Key->ShapesArray();
 								for (int32 FromIndex = 0; FromIndex < AllFromShapes.Num(); ++FromIndex)
