Resim HeadlessChaos test fix1

diff --git a/Engine/Source/Runtime/Experimental/Chaos/Private/PhysicsProxy/SingleParticlePhysicsProxy.cpp b/Engine/Source/Runtime/Experimental/Chaos/Private/PhysicsProxy/SingleParticlePhysicsProxy.cpp
--- a/Engine/Source/Runtime/Experimental/Chaos/Private/PhysicsProxy/SingleParticlePhysicsProxy.cpp
+++ b/Engine/Source/Runtime/Experimental/Chaos/Private/PhysicsProxy/SingleParticlePhysicsProxy.cpp
@@ -289,12 +289,12 @@ bool FSingleParticlePhysicsProxy::PullFromPhysicsState(const Chaos::FDirtyRigidP
 		const FSingleParticleProxyTimestamp* ProxyTimestamp = PullData.GetTimestamp();
 		
 #if RENDERINTERP_ERRORVELOCITYSMOOTHING
-		const int32 RenderInterpErrorVelocitySmoothingDurationTicks = FMath::CeilToInt32(RenderInterpErrorVelocitySmoothingDuration / AsyncFixedTimeStep); // Convert duration from seconds to simulation ticks
+		const int32 RenderInterpErrorVelocitySmoothingDurationTicks = FMath::FloorToInt32(RenderInterpErrorVelocitySmoothingDuration / AsyncFixedTimeStep); // Convert duration from seconds to simulation ticks
 #endif
 
 		if (Error)
 		{
-			const int32 RenderInterpErrorCorrectionDurationTicks = FMath::CeilToInt32(RenderInterpErrorCorrectionDuration / AsyncFixedTimeStep); // Convert duration from seconds to simulation ticks
+			const int32 RenderInterpErrorCorrectionDurationTicks = FMath::FloorToInt32(RenderInterpErrorCorrectionDuration / AsyncFixedTimeStep); // Convert duration from seconds to simulation ticks
 			InterpolationData.AccumlateErrorXR(Error->ErrorX, Error->ErrorR, SolverSyncTimestamp, RenderInterpErrorCorrectionDurationTicks);
 #if RENDERINTERP_ERRORVELOCITYSMOOTHING
 			InterpolationData.SetVelocitySmoothing(Rigid->V(), Rigid->X(), RenderInterpErrorVelocitySmoothingDurationTicks);
