Exposed GetAccumulatedTime() for obtaining the AccumulateTime in fixed step mode

diff --git a/Engine/Source/Runtime/Experimental/Chaos/Public/Chaos/Framework/PhysicsSolverBase.h b/Engine/Source/Runtime/Experimental/Chaos/Public/Chaos/Framework/PhysicsSolverBase.h
--- a/Engine/Source/Runtime/Experimental/Chaos/Public/Chaos/Framework/PhysicsSolverBase.h
+++ b/Engine/Source/Runtime/Experimental/Chaos/Public/Chaos/Framework/PhysicsSolverBase.h
@@ -615,6 +615,11 @@ namespace Chaos
 			return bSolverSubstep_External;
 		}
 
+		FReal GetAccumulatedTime() const
+		{
+			return AccumulatedTime;
+		}
+
 		/** Returns the time used by physics results. If fixed dt is used this will be the interpolated time */
 		FReal GetPhysicsResultsTime_External() const
 		{
