Optimize data marshalling of geometry collectionsbetween PT and GT
- Remove global transforms from the payloads
- Only send data for the transform indices that have been changes ( reducing copy and allocation )
- Do not send any more data from BufferPhysicsResults_External ( removing it from showing in profiling )
- move UpdateRemovalIfNeeded to only be called if the root is broken

diff --git a/Engine/Source/Runtime/Experimental/Chaos/Private/PhysicsProxy/GeometryCollectionPhysicsProxy.cpp b/Engine/Source/Runtime/Experimental/Chaos/Private/PhysicsProxy/GeometryCollectionPhysicsProxy.cpp
--- a/Engine/Source/Runtime/Experimental/Chaos/Private/PhysicsProxy/GeometryCollectionPhysicsProxy.cpp
+++ b/Engine/Source/Runtime/Experimental/Chaos/Private/PhysicsProxy/GeometryCollectionPhysicsProxy.cpp
@@ -117,12 +117,13 @@ FGeometryCollectionResults::FGeometryCollectionResults()
 void FGeometryCollectionResults::Reset()
 {
 	SolverDt = 0.0f;
-	States.SetNum(0);
-	GlobalTransforms.SetNum(0);
-	ParticleXs.SetNum(0);
-	ParticleRs.SetNum(0);
-	ParticleVs.SetNum(0);
-	ParticleWs.SetNum(0);
+
+	ModifiedTransformIndices.Reset();
+	States.Reset();
+	Positions.Reset();
+	Velocities.Reset();
+	Transforms.Reset();
+
 	IsObjectDynamic = false;
 	IsObjectLoading = false;
 }
@@ -2885,14 +2886,11 @@ void FGeometryCollectionPhysicsProxy::PrepareBufferData(Chaos::FDirtyGeometryCol
 	BufferData.SetProxy(*this);
 
 	IsObjectDynamic = false;
-	FGeometryCollectionResults& TargetResults = BufferData.Results;
-	TargetResults.SolverDt = SolverLastDt;
+	FGeometryCollectionResults& Results = BufferData.Results;
+	Results.SetSolverDt(SolverLastDt);
 
 	const int32 NumTransformGroupElements = ThreadCollection.NumElements(FGeometryCollection::TransformGroup);
-	if (TargetResults.NumTransformGroup() != NumTransformGroupElements)
-	{
-		TargetResults.InitArrays(ThreadCollection);
-	}
+	Results.InitArrays(ThreadCollection);
 }
 
 void FGeometryCollectionPhysicsProxy::BufferPhysicsResults_External(Chaos::FDirtyGeometryCollectionData& BufferData)
@@ -2902,59 +2900,11 @@ void FGeometryCollectionPhysicsProxy::BufferPhysicsResults_External(Chaos::FDirt
 	* Called per-tick when async is on after the simulation has completed.
 	* goal is collect current game thread data of the proxy so it can be used if no previous physics thread data is available for interpolating  
 	*/
-	//SCOPE_CYCLE_COUNTER(STAT_GeometryCollection_BufferPhysicsResults_External);
-	if (IsObjectDeleting) return;
-	PrepareBufferData(BufferData, GameThreadCollection);
-	
-	FGeometryCollectionResults& TargetResults = BufferData.Results;
-	//TargetResults.SolverDt = CurrentSolver->GetLastDt();	//todo: should this use timestamp for async mode?
-
-	// not interpolatable, so we don't fill it 
-	//	TargetResults.DisabledStates;
-	//	TargetResults.Parent;
-	//
-	//TargetResults.DynamicState; // ObjectState
-	const TManagedArray<FVector3f>* LinearVelocities = GameThreadCollection.FindAttributeTyped<FVector3f>("LinearVelocity", FTransformCollection::TransformGroup);
-	const TManagedArray<FVector3f>* AngularVelocities = GameThreadCollection.FindAttributeTyped<FVector3f>("AngularVelocity", FTransformCollection::TransformGroup);
-	
-	for (int32 TransformGroupIndex = 0; TransformGroupIndex < TargetResults.Transforms.Num(); ++TransformGroupIndex)
-	{
-		if (GTParticles[TransformGroupIndex] == nullptr)
-		{
-			continue;
-		}
-
-		FParticle& GTParticle = *GTParticles[TransformGroupIndex];
-		TargetResults.ParticleXs[TransformGroupIndex] = GTParticle.X();
-		GC_PHYSICSPROXY_CHECK_FOR_NAN(TargetResults.ParticleXs[TransformGroupIndex]);
-
-		TargetResults.ParticleRs[TransformGroupIndex] = GTParticle.R();
-		TargetResults.States[TransformGroupIndex].DynamicState = static_cast<uint8>(GetObjectStateEnumFromObjectState(GTParticle.ObjectState()));
-		TargetResults.States[TransformGroupIndex].DisabledState = GTParticle.Disabled();
-		TargetResults.States[TransformGroupIndex].DynamicInternalClusterParent = false;
-		TargetResults.States[TransformGroupIndex].HasInternalClusterParent = false;
-
-		if (LinearVelocities && AngularVelocities)
-		{
-			TargetResults.ParticleVs[TransformGroupIndex] = Chaos::FVec3((*LinearVelocities)[TransformGroupIndex]);
-			TargetResults.ParticleWs[TransformGroupIndex] = Chaos::FVec3((*AngularVelocities)[TransformGroupIndex]);
-		}
-
-		TargetResults.Parent[TransformGroupIndex] = GameThreadCollection.Parent[TransformGroupIndex];
-		TargetResults.Transforms[TransformGroupIndex] = GameThreadCollection.Transform[TransformGroupIndex];
-
-		TargetResults.InternalClusterUniqueIdx[TransformGroupIndex] = INDEX_NONE;
-		
-#if WITH_EDITORONLY_DATA
-		TargetResults.DamageInfo[TransformGroupIndex].Damage = 0.0f;
-		TargetResults.DamageInfo[TransformGroupIndex].DamageThreshold = 0.0f;
-#endif
-	}
 	
-	// TargetResults.GlobalTransforms; // no need, can be recomputed ?  
-	//
-	// TargetResults.IsObjectDynamic;
-	// TargetResults.IsObjectLoading;
+	// for geometry collection the physics thread is the authority
+	// we leave the results empty, meaning that nothing has changed 
+	// PullFromPhysics state can then use the GT state to interpolate from if necessary
+	PrepareBufferData(BufferData, PhysicsThreadCollection);
 }
 
 static void UpdateParticleHandleTransform(Chaos::FPBDRigidsSolver& CurrentSolver, Chaos::FPBDRigidClusteredParticleHandle& Handle, const Chaos::FRigidTransform3& NewTransform)
@@ -2979,7 +2929,19 @@ static void UpdateParticleHandleTransformIfNeeded(Chaos::FPBDRigidsSolver& Curre
 	}
 }
 
-
+// remove a transform from its parent in the physics thread collection
+// return true if the removal was successfull ( false if the child has already no parent )
+inline static bool RemoveFromParentInCollection_Internal(FGeometryDynamicCollection& PhysicsThreadCollection, int32 TransformGroupIndex)
+{
+	// CONTEXT : PHYSICS THREAD
+	if (PhysicsThreadCollection.Parent[TransformGroupIndex] != INDEX_NONE)
+	{
+		PhysicsThreadCollection.Children[PhysicsThreadCollection.Parent[TransformGroupIndex]].Remove(TransformGroupIndex);
+		PhysicsThreadCollection.Parent[TransformGroupIndex] = INDEX_NONE;
+		return true;
+	}
+	return false;
+}
 
 void FGeometryCollectionPhysicsProxy::BufferPhysicsResults_Internal(Chaos::FPBDRigidsSolver* CurrentSolver, Chaos::FDirtyGeometryCollectionData& BufferData)
 {
@@ -2996,234 +2958,176 @@ void FGeometryCollectionPhysicsProxy::BufferPhysicsResults_Internal(Chaos::FPBDR
 	PrepareBufferData(BufferData, PhysicsThreadCollection, CurrentSolver->GetLastDt());
 
 	IsObjectDynamic = false;
-	FGeometryCollectionResults& TargetResults = BufferData.Results;
+	FGeometryCollectionResults& Results = BufferData.Results;
 
 	const FTransform& ActorToWorld = Parameters.WorldTransform;
-	const TManagedArray<int32>& Parent = PhysicsThreadCollection.Parent;
-	const TManagedArray<TSet<int32>>& Children = PhysicsThreadCollection.Children;
 	const bool IsActorScaled = !ActorToWorld.GetScale3D().Equals(FVector::OneVector);
 	const FTransform ActorScaleTransform(FQuat::Identity,  FVector::ZeroVector, ActorToWorld.GetScale3D());
 
 	UniqueIdxToInternalClusterHandle.Reset();
 	
-	const int32 NumTransformGroupElements = TargetResults.Transforms.Num(); 
-	if(NumTransformGroupElements > 0)
+	const int32 NumTransforms = PhysicsThreadCollection.Transform.Num();
+	if(NumTransforms > 0)
 	{ 
 		SCOPE_CYCLE_COUNTER(STAT_CalcParticleToWorld);
 
-		for (int32 TransformGroupIndex = 0; TransformGroupIndex < NumTransformGroupElements; ++TransformGroupIndex)
+		for (int32 TransformGroupIndex = 0; TransformGroupIndex < NumTransforms; ++TransformGroupIndex)
 		{
-			const bool bWasActive = PhysicsThreadCollection.Active[TransformGroupIndex];
-
-			TargetResults.Transforms[TransformGroupIndex] = PhysicsThreadCollection.Transform[TransformGroupIndex];
-			TargetResults.Parent[TransformGroupIndex] = PhysicsThreadCollection.Parent[TransformGroupIndex];
-			TargetResults.InternalClusterUniqueIdx[TransformGroupIndex] = INDEX_NONE;
-
-			TargetResults.States[TransformGroupIndex].DisabledState = true;
-			TargetResults.States[TransformGroupIndex].HasInternalClusterParent = false;
-			TargetResults.States[TransformGroupIndex].DynamicInternalClusterParent = false;
 			Chaos::FPBDRigidClusteredParticleHandle* Handle = SolverParticleHandles[TransformGroupIndex];
 			if (!Handle)
 			{
-				PhysicsThreadCollection.Active[TransformGroupIndex] = !TargetResults.States[TransformGroupIndex].DisabledState;
+				PhysicsThreadCollection.Active[TransformGroupIndex] = false;
 				continue;
 			}
 
-			// Dynamic state is also updated by the solver during field interaction.
-			TargetResults.States[TransformGroupIndex].DynamicState = static_cast<int8>(GetObjectStateFromHandle(Handle));
+			const bool bWasActive = PhysicsThreadCollection.Active[TransformGroupIndex];
+			const bool bIsActive = !Handle->Disabled();
+
+			FGeometryCollectionResults::FPositionData PositionData;
+			PositionData.ParticleX = Handle->X();
+			PositionData.ParticleR = Handle->R();
+
+			FGeometryCollectionResults::FVelocityData VelocityData;
+			VelocityData.ParticleV = Handle->V();
+			VelocityData.ParticleW = Handle->W();
+
+			FGeometryCollectionResults::FStateData StateData;
+			StateData.TransformIndex = TransformGroupIndex;
+			StateData.ParentTransformIndex = PhysicsThreadCollection.Parent[TransformGroupIndex];
+			StateData.InternalClusterUniqueIdx = INDEX_NONE;
+			StateData.State.DisabledState = Handle->Disabled();
+			StateData.State.HasInternalClusterParent = false;
+			StateData.State.DynamicInternalClusterParent = false;
+			StateData.State.DynamicState = static_cast<int8>(GetObjectStateFromHandle(Handle));
+
+			bool bHasChanged = false;
 
 			// Update the transform and parent hierarchy of the active rigid bodies. Active bodies can be either
 			// rigid geometry defined from the leaf nodes of the collection, or cluster bodies that drive an entire
 			// branch of the hierarchy within the GeometryCollection.
 			// - Active bodies are directly driven from the global position of the corresponding
-			//   rigid bodies within the solver ( cases where RigidBodyID[TransformGroupIndex] is not disabled ). 
+			//   rigid bodies within the solver 
 			// - Deactivated bodies are driven from the transforms of their active parents. However the solver can
 			//   take ownership of the parents during the simulation, so it might be necessary to force deactivated
 			//   bodies out of the collections hierarchy during the simulation.  
-			if (!Handle->Disabled())
+
+			if (bIsActive || (bWasActive != bIsActive))
 			{
-				// Update the transform of the active body. The active body can be either a single rigid
-				// or a collection of rigidly attached geometries (Clustering). The cluster is represented as a
-				// single transform in the GeometryCollection, and all children are stored in the local space
-				// of the parent cluster.
-
-				TargetResults.ParticleXs[TransformGroupIndex] = Handle->X();
-				GC_PHYSICSPROXY_CHECK_FOR_NAN(TargetResults.ParticleXs[TransformGroupIndex]);
-
-				TargetResults.ParticleRs[TransformGroupIndex] = Handle->R();
-				TargetResults.ParticleVs[TransformGroupIndex] = Handle->V();
-				TargetResults.ParticleWs[TransformGroupIndex] = Handle->W();
-				FRigidTransform3 ParticleToWorld(Handle->X(), Handle->R());
-				const FTransform MassToLocal = PhysicsThreadCollection.MassToLocal[TransformGroupIndex];
-
-				TargetResults.Transforms[TransformGroupIndex] = MassToLocal.GetRelativeTransformReverse(ParticleToWorld).GetRelativeTransform(ActorToWorld);
-				TargetResults.Transforms[TransformGroupIndex].NormalizeRotation();
-				if (IsActorScaled)
-				{
-					TargetResults.Transforms[TransformGroupIndex] = MassToLocal.Inverse() * ActorScaleTransform * MassToLocal * TargetResults.Transforms[TransformGroupIndex];
-				}
+				bHasChanged = true;
+			}
+					
+			Chaos::FPBDRigidClusteredParticleHandle* ClusterParent = Handle->Parent();
 
-				PhysicsThreadCollection.Transform[TransformGroupIndex] = TargetResults.Transforms[TransformGroupIndex];
+			// Has parent changed? ( can be a new one or nullptr)
+			const bool bHasNewParent = SolverClusterID[TransformGroupIndex] != ClusterParent;
+			if (bHasNewParent)
+			{
+				// Force all driven rigid bodies out of the transform hierarchy ( because the new parent can only be an internal cluster )
+				if (RemoveFromParentInCollection_Internal(PhysicsThreadCollection, TransformGroupIndex))
+				{
+					// If the parent of this NON DISABLED body is set to anything other than INDEX_NONE,
+					// then it was just unparented, likely either by rigid clustering or by fields.  We
+					// need to force all such enabled rigid bodies out of the transform hierarchy.
+					StateData.ParentTransformIndex = INDEX_NONE;
 
-				// Indicate that this object needs to be updated and the proxy is active.
-				TargetResults.States[TransformGroupIndex].DisabledState = false;
-				TargetResults.States[TransformGroupIndex].HasInternalClusterParent = false;
-				TargetResults.States[TransformGroupIndex].DynamicInternalClusterParent = false;
-				IsObjectDynamic = true;
+					// Indicate that this object needs to be updated and the proxy is active.
+					StateData.State.DisabledState = false;
 
-				// If the parent of this NON DISABLED body is set to anything other than INDEX_NONE,
-				// then it was just unparented, likely either by rigid clustering or by fields.  We
-				// need to force all such enabled rigid bodies out of the transform hierarchy.
-				TargetResults.Parent[TransformGroupIndex] = INDEX_NONE;
-				if (PhysicsThreadCollection.Parent[TransformGroupIndex] != INDEX_NONE)
-				{
-					//GeometryCollectionAlgo::UnparentTransform(&PhysicsThreadCollection,TransformGroupIndex);
-					PhysicsThreadCollection.Children[PhysicsThreadCollection.Parent[TransformGroupIndex]].Remove(TransformGroupIndex);
-					PhysicsThreadCollection.Parent[TransformGroupIndex] = INDEX_NONE;
+					bHasChanged = true;
 				}
-
-				// When a leaf node rigid body is removed from a cluster, the rigid
-				// body will become active and needs its clusterID updated.  This just
-				// syncs the clusterID all the time.
-				TPBDRigidParticleHandle<Chaos::FReal, 3>* ClusterParentId = Handle->ClusterIds().Id;
-				SolverClusterID[TransformGroupIndex] = ClusterParentId;
+				// set new parent ( likely an internal cluster )
+				SolverClusterID[TransformGroupIndex] = ClusterParent;
 			}
-			else    // Handle->Disabled()
+
+			// do we have an internal cluster parent ( it also implies we are a disabled particle )
+			if (ClusterParent && ClusterParent->InternalCluster())
 			{
-				// if the particle was active and now is not we still need to properly set the positions, velocities and transform values
-				// this is because interpolation may need this data when processed on the gamethread ( see PullFromPhysicsState method )
-				if (bWasActive)
-				{ 
-					TargetResults.ParticleXs[TransformGroupIndex] = Handle->X();
-					GC_PHYSICSPROXY_CHECK_FOR_NAN(TargetResults.ParticleXs[TransformGroupIndex]);
-
-					TargetResults.ParticleRs[TransformGroupIndex] = Handle->R();
-					TargetResults.ParticleVs[TransformGroupIndex] = Handle->V();
-					TargetResults.ParticleWs[TransformGroupIndex] = Handle->W();
-					FRigidTransform3 ParticleToWorld(Handle->X(), Handle->R());
-					const FTransform MassToLocal = PhysicsThreadCollection.MassToLocal[TransformGroupIndex];
+				const bool bIsOwnedInternalCluster = ClusterParent->PhysicsProxy() == this;
 
-					TargetResults.Transforms[TransformGroupIndex] = MassToLocal.GetRelativeTransformReverse(ParticleToWorld).GetRelativeTransform(ActorToWorld);
-					TargetResults.Transforms[TransformGroupIndex].NormalizeRotation();
-					if (IsActorScaled)
+				// We don't want the to keep track of any external internal clusters (i.e. cluster unions).
+				if (bIsOwnedInternalCluster)
+				{
+					StateData.InternalClusterUniqueIdx = ClusterParent->UniqueIdx().Idx;
+					if (!UniqueIdxToInternalClusterHandle.Contains(StateData.InternalClusterUniqueIdx))
 					{
-						TargetResults.Transforms[TransformGroupIndex] = MassToLocal.Inverse() * ActorScaleTransform * MassToLocal * TargetResults.Transforms[TransformGroupIndex];
+						UniqueIdxToInternalClusterHandle.Add(StateData.InternalClusterUniqueIdx, ClusterParent);
 					}
-
-					PhysicsThreadCollection.Transform[TransformGroupIndex] = TargetResults.Transforms[TransformGroupIndex];
 				}
 
-				// The rigid body parent cluster has changed within the solver, and its
-				// parent body is not tracked within the geometry collection. So we need to
-				// pull the rigid bodies out of the transform hierarchy, and just drive
-				// the positions directly from the solvers cluster particle.
-				if(TPBDRigidParticleHandle<Chaos::FReal, 3>* ClusterParentBase = Handle->ClusterIds().Id)
-				{
-					if(Chaos::FPBDRigidClusteredParticleHandle* ClusterParent = ClusterParentBase->CastToClustered())
-					{
-						// synchronize parents if it has changed.
-						if(SolverClusterID[TransformGroupIndex] != ClusterParent)
-						{
-							// Force all driven rigid bodies out of the transform hierarchy
-							if(Parent[TransformGroupIndex] != INDEX_NONE)
-							{
-								// If the parent of this NON DISABLED body is set to anything other than INDEX_NONE,
-								// then it was just unparented, likely either by rigid clustering or by fields.  We
-								// need to force all such enabled rigid bodies out of the transform hierarchy.
-								TargetResults.Parent[TransformGroupIndex] = INDEX_NONE;
-
-								// GeometryCollectionAlgo::UnparentTransform(&PhysicsThreadCollection, ChildIndex);
-								PhysicsThreadCollection.Children[PhysicsThreadCollection.Parent[TransformGroupIndex]].Remove(TransformGroupIndex);
-								PhysicsThreadCollection.Parent[TransformGroupIndex] = INDEX_NONE;
-
-								// Indicate that this object needs to be updated and the proxy is active.
-								TargetResults.States[TransformGroupIndex].DisabledState = false;
-								IsObjectDynamic = true;
-							}
-							SolverClusterID[TransformGroupIndex] = Handle->ClusterIds().Id;
-						}
+				// child of internal clusters may not have their position up to date 
+				// so we need to compute it from the ChildToParent transform 
+				const FTransform ParticleToWorld = Handle->ChildToParent() * FRigidTransform3(ClusterParent->X(), ClusterParent->R()); // aka ClusterChildToWorld
+				PositionData.ParticleX = ParticleToWorld.GetTranslation();
+				PositionData.ParticleR = ParticleToWorld.GetRotation();
+				
+				// Indicate that this object needs to be updated and the proxy is active.
+				StateData.State.DisabledState = false;
+				StateData.State.HasInternalClusterParent = true;
+				StateData.State.DynamicInternalClusterParent = (ClusterParent->IsDynamic());
 
-						if(ClusterParent->InternalCluster())
-						{
-							const bool bIsOwnedInternalCluster = ClusterParent->PhysicsProxy() == this;
+				bHasChanged = true;
 
-							// We don't want the to keep track of any external internal clusters (i.e. cluster unions).
-							if (bIsOwnedInternalCluster)
-							{
-								const int32 InternalClusterParentUniqueIdx = ClusterParent->UniqueIdx().Idx;
-								if (!UniqueIdxToInternalClusterHandle.Contains(InternalClusterParentUniqueIdx))
-								{
-									UniqueIdxToInternalClusterHandle.Add(InternalClusterParentUniqueIdx, ClusterParent);
-								}
-								TargetResults.InternalClusterUniqueIdx[TransformGroupIndex] = ClusterParent->UniqueIdx().Idx;
-							}
-							
-							const FTransform ParticleToWorld = Handle->ChildToParent() * FRigidTransform3(ClusterParent->X(), ClusterParent->R());    // aka ClusterChildToWorld
-							TargetResults.ParticleXs[TransformGroupIndex] = ParticleToWorld.GetTranslation();
-							GC_PHYSICSPROXY_CHECK_FOR_NAN(TargetResults.ParticleXs[TransformGroupIndex]);
-
-							TargetResults.ParticleRs[TransformGroupIndex] = ParticleToWorld.GetRotation();
-							TargetResults.ParticleVs[TransformGroupIndex] = ClusterParent->V();
-							TargetResults.ParticleWs[TransformGroupIndex] = ClusterParent->W();
-
-							// GeomToActor = ActorToWorld.Inv() * ClusterChildToWorld * MassToLocal.Inv();
-							const FTransform MassToLocal                  = PhysicsThreadCollection.MassToLocal[TransformGroupIndex];
-							TargetResults.Transforms[TransformGroupIndex] = MassToLocal.GetRelativeTransformReverse(ParticleToWorld).GetRelativeTransform(ActorToWorld);
-							TargetResults.Transforms[TransformGroupIndex].NormalizeRotation();
-							if (IsActorScaled)
-							{
-								TargetResults.Transforms[TransformGroupIndex] = MassToLocal.Inverse() * ActorScaleTransform * MassToLocal * TargetResults.Transforms[TransformGroupIndex];
-							}
+				// as we just transitioned from disabled to non disabled the update is unconditional 
+				UpdateParticleHandleTransform(*CurrentSolver, *Handle, ParticleToWorld);
+			}
 
-							PhysicsThreadCollection.Transform[TransformGroupIndex] = TargetResults.Transforms[TransformGroupIndex];
+			if (bGeometryCollectionEnabledNestedChildTransformUpdates)
+			{
+				if (ClusterParent && !ClusterParent->Disabled())
+				{
+					const FRigidTransform3 ChildToWorld = Handle->ChildToParent() * FRigidTransform3(ClusterParent->X(), ClusterParent->R());
+					UpdateParticleHandleTransformIfNeeded(*CurrentSolver, *Handle, ChildToWorld);
+					// fields may have applied velocities, we need to make sure to clear that up, so that we don't accumulate
+					VelocityData.ParticleV = FVec3::ZeroVector;
+					VelocityData.ParticleW = FVec3::ZeroVector;
 
-							// Indicate that this object needs to be updated and the proxy is active.
-							TargetResults.States[TransformGroupIndex].DisabledState = false;
-							TargetResults.States[TransformGroupIndex].HasInternalClusterParent = true;
-							TargetResults.States[TransformGroupIndex].DynamicInternalClusterParent = (ClusterParent->IsDynamic());
-							IsObjectDynamic = true;
+					bHasChanged = true;
+				}
+			}
 
-							// as we just transitioned from disabled to non disabled the update is unconditional 
-							UpdateParticleHandleTransform(*CurrentSolver, *Handle, ParticleToWorld);
-						}
+			// Add to the results if anything has been marked as changed 
+			if (bHasChanged)
+			{
+				// default is what we have on the collection
+				FTransform ParentSpaceTransform = PhysicsThreadCollection.Transform[TransformGroupIndex];
 
-						if (bGeometryCollectionEnabledNestedChildTransformUpdates)
-						{
-							if (!ClusterParent->Disabled())
-							{
-								const FRigidTransform3 ChildToWorld = Handle->ChildToParent() * FRigidTransform3(ClusterParent->X(), ClusterParent->R());
-								UpdateParticleHandleTransformIfNeeded(*CurrentSolver, *Handle, ChildToWorld);
-								// fields may have applied velocities, we need to make sure to clear that up, so that we don't accumulate
-								Handle->SetV(FVec3::Zero());
-								Handle->SetW(FVec3::Zero());
-							}
-						}
+				// recompute parent space transform if there's no more parent or if the parent is an internal cluster 
+				if (!ClusterParent || ClusterParent->InternalCluster())
+				{
+					const FRigidTransform3 ParticleToWorld(PositionData.ParticleX, PositionData.ParticleR);
+					const FTransform MassToLocal = PhysicsThreadCollection.MassToLocal[TransformGroupIndex];
+					ParentSpaceTransform = MassToLocal.GetRelativeTransformReverse(ParticleToWorld).GetRelativeTransform(ActorToWorld);
+					ParentSpaceTransform.NormalizeRotation();
+					if (IsActorScaled)
+					{
+						ParentSpaceTransform = MassToLocal.Inverse() * ActorScaleTransform * MassToLocal * ParentSpaceTransform;
 					}
 				}
-			}    // end if
 
-			PhysicsThreadCollection.Active[TransformGroupIndex] = !TargetResults.States[TransformGroupIndex].DisabledState;
+				const int32 EntryIndex = Results.AddEntry(TransformGroupIndex);
+				Results.SetState(EntryIndex, StateData);
+				Results.SetPositions(EntryIndex, PositionData);
+				Results.SetVelocities(EntryIndex, VelocityData);
+				Results.SetTransform(EntryIndex, ParentSpaceTransform);
+				IsObjectDynamic = true;
+			}
+
+			// update collection active state 
+			PhysicsThreadCollection.Active[TransformGroupIndex] = !StateData.State.DisabledState;
 
 #if WITH_EDITORONLY_DATA
-			FGeometryCollectionResults::FDamageInfo& DamageInfo = TargetResults.DamageInfo[TransformGroupIndex];
-			DamageInfo.Damage = static_cast<float>(Handle->CollisionImpulse());
-			DamageInfo.DamageThreshold = static_cast<float>(Handle->GetInternalStrains());
+			// collect damage information to display in the UI 
+			FGeometryCollectionResults::FDamageData DamageData;
+			DamageData.Damage = FMath::Max(Handle->CollisionImpulse(), Handle->GetExternalStrain());
+			DamageData.DamageThreshold = Handle->GetInternalStrains();
+			Results.SetDamages(TransformGroupIndex, DamageData);
 #endif			
 		}    // end for
 	}        // STAT_CalcParticleToWorld scope
 	
-	// If object is dynamic, compute global matrices	
-	if (IsObjectDynamic || TargetResults.GlobalTransforms.Num() == 0)
-	{
-		SCOPE_CYCLE_COUNTER(STAT_CalcGlobalGCMatrices);
-		check(TargetResults.Transforms.Num() == TargetResults.Parent.Num());
-		GeometryCollectionAlgo::GlobalMatrices(TargetResults.Transforms, TargetResults.Parent, TargetResults.GlobalTransforms);
-	}
-	
-
 	// Advertise to game thread
-	TargetResults.IsObjectDynamic = IsObjectDynamic;
-	TargetResults.IsObjectLoading = IsObjectLoading;
+	Results.IsObjectDynamic = IsObjectDynamic;
+	Results.IsObjectLoading = IsObjectLoading;
 }
 
 void FGeometryCollectionPhysicsProxy::FlipBuffer()
@@ -3238,6 +3142,183 @@ void FGeometryCollectionPhysicsProxy::FlipBuffer()
 	PhysToGameInterchange.FlipProducer();
 }
 
+// Update game thread particle X and R properties and retun true if they changed from previous values
+static inline bool UpdateGTParticleXR(Chaos::FPBDRigidParticle& GTParticle, const Chaos::FVec3& NewX, const Chaos::FRotation3& NewR)
+{
+	GC_PHYSICSPROXY_CHECK_FOR_NAN(NewX);
+
+	const Chaos::FVec3 OldX = GTParticle.X();
+	const bool bNeedUpdateX = (NewX != OldX);
+	if (bNeedUpdateX)
+	{
+		GTParticle.SetX(NewX, false);
+	}
+
+	const Chaos::FRotation3 OldR = GTParticle.R();
+	const bool bNeedUpdateR = (NewR != OldR);
+	if (bNeedUpdateR)
+	{
+		GTParticle.SetR(NewR, false);
+	}
+
+	const bool bChanged = (bNeedUpdateX || bNeedUpdateR);
+	if (bChanged)
+	{
+		GTParticle.UpdateShapeBounds();
+	}
+
+	return bChanged;
+}
+
+// Update a value and return true if the value was different from the old one 
+template <typename TOld, typename TNew>
+static inline bool UpdateValue(TOld& ValueInOut,const TNew& NewValue)
+{
+	const bool bValueDifferent = (NewValue != ValueInOut);
+	if (bValueDifferent)
+	{
+		ValueInOut = NewValue;
+	}
+	return bValueDifferent;
+}
+
+static inline bool UpdateValue(FBitReference ValueInOut, const bool& NewValue)
+{
+	const bool bValueDifferent = (NewValue != ValueInOut);
+	if (bValueDifferent)
+	{
+		ValueInOut = NewValue;
+	}
+	return bValueDifferent;
+}
+
+static inline bool UpdateTransform(FTransform& ValueInOut, const FTransform& NewValue)
+{
+	const bool bValueDifferent = !NewValue.Equals(ValueInOut);
+	if (bValueDifferent)
+	{
+		ValueInOut = NewValue;
+	}
+	return bValueDifferent;
+}
+
+// this helper class helps getting the right interpolated values )
+// it also handles the case where entries for a specific transfrom index may be available in one or the other the other results 
+// and properly fallback to interpolate 
+struct FResultInterpolator
+{
+public:
+	FResultInterpolator(
+		int32 TransformIndexIn,
+		const FGeometryCollectionResults& ResultsIn, 
+		const FGeometryCollectionResults& NextResultsIn,
+		const TManagedArray<FTransform>& TransformsIn,
+		Chaos::FPBDRigidParticle& GTParticleIn,
+		Chaos::FRealSingle AlphaIn
+	)
+		: TransformIndex(TransformIndexIn)
+		, Results(&ResultsIn)
+		, NextResults(&NextResultsIn)
+		, Transforms(TransformsIn)
+		, GTParticle(GTParticleIn)
+		, Alpha(AlphaIn)
+	{
+		EntryIndex = Results->GetEntryIndexByTransformIndex(TransformIndexIn);
+		NextEntryIndex = NextResults->GetEntryIndexByTransformIndex(TransformIndexIn);
+	}
+
+	bool HasNoEntry() const
+	{
+		return (EntryIndex == INDEX_NONE && NextEntryIndex == INDEX_NONE);
+	}
+
+	void GetPositions(Chaos::FVec3& PosOut, Chaos::FRotation3& RotOut) const
+	{
+		if (NextEntryIndex == INDEX_NONE)
+		{
+			const FGeometryCollectionResults::FPositionData& Data = Results->GetPositions(EntryIndex);
+			// not present in next results just use results
+			PosOut = Data.ParticleX;
+			RotOut = Data.ParticleR;
+		}
+		else if (EntryIndex == INDEX_NONE)
+		{
+			// not present in current results, we need to interpolate from the actual GTparticle
+			const FGeometryCollectionResults::FPositionData& NextData = NextResults->GetPositions(NextEntryIndex);
+			PosOut = FMath::Lerp(GTParticle.X(), NextData.ParticleX , Alpha);
+			RotOut = FQuat::Slerp(GTParticle.R(), NextData.ParticleR, Alpha);
+		}
+		else
+		{
+			// both available 
+			const FGeometryCollectionResults::FPositionData& Data = Results->GetPositions(EntryIndex);
+			const FGeometryCollectionResults::FPositionData& NextData = NextResults->GetPositions(NextEntryIndex);
+			PosOut = FMath::Lerp(Data.ParticleX, NextData.ParticleX, Alpha);
+			RotOut = FQuat::Slerp(Data.ParticleR, NextData.ParticleR, Alpha);
+		}
+	}
+
+	void GetVelocities(FVector3f& LinearVelOut, FVector3f& AngularVelOut) const
+	{
+		if (NextEntryIndex == INDEX_NONE)
+		{
+			// not present in next results just use results
+			const FGeometryCollectionResults::FVelocityData& Data = Results->GetVelocities(EntryIndex);
+			LinearVelOut = Data.ParticleV;
+			AngularVelOut = Data.ParticleW;
+		}
+		else if (EntryIndex == INDEX_NONE)
+		{
+			// not present in current results, we need to interpolate from the actual GTparticle
+			const FGeometryCollectionResults::FVelocityData& NextData = NextResults->GetVelocities(NextEntryIndex);
+			LinearVelOut = FMath::Lerp(Chaos::FVec3f(GTParticle.V()), NextData.ParticleV, Alpha);
+			AngularVelOut = FMath::Lerp(Chaos::FVec3f(GTParticle.W()), NextData.ParticleW, Alpha);
+		}
+		else
+		{
+			// both available 
+			const FGeometryCollectionResults::FVelocityData& Data = Results->GetVelocities(EntryIndex);
+			const FGeometryCollectionResults::FVelocityData& NextData = NextResults->GetVelocities(NextEntryIndex);
+			LinearVelOut = FMath::Lerp(Data.ParticleV, NextData.ParticleV, Alpha);
+			AngularVelOut = FMath::Lerp(Data.ParticleW, NextData.ParticleW, Alpha);
+		}
+	}
+
+	void GetTransform(FTransform& TransformOut, const FTransform& MassToLocal) const
+	{
+		if (NextEntryIndex == INDEX_NONE)
+		{
+			// not present in next results just use results
+			TransformOut = Results->GetTransform(EntryIndex);
+		}
+		else if (EntryIndex == INDEX_NONE)
+		{
+			// not present in current results, we need to interpolate from the actual GTparticle
+			FTransform InterpolatedTransform;
+			InterpolatedTransform.Blend(MassToLocal * Transforms[TransformIndex], MassToLocal * NextResults->GetTransform(NextEntryIndex), Alpha);
+			TransformOut = MassToLocal.Inverse() * InterpolatedTransform;
+		}
+		else
+		{
+			// both available 
+			FTransform InterpolatedTransform;
+			InterpolatedTransform.Blend(MassToLocal * Results->GetTransform(EntryIndex), MassToLocal * NextResults->GetTransform(NextEntryIndex), Alpha);
+			TransformOut = MassToLocal.Inverse() * InterpolatedTransform;
+		}
+	}
+
+private:
+	const int32 TransformIndex;
+	const FGeometryCollectionResults* Results;
+	FGeometryCollectionResults::FEntryIndex EntryIndex;
+	const FGeometryCollectionResults* NextResults;
+	FGeometryCollectionResults::FEntryIndex NextEntryIndex;
+	const TManagedArray<FTransform>& Transforms;
+	Chaos::FPBDRigidParticle& GTParticle;
+	Chaos::FRealSingle Alpha;
+
+};
+
 // Called from FPhysScene_ChaosInterface::SyncBodies(), NOT the solver.
 bool FGeometryCollectionPhysicsProxy::PullFromPhysicsState(const Chaos::FDirtyGeometryCollectionData& PullData, const int32 SolverSyncTimestamp, const Chaos::FDirtyGeometryCollectionData* NextPullData, const Chaos::FRealSingle* Alpha)
 {
@@ -3254,72 +3335,100 @@ bool FGeometryCollectionPhysicsProxy::PullFromPhysicsState(const Chaos::FDirtyGe
 	 *
 	 */
 
-	const FGeometryCollectionResults& TargetResults = PullData.Results;
+	// earlier exit
+	const bool HasResults = (PullData.Results.GetNumEntries() > 0);
+	const bool HasNextResults = (NextPullData && NextPullData->Results.GetNumEntries() > 0);
+	if (!HasResults && !HasNextResults)
+	{
+		return false;
+	}
 
-	TManagedArray<FVector3f>* LinearVelocities = GameThreadCollection.FindAttributeTyped<FVector3f>("LinearVelocity", FTransformCollection::TransformGroup);
-	TManagedArray<FVector3f>* AngularVelocities = GameThreadCollection.FindAttributeTyped<FVector3f>("AngularVelocity", FTransformCollection::TransformGroup);
-	TManagedArray<uint8>* InternalClusterParentTypeArray = GameThreadCollection.FindAttributeTyped<uint8>("InternalClusterParentTypeArray", FTransformCollection::TransformGroup);
+	const FGeometryCollectionResults& Results = PullData.Results;
 
-	bool bIsCollectionDirty = false;
-	
-	// We should never be changing the number of entries, this would break other 
-	// attributes in the transform group.
 	const int32 NumTransforms = GameThreadCollection.Transform.Num();
-	if (ensure(NumTransforms == TargetResults.Transforms.Num()))
+	const bool bNeedInterpolation = (NextPullData != nullptr);
+
+	bool bIsCollectionDirty = false;
+
+	if (NumTransforms > 0)
 	{
-		// first : copy the non interpolate-able values
-		for (int32 TransformGroupIndex = 0; TransformGroupIndex < NumTransforms; ++TransformGroupIndex)
+		static const FName LinearVelocityAttributeName = "LinearVelocity";
+		static const FName AngularVelocitiesAttributeName = "AngularVelocity";
+		static const FName InternalClusterParentTypeArrayAttributeName = "InternalClusterParentTypeArray";
+		static const FName AnimateTransformAttributeName = "AnimateTransformAttribute";
+
+		TManagedArray<FVector3f>* LinearVelocities = GameThreadCollection.FindAttributeTyped<FVector3f>(LinearVelocityAttributeName, FTransformCollection::TransformGroup);
+		TManagedArray<FVector3f>* AngularVelocities = GameThreadCollection.FindAttributeTyped<FVector3f>(AngularVelocitiesAttributeName, FTransformCollection::TransformGroup);
+		TManagedArray<uint8>* InternalClusterParentTypeArray = GameThreadCollection.FindAttributeTyped<uint8>(InternalClusterParentTypeArrayAttributeName, FTransformCollection::TransformGroup);
+		const TManagedArray<bool>* AnimationsActive = GameThreadCollection.FindAttribute<bool>(AnimateTransformAttributeName, FGeometryCollection::TransformGroup);
+
+		// first step : process the non values that do not need to be interpolated 
+		for (int32 EntryIndex = 0; EntryIndex < Results.GetNumEntries(); EntryIndex++)
 		{
+			const FGeometryCollectionResults::FStateData& StateData = Results.GetState(EntryIndex);
+			const int32 TransformGroupIndex = StateData.TransformIndex;
+
 			if (GTParticles[TransformGroupIndex] == nullptr)
 			{
 				continue;
 			}
 			FParticle& GTParticle = *GTParticles[TransformGroupIndex];
 
-			const bool bIsActive = !TargetResults.States[TransformGroupIndex].DisabledState;
-			if (GameThreadCollection.Active[TransformGroupIndex] != bIsActive)
+			const bool bIsActive = !StateData.State.DisabledState;
+			if (UpdateValue(GameThreadCollection.Active[TransformGroupIndex], bIsActive))
 			{
-				GameThreadCollection.Active[TransformGroupIndex] = bIsActive;
 				GTParticle.SetDisabled(!bIsActive);
 				bIsCollectionDirty = true;
 			}
 			
-			const int32 NewState = TargetResults.States[TransformGroupIndex].DynamicState;
-			if (GameThreadCollection.DynamicState[TransformGroupIndex] != NewState)
+
+			if (UpdateValue(GameThreadCollection.DynamicState[TransformGroupIndex], StateData.State.DynamicState))
 			{
-				GameThreadCollection.DynamicState[TransformGroupIndex] = NewState;
-				GTParticle.SetObjectState(static_cast<Chaos::EObjectStateType>(NewState));
+				GTParticle.SetObjectState(static_cast<Chaos::EObjectStateType>(StateData.State.DynamicState));
 				bIsCollectionDirty = true;
 			}
 
-			if (bIsActive)
+			if (UpdateValue(GameThreadCollection.Parent[TransformGroupIndex], StateData.ParentTransformIndex))
 			{
-				const int32 NewParent = TargetResults.Parent[TransformGroupIndex];
-				if (GameThreadCollection.Parent[TransformGroupIndex] != NewParent)
-				{
-					GameThreadCollection.Parent[TransformGroupIndex] = NewParent;
-					bIsCollectionDirty = true;
-				}
-				// if (NewState != (int32)EObjectStateTypeEnum::Chaos_Object_Static && NewState != (int32)EObjectStateTypeEnum::Chaos_Object_Sleeping)
-				// {
-				// 	bIsCollectionDirty = true;
-				// }
+				bIsCollectionDirty = true;
 			}
 
 			if (InternalClusterParentTypeArray)
 			{
 				Chaos::EInternalClusterType ParentType = Chaos::EInternalClusterType::None;
-				if (TargetResults.States[TransformGroupIndex].HasInternalClusterParent != 0)
+				if (StateData.State.HasInternalClusterParent != 0)
 				{
-					ParentType = (TargetResults.States[TransformGroupIndex].DynamicInternalClusterParent != 0)? Chaos::EInternalClusterType::Dynamic: Chaos::EInternalClusterType::KinematicOrStatic; 
+					ParentType = (StateData.State.DynamicInternalClusterParent != 0)? Chaos::EInternalClusterType::Dynamic: Chaos::EInternalClusterType::KinematicOrStatic;
 				}
 				const uint8 ParentTypeUInt8 = static_cast<uint8>(ParentType);
-				if ((*InternalClusterParentTypeArray)[TransformGroupIndex] != ParentTypeUInt8)
+				if (UpdateValue((*InternalClusterParentTypeArray)[TransformGroupIndex], ParentTypeUInt8))
 				{
-					(*InternalClusterParentTypeArray)[TransformGroupIndex] = ParentTypeUInt8;
 					bIsCollectionDirty = true;
 				}
 			}
+
+			// if interpolation is off , we need to apply XR, VW and transforms
+			if (bNeedInterpolation == false)
+			{
+				const bool bAnimatingWhileDisabled = AnimationsActive ? (*AnimationsActive)[TransformGroupIndex] : false;
+				if (bIsActive || bAnimatingWhileDisabled)
+				{
+					const FGeometryCollectionResults::FPositionData& PositionData = Results.GetPositions(EntryIndex);
+
+					const bool XRModified = UpdateGTParticleXR(GTParticle, PositionData.ParticleX, PositionData.ParticleR);
+					bIsCollectionDirty |= XRModified;
+
+					if (LinearVelocities && AngularVelocities)
+					{
+						const FGeometryCollectionResults::FVelocityData& VelocityData = Results.GetVelocities(EntryIndex);
+						bIsCollectionDirty |= UpdateValue((*LinearVelocities)[TransformGroupIndex], VelocityData.ParticleV);
+						bIsCollectionDirty |= UpdateValue((*AngularVelocities)[TransformGroupIndex], VelocityData.ParticleW);
+					}
+
+					const FTransform& NewTransform = Results.GetTransform(EntryIndex);
+					bIsCollectionDirty |= UpdateTransform(GameThreadCollection.Transform[TransformGroupIndex], NewTransform);
+				}
+			}
 		}
 
 #if WITH_EDITORONLY_DATA
@@ -3327,134 +3436,77 @@ bool FGeometryCollectionPhysicsProxy::PullFromPhysicsState(const Chaos::FDirtyGe
 		{
 			for (int32 TransformGroupIndex = 0; TransformGroupIndex < NumTransforms; ++TransformGroupIndex)
 			{
-				const FGeometryCollectionResults::FDamageInfo& DamageInfo = TargetResults.DamageInfo[TransformGroupIndex]; 
-				Collector->SampleDamage(TransformGroupIndex, DamageInfo.Damage, DamageInfo.DamageThreshold);
+				const FGeometryCollectionResults::FDamageData& DamageData = Results.GetDamages(TransformGroupIndex);
+				Collector->SampleDamage(TransformGroupIndex, DamageData.Damage, DamageData.DamageThreshold);
 			}
 		}
 #endif
-		TManagedArray<bool>* AnimationsActive =
-			GameThreadCollection.FindAttribute<bool>("AnimateTransformAttribute", FGeometryCollection::TransformGroup);
 
 		// second : interpolate-able ones
-		const bool bNeedInterpolation = (NextPullData!= nullptr);
 		if (bNeedInterpolation)
 		{
-			const FGeometryCollectionResults& Prev = PullData.Results;
-			const FGeometryCollectionResults& Next = NextPullData->Results;
+			const FGeometryCollectionResults& NextResults = NextPullData->Results;
+
+			// for that case we cannot just go through the list of entries since Results and NextResults may have different number of entries that don't always match
+			// so we need to go through the transform indices and find the matching entries on both side 
+			// this is helped by the FResultAccessor structure
 			for (int32 TransformGroupIndex = 0; TransformGroupIndex < NumTransforms; ++TransformGroupIndex)
 			{
-				const bool bAnimatingWhileDisabled = AnimationsActive ? (*AnimationsActive)[TransformGroupIndex] : false;
-				const bool WasDisabled = Prev.States[TransformGroupIndex].DisabledState;
-
-				if (!WasDisabled || bAnimatingWhileDisabled)
+				if (GTParticles[TransformGroupIndex] == nullptr)
 				{
-					if (GTParticles[TransformGroupIndex] == nullptr)
-					{
-						continue;
-					}
-					FParticle& GTParticle = *GTParticles[TransformGroupIndex];
-
-					GC_PHYSICSPROXY_CHECK_FOR_NAN(Prev.ParticleXs[TransformGroupIndex]);
-					GC_PHYSICSPROXY_CHECK_FOR_NAN(Next.ParticleXs[TransformGroupIndex]);
-
-					const Chaos::FVec3 OldX = GTParticle.X();
-					const Chaos::FVec3 NewX = FMath::Lerp(Prev.ParticleXs[TransformGroupIndex], Next.ParticleXs[TransformGroupIndex], *Alpha);
-					GC_PHYSICSPROXY_CHECK_FOR_NAN(NewX);
-
-					const bool bNeedUpdateX = (NewX != OldX);
-					if (bNeedUpdateX)
-					{
-						GTParticle.SetX(NewX, false);
-					}
-					
-					const Chaos::FRotation3 OldR = GTParticle.R();
-					const Chaos::FRotation3 NewR = FQuat::Slerp(Prev.ParticleRs[TransformGroupIndex], Next.ParticleRs[TransformGroupIndex], *Alpha);
-					const bool bNeedUpdateR = (NewR != OldR); 
-					if (bNeedUpdateR)
-					{
-						GTParticle.SetR(NewR, false);
-					}
-
-					if (bNeedUpdateX || bNeedUpdateR)
-					{
-						GTParticle.UpdateShapeBounds();
-						bIsCollectionDirty = true;
-					}
-
-					
-					// need to interpolate using MassToLocal
-					const FTransform& MassToLocal = GameThreadCollection.MassToLocal[TransformGroupIndex];
-					FTransform InterpolatedTransform;
-					InterpolatedTransform.Blend(MassToLocal * Prev.Transforms[TransformGroupIndex], MassToLocal * Next.Transforms[TransformGroupIndex], *Alpha);
-					GameThreadCollection.Transform[TransformGroupIndex] = MassToLocal.Inverse() * InterpolatedTransform;
-
-					if(LinearVelocities && AngularVelocities)
-					{
-						// LWC : potential loss of precision if the velocity is insanely big ( unlikely ) 
-						(*LinearVelocities)[TransformGroupIndex] = FVector3f(FMath::Lerp(Prev.ParticleVs[TransformGroupIndex], Next.ParticleVs[TransformGroupIndex], *Alpha)); 
-						(*AngularVelocities)[TransformGroupIndex] = FVector3f(FMath::Lerp(Prev.ParticleWs[TransformGroupIndex], Next.ParticleWs[TransformGroupIndex], *Alpha));
-					}
+					continue;
 				}
-			}
-		}
-		else
-		{
-			for (int32 TransformGroupIndex = 0; TransformGroupIndex < NumTransforms; ++TransformGroupIndex)
-			{
-				bool bAnimatingWhileDisabled = AnimationsActive ? (*AnimationsActive)[TransformGroupIndex] : false;
 
-				if (!TargetResults.States[TransformGroupIndex].DisabledState || bAnimatingWhileDisabled)
+				FParticle& GTParticle = *GTParticles[TransformGroupIndex];
+				const FResultInterpolator ResultInterpolator(TransformGroupIndex, Results, NextResults, GameThreadCollection.Transform, GTParticle, *Alpha);
+				if (ResultInterpolator.HasNoEntry())
 				{
-					if (GTParticles[TransformGroupIndex] == nullptr)
-					{
-						continue;
-					}
-					FParticle& GTParticle = *GTParticles[TransformGroupIndex];
-
-					const Chaos::FVec3 OldX = GTParticle.X();
-					const Chaos::FVec3& NewX = TargetResults.ParticleXs[TransformGroupIndex];
-					const bool bNeedUpdateX = (NewX != OldX);
-					if (bNeedUpdateX)
-					{
-						GTParticle.SetX(NewX, false);
-					}
-					
-					const Chaos::FRotation3 OldR = GTParticle.R();
-					const Chaos::FRotation3& NewR = TargetResults.ParticleRs[TransformGroupIndex];
-					const bool bNeedUpdateR = (NewR != OldR);
-					if (bNeedUpdateR)
-					{
-						GTParticle.SetR(NewR, false);
-					}
+					continue;
+				}
 
-					if (bNeedUpdateX || bNeedUpdateR)
-					{
-						GTParticle.UpdateShapeBounds();
-						bIsCollectionDirty = true;
-					}
+				const bool bAnimatingWhileDisabled = AnimationsActive ? (*AnimationsActive)[TransformGroupIndex] : false;
+				const bool bActive = GameThreadCollection.Active[TransformGroupIndex];
 
-					GameThreadCollection.Transform[TransformGroupIndex] = TargetResults.Transforms[TransformGroupIndex];
+				if (bActive || bAnimatingWhileDisabled)
+				{
+					Chaos::FVec3 NewX;
+					Chaos::FRotation3 NewR;
+					ResultInterpolator.GetPositions(NewX, NewR);
+					const bool XRModified = UpdateGTParticleXR(GTParticle, NewX, NewR);
+					bIsCollectionDirty |= XRModified;
 
-					if(LinearVelocities && AngularVelocities)
+					if (LinearVelocities && AngularVelocities)
 					{
-						// LWC : potential loss of precision if the velocity is insanely big ( unlikely ) 
-						(*LinearVelocities)[TransformGroupIndex] = FVector3f(TargetResults.ParticleVs[TransformGroupIndex]); 
-						(*AngularVelocities)[TransformGroupIndex] = FVector3f(TargetResults.ParticleWs[TransformGroupIndex]);
+						Chaos::FVec3f NewV;
+						Chaos::FVec3f NewW;
+						ResultInterpolator.GetVelocities(NewV, NewW);
+						bIsCollectionDirty |= UpdateValue((*LinearVelocities)[TransformGroupIndex], NewV);
+						bIsCollectionDirty |= UpdateValue((*AngularVelocities)[TransformGroupIndex], NewW);
 					}
+
+					FTransform NewTransform{ FTransform::Identity };
+					const FTransform& MassToLocal = GameThreadCollection.MassToLocal[TransformGroupIndex];
+					ResultInterpolator.GetTransform(NewTransform, MassToLocal);
+					bIsCollectionDirty |= UpdateTransform(GameThreadCollection.Transform[TransformGroupIndex], NewTransform);
 				}
 			}
-		}
+		} 
 
 		// internal cluster index map update
 		GTParticlesToInternalClusterUniqueIdx.Reset();
 		InternalClusterUniqueIdxToChildrenTransformIndices.Reset();
 		for (int32 TransformGroupIndex = 0; TransformGroupIndex < NumTransforms; ++TransformGroupIndex)
 		{
-			const int32 InternalClusterUniqueIdx = TargetResults.InternalClusterUniqueIdx[TransformGroupIndex]; 
-			if (InternalClusterUniqueIdx > INDEX_NONE)
+			const FGeometryCollectionResults::FEntryIndex EntryIndex = Results.GetEntryIndexByTransformIndex(TransformGroupIndex);
+			if (EntryIndex != INDEX_NONE)
 			{
-				GTParticlesToInternalClusterUniqueIdx.Add(GTParticles[TransformGroupIndex].Get(), InternalClusterUniqueIdx);
-				InternalClusterUniqueIdxToChildrenTransformIndices.FindOrAdd(InternalClusterUniqueIdx).Add(TransformGroupIndex); 
+				const FGeometryCollectionResults::FStateData& StateData = Results.GetState(EntryIndex);
+
+				if (StateData.InternalClusterUniqueIdx > INDEX_NONE)
+				{
+					GTParticlesToInternalClusterUniqueIdx.Add(GTParticles[TransformGroupIndex].Get(), StateData.InternalClusterUniqueIdx);
+					InternalClusterUniqueIdxToChildrenTransformIndices.FindOrAdd(StateData.InternalClusterUniqueIdx).Add(TransformGroupIndex);
+				}
 			}
 		}
 
@@ -3475,7 +3527,6 @@ bool FGeometryCollectionPhysicsProxy::PullFromPhysicsState(const Chaos::FDirtyGe
 	return true;
 }
 
-
 void FGeometryCollectionPhysicsProxy::UpdateFilterData_External(const FCollisionFilterData& NewSimFilter, const FCollisionFilterData& NewQueryFilter)
 {
 	SCOPE_CYCLE_COUNTER(STAT_GCUpdateFilterData);
@@ -4819,4 +4870,4 @@ void FRuntimeDataCollector::RemoveCollector(const FGuid& Guid)
 {
 	Collectors.Remove(Guid);
 	
-}
+}
\ No newline at end of file
diff --git a/Engine/Source/Runtime/Experimental/Chaos/Public/GeometryCollectionProxyData.h b/Engine/Source/Runtime/Experimental/Chaos/Public/GeometryCollectionProxyData.h
--- a/Engine/Source/Runtime/Experimental/Chaos/Public/GeometryCollectionProxyData.h
+++ b/Engine/Source/Runtime/Experimental/Chaos/Public/GeometryCollectionProxyData.h
@@ -276,28 +276,24 @@ class FGeometryCollectionResults
 public:
 	FGeometryCollectionResults();
 
-	void Reset();
+	int32 GetNumEntries() const { return Transforms.Num(); }
 
-	int32 NumTransformGroup() const { return Transforms.Num(); }
+	void Reset();
 
-	void InitArrays(const FGeometryDynamicCollection& Other)
+	void InitArrays(const FGeometryDynamicCollection& Collection)
 	{
-		const int32 NumTransforms = Other.NumElements(FGeometryCollection::TransformGroup);
-		States.SetNumUninitialized(NumTransforms);
-		GlobalTransforms.SetNumUninitialized(NumTransforms);
-		ParticleXs.SetNumUninitialized(NumTransforms);
-		ParticleRs.SetNumUninitialized(NumTransforms);
-		ParticleVs.SetNumUninitialized(NumTransforms);
-		ParticleWs.SetNumUninitialized(NumTransforms);
-		
-		Transforms.SetNumUninitialized(NumTransforms);
-		Parent.SetNumUninitialized(NumTransforms);
-		InternalClusterUniqueIdx.SetNumUninitialized(NumTransforms);
+		const int32 NumTransforms = Collection.NumElements(FGeometryCollection::TransformGroup);
+		ModifiedTransformIndices.Init(false, NumTransforms);
 #if WITH_EDITORONLY_DATA
-		DamageInfo.SetNumUninitialized(NumTransforms);
-#endif		
+		if (Damages.Num() != NumTransforms)
+		{
+			Damages.SetNumUninitialized(NumTransforms);
+		}
+#endif	
 	}
 
+	using FEntryIndex = int32;
+
 	struct FState
 	{
 		uint16 DynamicState: 8; // need to fit EObjectStateTypeEnum
@@ -307,29 +303,128 @@ public:
 		// 5 bits left
 	};
 
-	Chaos::FReal SolverDt;
-	TArray<FState> States;
-	TArray<FTransform> GlobalTransforms;
-	TArray<Chaos::FVec3> ParticleXs;
-	TArray<Chaos::FRotation3> ParticleRs;
-	TArray<Chaos::FVec3> ParticleVs;
-	TArray<Chaos::FVec3> ParticleWs;
+	struct FStateData
+	{
+		int32  TransformIndex;
+		int32  ParentTransformIndex;
+		int32  InternalClusterUniqueIdx;
+		FState State;
+	};
 
-	TArray<FTransform> Transforms;
-	TArray<int32> Parent;
-	TArray<int32> InternalClusterUniqueIdx;
+	struct FPositionData
+	{
+		Chaos::FVec3 ParticleX;
+		Chaos::FRotation3 ParticleR;
+	};
+
+	struct FVelocityData
+	{
+		Chaos::FVec3f ParticleV;
+		Chaos::FVec3f ParticleW;
+	};
 
 #if WITH_EDITORONLY_DATA
-	struct FDamageInfo
+	struct FDamageData
 	{
 		float Damage = 0;
 		float DamageThreshold = 0;
 	};
-	
+
+	void SetDamages(int32 TransformIndex, const FDamageData& DamageData)
+	{
+		Damages[TransformIndex] = DamageData;
+	}
+
+	const FDamageData& GetDamages(int32 TransformIndex) const
+	{
+		return Damages[TransformIndex];
+	}
+#endif
+
+	inline FEntryIndex GetEntryIndexByTransformIndex(int32 TransformIndex) const
+	{
+		if (ModifiedTransformIndices[TransformIndex])
+		{
+			return ModifiedTransformIndices.CountSetBits(0, TransformIndex + 1) - 1;
+		}
+		return INDEX_NONE;
+	}
+
+	inline const FStateData& GetState(FEntryIndex EntryIndex) const
+	{
+		return States[EntryIndex];
+	}
+
+	inline const FPositionData& GetPositions(FEntryIndex EntryIndex) const
+	{
+		return Positions[EntryIndex];
+	}
+
+	inline const FVelocityData& GetVelocities(FEntryIndex EntryIndex) const
+	{
+		return Velocities[EntryIndex];
+	}
+
+	inline const FTransform& GetTransform(FEntryIndex EntryIndex) const
+	{
+		return Transforms[EntryIndex];
+	}
+
+	inline void SetSolverDt(const Chaos::FReal SolverDtIn)
+	{
+		SolverDt = SolverDtIn;
+	}
+
+	inline void SetState(int32 EntryIndex, const FStateData& StateData)
+	{
+		States[EntryIndex] = StateData;
+	}
+
+	inline FEntryIndex AddEntry(int32 TransformIndex)
+	{
+		ModifiedTransformIndices[TransformIndex] = true;
+		const FEntryIndex EntryIndex = States.AddDefaulted();
+		ensure(GetEntryIndexByTransformIndex(TransformIndex) == EntryIndex);
+		Positions.AddDefaulted();
+		Velocities.AddDefaulted();
+		Transforms.AddDefaulted();
+		return EntryIndex;
+	}
+
+	inline void SetPositions(FEntryIndex EntryIndex, const FPositionData& PositionData)
+	{
+		Positions[EntryIndex] = PositionData;
+	}
+
+	inline void SetVelocities(FEntryIndex EntryIndex, const FVelocityData& VelocityData)
+	{
+		Velocities[EntryIndex] = VelocityData;
+	}
+
+	inline void SetTransform(FEntryIndex EntryIndex, const FTransform& Transform)
+	{
+		Transforms[EntryIndex] = Transform;
+	}
+
+private:
+	Chaos::FReal SolverDt;
+
+	// we only store the data for modified transforms
+	// ModifiedTransformIndices contains which transform has been set 
+	// use the API to retrieve the entry Index matching a specific transform index
+	TBitArray<> ModifiedTransformIndices;
+	TArray<FStateData> States;
+	TArray<FPositionData> Positions;
+	TArray<FVelocityData> Velocities;
+	TArray<FTransform> Transforms;
+
+#if WITH_EDITORONLY_DATA
 	// use to display impulse statistics in editor
-	TArray<FDamageInfo> DamageInfo;
+	// this is indexed on the transform index
+	TArray<FDamageData> Damages;
 #endif
-	
-	bool IsObjectDynamic;
-	bool IsObjectLoading;
+
+public:
+	uint8 IsObjectDynamic: 1;
+	uint8 IsObjectLoading: 1;
 };
diff --git a/Engine/Source/Runtime/Experimental/GeometryCollectionEngine/Private/GeometryCollection/GeometryCollectionComponent.cpp b/Engine/Source/Runtime/Experimental/GeometryCollectionEngine/Private/GeometryCollection/GeometryCollectionComponent.cpp
--- a/Engine/Source/Runtime/Experimental/GeometryCollectionEngine/Private/GeometryCollection/GeometryCollectionComponent.cpp
+++ b/Engine/Source/Runtime/Experimental/GeometryCollectionEngine/Private/GeometryCollection/GeometryCollectionComponent.cpp
@@ -682,18 +682,12 @@ FBoxSphereBounds UGeometryCollectionComponent::CalcBounds(const FTransform& Loca
 {	
 	SCOPE_CYCLE_COUNTER(STAT_GCCUpdateBounds);
 
-	// #todo(dmp): hack to make bounds calculation work when we don't have valid physics proxy data.  This will
-	// force bounds calculation.
-
-	const FGeometryCollectionResults* Results = PhysicsProxy ? PhysicsProxy->GetConsumerResultsGT() : nullptr;
-	const int32 NumTransforms = Results ? Results->GlobalTransforms.Num() : 0;
-
 	if (bChaos_GC_CacheComponentSpaceBounds)
 	{
 		bool NeedBoundsUpdate = false;
 		NeedBoundsUpdate |= (ComponentSpaceBounds.GetSphere().W < 1e-5);
 		NeedBoundsUpdate |= CachePlayback;
-		NeedBoundsUpdate |= (NumTransforms > 0);
+//		NeedBoundsUpdate |= (PhysicsProxy == nullptr);
 		NeedBoundsUpdate |= (DynamicCollection && DynamicCollection->IsDirty());
 		
 		if (NeedBoundsUpdate)
@@ -2767,8 +2761,8 @@ void UGeometryCollectionComponent::ResetDynamicCollection()
 	if (RestCollection)
 	{
 		CalculateGlobalMatrices();
-		CalculateLocalBounds();
 	}
+	UpdateCachedBounds();
 }
 
 void UGeometryCollectionComponent::OnCreatePhysicsState()
@@ -3091,6 +3085,9 @@ void UGeometryCollectionComponent::OnPostPhysicsSync()
 
 			PrimaryComponentTick.SetTickFunctionEnable(true);
 		}
+
+		// only update removal if the root is broken
+		UpdateRemovalIfNeeded();
 	}
 
 	const bool bDynamicDataIsDirty = (DynamicCollection && DynamicCollection->IsDirty() && HasVisibleGeometry());
@@ -3119,8 +3116,6 @@ void UGeometryCollectionComponent::UpdateRenderSystemsIfNeeded(bool bDynamicColl
 
 	if (bDynamicCollectionDirty)
 	{
-		UpdateRemovalIfNeeded();
-
 		// #todo review: When we've made changes to ISMC, we need to move this function call to SetRenderDynamicData_Concurrent
 		RefreshEmbeddedGeometry();
 
@@ -4449,31 +4444,19 @@ void UGeometryCollectionComponent::CalculateGlobalMatrices()
 {
 	SCOPE_CYCLE_COUNTER(STAT_GCCUGlobalMatrices);
 
-	const FGeometryCollectionResults* Results = PhysicsProxy ? PhysicsProxy->GetConsumerResultsGT() : nullptr;
+	// If hierarchy topology has changed, the RestTransforms is invalidated.
+	if (RestTransforms.Num() != GetTransformArray().Num())
+	{
+		RestTransforms.Empty();
+	}
 
-	const int32 NumTransforms = Results ? Results->GlobalTransforms.Num() : 0;
-	if(NumTransforms > 0)
+	if (!DynamicCollection && RestTransforms.Num() > 0)
 	{
-		// Just calc from results
-		ComponentSpaceTransforms.Reset();
-		ComponentSpaceTransforms.Append(Results->GlobalTransforms);
+		GeometryCollectionAlgo::GlobalMatrices(RestTransforms, GetParentArray(), ComponentSpaceTransforms);
 	}
 	else
 	{
-		// If hierarchy topology has changed, the RestTransforms is invalidated.
-		if (RestTransforms.Num() != GetTransformArray().Num())
-		{
-			RestTransforms.Empty();
-		}
-
-		if (!DynamicCollection && RestTransforms.Num() > 0)
-		{
-			GeometryCollectionAlgo::GlobalMatrices(RestTransforms, GetParentArray(), ComponentSpaceTransforms);
-		}
-		else
-		{
-			GeometryCollectionAlgo::GlobalMatrices(GetTransformArray(), GetParentArray(), ComponentSpaceTransforms);
-		}
+		GeometryCollectionAlgo::GlobalMatrices(GetTransformArray(), GetParentArray(), ComponentSpaceTransforms);
 	}
 
 #if WITH_EDITOR
