Geometry collection : expose component DamageModel property to blueporints

diff --git a/Engine/Source/Runtime/Experimental/GeometryCollectionEngine/Public/GeometryCollection/GeometryCollectionComponent.h b/Engine/Source/Runtime/Experimental/GeometryCollectionEngine/Public/GeometryCollection/GeometryCollectionComponent.h
--- a/Engine/Source/Runtime/Experimental/GeometryCollectionEngine/Public/GeometryCollection/GeometryCollectionComponent.h
+++ b/Engine/Source/Runtime/Experimental/GeometryCollectionEngine/Public/GeometryCollection/GeometryCollectionComponent.h
@@ -795,7 +795,7 @@ public:
 	int32 MaxSimulatedLevel;
 
 	/** Damage model to use for evaluating destruction. */
-	UPROPERTY(EditAnywhere, Category = "ChaosPhysics|Damage")
+	UPROPERTY(EditAnywhere, BlueprintReadWrite, Category = "ChaosPhysics|Damage")
 	EDamageModelTypeEnum DamageModel;
 
 	UPROPERTY(EditAnywhere, BlueprintReadWrite, BlueprintGetter=GetDamageThreshold, BlueprintSetter=SetDamageThreshold, Category = "ChaosPhysics|Damage", meta = (EditCondition = "!bUseSizeSpecificDamageThreshold && DamageModel == EDamageModelTypeEnum::Chaos_Damage_Model_UserDefined_Damage_Threshold"))
