Iris
- Added PacketId to NetSerializationContext

diff --git a/Engine/Source/Runtime/Engine/Private/Net/Experimental/Iris/DataStreamChannel.cpp b/Engine/Source/Runtime/Engine/Private/Net/Experimental/Iris/DataStreamChannel.cpp
--- a/Engine/Source/Runtime/Engine/Private/Net/Experimental/Iris/DataStreamChannel.cpp
+++ b/Engine/Source/Runtime/Engine/Private/Net/Experimental/Iris/DataStreamChannel.cpp
@@ -106,6 +106,7 @@ void UDataStreamChannel::ReceivedBunch(FInBunch& Bunch)
 	BitReader.Seek(Bunch.GetPosBits());
 
 	FNetSerializationContext SerializationContext(&BitReader);
+	SerializationContext.SetPacketId(Connection->InPacketId);
 
 	// For packet stats
 	SerializationContext.SetTraceCollector(Connection->GetInTraceCollector());
@@ -208,6 +209,7 @@ void UDataStreamChannel::Tick()
 		BitWriter.InitBytes(BitStreamBuffer, MaxBytes);
 
 		FNetSerializationContext SerializationContext(&BitWriter);
+		SerializationContext.SetPacketId(Connection->OutPacketId);
 
 #if UE_NET_TRACE_ENABLED	
 		// For Iris we can use the connection trace collector as long as we make sure that the packet is prepared 
diff --git a/Engine/Source/Runtime/Experimental/Iris/Core/Public/Iris/Serialization/NetSerializationContext.h b/Engine/Source/Runtime/Experimental/Iris/Core/Public/Iris/Serialization/NetSerializationContext.h
--- a/Engine/Source/Runtime/Experimental/Iris/Core/Public/Iris/Serialization/NetSerializationContext.h
+++ b/Engine/Source/Runtime/Experimental/Iris/Core/Public/Iris/Serialization/NetSerializationContext.h
@@ -75,6 +75,9 @@ public:
 	void SetIsInitializingDefaultState(bool bInIsInitializingDefaultState) { bIsInitializingDefaultState = bInIsInitializingDefaultState; }
 	bool IsInitializingDefaultState() const { return bIsInitializingDefaultState; }
 
+	void SetPacketId(int32 InPacketId) { PacketId = InPacketId; }
+	int32 GetPacketId() const { return PacketId; }
+
 private:
 	IRISCORE_API bool IsBitStreamOverflown() const;
 	IRISCORE_API void SetBitStreamOverflow();
@@ -88,6 +91,7 @@ private:
 	const FNetBitArrayView* ChangeMask;
 	INetBlobReceiver* NetBlobReceiver;
 	uint32 LocalConnectionId;
+	int32 PacketId;
 	/** Set when replicated objects send their very first state. */
 	uint32 bIsInitState : 1;
 	/** Set only when dealing with a default state. */
@@ -104,6 +108,7 @@ inline FNetSerializationContext::FNetSerializationContext(FNetBitStreamReader* I
 , ChangeMask(nullptr)
 , NetBlobReceiver(nullptr)
 , LocalConnectionId(0)
+, PacketId(-1)
 , bIsInitState(0)
 , bIsInitializingDefaultState(0)
 {
