I/O - extend on demand I/O store TOC with source chunk hash and compression block hashes

#rnx

diff --git a/Engine/Source/Runtime/Core/Private/IO/IoChunkEncoding.cpp b/Engine/Source/Runtime/Core/Private/IO/IoChunkEncoding.cpp
--- a/Engine/Source/Runtime/Core/Private/IO/IoChunkEncoding.cpp
+++ b/Engine/Source/Runtime/Core/Private/IO/IoChunkEncoding.cpp
@@ -361,3 +361,25 @@ uint64 FIoChunkEncoding::GetTotalEncodedSize(TConstArrayView<uint32> EncodedBloc
 
 	return TotalEncodedSize;
 }
+
+FIoStatus FIoChunkEncoding::HashBlocks(const FIoChunkEncoding::FHeader& Header, FMemoryView EncodedData, TArray<FIoHash>& OutHashes)
+{
+	TConstArrayView<uint32> Blocks = Header.GetBlocks();
+	if (Blocks.IsEmpty())
+	{
+		return FIoStatus(EIoErrorCode::InvalidParameter);
+	}
+
+	OutHashes.Reserve(Blocks.Num());
+
+	for (uint32 BlockSize : Blocks)
+	{
+		check(!EncodedData.IsEmpty());
+		OutHashes.Add(FIoHash::HashBuffer(EncodedData.Left(BlockSize)));
+		EncodedData += BlockSize;
+	}
+
+	check(EncodedData.IsEmpty());
+	
+	return FIoStatus::Ok;
+}
diff --git a/Engine/Source/Runtime/Core/Public/IO/IoChunkEncoding.h b/Engine/Source/Runtime/Core/Public/IO/IoChunkEncoding.h
--- a/Engine/Source/Runtime/Core/Public/IO/IoChunkEncoding.h
+++ b/Engine/Source/Runtime/Core/Public/IO/IoChunkEncoding.h
@@ -8,6 +8,7 @@
 #include "UObject/NameTypes.h"
 
 class FIoBuffer;
+struct FIoHash;
 
 /** I/O chunk encryption method. */
 enum class EIoEncryptionMethod : uint8
@@ -76,4 +77,5 @@ public:
 	static TIoStatusOr<FIoOffsetAndLength> GetChunkRange(uint64 TotalRawSize, uint32 RawBlockSize, TConstArrayView<uint32> EncodedBlockSize, uint64 RawOffset, uint64 RawSize);
 	static TIoStatusOr<FIoOffsetAndLength> GetChunkRange(const FIoChunkDecodingParams& Params, uint64 RawSize);
 	static uint64 GetTotalEncodedSize(TConstArrayView<uint32> EncodedBlockSize);
+	static FIoStatus HashBlocks(const FHeader& Header, FMemoryView EncodedData, TArray<FIoHash>& OutHashes);
 };
diff --git a/Engine/Source/Runtime/Experimental/IoStoreOnDemand/Private/IoStoreOnDemand.cpp b/Engine/Source/Runtime/Experimental/IoStoreOnDemand/Private/IoStoreOnDemand.cpp
--- a/Engine/Source/Runtime/Experimental/IoStoreOnDemand/Private/IoStoreOnDemand.cpp
+++ b/Engine/Source/Runtime/Experimental/IoStoreOnDemand/Private/IoStoreOnDemand.cpp
@@ -57,6 +57,7 @@ FCbWriter& operator<<(FCbWriter& Writer, const FOnDemandTocEntry& Entry)
 {
 	Writer.BeginObject();
 	Writer.AddHash(UTF8TEXTVIEW("Hash"), Entry.Hash);
+	Writer.AddHash(UTF8TEXTVIEW("RawHash"), Entry.RawHash);
 	Writer << UTF8TEXTVIEW("ChunkId") << Entry.ChunkId;
 	Writer.AddInteger(UTF8TEXTVIEW("RawSize"), Entry.RawSize);
 	Writer.AddInteger(UTF8TEXTVIEW("EncodedSize"), Entry.EncodedSize);
@@ -77,6 +78,7 @@ bool LoadFromCompactBinary(FCbFieldView Field, FOnDemandTocEntry& OutTocEntry)
 		}
 
 		OutTocEntry.Hash = Obj["Hash"].AsHash();
+		OutTocEntry.RawHash = Obj["RawHash"].AsHash();
 		OutTocEntry.RawSize = Obj["RawSize"].AsUInt64(~uint64(0));
 		OutTocEntry.EncodedSize = Obj["EncodedSize"].AsUInt64(~uint64(0));
 		OutTocEntry.BlockOffset = Obj["BlockOffset"].AsUInt32(~uint32(0));
@@ -111,6 +113,13 @@ FCbWriter& operator<<(FCbWriter& Writer, const FOnDemandTocContainerEntry& Conta
 	}
 	Writer.EndArray();
 
+	Writer.BeginArray(UTF8TEXTVIEW("BlockHashes"));
+	for (const FIoHash& BlockHash : ContainerEntry.BlockHashes)
+	{
+		Writer << BlockHash;
+	}
+	Writer.EndArray();
+
 	Writer.EndObject();
 
 	return Writer;
@@ -140,6 +149,13 @@ bool LoadFromCompactBinary(FCbFieldView Field, FOnDemandTocContainerEntry& OutCo
 			OutContainer.BlockSizes.Add(ArrayField.AsUInt32());
 		}
 
+		FCbArrayView BlockHashes = Obj["BlockHashes"].AsArrayView();
+		OutContainer.BlockHashes.Reserve(int32(BlockHashes.Num()));
+		for (FCbFieldView ArrayField : BlockHashes)
+		{
+			OutContainer.BlockHashes.Add(ArrayField.AsHash());
+		}
+
 		return true;
 	}
 
diff --git a/Engine/Source/Runtime/Experimental/IoStoreOnDemand/Private/OnDemandIoStoreWriter.cpp b/Engine/Source/Runtime/Experimental/IoStoreOnDemand/Private/OnDemandIoStoreWriter.cpp
--- a/Engine/Source/Runtime/Experimental/IoStoreOnDemand/Private/OnDemandIoStoreWriter.cpp
+++ b/Engine/Source/Runtime/Experimental/IoStoreOnDemand/Private/OnDemandIoStoreWriter.cpp
@@ -27,7 +27,9 @@ struct FPendingWrite
 	FIoWriteOptions WriteOptions;
 	FIoBuffer ChunkBuffer;
 	FIoBuffer ChunkHeader;
+	FIoHash ChunkRawHash;
 	FIoHash ChunkHash;
+	TArray<FIoHash> BlockHashes;
 	FString ErrorText;
 	uint64 RawSize{0};
 	uint64 EncodedSize{0};
@@ -152,6 +154,7 @@ void FOnDemandIoStoreWriter::Flush()
 			}
 
 			FOnDemandTocEntry& TocEntry = ContainerEntry.Entries.AddDefaulted_GetRef();
+			TocEntry.RawHash = WriteEntry.ChunkRawHash;
 			TocEntry.Hash = WriteEntry.ChunkHash;
 			TocEntry.ChunkId = WriteEntry.ChunkId;
 			TocEntry.RawSize = WriteEntry.RawSize;
@@ -161,8 +164,12 @@ void FOnDemandIoStoreWriter::Flush()
 			const FIoChunkEncoding::FHeader* Header = FIoChunkEncoding::FHeader::Decode(WriteEntry.ChunkHeader.GetView());
 			check(Header);
 			TConstArrayView<uint32> Blocks = Header->GetBlocks();
+			TArray<FIoHash> BlockHashes = MoveTemp(WriteEntry.BlockHashes);
+			check(Blocks.Num() == BlockHashes.Num());
+
 			TocEntry.BlockCount = Blocks.Num();
 			ContainerEntry.BlockSizes.Append(Blocks);
+			ContainerEntry.BlockHashes.Append(BlockHashes);
 
 			WriteResult.UncompressedContainerSize += WriteEntry.RawSize;
 			WriteResult.CompressedContainerSize += WriteEntry.EncodedSize;
@@ -229,7 +236,7 @@ void FOnDemandIoStoreWriter::Append(
 						{
 							PendingWrite->ChunkBuffer = *SourceBuffer;
 							PendingWrite->ChunkBuffer.EnsureOwned();
-							PendingWrite->ChunkHash = FIoHash::HashBuffer(PendingWrite->ChunkBuffer.GetView());
+							PendingWrite->ChunkRawHash = FIoHash::HashBuffer(PendingWrite->ChunkBuffer.GetView());
 							PendingWrite->RawSize = PendingWrite->ChunkBuffer.GetSize();
 
 							PendingWrite->WriteRequest->FreeSourceBuffer();
@@ -244,6 +251,14 @@ void FOnDemandIoStoreWriter::Append(
 							{
 								PendingWrite->EncodedSize = PendingWrite->ChunkBuffer.GetSize();
 								PendingWrite->ChunkHash = FIoHash::HashBuffer(PendingWrite->ChunkBuffer.GetView());
+
+								if (FIoStatus Status = FIoChunkEncoding::HashBlocks(
+									*FIoChunkEncoding::FHeader::Decode(PendingWrite->ChunkHeader.GetView()),
+									PendingWrite->ChunkBuffer.GetView(),
+									PendingWrite->BlockHashes); !Status.IsOk())
+								{
+									PendingWrite->ErrorText = FString::Printf(TEXT("Failed to hash encoded blocks '%s'"), *PendingWrite->WriteOptions.FileName);
+								}
 							}
 							else
 							{
diff --git a/Engine/Source/Runtime/Experimental/IoStoreOnDemand/Public/IO/IoStoreOnDemand.h b/Engine/Source/Runtime/Experimental/IoStoreOnDemand/Public/IO/IoStoreOnDemand.h
--- a/Engine/Source/Runtime/Experimental/IoStoreOnDemand/Public/IO/IoStoreOnDemand.h
+++ b/Engine/Source/Runtime/Experimental/IoStoreOnDemand/Public/IO/IoStoreOnDemand.h
@@ -61,6 +61,7 @@ UE_API bool LoadFromCompactBinary(FCbFieldView Field, FOnDemandTocHeader& OutToc
 struct FOnDemandTocEntry
 {
 	FIoHash Hash = FIoHash::Zero;
+	FIoHash RawHash = FIoHash::Zero;
 	FIoChunkId ChunkId = FIoChunkId::InvalidChunkId;
 	uint64 RawSize = 0;
 	uint64 EncodedSize = 0;
@@ -78,6 +79,7 @@ struct FOnDemandTocContainerEntry
 	FString EncryptionKeyGuid;
 	TArray<FOnDemandTocEntry> Entries;
 	TArray<uint32> BlockSizes;
+	TArray<FIoHash> BlockHashes;
 
 	UE_API friend FCbWriter& operator<<(FCbWriter& Writer, const FOnDemandTocContainerEntry& ContainerEntry);
 };
