Iris
- Fixed issue with new objects that was initially culled out as not relevant not copying and quantizing the full initial state

diff --git a/Engine/Source/Runtime/Experimental/Iris/Core/Private/Iris/ReplicationSystem/NetRefHandleManager.cpp b/Engine/Source/Runtime/Experimental/Iris/Core/Private/Iris/ReplicationSystem/NetRefHandleManager.cpp
--- a/Engine/Source/Runtime/Experimental/Iris/Core/Private/Iris/ReplicationSystem/NetRefHandleManager.cpp
+++ b/Engine/Source/Runtime/Experimental/Iris/Core/Private/Iris/ReplicationSystem/NetRefHandleManager.cpp
@@ -116,6 +116,9 @@ FInternalNetRefIndex FNetRefHandleManager::InternalCreateNetObject(const FNetRef
 		// When a handle is first created, it is not set to be a subobject
 		SubObjectInternalIndices.ClearBit(InternalIndex);
 
+		// Need a full copy if set, normally only needed for new objects.
+		Data.bNeedsFullCopyAndQuantize = 1U;
+
 		ReplicatedObjectRefCount[InternalIndex] = 0;
 
 		return InternalIndex;
diff --git a/Engine/Source/Runtime/Experimental/Iris/Core/Private/Iris/ReplicationSystem/NetRefHandleManager.h b/Engine/Source/Runtime/Experimental/Iris/Core/Private/Iris/ReplicationSystem/NetRefHandleManager.h
--- a/Engine/Source/Runtime/Experimental/Iris/Core/Private/Iris/ReplicationSystem/NetRefHandleManager.h
+++ b/Engine/Source/Runtime/Experimental/Iris/Core/Private/Iris/ReplicationSystem/NetRefHandleManager.h
@@ -96,8 +96,9 @@ public:
 				uint32 bIsDependentObject : 1U;
 				uint32 bHasDependentObjects : 1U;
 				uint32 bAllowDestroyInstanceFromRemote : 1U;
+				uint32 bNeedsFullCopyAndQuantize : 1U;
 				// Padding. Adjust when adding or removing flags.
-				uint32 Padding : 26U;
+				uint32 Padding : 25U;
 			};
 		};
 	
diff --git a/Engine/Source/Runtime/Experimental/Iris/Core/Private/Iris/ReplicationSystem/ReplicationOperationsInternal.cpp b/Engine/Source/Runtime/Experimental/Iris/Core/Private/Iris/ReplicationSystem/ReplicationOperationsInternal.cpp
--- a/Engine/Source/Runtime/Experimental/Iris/Core/Private/Iris/ReplicationSystem/ReplicationOperationsInternal.cpp
+++ b/Engine/Source/Runtime/Experimental/Iris/Core/Private/Iris/ReplicationSystem/ReplicationOperationsInternal.cpp
@@ -72,7 +72,7 @@ uint32 FReplicationInstanceOperationsInternal::CopyObjectStateData(FNetBitStream
 {
 	if (NetRefHandleManager.IsScopableIndex(InternalIndex))
 	{
-		const FNetRefHandleManager::FReplicatedObjectData& Object = NetRefHandleManager.GetReplicatedObjectDataNoCheck(InternalIndex);
+		FNetRefHandleManager::FReplicatedObjectData& Object = NetRefHandleManager.GetReplicatedObjectDataNoCheck(InternalIndex);
 
 		// We cannot copy state data for zero sized objects or objects that no longer has an instance protocol.
 		if (Object.InstanceProtocol && Object.Protocol->InternalTotalSize > 0U)
@@ -82,8 +82,9 @@ uint32 FReplicationInstanceOperationsInternal::CopyObjectStateData(FNetBitStream
 			// if the object was scopable prev frame we can do partial copy
 			bool bShouldPropagateChangedStates = Object.bShouldPropagateChangedStates;
 
-			// We do a full CopyAndQunatize if the cvar bCVarForceFullCopyAndQuantize is set or that the object is new
-			const bool bUseFullCopyAndQuantize = bCVarForceFullCopyAndQuantize || !NetRefHandleManager.GetPrevFrameScopableInternalIndices().GetBit(InternalIndex);
+			// We do a full CopyAndQunatize if the cvar bCVarForceFullCopyAndQuantize is set or that the object is marked as needing a FullCopyAndQuantize
+			const bool bUseFullCopyAndQuantize = bCVarForceFullCopyAndQuantize || Object.bNeedsFullCopyAndQuantize;
+			Object.bNeedsFullCopyAndQuantize = 0U;
 
 			// Copy dirty state
 			{
