Optimize removal update code to avoid evaluating Gettransform at each step of the loop
also remove the overse transform if scaling is not enabled

diff --git a/Engine/Source/Runtime/Experimental/GeometryCollectionEngine/Private/GeometryCollection/GeometryCollectionComponent.cpp b/Engine/Source/Runtime/Experimental/GeometryCollectionEngine/Private/GeometryCollection/GeometryCollectionComponent.cpp
--- a/Engine/Source/Runtime/Experimental/GeometryCollectionEngine/Private/GeometryCollection/GeometryCollectionComponent.cpp
+++ b/Engine/Source/Runtime/Experimental/GeometryCollectionEngine/Private/GeometryCollection/GeometryCollectionComponent.cpp
@@ -3255,17 +3255,19 @@ void UGeometryCollectionComponent::UpdateRemovalIfNeeded()
 		{
 			const FTransform ZeroScaleTransform(FQuat::Identity, FVector::Zero(), FVector(0, 0, 0));
 
-			const FTransform InverseComponentTransform = GetComponentTransform().Inverse();
-			for (int32 Idx = 0; Idx < GetTransformArray().Num(); ++Idx)
+			const FTransform InverseComponentTransform = (RestCollection->bScaleOnRemoval) ? GetComponentTransform().Inverse() : FTransform::Identity;
+
+			const int32 NumTransforms = DecayFacade.GetDecayAttributeSize();
+			for (int32 TransformIndex = 0; TransformIndex < NumTransforms; ++TransformIndex)
 			{
 				// only update values if the decay has changed 
-				const float Decay = DecayFacade.GetDecay(Idx);
+				const float Decay = DecayFacade.GetDecay(TransformIndex);
 				if (Decay > 0.f && Decay <= 1.f)
 				{
 					const float Scale = 1.0 - Decay;
 					if (Scale < UE_SMALL_NUMBER)
 					{
-						DynamicCollection->Transform[Idx].SetScale3D(FVector::ZeroVector);
+						DynamicCollection->Transform[TransformIndex].SetScale3D(FVector::ZeroVector);
 					}
 					// do not try to get this condition out of the loop as this may cause some optimizer related issues
 					else if (RestCollection->bScaleOnRemoval) 
@@ -3273,17 +3275,17 @@ void UGeometryCollectionComponent::UpdateRemovalIfNeeded()
 						float ShrinkRadius = 0.0f;
 						UE::Math::TSphere<double> AccumulatedSphere;
 						// todo(chaos) : find a faster way to do that ( precompute the data ? )
-						if (CalculateInnerSphere(Idx, AccumulatedSphere))
+						if (CalculateInnerSphere(TransformIndex, AccumulatedSphere))
 						{
 							ShrinkRadius = -AccumulatedSphere.W;
 						}
 
-						const FQuat LocalRotation = (InverseComponentTransform * ComponentSpaceTransforms[Idx].Inverse()).GetRotation();
+						const FQuat LocalRotation = (InverseComponentTransform * ComponentSpaceTransforms[TransformIndex].Inverse()).GetRotation();
 						const FVector LocalDown = LocalRotation.RotateVector(FVector(0.f, 0.f, ShrinkRadius));
-						const FVector CenterOfMass = DynamicCollection->MassToLocal[Idx].GetTranslation();
+						const FVector CenterOfMass = DynamicCollection->MassToLocal[TransformIndex].GetTranslation();
 						const FVector ScaleCenter = LocalDown + CenterOfMass;
 						const FTransform ScaleTransform(FQuat::Identity, ScaleCenter * FVector::FReal(1.f - Scale), FVector(Scale));
-						DynamicCollection->Transform[Idx] = ScaleTransform * DynamicCollection->Transform[Idx];
+						DynamicCollection->Transform[TransformIndex] = ScaleTransform * DynamicCollection->Transform[TransformIndex];
 					}
 				}
 			}
diff --git a/Engine/Source/Runtime/Experimental/GeometryCollectionEngine/Private/GeometryCollection/GeometryCollectionEngineRemoval.cpp b/Engine/Source/Runtime/Experimental/GeometryCollectionEngine/Private/GeometryCollection/GeometryCollectionEngineRemoval.cpp
--- a/Engine/Source/Runtime/Experimental/GeometryCollectionEngine/Private/GeometryCollection/GeometryCollectionEngineRemoval.cpp
+++ b/Engine/Source/Runtime/Experimental/GeometryCollectionEngine/Private/GeometryCollection/GeometryCollectionEngineRemoval.cpp
@@ -240,3 +240,8 @@ void FGeometryCollectionDecayDynamicFacade::SetDecay(int32 TransformIndex, float
 {
 	DecayAttribute.ModifyAt(TransformIndex, DecayValue);
 }
+
+int32 FGeometryCollectionDecayDynamicFacade::GetDecayAttributeSize() const
+{
+	return DecayAttribute.Num();
+}
\ No newline at end of file
diff --git a/Engine/Source/Runtime/Experimental/GeometryCollectionEngine/Public/GeometryCollection/GeometryCollectionEngineRemoval.h b/Engine/Source/Runtime/Experimental/GeometryCollectionEngine/Public/GeometryCollection/GeometryCollectionEngineRemoval.h
--- a/Engine/Source/Runtime/Experimental/GeometryCollectionEngine/Public/GeometryCollection/GeometryCollectionEngineRemoval.h
+++ b/Engine/Source/Runtime/Experimental/GeometryCollectionEngine/Public/GeometryCollection/GeometryCollectionEngineRemoval.h
@@ -159,6 +159,9 @@ public:
 	/** Set decay value for a specific transform index */
 	void SetDecay(int32 TransformIndex, float DecayValue);
 
+	/** Get the size of the decay attribute - this should match the number of transforms of the collection */
+	int32 GetDecayAttributeSize() const;
+
 private:
 	/** state of decay ([0-1] range) */
 	TManagedArrayAccessor<float> DecayAttribute;
