add convex decomposition options to leaf node convex hull dataflow node
+ add a 'max hulls' option to convex decomposition, which functions like 'target num hulls' but can still cap the number of hulls when an error tolerance is specified (target num hulls is ignored when error tolerance is set)

diff --git a/Engine/Plugins/Experimental/GeometryCollectionPlugin/Source/GeometryCollectionNodes/Private/Dataflow/GeometryCollectionUtilityNodes.cpp b/Engine/Plugins/Experimental/GeometryCollectionPlugin/Source/GeometryCollectionNodes/Private/Dataflow/GeometryCollectionUtilityNodes.cpp
--- a/Engine/Plugins/Experimental/GeometryCollectionPlugin/Source/GeometryCollectionNodes/Private/Dataflow/GeometryCollectionUtilityNodes.cpp
+++ b/Engine/Plugins/Experimental/GeometryCollectionPlugin/Source/GeometryCollectionNodes/Private/Dataflow/GeometryCollectionUtilityNodes.cpp
@@ -23,6 +23,7 @@ namespace Dataflow
 	{
 		static const FLinearColor CDefaultNodeBodyTintColor = FLinearColor(0.f, 0.f, 0.f, 0.5f);
 
+		DATAFLOW_NODE_REGISTER_CREATION_FACTORY(FMakeDataflowConvexDecompositionSettingsNode);
 		DATAFLOW_NODE_REGISTER_CREATION_FACTORY(FCreateLeafConvexHullsDataflowNode);
 		DATAFLOW_NODE_REGISTER_CREATION_FACTORY(FSimplifyConvexHullsDataflowNode);
 		DATAFLOW_NODE_REGISTER_CREATION_FACTORY(FCreateNonOverlappingConvexHullsDataflowNode);
@@ -33,18 +34,46 @@ namespace Dataflow
 	}
 }
 
+FMakeDataflowConvexDecompositionSettingsNode::FMakeDataflowConvexDecompositionSettingsNode(const Dataflow::FNodeParameters& InParam, FGuid InGuid)
+	: FDataflowNode(InParam, InGuid)
+{
+	RegisterInputConnection(&MinSizeToDecompose);
+	RegisterInputConnection(&MaxGeoToHullVolumeRatioToDecompose);
+	RegisterInputConnection(&ErrorTolerance);
+	RegisterInputConnection(&MaxHullsPerGeometry);
+	RegisterInputConnection(&MinThicknessTolerance);
+	RegisterInputConnection(&NumAdditionalSplits);
+	RegisterOutputConnection(&DecompositionSettings);
+}
+
+void FMakeDataflowConvexDecompositionSettingsNode::Evaluate(Dataflow::FContext& Context, const FDataflowOutput* Out) const
+{
+	if (Out->IsA(&DecompositionSettings))
+	{
+		FDataflowConvexDecompositionSettings OutSettings;
+		OutSettings.MinSizeToDecompose = GetValue(Context, &MinSizeToDecompose);
+		OutSettings.MaxGeoToHullVolumeRatioToDecompose = GetValue(Context, &MaxGeoToHullVolumeRatioToDecompose);
+		OutSettings.ErrorTolerance = GetValue(Context, &ErrorTolerance);
+		OutSettings.MaxHullsPerGeometry = GetValue(Context, &MaxHullsPerGeometry);
+		OutSettings.MinThicknessTolerance = GetValue(Context, &MinThicknessTolerance);
+		OutSettings.NumAdditionalSplits = GetValue(Context, &NumAdditionalSplits);
+		SetValue(Context, OutSettings, &DecompositionSettings);
+	}
+}
+
 FCreateLeafConvexHullsDataflowNode::FCreateLeafConvexHullsDataflowNode(const Dataflow::FNodeParameters& InParam, FGuid InGuid)
 	: FDataflowNode(InParam, InGuid)
 {
 	RegisterInputConnection(&Collection);
 	RegisterInputConnection(&OptionalSelectionFilter);
 	RegisterInputConnection(&SimplificationDistanceThreshold);
+	RegisterInputConnection(&ConvexDecompositionSettings);
 	RegisterOutputConnection(&Collection);
 }
 
 void FCreateLeafConvexHullsDataflowNode::Evaluate(Dataflow::FContext& Context, const FDataflowOutput* Out) const
 {
-	if (Out->IsA(&Collection) && IsConnected(&Collection))
+	if (Out->IsA(&Collection))
 	{
 		const FManagedArrayCollection& InCollection = GetValue(Context, &Collection);
 		if (InCollection.NumElements(FGeometryCollection::TransformGroup) == 0)
@@ -65,10 +94,18 @@ void FCreateLeafConvexHullsDataflowNode::Evaluate(Dataflow::FContext& Context, c
 			}
 
 			float InSimplificationDistanceThreshold = GetValue(Context, &SimplificationDistanceThreshold);
-			FGeometryCollectionConvexUtility::FIntersectionFilters IntersectionFilters;
-			IntersectionFilters.OnlyIntersectIfComputedIsSmallerFactor = IntersectIfComputedIsSmallerByFactor;
-			IntersectionFilters.MinExternalVolumeToIntersect = MinExternalVolumeToIntersect;
-			FGeometryCollectionConvexUtility::GenerateLeafConvexHulls(*GeomCollection, bRestrictToSelection, SelectedBones, InSimplificationDistanceThreshold, GenerateMethod, IntersectionFilters);
+
+			FGeometryCollectionConvexUtility::FLeafConvexHullSettings LeafSettings(InSimplificationDistanceThreshold, GenerateMethod);
+			LeafSettings.IntersectFilters.OnlyIntersectIfComputedIsSmallerFactor = IntersectIfComputedIsSmallerByFactor;
+			LeafSettings.IntersectFilters.MinExternalVolumeToIntersect = MinExternalVolumeToIntersect;
+			FDataflowConvexDecompositionSettings InDecompSettings = GetValue(Context, &ConvexDecompositionSettings);
+			LeafSettings.DecompositionSettings.MaxGeoToHullVolumeRatioToDecompose = InDecompSettings.MaxGeoToHullVolumeRatioToDecompose;
+			LeafSettings.DecompositionSettings.MinGeoVolumeToDecompose = InDecompSettings.MinSizeToDecompose * InDecompSettings.MinSizeToDecompose * InDecompSettings.MinSizeToDecompose;
+			LeafSettings.DecompositionSettings.ErrorTolerance = InDecompSettings.ErrorTolerance;
+			LeafSettings.DecompositionSettings.MaxHullsPerGeometry = InDecompSettings.MaxHullsPerGeometry;
+			LeafSettings.DecompositionSettings.MinThicknessTolerance = InDecompSettings.MinThicknessTolerance;
+			LeafSettings.DecompositionSettings.NumAdditionalSplits = InDecompSettings.NumAdditionalSplits;
+			FGeometryCollectionConvexUtility::GenerateLeafConvexHulls(*GeomCollection, bRestrictToSelection, SelectedBones, LeafSettings);
 			SetValue(Context, (const FManagedArrayCollection&)(*GeomCollection), &Collection);
 		}
 	}
diff --git a/Engine/Plugins/Experimental/GeometryCollectionPlugin/Source/GeometryCollectionNodes/Public/Dataflow/GeometryCollectionUtilityNodes.h b/Engine/Plugins/Experimental/GeometryCollectionPlugin/Source/GeometryCollectionNodes/Public/Dataflow/GeometryCollectionUtilityNodes.h
--- a/Engine/Plugins/Experimental/GeometryCollectionPlugin/Source/GeometryCollectionNodes/Public/Dataflow/GeometryCollectionUtilityNodes.h
+++ b/Engine/Plugins/Experimental/GeometryCollectionPlugin/Source/GeometryCollectionNodes/Public/Dataflow/GeometryCollectionUtilityNodes.h
@@ -23,6 +23,82 @@ enum class EConvexOverlapRemovalMethodEnum : uint8
 	Dataflow_Max                UMETA(Hidden)
 };
 
+USTRUCT()
+struct FDataflowConvexDecompositionSettings
+{
+	GENERATED_USTRUCT_BODY()
+
+public:
+	// If greater than zero, the minimum geometry size (cube root of volume) to consider for convex decomposition
+	UPROPERTY(EditAnywhere, Category = Filter, meta = (ClampMin = 0.0))
+	float MinSizeToDecompose = 0.f;
+
+	// If the geo volume / hull volume ratio is greater than this, do not consider convex decomposition
+	UPROPERTY(EditAnywhere, Category = Filter, meta = (ClampMin = 0.0, ClampMax = 1.0))
+	float MaxGeoToHullVolumeRatioToDecompose = 1.f;
+
+	// Stop splitting when hulls have error less than this (expressed in cm; will be cubed for volumetric error).
+	// Note: ErrorTolerance must be > 0 or MaxHullsPerGeometry > 1, or decomposition will not be performed.
+	UPROPERTY(EditAnywhere, Category = Decomposition, meta = (ClampMin = 0.0))
+	float ErrorTolerance = 0.f;
+
+	// If greater than zero, maximum number of convex hulls to use in each convex decomposition.
+	//Note: ErrorTolerance must be > 0 or MaxHullsPerGeometry > 1, or decomposition will not be performed.
+	UPROPERTY(EditAnywhere, Category = Decomposition, meta = (ClampMin = -1))
+	int32 MaxHullsPerGeometry = -1;
+
+	// Optionally specify a minimum thickness (in cm) for convex parts; parts below this thickness will always be merged away. Overrides NumOutputHulls and ErrorTolerance when needed.
+	UPROPERTY(EditAnywhere, Category = Decomposition, meta = (ClampMin = 0.0))
+	float MinThicknessTolerance = 0.f;
+
+	// Control the search effort spent per convex decomposition: larger values will require more computation but may find better convex decompositions
+	UPROPERTY(EditAnywhere, Category = Decomposition, meta = (ClampMin = 0))
+	int32 NumAdditionalSplits = 4;
+};
+
+
+//~ TODO: Ideally this would be generated from the above FDataflowConvexDecompositionSettings struct
+// Provide settings for running convex decomposition of geometry
+USTRUCT(meta = (DataflowGeometryCollection))
+struct FMakeDataflowConvexDecompositionSettingsNode : public FDataflowNode
+{
+	GENERATED_USTRUCT_BODY()
+	DATAFLOW_NODE_DEFINE_INTERNAL(FMakeDataflowConvexDecompositionSettingsNode, "MakeConvexDecompositionSettings", "GeometryCollection|Utilities", "")
+
+public:
+	// If greater than zero, the minimum geometry size (cube root of volume) to consider for convex decomposition
+	UPROPERTY(EditAnywhere, Category = Filter, meta = (ClampMin = 0.0, DataflowInput))
+	float MinSizeToDecompose = 0.f;
+	
+	// If the geo volume / hull volume ratio is greater than this, do not consider convex decomposition
+	UPROPERTY(EditAnywhere, Category = Filter, meta = (ClampMin = 0.0, ClampMax = 1.0, DataflowInput))
+	float MaxGeoToHullVolumeRatioToDecompose = 1.f;
+
+	// Stop splitting when hulls have error less than this (expressed in cm; will be cubed for volumetric error).
+	// Note: ErrorTolerance must be > 0 or MaxHullsPerGeometry > 1, or decomposition will not be performed.
+	UPROPERTY(EditAnywhere, Category = Decomposition, meta = (ClampMin = 0.0, DataflowInput))
+	float ErrorTolerance = 0.f;
+
+	// If greater than zero, maximum number of convex hulls to use in each convex decomposition.
+	//Note: ErrorTolerance must be > 0 or MaxHullsPerGeometry > 1, or decomposition will not be performed.
+	UPROPERTY(EditAnywhere, Category = Decomposition, meta = (ClampMin = -1, DataflowInput))
+	int32 MaxHullsPerGeometry = -1;
+
+	// Optionally specify a minimum thickness (in cm) for convex parts; parts below this thickness will always be merged away. Overrides NumOutputHulls and ErrorTolerance when needed.
+	UPROPERTY(EditAnywhere, Category = Decomposition, meta = (ClampMin = 0.0, DataflowInput))
+	float MinThicknessTolerance = 0.f;
+
+	// Control the search effort spent per convex decomposition: larger values will require more computation but may find better convex decompositions
+	UPROPERTY(EditAnywhere, Category = Decomposition, meta = (ClampMin = 0, DataflowInput))
+	int32 NumAdditionalSplits = 4;
+
+	UPROPERTY(meta = (DataflowOutput))
+	FDataflowConvexDecompositionSettings DecompositionSettings;
+
+	FMakeDataflowConvexDecompositionSettingsNode(const Dataflow::FNodeParameters& InParam, FGuid InGuid = FGuid::NewGuid());
+
+	virtual void Evaluate(Dataflow::FContext& Context, const FDataflowOutput* Out) const override;
+};
 
 USTRUCT(meta = (DataflowGeometryCollection))
 struct FCreateLeafConvexHullsDataflowNode : public FDataflowNode
@@ -51,9 +127,12 @@ public:
 	float MinExternalVolumeToIntersect = 0.0f;
 
 	/** Computed convex hulls are simplified to keep points spaced at least this far apart (except where needed to keep the hull from collapsing to zero volume) */
-	UPROPERTY(EditAnywhere, Category = "Convex", meta = (DataflowInput, UIMin = 0.f, EditCondition = "GenerateMethod != EGenerateConvexMethod::ExternalCollision"))
+	UPROPERTY(EditAnywhere, Category = Convex, meta = (DataflowInput, ClampMin = 0.f, EditCondition = "GenerateMethod != EGenerateConvexMethod::ExternalCollision"))
 	float SimplificationDistanceThreshold = 10.f;
 
+	UPROPERTY(EditAnywhere, Category = Convex, meta = (DataflowInput))
+	FDataflowConvexDecompositionSettings ConvexDecompositionSettings;
+
 	FCreateLeafConvexHullsDataflowNode(const Dataflow::FNodeParameters& InParam, FGuid InGuid = FGuid::NewGuid());
 
 	virtual void Evaluate(Dataflow::FContext& Context, const FDataflowOutput* Out) const override;
diff --git a/Engine/Source/Runtime/Experimental/Chaos/Private/GeometryCollection/GeometryCollectionConvexUtility.cpp b/Engine/Source/Runtime/Experimental/Chaos/Private/GeometryCollection/GeometryCollectionConvexUtility.cpp
--- a/Engine/Source/Runtime/Experimental/Chaos/Private/GeometryCollection/GeometryCollectionConvexUtility.cpp
+++ b/Engine/Source/Runtime/Experimental/Chaos/Private/GeometryCollection/GeometryCollectionConvexUtility.cpp
@@ -1161,6 +1161,55 @@
 }
 
 
+// helper to compute the volume of an individual piece of geometry
+double ComputeGeometryVolume(
+	const FManagedArrayCollection* Collection,
+	int32 GeometryIdx,
+	const FTransform& GlobalTransform,
+	double ScalePerDimension
+)
+{
+	GeometryCollection::Facades::FCollectionMeshFacade MeshFacade(*Collection);
+	const TManagedArray<int32>& VertexStart = MeshFacade.VertexStartAttribute.Get();
+	const TManagedArray<int32>& VertexCount = MeshFacade.VertexCountAttribute.Get();
+	const TManagedArray<FVector3f>& Vertex = MeshFacade.VertexAttribute.Get();
+	const TManagedArray<int32>& FaceStart = MeshFacade.FaceStartAttribute.Get();
+	const TManagedArray<int32>& FaceCount = MeshFacade.FaceCountAttribute.Get();
+	const TManagedArray<FIntVector>& Indices = MeshFacade.IndicesAttribute.Get();
+	int32 VStart = VertexStart[GeometryIdx];
+	int32 VEnd = VStart + VertexCount[GeometryIdx];
+	if (VStart == VEnd)
+	{
+		return 0.0;
+	}
+	FVector3d Center = FVector::ZeroVector;
+	for (int32 VIdx = VStart; VIdx < VEnd; VIdx++)
+	{
+		FVector Pos = GlobalTransform.TransformPosition((FVector)Vertex[VIdx]);
+		Center += (FVector3d)Pos;
+	}
+	Center /= double(VEnd - VStart);
+	int32 FStart = FaceStart[GeometryIdx];
+	int32 FEnd = FStart + FaceCount[GeometryIdx];
+	double VolOut = 0;
+	for (int32 FIdx = FStart; FIdx < FEnd; FIdx++)
+	{
+		FIntVector Tri = Indices[FIdx];
+		FVector3d V0 = (FVector3d)GlobalTransform.TransformPosition((FVector)Vertex[Tri.X]);
+		FVector3d V1 = (FVector3d)GlobalTransform.TransformPosition((FVector)Vertex[Tri.Y]);
+		FVector3d V2 = (FVector3d)GlobalTransform.TransformPosition((FVector)Vertex[Tri.Z]);
+
+		// add volume of the tetrahedron formed by the triangles and the reference point
+		FVector3d V1mRef = (V1 - Center) * ScalePerDimension;
+		FVector3d V2mRef = (V2 - Center) * ScalePerDimension;
+		FVector3d N = V2mRef.Cross(V1mRef);
+
+		VolOut += ((V0 - Center) * ScalePerDimension).Dot(N) / 6.0;
+	}
+	return VolOut;
+}
+
+
 /// Helper to get convex hulls from a geometry collection in the format required by CreateNonoverlappingConvexHulls
 void HullsFromGeometry(
 	FGeometryCollection& Geometry,
@@ -1173,7 +1222,8 @@
 	int32 RigidType,
 	double SimplificationDistanceThreshold,
 	double OverlapRemovalShrinkPercent,
-	TFunction<bool(int32)> SkipBoneFn = nullptr
+	TFunction<bool(int32)> SkipBoneFn = nullptr,
+	const FGeometryCollectionConvexUtility::FConvexDecompositionSettings* OptionalDecompositionSettings = nullptr
 )
 {
 	TArray<FVector> GlobalVertices;
@@ -1226,22 +1276,98 @@
 		}
 		else if (SimulationType[Idx] == RigidType && GeomIdx != INDEX_NONE)
 		{
-			int32 VStart = Geometry.VertexStart[GeomIdx];
-			int32 VCount = Geometry.VertexCount[GeomIdx];
-			int32 VEnd = VStart + VCount;
-			TArray<Chaos::FConvex::FVec3Type> HullPts;
-			HullPts.Reserve(VCount);
-			for (int32 VIdx = VStart; VIdx < VEnd; VIdx++)
+			auto ComputeHull = [&Geometry, &GlobalVertices, SimplificationDistanceThreshold, OverlapRemovalShrinkPercent, GeomIdx](FVector& PivotOut) -> TUniquePtr<::Chaos::FConvex>
 			{
-				HullPts.Add(GlobalVertices[VIdx]);
+				int32 VStart = Geometry.VertexStart[GeomIdx];
+				int32 VCount = Geometry.VertexCount[GeomIdx];
+				int32 VEnd = VStart + VCount;
+				TArray<Chaos::FConvex::FVec3Type> HullPts;
+				HullPts.Reserve(VCount);
+				for (int32 VIdx = VStart; VIdx < VEnd; VIdx++)
+				{
+					HullPts.Add(GlobalVertices[VIdx]);
+				}
+				ensure(HullPts.Num() > 0);
+				FilterHullPoints(HullPts, SimplificationDistanceThreshold);
+				PivotOut = ScaleHullPoints(HullPts, OverlapRemovalShrinkPercent);
+				return MakeUnique<Chaos::FConvex>(HullPts, UE_KINDA_SMALL_NUMBER);
+			};
+			TUniquePtr<Chaos::FConvex> Hull = nullptr;
+			FVector HullPivot;
+
+			if (OptionalDecompositionSettings)
+			{
+				bool bAttemptDecomposition = OptionalDecompositionSettings->MaxHullsPerGeometry > 1 || OptionalDecompositionSettings->ErrorTolerance > 0.0;
+				double InitialGeoVolume = 0.0, InitialHullVolume = 0.0;
+				if (bAttemptDecomposition && (OptionalDecompositionSettings->MinGeoVolumeToDecompose > 0.0 || OptionalDecompositionSettings->MaxGeoToHullVolumeRatioToDecompose < 1.0))
+				{
+					InitialGeoVolume = ComputeGeometryVolume(&Geometry, GeomIdx, GlobalTransformArray[Idx], 1.0);
+				}
+				if (bAttemptDecomposition && InitialGeoVolume < OptionalDecompositionSettings->MinGeoVolumeToDecompose)
+				{
+					bAttemptDecomposition = false; // too small to consider for decomposition
+				}
+				if (bAttemptDecomposition && OptionalDecompositionSettings->MaxGeoToHullVolumeRatioToDecompose < 1.0)
+				{
+					Hull = ComputeHull(HullPivot);
+					InitialHullVolume = (double)Hull->GetVolume();
+					// Hull was scaled down by the overlap removal shrink percent, so to compute a correct ratio we need to adjust the geo volume in the same way
+					double VolumeScale = ScaleFactor * ScaleFactor * ScaleFactor;
+					if ((VolumeScale * InitialGeoVolume) / InitialHullVolume >= OptionalDecompositionSettings->MaxGeoToHullVolumeRatioToDecompose)
+					{
+						bAttemptDecomposition = false;
+					}
+				}
+
+				if (bAttemptDecomposition)
+				{
+					UE::Geometry::FConvexDecomposition3 Decomposition;
+					int32 VertexStart = Geometry.VertexStart[GeomIdx];
+					TArrayView<const FVector3f> VerticesView(Geometry.Vertex.GetData() + VertexStart, Geometry.VertexCount[GeomIdx]);
+					TArrayView<const FIntVector3> FacesView(Geometry.Indices.GetData() + Geometry.FaceStart[GeomIdx], Geometry.FaceCount[GeomIdx]);
+					Decomposition.InitializeFromIndexMesh(VerticesView, FacesView, true, -VertexStart);
+					Decomposition.Compute(OptionalDecompositionSettings->MaxHullsPerGeometry, OptionalDecompositionSettings->NumAdditionalSplits, 
+						OptionalDecompositionSettings->ErrorTolerance, OptionalDecompositionSettings->MinThicknessTolerance, OptionalDecompositionSettings->MaxHullsPerGeometry);
+					int32 NumHulls = Decomposition.NumHulls();
+					if ((NumHulls > 0 && !Hull) || NumHulls > 1)
+					{
+						TArray<TUniquePtr<::Chaos::FConvex>> DecompHulls; DecompHulls.Reserve(NumHulls);
+						TArray<FVector> Pivots; Pivots.Reserve(NumHulls);
+						for (int32 HullIdx = 0; HullIdx < NumHulls; ++HullIdx)
+						{
+							TArray<FVector3f> OrigDecompVerts = Decomposition.GetVertices<float>(HullIdx);
+							// convert to the vector type expected by Chaos::FConvex
+							TArray<::Chaos::FConvex::FVec3Type> DecompVerts;
+							DecompVerts.SetNumUninitialized(OrigDecompVerts.Num());
+							for (int32 PtIdx = 0; PtIdx < OrigDecompVerts.Num(); ++PtIdx)
+							{
+								DecompVerts[PtIdx] = OrigDecompVerts[PtIdx];
+							}
+							FilterHullPoints(DecompVerts, SimplificationDistanceThreshold);
+							FVector Pivot = ScaleHullPoints(DecompVerts, OverlapRemovalShrinkPercent);
+							Pivots.Add(Pivot);
+							DecompHulls.Add(MakeUnique<::Chaos::FConvex>(DecompVerts, UE_KINDA_SMALL_NUMBER));
+						}
+						HullCS.Lock();
+						for (int32 HullIdx = 0; HullIdx < DecompHulls.Num(); ++HullIdx)
+						{
+							int32 ConvexIdx = Convexes.Add(MoveTemp(DecompHulls[HullIdx]));
+							AddPivot(Convexes, ConvexPivots, Pivots[HullIdx]);
+							TransformToConvexIndices[Idx].Add(ConvexIdx);
+						}
+						HullCS.Unlock();
+						return;
+					}
+				}
 			}
-			ensure(HullPts.Num() > 0);
-			FilterHullPoints(HullPts, SimplificationDistanceThreshold);
-			FVector Pivot = ScaleHullPoints(HullPts, OverlapRemovalShrinkPercent);
-			TUniquePtr Hull = MakeUnique<Chaos::FConvex>(HullPts, UE_KINDA_SMALL_NUMBER);
+			// We don't need a convex decomposition
+			if (!Hull)
+			{
+				Hull = ComputeHull(HullPivot);
+			}
 			HullCS.Lock();
 			int32 ConvexIdx = Convexes.Add(MoveTemp(Hull));
-			AddPivot(Convexes, ConvexPivots, Pivot);
+			AddPivot(Convexes, ConvexPivots, HullPivot);
 			HullCS.Unlock();
 			TransformToConvexIndices[Idx].Add(ConvexIdx);
 		}
@@ -1319,54 +1445,6 @@
 	return true;
 }
 
-// helper to compute the volume of an individual piece of geometry
-double ComputeGeometryVolume(
-	const FManagedArrayCollection* Collection,
-	int32 GeometryIdx,
-	const FTransform& GlobalTransform,
-	double ScalePerDimension
-)
-{
-	GeometryCollection::Facades::FCollectionMeshFacade MeshFacade(*Collection);
-	const TManagedArray<int32>& VertexStart = MeshFacade.VertexStartAttribute.Get();
-	const TManagedArray<int32>& VertexCount = MeshFacade.VertexCountAttribute.Get();
-	const TManagedArray<FVector3f>& Vertex = MeshFacade.VertexAttribute.Get();
-	const TManagedArray<int32>& FaceStart = MeshFacade.FaceStartAttribute.Get();
-	const TManagedArray<int32>& FaceCount = MeshFacade.FaceCountAttribute.Get();
-	const TManagedArray<FIntVector>& Indices = MeshFacade.IndicesAttribute.Get();
-	int32 VStart = VertexStart[GeometryIdx];
-	int32 VEnd = VStart + VertexCount[GeometryIdx];
-	if (VStart == VEnd)
-	{
-		return 0.0;
-	}
-	FVector3d Center = FVector::ZeroVector;
-	for (int32 VIdx = VStart; VIdx < VEnd; VIdx++)
-	{
-		FVector Pos = GlobalTransform.TransformPosition((FVector)Vertex[VIdx]);
-		Center += (FVector3d)Pos;
-	}
-	Center /= double(VEnd - VStart);
-	int32 FStart = FaceStart[GeometryIdx];
-	int32 FEnd = FStart + FaceCount[GeometryIdx];
-	double VolOut = 0;
-	for (int32 FIdx = FStart; FIdx < FEnd; FIdx++)
-	{
-		FIntVector Tri = Indices[FIdx];
-		FVector3d V0 = (FVector3d)GlobalTransform.TransformPosition((FVector)Vertex[Tri.X]);
-		FVector3d V1 = (FVector3d)GlobalTransform.TransformPosition((FVector)Vertex[Tri.Y]);
-		FVector3d V2 = (FVector3d)GlobalTransform.TransformPosition((FVector)Vertex[Tri.Z]);
-
-		// add volume of the tetrahedron formed by the triangles and the reference point
-		FVector3d V1mRef = (V1 - Center) * ScalePerDimension;
-		FVector3d V2mRef = (V2 - Center) * ScalePerDimension;
-		FVector3d N = V2mRef.Cross(V1mRef);
-
-		VolOut += ((V0 - Center) * ScalePerDimension).Dot(N) / 6.0;
-	}
-	return VolOut;
-}
-
 
 }
 
@@ -1390,7 +1468,8 @@
 }
 
 UE::GeometryCollectionConvexUtility::FConvexHulls
-FGeometryCollectionConvexUtility::ComputeLeafHulls(FGeometryCollection* GeometryCollection, const TArray<FTransform>& GlobalTransformArray, double SimplificationDistanceThreshold, double OverlapRemovalShrinkPercent, TFunction<bool(int32)> SkipBoneFn)
+FGeometryCollectionConvexUtility::ComputeLeafHulls(FGeometryCollection* GeometryCollection, const TArray<FTransform>& GlobalTransformArray, double SimplificationDistanceThreshold, double OverlapRemovalShrinkPercent,
+	TFunction<bool(int32)> SkipBoneFn, const FConvexDecompositionSettings* OptionalDecompositionSettings)
 {
 	check(GeometryCollection);
 
@@ -1407,7 +1486,7 @@
 	TFunctionRef<bool(int32)> HasCustomConvexFn = (CustomConvexFlags != nullptr) ? (TFunctionRef<bool(int32)>)ConvexFlagsFromArray : (TFunctionRef<bool(int32)>)ConvexFlagsAlwaysFalse;
 
 	HullsFromGeometry(*GeometryCollection, GlobalTransformArray, HasCustomConvexFn, Hulls.Hulls, Hulls.Pivots, Hulls.TransformToHullsIndices, GeometryCollection->SimulationType,
-		FGeometryCollection::ESimulationTypes::FST_Rigid, SimplificationDistanceThreshold, Hulls.OverlapRemovalShrinkPercent, SkipBoneFn);
+		FGeometryCollection::ESimulationTypes::FST_Rigid, SimplificationDistanceThreshold, Hulls.OverlapRemovalShrinkPercent, SkipBoneFn, OptionalDecompositionSettings);
 
 	return Hulls;
 }
@@ -1463,7 +1542,7 @@
 	return { TransformToConvexIndices, ConvexHull };
 }
 
-void FGeometryCollectionConvexUtility::GenerateLeafConvexHulls(FGeometryCollection& Collection, bool bRestrictToSelection, const TArrayView<const int32> TransformSubset, double SimplificationDistanceThreshold, EGenerateConvexMethod GenerateMethod, FIntersectionFilters IntersectFilters)
+void FGeometryCollectionConvexUtility::GenerateLeafConvexHulls(FGeometryCollection& Collection, bool bRestrictToSelection, const TArrayView<const int32> TransformSubset, const FLeafConvexHullSettings& Settings)
 {
 	using FSharedImplicit = TSharedPtr<Chaos::FImplicitObject, ESPMode::ThreadSafe>;
 
@@ -1483,9 +1562,9 @@
 			bHasExternalCollision = false;
 		}
 	}
-	bool bUseExternal = bHasExternalCollision && (GenerateMethod == EGenerateConvexMethod::ExternalCollision || GenerateMethod == EGenerateConvexMethod::IntersectExternalWithComputed);
-	bool bUseGenerated = !bUseExternal || GenerateMethod == EGenerateConvexMethod::IntersectExternalWithComputed;
-	bool bUseIntersect = bUseExternal && GenerateMethod == EGenerateConvexMethod::IntersectExternalWithComputed;
+	bool bUseExternal = bHasExternalCollision && (Settings.GenerateMethod == EGenerateConvexMethod::ExternalCollision || Settings.GenerateMethod == EGenerateConvexMethod::IntersectExternalWithComputed);
+	bool bUseGenerated = !bUseExternal || Settings.GenerateMethod == EGenerateConvexMethod::IntersectExternalWithComputed;
+	bool bUseIntersect = bUseExternal && Settings.GenerateMethod == EGenerateConvexMethod::IntersectExternalWithComputed;
 	
 	UE::GeometryCollectionConvexUtility::FConvexHulls ComputedHulls;
 	TArray<FTransformedConvex> ExternalHulls;
@@ -1583,7 +1662,7 @@
 				return SelectionSet.Contains(BoneIdx);
 			};
 		}
-		ComputedHulls = ComputeLeafHulls(&Collection, IdentityTransformArray, SimplificationDistanceThreshold, 0, SkipBoneFn);
+		ComputedHulls = ComputeLeafHulls(&Collection, IdentityTransformArray, Settings.SimplificationDistanceThreshold, 0, SkipBoneFn, &Settings.DecompositionSettings);
 
 		if (!bUseIntersect)
 		{
@@ -1595,6 +1674,11 @@
 	}
 	if (bUseIntersect)
 	{
+		// TODO: this logic should be rewritten directly in the ComputeLeafHulls function, which should take the external hulls as an optional parameter.
+		// Specifically, ComputeLeafHulls should intersect the union of the external hulls with the geometry of the leaves, and then run convex decomposition on the result.
+		// This will avoid the possible multiplication of hulls if there are multiple external hulls, and also will allow the convex decomposition to focus
+		// only on the geometry that is not removed by the intersection.
+
 		TArray<TUniquePtr<Chaos::FConvex>> FinalHulls;
 		for (int32 SourceTransformIdx : UseTransforms)
 		{
@@ -1605,7 +1689,7 @@
 
 			TSet<int32>& GeoHulls = ComputedHulls.TransformToHullsIndices[SourceTransformIdx];
 
-			if (IntersectFilters.OnlyIntersectIfComputedIsSmallerFactor < 1.0 || IntersectFilters.MinExternalVolumeToIntersect > 0.0)
+			if (Settings.IntersectFilters.OnlyIntersectIfComputedIsSmallerFactor < 1.0 || Settings.IntersectFilters.MinExternalVolumeToIntersect > 0.0)
 			{
 				double ComputedVolSum = 0.0;
 				for (int32 GeoHullIdx : GeoHulls)
@@ -1618,8 +1702,8 @@
 					ExternalVolSum += (double)ExternalHulls[ExtHull].Convex->GetVolume();
 				}
 				if (
-					(IntersectFilters.OnlyIntersectIfComputedIsSmallerFactor < 1.0 && ComputedVolSum >= ExternalVolSum * IntersectFilters.OnlyIntersectIfComputedIsSmallerFactor) ||
-					(ExternalVolSum < IntersectFilters.MinExternalVolumeToIntersect)
+					(Settings.IntersectFilters.OnlyIntersectIfComputedIsSmallerFactor < 1.0 && ComputedVolSum >= ExternalVolSum * Settings.IntersectFilters.OnlyIntersectIfComputedIsSmallerFactor) ||
+					(ExternalVolSum < Settings.IntersectFilters.MinExternalVolumeToIntersect)
 					)
 				{
 					// Filters allow us to skip computing the intersection and just use the external hulls
@@ -1632,10 +1716,6 @@
 				}
 			}
 
-			// Note: Currently computed hulls path only generates one hull per transform. If this changes, and we have multiple
-			// compute hulls AND multiple external hulls, consider running an additional merge step after intersection
-			// to avoid multiplying the number of hulls used.
-			ensure(GeoHulls.Num() < 2);
 			for (int32 GeoHullIdx : GeoHulls)
 			{
 				if (TransformToExternalHullsIndices.IsEmpty())
@@ -1876,7 +1956,7 @@
 				};
 				UE::Geometry::FConvexDecomposition3 ConvexDecomposition;
 				ConvexDecomposition.InitializeFromHulls(Hulls.Num(), GetHullVolume, GetHullNumVertices, GetHullVertex, HullProximity);
-				ConvexDecomposition.MergeBest(Settings.ConvexCount, Settings.ErrorToleranceInCm, 0, true /*bAllowCompact*/, false /*bRequireHullTriangles*/, Settings.EmptySpace, &ParentTransform);
+				ConvexDecomposition.MergeBest(Settings.ConvexCount, Settings.ErrorToleranceInCm, 0, true /*bAllowCompact*/, false /*bRequireHullTriangles*/, -1 /*MaxHulls*/, Settings.EmptySpace, &ParentTransform);
 				
 				// reset existing hulls for this transform index
 				for (int32 ConvexIndex : TransformToConvexIndices[TransformIndex])
diff --git a/Engine/Source/Runtime/Experimental/Chaos/Public/GeometryCollection/GeometryCollectionConvexUtility.h b/Engine/Source/Runtime/Experimental/Chaos/Public/GeometryCollection/GeometryCollectionConvexUtility.h
--- a/Engine/Source/Runtime/Experimental/Chaos/Public/GeometryCollection/GeometryCollectionConvexUtility.h
+++ b/Engine/Source/Runtime/Experimental/Chaos/Public/GeometryCollection/GeometryCollectionConvexUtility.h
@@ -140,8 +140,70 @@ public:
 		double OnlyIntersectIfComputedIsSmallerFactor;
 		double MinExternalVolumeToIntersect;
 	};
+
+	// Settings to control convex decompositions. Note the default values are set to *not* perform any decomposition
+	struct FConvexDecompositionSettings
+	{
+		FConvexDecompositionSettings() :
+			MinGeoVolumeToDecompose(0.f),
+			MaxGeoToHullVolumeRatioToDecompose(1.f),
+			MaxHullsPerGeometry(1),
+			ErrorTolerance(0.f),
+			MinThicknessTolerance(0.f),
+			NumAdditionalSplits(4)
+		{}
+
+		///
+		/// These settings are filters that control whether convex decomposition is performed
+		///
+		
+		// If greater than zero, the minimum geometry volume to consider for convex decomposition
+		float MinGeoVolumeToDecompose;
+
+		// If the geo volume / hull volume ratio is greater than this, do not consider convex decomposition
+		float MaxGeoToHullVolumeRatioToDecompose;
+
+
+		///
+		/// If we do perform a convex decomposition, these settings control the decomposition process
+		///
+		
+		// If > 0, specify the maximum number of convex hulls to create.  Hull merges will be attempted until this number of hulls remains, or no merges are possible.
+		int32 MaxHullsPerGeometry;
+
+		// Stop splitting when hulls have error less than this (expressed in cm; will be cubed for volumetric error). Overrides NumOutputHulls if non-zero.
+		float ErrorTolerance;
+
+		// Optionally specify a minimum thickness (in cm) for convex parts; parts below this thickness will always be merged away. Will merge hulls with greater error than ErrorTolerance when needed.
+		float MinThicknessTolerance;
+
+		// How far to go beyond the target number of hulls for initial decomposition (before merges) -- larger values will require more computation but can find better convex decompositions
+		int32 NumAdditionalSplits;
+
+	};
+
+	// Settings to control how convex hulls are generated for rigid/leaf nodes, from the geometry and/or imported collision shapes
+	struct FLeafConvexHullSettings
+	{
+		FLeafConvexHullSettings() : SimplificationDistanceThreshold(1.0), GenerateMethod(EGenerateConvexMethod::ExternalCollision)
+		{}
+
+		explicit FLeafConvexHullSettings(double SimplificationDistanceThreshold, EGenerateConvexMethod GenerateMethod = EGenerateConvexMethod::ExternalCollision) :
+			SimplificationDistanceThreshold(SimplificationDistanceThreshold),
+			GenerateMethod(GenerateMethod)
+		{}
+
+		double SimplificationDistanceThreshold;
+		EGenerateConvexMethod GenerateMethod;
+
+		// Intersection filters only apply if GenerateMethod == EGenerateConvexMethod::IntersectExternalWithComputed
+		FIntersectionFilters IntersectFilters;
+
+		// Convex decomposition settings, applied to convex hulls generated from geometry
+		FConvexDecompositionSettings DecompositionSettings;
+	};
 	
-	static void GenerateLeafConvexHulls(FGeometryCollection& Collection, bool bRestrictToSelection, const TArrayView<const int32> TransformSubset, double SimplificationDistanceThreshold, EGenerateConvexMethod GenerateMethod = EGenerateConvexMethod::ExternalCollision, FIntersectionFilters IntersectFilters = FIntersectionFilters());
+	static void GenerateLeafConvexHulls(FGeometryCollection& Collection, bool bRestrictToSelection, const TArrayView<const int32> TransformSubset, const FLeafConvexHullSettings& Settings);
 
 	/** Returns the convex hull of the vertices contained in the specified geometry. */
 	static TUniquePtr<Chaos::FConvex> FindConvexHull(const FGeometryCollection* GeometryCollection, int32 GeometryIndex);
@@ -182,7 +244,7 @@ public:
 	// (TODO: Make auto-embed use this instead of the full hulls?)
 	// @param GlobalTransformArray		GeometryCollection's transforms to global space, as computed by GeometryCollectionAlgo::GlobalMatrices
 	static UE::GeometryCollectionConvexUtility::FConvexHulls ComputeLeafHulls(FGeometryCollection* GeometryCollection, const TArray<FTransform>& GlobalTransformArray, double SimplificationDistanceThreshold = 0.0, double OverlapRemovalShrinkPercent = 0.0,
-		TFunction<bool(int32)> SkipBoneFn = nullptr);
+		TFunction<bool(int32)> SkipBoneFn = nullptr, const FConvexDecompositionSettings* OptionalDecompositionSettings = nullptr);
 
 	struct FTransformedConvex
 	{
diff --git a/Engine/Source/Runtime/GeometryCore/Private/CompGeom/ConvexDecomposition3.cpp b/Engine/Source/Runtime/GeometryCore/Private/CompGeom/ConvexDecomposition3.cpp
--- a/Engine/Source/Runtime/GeometryCore/Private/CompGeom/ConvexDecomposition3.cpp
+++ b/Engine/Source/Runtime/GeometryCore/Private/CompGeom/ConvexDecomposition3.cpp
@@ -825,6 +825,25 @@ FConvexDecomposition3::FConvexPart::FConvexPart(const FDynamicMesh3& SourceMesh,
 	// Copy out the source mesh
 	InternalGeo.Copy(SourceMesh, false, false, false, false);
 
+	InitializeFromInternalGeo(bMergeEdges, TransformOut);
+}
+
+FConvexDecomposition3::FConvexPart::FConvexPart(TArrayView<const FVector3f> Vertices, TArrayView<const FIntVector3> Faces, bool bMergeEdges, FTransformSRT3d& TransformOut, int32 FaceVertexOffset)
+{
+	for (const FVector3f& V : Vertices)
+	{
+		InternalGeo.AppendVertex((FVector)V);
+	}
+	for (FIntVector F : Faces)
+	{
+		InternalGeo.AppendTriangle(FIndex3i(F.X + FaceVertexOffset, F.Y + FaceVertexOffset, F.Z + FaceVertexOffset));
+	}
+
+	InitializeFromInternalGeo(bMergeEdges, TransformOut);
+}
+
+void FConvexDecomposition3::FConvexPart::InitializeFromInternalGeo(bool bMergeEdges, FTransformSRT3d& TransformOut)
+{
 	// Transform the mesh to a standard unit-cube-at-origin space, so threshold have a consistent meaning
 	FAxisAlignedBox3d InitialBounds = InternalGeo.GetBounds();
 	double InvScaleFactor = FMath::Clamp(InitialBounds.MaxDim(), KINDA_SMALL_NUMBER, 1e8);
@@ -848,6 +867,19 @@ FConvexDecomposition3::FConvexPart::FConvexPart(const FDynamicMesh3& SourceMesh,
 	ComputeStats();
 }
 
+
+void FConvexDecomposition3::InitializeFromIndexMesh(TArrayView<const FVector3f> Vertices, TArrayView<const FIntVector> Faces, bool bMergeEdges, int32 FaceVertexOffset)
+{
+	Decomposition.Empty();
+	FConvexPart* Convex = new FConvexPart(Vertices, Faces, bMergeEdges, ResultTransform, FaceVertexOffset);
+	if (Convex->IsFailed())
+	{
+		delete Convex;
+		return;
+	}
+	Decomposition.Add(Convex);
+}
+
 bool FConvexDecomposition3::FConvexPart::ComputeHull(bool bComputePlanes)
 {
 	FConvexHull3d HullCompute;
@@ -1292,9 +1324,16 @@ void FConvexDecomposition3::InitializeFromHulls(int32 NumHulls, TFunctionRef<dou
 	}
 }
 
-int32 FConvexDecomposition3::MergeBest(int32 TargetNumParts, double MaxErrorTolerance, double MinThicknessToleranceWorldSpace, bool bAllowCompact, bool bRequireHullTriangles, 
+int32 FConvexDecomposition3::MergeBest(int32 InTargetNumParts, double MaxErrorTolerance, double MinThicknessToleranceWorldSpace, bool bAllowCompact, bool bRequireHullTriangles, int32 MaxOutputHulls,
 	const FSphereCovering* OptionalNegativeSpace, const FTransform* OptionalTransformIntoNegativeSpace)
 {
+	int32 TargetNumParts = InTargetNumParts;
+	const bool bHasValidMaxHulls = MaxOutputHulls > 0;
+	if (bHasValidMaxHulls)
+	{
+		TargetNumParts = FMath::Min(TargetNumParts, MaxOutputHulls);
+	}
+
 	// Support having a max error tolerance
 	double VolumeTolerance = ConvertDistanceToleranceToLocalVolumeTolerance(MaxErrorTolerance);
 	double MinThicknessTolerance = ConvertDistanceToleranceToLocalSpace(MinThicknessToleranceWorldSpace);
@@ -1573,8 +1612,10 @@ int32 FConvexDecomposition3::MergeBest(int32 TargetNumParts, double MaxErrorTole
 				MergeCost = Proximities[ProxIdx].GetMergeCost(Decomposition);
 			}
 
+			bool bAboveMax = bHasValidMaxHulls && Decomposition.Num() > MaxOutputHulls;
+
 			if (!bIncludesMustMergePart && ( // if there's no must-merge part, skip if ...
-					(VolumeTolerance > 0 && MergeCost > VolumeTolerance) // it won't pass the volume tolerance
+					(!bAboveMax && VolumeTolerance > 0 && MergeCost > VolumeTolerance) // it won't pass the volume tolerance, and we already have <= max hulls
 						|| 
 					bOnlyAllowMustMerges // or we're currently only considering must-merge parts (e.g., already at or below desired number of parts)
 				))
diff --git a/Engine/Source/Runtime/GeometryCore/Public/CompGeom/ConvexDecomposition3.h b/Engine/Source/Runtime/GeometryCore/Public/CompGeom/ConvexDecomposition3.h
--- a/Engine/Source/Runtime/GeometryCore/Public/CompGeom/ConvexDecomposition3.h
+++ b/Engine/Source/Runtime/GeometryCore/Public/CompGeom/ConvexDecomposition3.h
@@ -127,6 +127,15 @@ public:
 
 	void InitializeFromMesh(const FDynamicMesh3& SourceMesh, bool bMergeEdges);
 
+	/**
+	 * Initialize convex decomposition with a triangle index mesh
+	 * @param Vertices			Vertex buffer for mesh to decompose
+	 * @param Faces				Triangle buffer for mesh to decompose
+	 * @param bMergeEdges		Whether to attempt to weld matching edges before computing the convex hull; this can help the convex decomposition find better cutting planes for meshes that have boundaries e.g. due to seams
+	 * @param FaceVertexOffset	Indices from the Faces array are optionally offset by this value.  Useful e.g. to take slices of the multi-geometry vertex and face buffers of FGeometryCollection.
+	 */
+	void InitializeFromIndexMesh(TArrayView<const FVector3f> Vertices, TArrayView<const FIntVector> Faces, bool bMergeEdges, int32 FaceVertexOffset = 0);
+
 	//
 	// Settings
 	//
@@ -162,14 +171,20 @@ public:
 
 	/**
 	 * Compute a decomposition with the desired number of hulls
+	 * Note: A future version of this function may replace NumOutputHulls with MaxOutputHulls, but this version keeps both parameters for compatibility / consistent behavior.
 	 * 
 	 * @param NumOutputHulls		Number of convex hulls to use in the final convex decomposition
 	 * @param NumAdditionalSplits	How far to go beyond the target number of hulls when initially the mesh into pieces -- larger values will require more computation but can find better convex decompositions
 	 * @param ErrorTolerance		Stop splitting when hulls have error less than this (expressed in cm; will be cubed for volumetric error). Overrides NumOutputHulls if specified
 	 * @param MinThicknessTolerance	Optionally specify a minimum thickness (in cm) for convex parts; parts below this thickness will always be merged away. Overrides NumOutputHulls and ErrorTolerance when needed
+	 * @param MaxOutputHulls		If > 0, maximum number of convex hulls to generate. Overrides ErrorTolerance and TargetNumParts when needed
 	 */
-	void Compute(int32 NumOutputHulls, int32 NumAdditionalSplits = 10, double ErrorTolerance = 0.0, double MinThicknessTolerance = 0)
+	void Compute(int32 NumOutputHulls, int32 NumAdditionalSplits = 10, double ErrorTolerance = 0.0, double MinThicknessTolerance = 0, int32 MaxOutputHulls = -1)
 	{
+		if (MaxOutputHulls > 0)
+		{
+			NumOutputHulls = FMath::Min(NumOutputHulls, MaxOutputHulls);
+		}
 		int32 TargetNumSplits = NumOutputHulls + NumAdditionalSplits;
 		for (int32 SplitIdx = 0; SplitIdx < TargetNumSplits; SplitIdx++)
 		{
@@ -181,7 +196,7 @@ public:
 		}
 		
 		constexpr bool bAllowCompact = true;
-		MergeBest(NumOutputHulls, ErrorTolerance, MinThicknessTolerance, bAllowCompact);
+		MergeBest(NumOutputHulls, ErrorTolerance, MinThicknessTolerance, bAllowCompact, false, MaxOutputHulls);
 	}
 
 	// Split the worst convex part, and return the increase in the total number of convex parts after splitting (can be more than 1 if result has multiple separate connected components)
@@ -190,6 +205,7 @@ public:
 	int32 SplitWorst(bool bCanSkipUnreliableGeoVolumes = false, double ErrorTolerance = 0.0);
 
 	// Merge the pairs of convex hulls in the decomposition that will least increase the error.  Intermediate results can be used across merges, so it is best to do all merges in one call.
+	// Note: A future version of this function may replace NumOutputHulls with MaxOutputHulls, but this version keeps both parameters for compatibility / consistent behavior.
 	// @param TargetNumParts		The target number of parts for the decomposition; will be overriden by non-default ErrorTolerance or MinPartThickness
 	// @param ErrorTolerance		If > 0, continue to merge (if there are possible merges) until the resulting error would be greater than this value. Overrides TargetNumParts as the stopping condition.
 	//								Note: ErrorTolerance is expressed in cm, and will be cubed for volumetric error.
@@ -197,10 +213,11 @@ public:
 	//								Note: These parts may be further split so they can be merged into multiple hulls
 	// @param bAllowCompact			Allow the algorithm to discard underlying geometry once it will no longer be used, resulting in a smaller representation & faster merges
 	// @param bRequireHullTriangles	Require all hulls to have associated triangles after MergeBest is completed. (If using InitializeFromHulls, will need to triangulate any un-merged hulls.)
+	// @param MaxOutputHulls		If > 0, maximum number of convex hulls to generate. Overrides ErrorTolerance and TargetNumParts when needed
 	// @param OptionalNegativeSpace Optional representation of negative space that should not be covered by merges
 	// @param OptionalNegativeSpaceTransform Optional transform from space of the convex hulls into the space of the sphere covering; if not provided, assume spheres are in the same coordinate space as the hulls
 	// @return						The number of merges performed
-	int32 MergeBest(int32 TargetNumParts, double ErrorTolerance = 0, double MinThicknessTolerance = 0, bool bAllowCompact = true, bool bRequireHullTriangles = false, 
+	int32 MergeBest(int32 TargetNumParts, double ErrorTolerance = 0, double MinThicknessTolerance = 0, bool bAllowCompact = true, bool bRequireHullTriangles = false, int MaxOutputHulls = -1,
 		const FSphereCovering* OptionalNegativeSpace = nullptr, const FTransform* OptionalTransformIntoNegativeSpace = nullptr);
 
 	// simple helper to convert an error tolerance expressed in world space to a local-space volume tolerance
@@ -296,6 +313,8 @@ public:
 	{
 		FConvexPart() {}
 		FConvexPart(const FDynamicMesh3& SourceMesh, bool bMergeEdges, FTransformSRT3d& TransformOut);
+		FConvexPart(TArrayView<const FVector3f> Vertices, TArrayView<const FIntVector3> Faces, bool bMergeEdges, FTransformSRT3d& TransformOut, int32 FaceVertexOffset = 0);
+
 		// Allow direct construction of a compact part (e.g. to allow construction of a pre-existing convex hull)
 		FConvexPart(bool bIsCompact) : bIsCompact(bIsCompact)
 		{}
@@ -379,6 +398,10 @@ public:
 	protected:
 		bool bIsCompact = false;
 		bool bFailed = false;
+
+	private:
+		// helper to initialize the part from the already-set InternalGeo member; used to support constructors
+		void InitializeFromInternalGeo(bool bMergeEdges, FTransformSRT3d& TransformOut);
 	};
 
 	// All convex hulls parts in the convex decomposition
