Geometry collection : fix recent optimization to account for non-clustered geometry collection
- Make sure we detect a change of parent if the parent particle is null but the collection still shows a parent index

diff --git a/Engine/Source/Runtime/Experimental/Chaos/Private/PhysicsProxy/GeometryCollectionPhysicsProxy.cpp b/Engine/Source/Runtime/Experimental/Chaos/Private/PhysicsProxy/GeometryCollectionPhysicsProxy.cpp
--- a/Engine/Source/Runtime/Experimental/Chaos/Private/PhysicsProxy/GeometryCollectionPhysicsProxy.cpp
+++ b/Engine/Source/Runtime/Experimental/Chaos/Private/PhysicsProxy/GeometryCollectionPhysicsProxy.cpp
@@ -3034,7 +3034,7 @@ void FGeometryCollectionPhysicsProxy::BufferPhysicsResults_Internal(Chaos::FPBDR
 			Chaos::FPBDRigidClusteredParticleHandle* ClusterParent = Handle->Parent();
 
 			// Has parent changed? ( can be a new one or nullptr)
-			const bool bHasNewParent = SolverClusterID[TransformGroupIndex] != ClusterParent;
+			const bool bHasNewParent = (SolverClusterID[TransformGroupIndex] != ClusterParent) || (!ClusterParent && PhysicsThreadCollection.Parent[TransformGroupIndex]);
 			if (bHasNewParent)
 			{
 				// Force all driven rigid bodies out of the transform hierarchy ( because the new parent can only be an internal cluster )
