Removed in-memory cache experiment

#rnx

diff --git a/Engine/Source/Runtime/Core/Internal/IO/IoMemoryCache.h b/Engine/Source/Runtime/Core/Internal/IO/IoMemoryCache.h
deleted file mode 100644
diff --git a/Engine/Source/Runtime/Core/Private/IO/IoMemoryCache.cpp b/Engine/Source/Runtime/Core/Private/IO/IoMemoryCache.cpp
deleted file mode 100644
diff --git a/Engine/Source/Runtime/Experimental/IoStoreOnDemand/Private/IoStoreOnDemand.cpp b/Engine/Source/Runtime/Experimental/IoStoreOnDemand/Private/IoStoreOnDemand.cpp
--- a/Engine/Source/Runtime/Experimental/IoStoreOnDemand/Private/IoStoreOnDemand.cpp
+++ b/Engine/Source/Runtime/Experimental/IoStoreOnDemand/Private/IoStoreOnDemand.cpp
@@ -7,7 +7,6 @@
 #include "HAL/FileManager.h"
 #include "IO/IoCache.h"
 #include "IO/IoFileCache.h"
-#include "IO/IoMemoryCache.h"
 #include "Misc/CommandLine.h"
 #include "Misc/ConfigCacheIni.h"
 #include "Modules/ModuleManager.h"
@@ -337,13 +336,6 @@ void FIoStoreOnDemandModule::StartupModule()
 	{
 		TSharedPtr<IIoCache> Cache;
 
-#if !UE_BUILD_SHIPPING
-		if (uint64 CacheSize = ParseSizeParam(TEXT("OnDemandMemoryCache=")); CacheSize > 0)
-		{
-			UE_LOG(LogIoStoreOnDemand, Log, TEXT("Using %lluB in memory LRU cache"), CacheSize);
-			Cache = MakeShareable(MakeMemoryIoCache(CacheSize ).Release());
-		}
-#endif
 		if (uint64 DiskSize = ParseSizeParam(TEXT("OnDemandFileCache=")); DiskSize > 0)
 		{
 			uint64 MemorySize = ParseSizeParam(TEXT("OnDemandFileCacheQueueSize="));
