Chaos: SQ - Reduced stack reserved hit buffer sizes down to 16K from 64K, still reserving 128 hits for queries prior to heap allocating.

diff --git a/Engine/Source/Runtime/Experimental/Chaos/Public/Chaos/Interface/SQTypes.h b/Engine/Source/Runtime/Experimental/Chaos/Public/Chaos/Interface/SQTypes.h
--- a/Engine/Source/Runtime/Experimental/Chaos/Public/Chaos/Interface/SQTypes.h
+++ b/Engine/Source/Runtime/Experimental/Chaos/Public/Chaos/Interface/SQTypes.h
@@ -167,7 +167,7 @@ namespace ChaosInterface
 	using TThreadSweepHit = std::conditional_t<Id == Chaos::EThreadContext::External, FSweepHit, FPTSweepHit>;
 
 #ifndef CHAOS_HIT_BUFFER_SIZE
-#define CHAOS_HIT_BUFFER_SIZE 512 // Preallocated hit buffer size for traces and sweeps.
+#define CHAOS_HIT_BUFFER_SIZE 128 // Preallocated hit buffer size for traces and sweeps.
 #endif // CHAOS_HIT_BUFFER_SIZE
 	static_assert(CHAOS_HIT_BUFFER_SIZE > 0, "Invalid Chaos hit buffer size.");
 
