Iris - Allow enum _MAX values to be replicated.

Fixes issue with arbitrary code using said value to mean "not set" and similar use cases

diff --git a/Engine/Source/Runtime/Experimental/Iris/Core/Private/Iris/Serialization/InternalEnumNetSerializers.cpp b/Engine/Source/Runtime/Experimental/Iris/Core/Private/Iris/Serialization/InternalEnumNetSerializers.cpp
--- a/Engine/Source/Runtime/Experimental/Iris/Core/Private/Iris/Serialization/InternalEnumNetSerializers.cpp
+++ b/Engine/Source/Runtime/Experimental/Iris/Core/Private/Iris/Serialization/InternalEnumNetSerializers.cpp
@@ -2,6 +2,8 @@
 
 #include "Iris/Serialization/InternalEnumNetSerializers.h"
 #include "Iris/Core/BitTwiddling.h"
+#include <limits>
+#include <type_traits>
 
 namespace UE::Net::Private
 {
@@ -12,16 +14,16 @@ class TInitEnumNetSerializerConfig
 public:
 	static bool Init(ConfigType& OutConfig, const UEnum* Enum)
 	{
-		// NumEnums actually also contain the generated _MAX enum value. Skip it!
-		const int32 EnumValueCount = Enum->NumEnums() - 1;
+		// NumEnums() include the potentially auto-generated _MAX enum value.
+		const int32 EnumValueCount = Enum->NumEnums();
 
-		using LargeIntegerType = typename TChooseClass<TIsSigned<SourceType>::Value, int64, uint64>::Result;
+		using LargeIntegerType = std::conditional_t<std::is_signed_v<SourceType>, int64, uint64>;
 
 		// Find smallest and largest values.
 		if (EnumValueCount <= 0)
 		{
 			// At the time of this writing this case should not happen due to errors when a UENUM does not contain values.
-			ensure(EnumValueCount > 0);
+			ensureAlways(EnumValueCount > 0);
 			OutConfig.LowerBound = 0;
 			OutConfig.UpperBound = 0;
 			OutConfig.BitCount = static_cast<uint8>(GetBitsNeededForRange(OutConfig.LowerBound, OutConfig.UpperBound));
@@ -30,8 +32,8 @@ public:
 		else
 		{
 			// Cannot use UEnum methods here due to issues with the generated _MAX value as well as uint64 values outside of the positive int64 range.
-			LargeIntegerType SmallestValue = TNumericLimits<LargeIntegerType>::Max();
-			LargeIntegerType LargestValue = TNumericLimits<LargeIntegerType>::Min();
+			LargeIntegerType SmallestValue = std::numeric_limits<LargeIntegerType>::max();
+			LargeIntegerType LargestValue = std::numeric_limits<LargeIntegerType>::min();
 			for (int32 EnumIt = 0, EnumEndIt = EnumValueCount; EnumIt != EnumEndIt; ++EnumIt)
 			{
 				const LargeIntegerType Value = static_cast<LargeIntegerType>(Enum->GetValueByIndex(EnumIt));
