geometry collection : fix crash when SetRestCollection is called at runtime while using custom renderer

diff --git a/Engine/Source/Runtime/Experimental/GeometryCollectionEngine/Private/GeometryCollection/GeometryCollectionComponent.cpp b/Engine/Source/Runtime/Experimental/GeometryCollectionEngine/Private/GeometryCollection/GeometryCollectionComponent.cpp
--- a/Engine/Source/Runtime/Experimental/GeometryCollectionEngine/Private/GeometryCollection/GeometryCollectionComponent.cpp
+++ b/Engine/Source/Runtime/Experimental/GeometryCollectionEngine/Private/GeometryCollection/GeometryCollectionComponent.cpp
@@ -3625,6 +3625,12 @@ void UGeometryCollectionComponent::SetRestCollection(const UGeometryCollection*
 			InitializeEmbeddedGeometry();
 		}
 
+		if (CustomRenderer)
+		{
+			UnregisterCustomRenderer();
+			RegisterCustomRenderer();
+		}
+
 		if (bApplyAssetDefaults)
 		{
 			ApplyAssetDefaults();
