Reuse NewMaterialHandles temporary alloc, alternatively we can inline this alloc for 1-2 entries, as that is by far the common size

diff --git a/Engine/Source/Runtime/Engine/Private/PhysicsEngine/Experimental/PhysInterface_Chaos.cpp b/Engine/Source/Runtime/Engine/Private/PhysicsEngine/Experimental/PhysInterface_Chaos.cpp
--- a/Engine/Source/Runtime/Engine/Private/PhysicsEngine/Experimental/PhysInterface_Chaos.cpp
+++ b/Engine/Source/Runtime/Engine/Private/PhysicsEngine/Experimental/PhysInterface_Chaos.cpp
@@ -757,7 +757,7 @@ void FPhysInterface_Chaos::SetMaterials(const FPhysicsShapeHandle& InShape, cons
 		NewMaterialHandles.Add(UnrealMaterial->GetPhysicsMaterial());
 	}
 
-	InShape.Shape->SetMaterials(NewMaterialHandles);
+	InShape.Shape->SetMaterials(MoveTemp(NewMaterialHandles));
 }
 
 void FPhysInterface_Chaos::SetMaterials(const FPhysicsShapeHandle& InShape, const TArrayView<UPhysicalMaterial*> InMaterials, const TArrayView<FPhysicalMaterialMaskParams>& InMaterialMasks)
diff --git a/Engine/Source/Runtime/Experimental/Chaos/Public/Chaos/ShapeInstance.h b/Engine/Source/Runtime/Experimental/Chaos/Public/Chaos/ShapeInstance.h
--- a/Engine/Source/Runtime/Experimental/Chaos/Public/Chaos/ShapeInstance.h
+++ b/Engine/Source/Runtime/Experimental/Chaos/Public/Chaos/ShapeInstance.h
@@ -116,6 +116,7 @@ namespace Chaos
 		const TArray<FMaterialHandle>& GetMaterials() const;
 		void SetMaterial(FMaterialHandle InMaterial);
 		void SetMaterials(const TArray<FMaterialHandle>& InMaterials);
+		void SetMaterials(TArray<FMaterialHandle>&& InMaterials);
 
 		const TArray<FMaterialMaskHandle>& GetMaterialMasks() const;
 		void SetMaterialMasks(const TArray<FMaterialMaskHandle>& InMaterialMasks);
@@ -294,6 +295,14 @@ namespace Chaos
 				});
 		}
 
+		void SetMaterials(TArray<FMaterialHandle>&& InMaterials)
+		{
+			Materials.Modify(true, DirtyFlags, Proxy, ShapeIdx, [&InMaterials](FMaterialData& Data)
+				{
+					Data.Materials = MoveTemp(InMaterials);
+				});
+		}
+
 		void SetMaterialMasks(const TArray<FMaterialMaskHandle>& InMaterialMasks)
 		{
 			Materials.Modify(true, DirtyFlags, Proxy, ShapeIdx, [&InMaterialMasks](FMaterialData& Data)
@@ -554,6 +563,7 @@ namespace Chaos
 
 		void SetMaterial(FMaterialHandle InMaterial) { SetMaterialImpl(InMaterial); }
 		void SetMaterials(const TArray<FMaterialHandle>& InMaterials) { GetMaterialDataImpl().Materials = InMaterials; }
+		void SetMaterials(TArray<FMaterialHandle>&& InMaterials) { GetMaterialDataImpl().Materials = InMaterials; }
 		void SetMaterialMasks(const TArray<FMaterialMaskHandle>& InMaterialMasks) { GetMaterialDataImpl().MaterialMasks = InMaterialMasks; }
 		void SetMaterialMaskMaps(const TArray<uint32>& InMaterialMaskMaps) { GetMaterialDataImpl().MaterialMaskMaps = InMaterialMaskMaps; }
 		void SetMaterialMaskMapMaterials(const TArray<FMaterialHandle>& InMaterialMaskMapMaterials) { GetMaterialDataImpl().MaterialMaskMapMaterials = InMaterialMaskMapMaterials; }
@@ -986,6 +996,11 @@ namespace Chaos
 		DownCast([&InMaterials](auto& ShapeInstance) { ShapeInstance.SetMaterials(InMaterials); });
 	}
 
+	inline void FPerShapeData::SetMaterials(TArray<FMaterialHandle>&& InMaterials)
+	{
+		DownCast([&InMaterials](auto& ShapeInstance) { ShapeInstance.SetMaterials(MoveTemp(InMaterials)); });
+	}
+
 	inline void FPerShapeData::SetMaterialMasks(const TArray<FMaterialMaskHandle>& InMaterialMasks)
 	{
 		DownCast([&InMaterialMasks](auto& ShapeInstance) { ShapeInstance.SetMaterialMasks(InMaterialMasks); });
