Geometry collection rendering optimization
- Optimize the ISM Pool to use the instance static mesh array API for passing transform , avoiding local array allocation
- avoid sorting the removal indices twice by using the new instance static mesh view

diff --git a/Engine/Source/Runtime/Experimental/GeometryCollectionEngine/Private/GeometryCollection/GeometryCollectionISMPoolComponent.cpp b/Engine/Source/Runtime/Experimental/GeometryCollectionEngine/Private/GeometryCollection/GeometryCollectionISMPoolComponent.cpp
--- a/Engine/Source/Runtime/Experimental/GeometryCollectionEngine/Private/GeometryCollection/GeometryCollectionISMPoolComponent.cpp
+++ b/Engine/Source/Runtime/Experimental/GeometryCollectionEngine/Private/GeometryCollection/GeometryCollectionISMPoolComponent.cpp
@@ -33,6 +33,11 @@ FGeometryCollectionMeshGroup::FMeshId FGeometryCollectionMeshGroup::AddMesh(cons
 }
 
 bool FGeometryCollectionMeshGroup::BatchUpdateInstancesTransforms(FGeometryCollectionISMPool& ISMPool, FMeshId MeshId, int32 StartInstanceIndex, const TArray<FTransform>& NewInstancesTransforms, bool bWorldSpace, bool bMarkRenderStateDirty, bool bTeleport)
+{
+	return BatchUpdateInstancesTransforms(ISMPool, MeshId, StartInstanceIndex, MakeArrayView(NewInstancesTransforms), bWorldSpace, bMarkRenderStateDirty, bTeleport);
+}
+
+bool FGeometryCollectionMeshGroup::BatchUpdateInstancesTransforms(FGeometryCollectionISMPool& ISMPool, FMeshId MeshId, int32 StartInstanceIndex, TArrayView<const FTransform> NewInstancesTransforms, bool bWorldSpace, bool bMarkRenderStateDirty, bool bTeleport)
 {
 	if (MeshInfos.IsValidIndex(MeshId))
 	{
@@ -187,6 +192,13 @@ FGeometryCollectionMeshInfo FGeometryCollectionISMPool::AddISM(UGeometryCollecti
 
 bool FGeometryCollectionISMPool::BatchUpdateInstancesTransforms(FGeometryCollectionMeshInfo& MeshInfo, int32 StartInstanceIndex, const TArray<FTransform>& NewInstancesTransforms, bool bWorldSpace, bool bMarkRenderStateDirty, bool bTeleport)
 {
+	return BatchUpdateInstancesTransforms(MeshInfo, StartInstanceIndex, MakeArrayView(NewInstancesTransforms), bWorldSpace, bMarkRenderStateDirty, bTeleport);
+}
+
+bool FGeometryCollectionISMPool::BatchUpdateInstancesTransforms(FGeometryCollectionMeshInfo& MeshInfo, int32 StartInstanceIndex, TArrayView<const FTransform> NewInstancesTransforms, bool bWorldSpace, bool bMarkRenderStateDirty, bool bTeleport)
+{
+	constexpr bool bUseArrayView = true;
+
 	if (ISMs.IsValidIndex(MeshInfo.ISMIndex))
 	{
 		FGeometryCollectionISM& ISM = ISMs[MeshInfo.ISMIndex];
@@ -199,26 +211,51 @@ bool FGeometryCollectionISMPool::BatchUpdateInstancesTransforms(FGeometryCollect
 		int32 StartIndex = ISM.InstanceIndexToRenderIndex[InstanceGroup.Start];
 		int32 TransformIndex = 0;
 		int32 BatchCount = 1;
-		TArray<FTransform> BatchTransforms; // Can't use TArrayView because blueprint function doesn't support that 
-		BatchTransforms.Reserve(NewInstancesTransforms.Num());
-		BatchTransforms.Add(NewInstancesTransforms[TransformIndex++]);
-		for (int InstanceIndex = StartInstanceIndex + 1; InstanceIndex < NewInstancesTransforms.Num(); ++InstanceIndex)
+
+		if constexpr (bUseArrayView)
 		{
-			// Flush batch for non-sequential instances.
-			int32 RenderIndex = ISM.InstanceIndexToRenderIndex[InstanceGroup.Start + InstanceIndex];
-			if (RenderIndex != (StartIndex + BatchCount))
+			for (int InstanceIndex = StartInstanceIndex + 1; InstanceIndex < NewInstancesTransforms.Num(); ++InstanceIndex)
 			{
-				ISM.ISMComponent->BatchUpdateInstancesTransforms(StartIndex, BatchTransforms, bWorldSpace, bMarkRenderStateDirty, bTeleport);
-				StartIndex = RenderIndex;
-				BatchTransforms.SetNum(0, false);
-				BatchCount = 0;
+				// Flush batch for non-sequential instances.
+				int32 RenderIndex = ISM.InstanceIndexToRenderIndex[InstanceGroup.Start + InstanceIndex];
+				if (RenderIndex != (StartIndex + BatchCount))
+				{
+					TArrayView<const FTransform> BatchedTransformsView = MakeArrayView(NewInstancesTransforms.GetData() + TransformIndex, BatchCount);
+					ISM.ISMComponent->BatchUpdateInstancesTransforms(StartIndex, BatchedTransformsView, bWorldSpace, bMarkRenderStateDirty, bTeleport);
+					StartIndex = RenderIndex;
+					TransformIndex += BatchCount;
+					BatchCount = 0;
+				}
+				BatchCount++;
 			}
-			
-			BatchTransforms.Add(NewInstancesTransforms[TransformIndex++]);
-			BatchCount++;
+
+			// last one
+			TArrayView<const FTransform> BatchedTransformsView = MakeArrayView(NewInstancesTransforms.GetData() + TransformIndex, BatchCount);
+			return ISM.ISMComponent->BatchUpdateInstancesTransforms(StartIndex, BatchedTransformsView, bWorldSpace, bMarkRenderStateDirty, bTeleport);
 		}
+		else
+		{
+			TArray<FTransform> BatchTransforms; // Can't use TArrayView because blueprint function doesn't support that 
+			BatchTransforms.Reserve(NewInstancesTransforms.Num());
+			BatchTransforms.Add(NewInstancesTransforms[TransformIndex++]);
+			for (int InstanceIndex = StartInstanceIndex + 1; InstanceIndex < NewInstancesTransforms.Num(); ++InstanceIndex)
+			{
+				// Flush batch for non-sequential instances.
+				int32 RenderIndex = ISM.InstanceIndexToRenderIndex[InstanceGroup.Start + InstanceIndex];
+				if (RenderIndex != (StartIndex + BatchCount))
+				{
+					ISM.ISMComponent->BatchUpdateInstancesTransforms(StartIndex, BatchTransforms, bWorldSpace, bMarkRenderStateDirty, bTeleport);
+					StartIndex = RenderIndex;
+					BatchTransforms.SetNum(0, false);
+					BatchCount = 0;
+				}
+
+				BatchTransforms.Add(NewInstancesTransforms[TransformIndex++]);
+				BatchCount++;
+			}
 
-		return ISM.ISMComponent->BatchUpdateInstancesTransforms(StartIndex, BatchTransforms, bWorldSpace, bMarkRenderStateDirty, bTeleport);
+			return ISM.ISMComponent->BatchUpdateInstancesTransforms(StartIndex, BatchTransforms, bWorldSpace, bMarkRenderStateDirty, bTeleport);
+		}
 	}
 	UE_LOG(LogChaos, Warning, TEXT("UGeometryCollectionISMPoolComponent : Invalid ISM Id (%d) when updating the transform "), MeshInfo.ISMIndex);
 	return false;
@@ -240,11 +277,14 @@ void FGeometryCollectionISMPool::RemoveISM(const FGeometryCollectionMeshInfo& Me
 			// Clear the stored render index since we're about to remove it.
 			ISM.InstanceIndexToRenderIndex[InstanceGroup.Start + InstanceIndex] = -1;
 		}
-		ISM.ISMComponent->RemoveInstances(InstancesToRemove);
+
+		// we sort the array on the spot because we use it after calling RemoveInstances to fix up our own indices
+		InstancesToRemove.Sort(TGreater<int32>());
+		constexpr bool bArrayAlreadySorted = true;
+		ISM.ISMComponent->RemoveInstances(InstancesToRemove, bArrayAlreadySorted);
 
 		// Fix up instance index remapping to match what will have happened in our ISM component in RemoveInstances()
 		check(ISM.ISMComponent->SupportsRemoveSwap());
-		InstancesToRemove.Sort(TGreater<int32>());
 		for (int32 RenderIndex : InstancesToRemove)
 		{
 			ISM.RenderIndexToInstanceIndex.RemoveAtSwap(RenderIndex, 1, false);
@@ -331,6 +371,11 @@ UGeometryCollectionISMPoolComponent::FMeshId UGeometryCollectionISMPoolComponent
 }
 
 bool UGeometryCollectionISMPoolComponent::BatchUpdateInstancesTransforms(FMeshGroupId MeshGroupId, FMeshId MeshId, int32 StartInstanceIndex, const TArray<FTransform>& NewInstancesTransforms, bool bWorldSpace, bool bMarkRenderStateDirty, bool bTeleport)
+{
+	return BatchUpdateInstancesTransforms(MeshGroupId, MeshId, StartInstanceIndex, MakeArrayView(NewInstancesTransforms), bWorldSpace, bMarkRenderStateDirty, bTeleport);
+}
+
+bool UGeometryCollectionISMPoolComponent::BatchUpdateInstancesTransforms(FMeshGroupId MeshGroupId, FMeshId MeshId, int32 StartInstanceIndex, TArrayView<const FTransform> NewInstancesTransforms, bool bWorldSpace, bool bMarkRenderStateDirty, bool bTeleport)
 {
 	if (FGeometryCollectionMeshGroup* MeshGroup = MeshGroups.Find(MeshGroupId))
 	{
diff --git a/Engine/Source/Runtime/Experimental/GeometryCollectionEngine/Private/GeometryCollection/GeometryCollectionISMPoolRenderer.cpp b/Engine/Source/Runtime/Experimental/GeometryCollectionEngine/Private/GeometryCollection/GeometryCollectionISMPoolRenderer.cpp
--- a/Engine/Source/Runtime/Experimental/GeometryCollectionEngine/Private/GeometryCollection/GeometryCollectionISMPoolRenderer.cpp
+++ b/Engine/Source/Runtime/Experimental/GeometryCollectionEngine/Private/GeometryCollection/GeometryCollectionISMPoolRenderer.cpp
@@ -168,9 +168,7 @@ void UGeometryCollectionISMPoolRenderer::UpdateMergedMeshTransforms(FTransform c
 		return;
 	}
 
-	TArray<FTransform> InstanceTransforms;
-	InstanceTransforms.Add(InBaseTransform);
-
+	TArrayView<const FTransform> InstanceTransforms(&InBaseTransform, 1);
 	for (int32 MeshIndex = 0; MeshIndex < MergedMeshGroup.MeshIds.Num(); MeshIndex++)
 	{
 		ISMPoolComponent->BatchUpdateInstancesTransforms(MergedMeshGroup.GroupIndex, MergedMeshGroup.MeshIds[MeshIndex], 0, InstanceTransforms, true/*bWorldSpace*/, false/*bMarkRenderStateDirty*/, false/*bTeleport*/);
@@ -212,7 +210,7 @@ void UGeometryCollectionISMPoolRenderer::UpdateInstanceTransforms(UGeometryColle
 			}
 		}
 		
-		ISMPoolComponent->BatchUpdateInstancesTransforms(InstancesGroup.GroupIndex, InstancesGroup.MeshIds[MeshIndex], 0, InstanceTransforms, true/*bWorldSpace*/, false/*bMarkRenderStateDirty*/, false/*bTeleport*/);
+		ISMPoolComponent->BatchUpdateInstancesTransforms(InstancesGroup.GroupIndex, InstancesGroup.MeshIds[MeshIndex], 0, MakeArrayView(InstanceTransforms), true/*bWorldSpace*/, false/*bMarkRenderStateDirty*/, false/*bTeleport*/);
 	}
 }
 
diff --git a/Engine/Source/Runtime/Experimental/GeometryCollectionEngine/Public/GeometryCollection/GeometryCollectionISMPoolComponent.h b/Engine/Source/Runtime/Experimental/GeometryCollectionEngine/Public/GeometryCollection/GeometryCollectionISMPoolComponent.h
--- a/Engine/Source/Runtime/Experimental/GeometryCollectionEngine/Public/GeometryCollection/GeometryCollectionISMPoolComponent.h
+++ b/Engine/Source/Runtime/Experimental/GeometryCollectionEngine/Public/GeometryCollection/GeometryCollectionISMPoolComponent.h
@@ -220,6 +220,8 @@ struct FGeometryCollectionMeshGroup
 	FMeshId AddMesh(const FGeometryCollectionStaticMeshInstance& MeshInstance, int32 InstanceCount, const FGeometryCollectionMeshInfo& ISMInstanceInfo);
 	/** Update instance transforms for a group of instances. */
 	bool BatchUpdateInstancesTransforms(FGeometryCollectionISMPool& ISMPool, FMeshId MeshId, int32 StartInstanceIndex, const TArray<FTransform>& NewInstancesTransforms, bool bWorldSpace, bool bMarkRenderStateDirty, bool bTeleport);
+	bool BatchUpdateInstancesTransforms(FGeometryCollectionISMPool& ISMPool, FMeshId MeshId, int32 StartInstanceIndex, TArrayView<const FTransform> NewInstancesTransforms, bool bWorldSpace, bool bMarkRenderStateDirty, bool bTeleport);
+
 	/** Remove all of our managed meshes and associated instances. */
 	void RemoveAllMeshes(FGeometryCollectionISMPool& ISMPool);
 
@@ -256,7 +258,11 @@ struct FGeometryCollectionISMPool
 
 	FISMIndex AddISM(UGeometryCollectionISMPoolComponent* OwningComponent, const FGeometryCollectionStaticMeshInstance& MeshInstance);
 	FGeometryCollectionMeshInfo AddISM(UGeometryCollectionISMPoolComponent* OwningComponent, const FGeometryCollectionStaticMeshInstance& MeshInstance, int32 InstanceCount, TArrayView<const float> CustomDataFloats);
+
 	bool BatchUpdateInstancesTransforms(FGeometryCollectionMeshInfo& MeshInfo, int32 StartInstanceIndex, const TArray<FTransform>& NewInstancesTransforms, bool bWorldSpace, bool bMarkRenderStateDirty, bool bTeleport);
+
+	bool BatchUpdateInstancesTransforms(FGeometryCollectionMeshInfo& MeshInfo, int32 StartInstanceIndex, TArrayView<const FTransform> NewInstancesTransforms, bool bWorldSpace, bool bMarkRenderStateDirty, bool bTeleport);
+
 	void RemoveISM(const FGeometryCollectionMeshInfo& MeshInfo);
 	
 	/** Clear all ISM components and associated data */
@@ -299,6 +305,9 @@ public:
 	FMeshId AddMeshToGroup(FMeshGroupId MeshGroupId, const FGeometryCollectionStaticMeshInstance& MeshInstance, int32 InstanceCount, TArrayView<const float> CustomDataFloats);
 
 	/** Add a static mesh for a mesh group */
+	bool BatchUpdateInstancesTransforms(FMeshGroupId MeshGroupId, FMeshId MeshId, int32 StartInstanceIndex, TArrayView<const FTransform> NewInstancesTransforms, bool bWorldSpace = false, bool bMarkRenderStateDirty = false, bool bTeleport = false);
+
+	UE_DEPRECATED(5.3, "BatchUpdateInstancesTransforms Array parameter version is deprecated, use the TArrayView version instead")
 	bool BatchUpdateInstancesTransforms(FMeshGroupId MeshGroupId, FMeshId MeshId, int32 StartInstanceIndex, const TArray<FTransform>& NewInstancesTransforms, bool bWorldSpace = false, bool bMarkRenderStateDirty = false, bool bTeleport = false);
 
 	/** 
