Iris - Debugging info for crash.

#rnx

diff --git a/Engine/Source/Runtime/Experimental/Iris/Core/Private/Iris/ReplicationSystem/ReplicationWriter.cpp b/Engine/Source/Runtime/Experimental/Iris/Core/Private/Iris/ReplicationSystem/ReplicationWriter.cpp
--- a/Engine/Source/Runtime/Experimental/Iris/Core/Private/Iris/ReplicationSystem/ReplicationWriter.cpp
+++ b/Engine/Source/Runtime/Experimental/Iris/Core/Private/Iris/ReplicationSystem/ReplicationWriter.cpp
@@ -679,6 +679,7 @@ void FReplicationWriter::UpdateScope(const FNetBitArrayView& UpdatedScope)
 				FReplicationInfo& OwnerInfo = ReplicatedObjects[ObjectData.SubObjectRootIndex];
 				if (OwnerInfo.GetState() < EReplicatedObjectState::PendingDestroy)
 				{
+					ensureAlwaysMsgf(!OwnerInfo.TearOff, TEXT("Parent is tearing off ( InternalIndex: %u ) currently in State: %s "), ObjectData.SubObjectRootIndex, LexToString(OwnerInfo.GetState()));
 					OwnerInfo.HasDirtySubObjects |= Info.HasDirtyChangeMask;
 					ObjectsWithDirtyChanges.SetBitValue(ObjectData.SubObjectRootIndex, ObjectsWithDirtyChanges.GetBit(ObjectData.SubObjectRootIndex) || Info.HasDirtyChangeMask);
 				}
@@ -688,8 +689,8 @@ void FReplicationWriter::UpdateScope(const FNetBitArrayView& UpdatedScope)
 		{
 			UE_LOG_REPLICATIONWRITER_CONN(TEXT("New object added to scope, Waiting to start replication for ( InternalIndex: %u ) currently in State: %s "), Index, LexToString(State));
 
-			check(ObjectsWithDirtyChanges.GetBit(Index) == false);
-			check(Info.HasDirtyChangeMask == 0U);
+			ensureAlwaysMsgf(!ObjectsWithDirtyChanges.GetBit(Index) , TEXT("New object added to scope, Waiting to start replication for ( InternalIndex: %u ) currently in State: %s "), Index, LexToString(State));
+			ensureAlwaysMsgf(!Info.HasDirtyChangeMask, TEXT("New object added to scope, Waiting to start replication for ( InternalIndex: %u ) currently in State: %s "), Index, LexToString(State));
 		}
 	};
 
