Iris
- Implemented support for capturing object references when using the LastResortNetSerializer

diff --git a/Engine/Plugins/Runtime/ReplicationSystemTestPlugin/Source/Private/Tests/Serialization/TestLastResortPropertyNetSerializer.cpp b/Engine/Plugins/Runtime/ReplicationSystemTestPlugin/Source/Private/Tests/Serialization/TestLastResortPropertyNetSerializer.cpp
--- a/Engine/Plugins/Runtime/ReplicationSystemTestPlugin/Source/Private/Tests/Serialization/TestLastResortPropertyNetSerializer.cpp
+++ b/Engine/Plugins/Runtime/ReplicationSystemTestPlugin/Source/Private/Tests/Serialization/TestLastResortPropertyNetSerializer.cpp
@@ -4,6 +4,8 @@
 #include "TestNetSerializerFixture.h"
 #include "Iris/ReplicationState/ReplicationStateDescriptorBuilder.h"
 #include "Iris/Serialization/InternalNetSerializers.h"
+#include "Iris/Serialization/InternalNetSerializationContext.h"
+#include "Iris/Serialization/IrisObjectReferencePackageMap.h"
 
 namespace UE::Net
 {
@@ -47,6 +49,7 @@ protected:
 	static const FText TestValues[];
 	static const SIZE_T TestValueCount;
 
+	Private::FInternalNetSerializationContext InternalContext;
 	TRefCountPtr<const FReplicationStateDescriptor> ReplicationStateDescriptor;
 	FLastResortPropertyNetSerializerConfig* SerializerConfig;
 };
@@ -84,6 +87,9 @@ void FTestLastResortPropertyNetSerializer::SetUp()
 {
 	Super::SetUp();
 
+	InternalContext.PackageMap = NewObject<UIrisObjectReferencePackageMap>();
+	Context.SetInternalContext(&InternalContext);
+
 	ReplicationStateDescriptor = FReplicationStateDescriptorBuilder::CreateDescriptorForStruct(StaticStruct<FStructForLastResortPropertyNetSerializerTest>());
 	UE_NET_ASSERT_EQ_MSG(ReplicationStateDescriptor->MemberCount, uint16(1), "Expected FStructForLastResortPropertyNetSerializerTest to contain exactly one member");
 
diff --git a/Engine/Source/Runtime/Experimental/Iris/Core/Private/Iris/Core/IrisDebugging.cpp b/Engine/Source/Runtime/Experimental/Iris/Core/Private/Iris/Core/IrisDebugging.cpp
--- a/Engine/Source/Runtime/Experimental/Iris/Core/Private/Iris/Core/IrisDebugging.cpp
+++ b/Engine/Source/Runtime/Experimental/Iris/Core/Private/Iris/Core/IrisDebugging.cpp
@@ -282,6 +282,7 @@ void NetObjectStateToString(FStringBuilderBase& StringBuilder, FNetRefHandle Ref
 	FInternalNetSerializationContext InternalContext;
 	FInternalNetSerializationContext::FInitParameters InternalContextInitParams;
 	InternalContextInitParams.ReplicationSystem = ReplicationSystem;
+	InternalContextInitParams.PackageMap = ReplicationSystemInternal->GetIrisObjectReferencePackageMap();
 	InternalContextInitParams.ObjectResolveContext.RemoteNetTokenStoreState = TokenStoreState;
 	InternalContextInitParams.ObjectResolveContext.ConnectionId = (FirstValidConnectionId == FNetBitArray::InvalidIndex ? InvalidConnectionId : FirstValidConnectionId);
 	InternalContext.Init(InternalContextInitParams);
diff --git a/Engine/Source/Runtime/Experimental/Iris/Core/Private/Iris/ReplicationSystem/ReplicationReader.cpp b/Engine/Source/Runtime/Experimental/Iris/Core/Private/Iris/ReplicationSystem/ReplicationReader.cpp
--- a/Engine/Source/Runtime/Experimental/Iris/Core/Private/Iris/ReplicationSystem/ReplicationReader.cpp
+++ b/Engine/Source/Runtime/Experimental/Iris/Core/Private/Iris/ReplicationSystem/ReplicationReader.cpp
@@ -1485,6 +1485,7 @@ void FReplicationReader::ResolveAndDispatchUnresolvedReferences()
 	FInternalNetSerializationContext::FInitParameters InternalContextInitParams;
 	InternalContextInitParams.ReplicationSystem = Parameters.ReplicationSystem;
 	InternalContextInitParams.ObjectResolveContext = ResolveContext;
+	InternalContextInitParams.PackageMap = ReplicationSystemInternal->GetIrisObjectReferencePackageMap();
 	InternalContext.Init(InternalContextInitParams);
 
 	FNetSerializationContext Context;
@@ -1727,6 +1728,7 @@ void FReplicationReader::ProcessQueuedBatches()
 	FInternalNetSerializationContext::FInitParameters InternalContextInitParams;
 	InternalContextInitParams.ReplicationSystem = Parameters.ReplicationSystem;
 	InternalContextInitParams.ObjectResolveContext = ResolveContext;
+	InternalContextInitParams.PackageMap = ReplicationSystemInternal->GetIrisObjectReferencePackageMap();
 	InternalContext.Init(InternalContextInitParams);
 
 	FNetBitStreamReader Reader;
@@ -1854,6 +1856,7 @@ void FReplicationReader::Read(FNetSerializationContext& Context)
 	FInternalNetSerializationContext InternalContext;
 	FInternalNetSerializationContext::FInitParameters InternalContextInitParams;
 	InternalContextInitParams.ReplicationSystem = Parameters.ReplicationSystem;
+	InternalContextInitParams.PackageMap = ReplicationSystemInternal->GetIrisObjectReferencePackageMap();
 	InternalContextInitParams.ObjectResolveContext = ResolveContext;
 	InternalContext.Init(InternalContextInitParams);
 	
diff --git a/Engine/Source/Runtime/Experimental/Iris/Core/Private/Iris/ReplicationSystem/ReplicationSystem.cpp b/Engine/Source/Runtime/Experimental/Iris/Core/Private/Iris/ReplicationSystem/ReplicationSystem.cpp
--- a/Engine/Source/Runtime/Experimental/Iris/Core/Private/Iris/ReplicationSystem/ReplicationSystem.cpp
+++ b/Engine/Source/Runtime/Experimental/Iris/Core/Private/Iris/ReplicationSystem/ReplicationSystem.cpp
@@ -21,6 +21,7 @@
 #include "Iris/Serialization/InternalNetSerializationContext.h"
 #include "Iris/Serialization/NetBitStreamWriter.h"
 #include "Iris/Serialization/NetSerializer.h"
+#include "Iris/Serialization/IrisObjectReferencePackageMap.h"
 #include "HAL/IConsoleManager.h"
 #include "ProfilingDebugging/CsvProfiler.h"
 #include "UObject/UObjectGlobals.h"
@@ -201,6 +202,13 @@ public:
 
 		ReplicationSystemInternal.GetObjectReferenceCache().Init(ReplicationSystem);
 
+		// Init custom packagemap we use for capturing references for backwards compatible NetSerializers
+		{
+			UIrisObjectReferencePackageMap* ObjectReferencePackageMap = NewObject<UIrisObjectReferencePackageMap>();
+			ObjectReferencePackageMap->AddToRoot();
+			ReplicationSystemInternal.SetIrisObjectReferencePackageMap(ObjectReferencePackageMap);
+		}
+
 		FNetBlobManager& BlobManager = ReplicationSystemInternal.GetNetBlobManager();
 		{
 			FNetBlobManagerInitParams InitParams = {};
@@ -216,6 +224,13 @@ public:
 
 		// Reset replication bridge
 		ReplicationSystemInternal.GetReplicationBridge()->Deinitialize();
+
+		if (UIrisObjectReferencePackageMap* ObjectReferencePackageMap = ReplicationSystemInternal.GetIrisObjectReferencePackageMap())
+		{
+			ObjectReferencePackageMap->RemoveFromRoot();
+			ObjectReferencePackageMap->MarkAsGarbage();
+			ReplicationSystemInternal.SetIrisObjectReferencePackageMap(static_cast<UIrisObjectReferencePackageMap*>(nullptr));
+		}
 	}
 
 	void UpdateDirtyObjectList()
diff --git a/Engine/Source/Runtime/Experimental/Iris/Core/Private/Iris/ReplicationSystem/ReplicationSystemInternal.h b/Engine/Source/Runtime/Experimental/Iris/Core/Private/Iris/ReplicationSystem/ReplicationSystemInternal.h
--- a/Engine/Source/Runtime/Experimental/Iris/Core/Private/Iris/ReplicationSystem/ReplicationSystemInternal.h
+++ b/Engine/Source/Runtime/Experimental/Iris/Core/Private/Iris/ReplicationSystem/ReplicationSystemInternal.h
@@ -26,6 +26,8 @@
 #include "Iris/ReplicationState/ReplicationStateDescriptorRegistry.h"
 #include "Iris/Stats/NetStats.h"
 
+class UIrisObjectReferencePackageMap;
+
 namespace UE::Net::Private
 {
 
@@ -41,6 +43,8 @@ public:
 	explicit FReplicationSystemInternal(const FReplicationSystemInternalInitParams& Params)
 	: NetRefHandleManager(ReplicationProtocolManager, Params.ReplicationSystemId, Params.MaxReplicatedObjectCount)
 	, DirtyNetObjectTracker()
+	, ReplicationBridge(nullptr)
+	, IrisObjectReferencePackageMap(nullptr)
 	, StringTokenStore(NetTokenStore)
 	, Id(Params.ReplicationSystemId)
 	{}
@@ -62,6 +66,9 @@ public:
 	UReplicationBridge* GetReplicationBridge() const { return ReplicationBridge; }
 	UReplicationBridge* GetReplicationBridge(FNetRefHandle Handle) const { return ReplicationBridge; }
 
+	void SetIrisObjectReferencePackageMap(UIrisObjectReferencePackageMap* InIrisObjectReferencePackageMap) { IrisObjectReferencePackageMap = InIrisObjectReferencePackageMap; }
+	UIrisObjectReferencePackageMap* GetIrisObjectReferencePackageMap() { return IrisObjectReferencePackageMap; }
+
 	FChangeMaskCache& GetChangeMaskCache() { return ChangeMaskCache; }
 
 	FReplicationConnections& GetConnections() { return Connections; }
@@ -103,6 +110,7 @@ private:
 	FReplicationStateStorage ReplicationStateStorage;
 	FReplicationStateDescriptorRegistry ReplicationStateDescriptorRegistry;
 	UReplicationBridge* ReplicationBridge;
+	UIrisObjectReferencePackageMap* IrisObjectReferencePackageMap;
 	FChangeMaskCache ChangeMaskCache;
 	FReplicationConnections Connections;
 	FReplicationFiltering Filtering;
diff --git a/Engine/Source/Runtime/Experimental/Iris/Core/Private/Iris/Serialization/InternalNetSerializationContext.cpp b/Engine/Source/Runtime/Experimental/Iris/Core/Private/Iris/Serialization/InternalNetSerializationContext.cpp
--- a/Engine/Source/Runtime/Experimental/Iris/Core/Private/Iris/Serialization/InternalNetSerializationContext.cpp
+++ b/Engine/Source/Runtime/Experimental/Iris/Core/Private/Iris/Serialization/InternalNetSerializationContext.cpp
@@ -11,6 +11,7 @@ namespace UE::Net::Private
 FInternalNetSerializationContext::FInternalNetSerializationContext(UReplicationSystem* InReplicationSystem)
 : ReplicationSystem(InReplicationSystem)
 , ObjectReferenceCache(&InReplicationSystem->GetReplicationSystemInternal()->GetObjectReferenceCache())
+, PackageMap(InReplicationSystem->GetReplicationSystemInternal()->GetIrisObjectReferencePackageMap())
 , bDowngradeAutonomousProxyRole(0)
 , bInlineObjectReferenceExports(0)
 {
@@ -35,6 +36,7 @@ void FInternalNetSerializationContext::Init(const FInitParameters& InitParams)
 {
 	ReplicationSystem = InitParams.ReplicationSystem;
 	ObjectReferenceCache = &InitParams.ReplicationSystem->GetReplicationSystemInternal()->GetObjectReferenceCache();
+	PackageMap = InitParams.PackageMap;
 	ResolveContext = InitParams.ObjectResolveContext;
 }
 
diff --git a/Engine/Source/Runtime/Experimental/Iris/Core/Private/Iris/Serialization/InternalNetSerializationContext.h b/Engine/Source/Runtime/Experimental/Iris/Core/Private/Iris/Serialization/InternalNetSerializationContext.h
--- a/Engine/Source/Runtime/Experimental/Iris/Core/Private/Iris/Serialization/InternalNetSerializationContext.h
+++ b/Engine/Source/Runtime/Experimental/Iris/Core/Private/Iris/Serialization/InternalNetSerializationContext.h
@@ -8,6 +8,7 @@
 
 class UReplicationSystem;
 class UObject;
+class UIrisObjectReferencePackageMap;
 
 namespace UE::Net::Private
 {
@@ -24,6 +25,7 @@ public:
 	{
 		UReplicationSystem* ReplicationSystem = nullptr;
 		FNetObjectResolveContext ObjectResolveContext;
+		UIrisObjectReferencePackageMap* PackageMap = nullptr;			
 	};
 
 	FInternalNetSerializationContext();
@@ -40,6 +42,9 @@ public:
 	FObjectReferenceCache* ObjectReferenceCache;
 	FNetObjectResolveContext ResolveContext;
 
+	// We have a special implementation of UPackageMap to capture references when using LastResortNetSerializer
+	UIrisObjectReferencePackageMap* PackageMap = nullptr;
+
 	// $IRIS TODO Roles really shouldn't be replicated as properties. In dire need of a proper authority system.
 	// This is ONLY to be used by role serialization.
 	uint32 bDowngradeAutonomousProxyRole : 1;
@@ -51,6 +56,7 @@ public:
 inline FInternalNetSerializationContext::FInternalNetSerializationContext()
 : ReplicationSystem(nullptr)
 , ObjectReferenceCache(nullptr)
+, PackageMap(nullptr)
 , bDowngradeAutonomousProxyRole(0)
 , bInlineObjectReferenceExports(0)
 {}
diff --git a/Engine/Source/Runtime/Experimental/Iris/Core/Private/Iris/Serialization/InternalNetSerializers.h b/Engine/Source/Runtime/Experimental/Iris/Core/Private/Iris/Serialization/InternalNetSerializers.h
--- a/Engine/Source/Runtime/Experimental/Iris/Core/Private/Iris/Serialization/InternalNetSerializers.h
+++ b/Engine/Source/Runtime/Experimental/Iris/Core/Private/Iris/Serialization/InternalNetSerializers.h
@@ -73,6 +73,10 @@ struct FLastResortPropertyNetSerializerConfig : public FNetSerializerConfig
 public:
 	UPROPERTY()
 	TFieldPath<FProperty> Property;
+
+	// Value used to sanity check incoming data so that we do not over-allocate dynamic memory
+	UPROPERTY()
+	uint32 MaxAllowedObjectReferences = 128; 
 };
 
 // ENetRole. With role swapping at deserialization
diff --git a/Engine/Source/Runtime/Experimental/Iris/Core/Private/Iris/Serialization/LastResortPropertyNetSerializer.cpp b/Engine/Source/Runtime/Experimental/Iris/Core/Private/Iris/Serialization/LastResortPropertyNetSerializer.cpp
--- a/Engine/Source/Runtime/Experimental/Iris/Core/Private/Iris/Serialization/LastResortPropertyNetSerializer.cpp
+++ b/Engine/Source/Runtime/Experimental/Iris/Core/Private/Iris/Serialization/LastResortPropertyNetSerializer.cpp
@@ -4,8 +4,34 @@
 #include "UObject/CoreNet.h"
 #include "Iris/Serialization/NetBitStreamReader.h"
 #include "Iris/Serialization/NetBitStreamUtil.h"
+#include "Iris/Serialization/NetSerializers.h"
 #include "Iris/Serialization/NetBitStreamWriter.h"
 #include "Iris/Serialization/InternalNetSerializationContext.h"
+#include "Iris/Serialization/NetReferenceCollector.h"
+#include "Iris/Serialization/NetSerializerArrayStorage.h"
+#include "Iris/Serialization/IrisObjectReferencePackageMap.h"
+#include "Net/Core/Trace/NetTrace.h"
+
+namespace UE::Net
+{
+
+struct FFLastResortPropertyNetSerializerQuantizedType
+{
+	static constexpr uint32 MaxInlinedObjectRefs = 4;
+	typedef FNetSerializerArrayStorage<FNetObjectReference, AllocationPolicies::TInlinedElementAllocationPolicy<MaxInlinedObjectRefs>> FObjectReferenceStorage;
+
+	FObjectReferenceStorage ObjectReferenceStorage;
+
+	// How many bytes the current allocation can hold.
+	uint16 ByteCapacity = 0;
+	// How many bits are valid
+	uint16 BitCount = 0;
+	void* Storage = nullptr;
+};
+
+}
+
+template <> struct TIsPODType<UE::Net::FFLastResortPropertyNetSerializerQuantizedType> { enum { Value = true }; };
 
 namespace UE::Net
 {
@@ -17,20 +43,13 @@ struct FLastResortPropertyNetSerializer
 
 	// Traits
 	static constexpr bool bHasDynamicState = true;
+	static constexpr bool bHasCustomNetReference = true;
 
 	// Types
-	struct FQuantizedType
-	{
-		// How many bytes the current allocation can hold.
-		uint16 ByteCapacity;
-		// How many bits are valid
-		uint16 BitCount;
-		void* Storage;
-	};
 
 	// SourceType is unknown
 	typedef void SourceType;
-	typedef FQuantizedType QuantizedType;
+	typedef FFLastResortPropertyNetSerializerQuantizedType QuantizedType;
 	typedef FLastResortPropertyNetSerializerConfig ConfigType;
 
 	//
@@ -45,6 +64,8 @@ struct FLastResortPropertyNetSerializer
 	static void CloneDynamicState(FNetSerializationContext&, const FNetCloneDynamicStateArgs&);
 	static void FreeDynamicState(FNetSerializationContext&, const FNetFreeDynamicStateArgs&);
 
+	static void CollectNetReferences(FNetSerializationContext&, const FNetCollectReferencesArgs&);
+
 private:
 	static constexpr uint32 AllocationAlignment = 4U;
 
@@ -52,14 +73,35 @@ private:
 	static void GrowDynamicStateInternal(FNetSerializationContext&, QuantizedType& Value, uint16 NewBitCount);
 	static void ShrinkDynamicStateInternal(FNetSerializationContext&, QuantizedType& Value, uint16 NewBitCount);
 	static void AdjustStorageSize(FNetSerializationContext&, QuantizedType& Value, uint16 NewBitCount);
+
+	static inline const FNetSerializer* ObjectNetSerializer = &UE_NET_GET_SERIALIZER(FObjectNetSerializer);
 };
 UE_NET_IMPLEMENT_SERIALIZER_INTERNAL(FLastResortPropertyNetSerializer);
 
 void FLastResortPropertyNetSerializer::Serialize(FNetSerializationContext& Context, const FNetSerializeArgs& Args)
 {
+	const ConfigType* Config = static_cast<const ConfigType*>(Args.NetSerializerConfig);
 	const QuantizedType& Value = *reinterpret_cast<const QuantizedType*>(Args.Source);
-
 	FNetBitStreamWriter* Writer = Context.GetBitStreamWriter();
+	const uint32 NumReferences = Value.ObjectReferenceStorage.Num();
+
+	UE_NET_TRACE_DYNAMIC_NAME_SCOPE(Config->Property.Get()->GetName(), *Writer, Context.GetTraceCollector(), ENetTraceVerbosity::VeryVerbose);
+
+	// If we have any references, export them!
+	if (Writer->WriteBool(NumReferences != 0))
+	{
+		UE::Net::WritePackedUint32(Writer, NumReferences);		
+		FObjectNetSerializerConfig ObjectSerializerConfig;
+		for (const FNetObjectReference& Ref : MakeArrayView(Value.ObjectReferenceStorage.GetData(), NumReferences))
+		{
+			FNetSerializeArgs ObjectArgs;
+			ObjectArgs.NetSerializerConfig = &ObjectSerializerConfig;
+			ObjectArgs.Source = NetSerializerValuePointer(&Ref);
+
+			ObjectNetSerializer->Serialize(Context, ObjectArgs);
+		}
+	}
+
 	WritePackedUint32(Writer, Value.BitCount);
 	if (Value.BitCount > 0)
 	{
@@ -69,10 +111,43 @@ void FLastResortPropertyNetSerializer::Serialize(FNetSerializationContext& Conte
 
 void FLastResortPropertyNetSerializer::Deserialize(FNetSerializationContext& Context, const FNetDeserializeArgs& Args)
 {
+	const ConfigType* Config = static_cast<const ConfigType*>(Args.NetSerializerConfig);
 	QuantizedType& Value = *reinterpret_cast<QuantizedType*>(Args.Target);
 	const uint32 CurrentBitCount = Value.BitCount;
 
 	FNetBitStreamReader* Reader = Context.GetBitStreamReader();
+
+	UE_NET_TRACE_DYNAMIC_NAME_SCOPE(Config->Property.Get()->GetName(), *Reader, Context.GetTraceCollector(), ENetTraceVerbosity::VeryVerbose);
+
+	// Read any object references
+	const bool bHasObjectReferences = Reader->ReadBool();
+	if (bHasObjectReferences)
+	{
+		const uint32 NumReferences = UE::Net::ReadPackedUint32(Reader);
+
+		if (NumReferences > Config->MaxAllowedObjectReferences)
+		{
+			Context.SetError(GNetError_ArraySizeTooLarge);
+			return;
+		}
+
+		Value.ObjectReferenceStorage.AdjustSize(Context, NumReferences);
+
+		FObjectNetSerializerConfig ObjectSerializerConfig;
+		for (FNetObjectReference& Ref : MakeArrayView(Value.ObjectReferenceStorage.GetData(), Value.ObjectReferenceStorage.Num()))
+		{
+			FNetDeserializeArgs ObjectArgs;
+			ObjectArgs.NetSerializerConfig = &ObjectSerializerConfig;
+			ObjectArgs.Target = NetSerializerValuePointer(&Ref);
+
+			ObjectNetSerializer->Deserialize(Context, ObjectArgs);
+		}
+	}
+	else
+	{
+		Value.ObjectReferenceStorage.Free(Context);
+	}
+
 	const uint32 NewBitCount = ReadPackedUint32(Reader);
 	if (NewBitCount > 65535U)
 	{
@@ -88,10 +163,29 @@ void FLastResortPropertyNetSerializer::Quantize(FNetSerializationContext& Contex
 {
 	const ConfigType* Config = static_cast<const ConfigType*>(Args.NetSerializerConfig);
 	const FProperty* Property = Config->Property.Get();
+	QuantizedType& Value = *reinterpret_cast<QuantizedType*>(Args.Target);
+
+	// Capture references
+	Private::FInternalNetSerializationContext* InternalContext = Context.GetInternalContext();
+	UIrisObjectReferencePackageMap* PackageMap = InternalContext ? InternalContext->PackageMap : nullptr;
+
+	// We do not store any defaults for properties using the LastResortPropertyNetSerializer.
+	if (Context.IsInitializingDefaultState())
+	{
+		Value.ObjectReferenceStorage.AdjustSize(Context, 0);
+		AdjustStorageSize(Context, Value, 0);
+		return;
+	}
+	
+	UIrisObjectReferencePackageMap::FObjectReferenceArray ObjectReferences;
+	if (PackageMap)
+	{
+		PackageMap->InitForWrite(&ObjectReferences);
+	}
 
 	// Use the Property serialization and store as binary blob.
-	FNetBitWriter Archive(8192);
-	Property->NetSerializeItem(Archive, static_cast<UPackageMap*>(nullptr), reinterpret_cast<void*>(Args.Source));
+	FNetBitWriter Archive(PackageMap, 8192);
+	Property->NetSerializeItem(Archive, PackageMap, reinterpret_cast<void*>(Args.Source));
 
 	const uint64 BitCount = Archive.GetNumBits();
 	if (BitCount > 65535U)
@@ -100,7 +194,27 @@ void FLastResortPropertyNetSerializer::Quantize(FNetSerializationContext& Contex
 		return;
 	}
 
-	QuantizedType& Value = *reinterpret_cast<QuantizedType*>(Args.Target);
+
+	// Capture any references
+	const uint32 NumObjectReferences = ObjectReferences.Num();
+	Value.ObjectReferenceStorage.AdjustSize(Context, NumObjectReferences);
+	if (NumObjectReferences > 0)
+	{
+		FObjectNetSerializerConfig ObjectNetSerializerConfig;
+		const TObjectPtr<UObject>* SourceReferences = ObjectReferences.GetData();
+		FNetObjectReference* TargetReferences = Value.ObjectReferenceStorage.GetData();
+		for (uint32 ReferenceIndex = 0; ReferenceIndex < NumObjectReferences; ++ReferenceIndex)
+		{
+			FNetQuantizeArgs ObjectArgs;
+			ObjectArgs.NetSerializerConfig = &ObjectNetSerializerConfig;
+			ObjectArgs.Source = NetSerializerValuePointer(SourceReferences + ReferenceIndex);
+			ObjectArgs.Target = NetSerializerValuePointer(TargetReferences + ReferenceIndex);
+
+			ObjectNetSerializer->Quantize(Context, ObjectArgs);
+		}
+	}
+
+	// Deal with serialized data
 	AdjustStorageSize(Context, Value, static_cast<uint16>(BitCount));
 	if (BitCount > 0)
 	{
@@ -114,8 +228,43 @@ void FLastResortPropertyNetSerializer::Dequantize(FNetSerializationContext& Cont
 	const FProperty* Property = Config->Property.Get();
 
 	QuantizedType& Source = *reinterpret_cast<QuantizedType*>(Args.Source);
-	FNetBitReader  Archive(nullptr, static_cast<uint8*>(Source.Storage), Source.BitCount);
-	Property->NetSerializeItem(Archive, static_cast<UPackageMap*>(nullptr), reinterpret_cast<void*>(Args.Target));
+
+	// Inject references
+	Private::FInternalNetSerializationContext* InternalContext = Context.GetInternalContext();
+	UIrisObjectReferencePackageMap* PackageMap = InternalContext ? InternalContext->PackageMap : nullptr;
+	UIrisObjectReferencePackageMap::FObjectReferenceArray ObjectReferences;
+
+	// References
+	const uint32 NumObjectReferences = Source.ObjectReferenceStorage.Num();
+	if (NumObjectReferences > 0U && ensureAlwaysMsgf(PackageMap, TEXT("FLastResortPropertyNetSerializer::Dequantize must have a packagemap to be able to dequantize object references. Make sure it is set in the InternalContext.")))
+	{		
+		ObjectReferences.SetNumUninitialized(NumObjectReferences);
+
+		FObjectNetSerializerConfig ObjectNetSerializerConfig;
+		const FNetObjectReference* SourceReferences = Source.ObjectReferenceStorage.GetData();
+		TObjectPtr<UObject>* TargetReferences = ObjectReferences.GetData();
+		for (uint32 ReferenceIndex = 0; ReferenceIndex < NumObjectReferences; ++ReferenceIndex)
+		{
+			FNetDequantizeArgs ObjectArgs;
+			ObjectArgs.NetSerializerConfig = &ObjectNetSerializerConfig;
+			ObjectArgs.Source = NetSerializerValuePointer(SourceReferences + ReferenceIndex);
+			ObjectArgs.Target = NetSerializerValuePointer(TargetReferences + ReferenceIndex);
+
+			ObjectNetSerializer->Dequantize(Context, ObjectArgs);
+		}
+
+		PackageMap->InitForRead(&ObjectReferences);
+	}
+
+	if (Source.BitCount)
+	{
+		FNetBitReader Archive(PackageMap, static_cast<uint8*>(Source.Storage), Source.BitCount);
+		Property->NetSerializeItem(Archive, PackageMap, reinterpret_cast<void*>(Args.Target));
+	}
+	else
+	{
+		Property->ClearValue(reinterpret_cast<void*>(Args.Target));
+	}
 }
 
 bool FLastResortPropertyNetSerializer::IsEqual(FNetSerializationContext& Context, const FNetIsEqualArgs& Args)
@@ -124,11 +273,17 @@ bool FLastResortPropertyNetSerializer::IsEqual(FNetSerializationContext& Context
 	{
 		const QuantizedType& Value0 = *reinterpret_cast<const QuantizedType*>(Args.Source0);
 		const QuantizedType& Value1 = *reinterpret_cast<const QuantizedType*>(Args.Source1);
-		if (Value0.BitCount != Value1.BitCount)
+		if ((Value0.BitCount != Value1.BitCount) || (Value0.ObjectReferenceStorage.Num() != Value1.ObjectReferenceStorage.Num()))
 		{
 			return false;
 		}
-		const bool bIsEqual = FMemory::Memcmp(Value0.Storage, Value1.Storage, Align((Value0.BitCount + 7U)/8U, AllocationAlignment)) == 0;
+
+		if (Value0.ObjectReferenceStorage.Num() > 0 && FMemory::Memcmp(Value0.ObjectReferenceStorage.GetData(), Value1.ObjectReferenceStorage.GetData(), sizeof(FNetObjectReference) * Value0.ObjectReferenceStorage.Num()) != 0)
+		{
+			return false;
+		}
+
+		const bool bIsEqual = (Value0.BitCount == 0U) || FMemory::Memcmp(Value0.Storage, Value1.Storage, Align((Value0.BitCount + 7U)/8U, AllocationAlignment)) == 0;
 		return bIsEqual;
 	}
 	else
@@ -159,6 +314,8 @@ void FLastResortPropertyNetSerializer::CloneDynamicState(FNetSerializationContex
 	Target.ByteCapacity = ByteCount;
 	Target.BitCount = Source.BitCount;
 	Target.Storage = Storage;
+
+	Target.ObjectReferenceStorage.Clone(Context, Source.ObjectReferenceStorage);
 }
 
 void FLastResortPropertyNetSerializer::FreeDynamicState(FNetSerializationContext& Context, const FNetFreeDynamicStateArgs& Args)
@@ -168,10 +325,13 @@ void FLastResortPropertyNetSerializer::FreeDynamicState(FNetSerializationContext
 
 void FLastResortPropertyNetSerializer::FreeDynamicStateInternal(FNetSerializationContext& Context, QuantizedType& Value)
 {
+	// Clear all info
+	Value.ObjectReferenceStorage.Free(Context);
 	Context.GetInternalContext()->Free(Value.Storage);
 
-	// Clear all info
-	Value = QuantizedType();
+	Value.BitCount = 0;
+	Value.ByteCapacity = 0;
+	Value.Storage = 0;
 }
 
 void FLastResortPropertyNetSerializer::GrowDynamicStateInternal(FNetSerializationContext& Context, QuantizedType& Value, uint16 NewBitCount)
@@ -223,4 +383,18 @@ bool InitLastResortPropertyNetSerializerConfigFromProperty(FLastResortPropertyNe
 	return Property != nullptr;
 }
 
+void FLastResortPropertyNetSerializer::CollectNetReferences(FNetSerializationContext& Context, const FNetCollectReferencesArgs& Args)
+{
+	const ConfigType& Config = *static_cast<const ConfigType*>(Args.NetSerializerConfig);
+	const QuantizedType& Value = *reinterpret_cast<const QuantizedType*>(Args.Source);
+	FNetReferenceCollector& Collector = *reinterpret_cast<UE::Net::FNetReferenceCollector*>(Args.Collector);
+
+	const FNetReferenceInfo ReferenceInfo(FNetReferenceInfo::EResolveType::ResolveOnClient);	
+	for (const FNetObjectReference& Ref : MakeArrayView(Value.ObjectReferenceStorage.GetData(), Value.ObjectReferenceStorage.Num()))
+	{
+		Collector.Add(ReferenceInfo, Ref, Args.ChangeMaskInfo);
+	}
+}
+
+
 }
diff --git a/Engine/Source/Runtime/Experimental/Iris/Core/Public/Iris/ReplicationState/PropertyNetSerializerInfoRegistry.h b/Engine/Source/Runtime/Experimental/Iris/Core/Public/Iris/ReplicationState/PropertyNetSerializerInfoRegistry.h
--- a/Engine/Source/Runtime/Experimental/Iris/Core/Public/Iris/ReplicationState/PropertyNetSerializerInfoRegistry.h
+++ b/Engine/Source/Runtime/Experimental/Iris/Core/Public/Iris/ReplicationState/PropertyNetSerializerInfoRegistry.h
@@ -208,7 +208,7 @@ UE::Net::FPropertyNetSerializerInfoRegistry::Unregister(&GetPropertyNetSerialize
 
 // Implement minimal required delegates for a NetSerializer
 #define UE_NET_IMPLEMENT_NETSERIALIZER_REGISTRY_DELEGATES(Name) \
-struct F##Name##NetSerializerRegistryDelegates : protected FNetSerializerRegistryDelegates \
+struct F##Name##NetSerializerRegistryDelegates : protected UE::Net::FNetSerializerRegistryDelegates \
 { \
 	~F##Name##NetSerializerRegistryDelegates() { UE_NET_UNREGISTER_NETSERIALIZER_INFO(PropertyNetSerializerRegistry_NAME_##Name); } \
 	virtual void OnPreFreezeNetSerializerRegistry() override { UE_NET_REGISTER_NETSERIALIZER_INFO(PropertyNetSerializerRegistry_NAME_##Name); } \
