When geometry collection is empty, don't build the render mesh data.
This fixes crash when trying to render with empty mesh data.

diff --git a/Engine/Source/Runtime/Experimental/GeometryCollectionEngine/Private/GeometryCollection/GeometryCollectionRenderData.cpp b/Engine/Source/Runtime/Experimental/GeometryCollectionEngine/Private/GeometryCollection/GeometryCollectionRenderData.cpp
--- a/Engine/Source/Runtime/Experimental/GeometryCollectionEngine/Private/GeometryCollection/GeometryCollectionRenderData.cpp
+++ b/Engine/Source/Runtime/Experimental/GeometryCollectionEngine/Private/GeometryCollection/GeometryCollectionRenderData.cpp
@@ -694,14 +694,16 @@ TUniquePtr<FGeometryCollectionRenderData> FGeometryCollectionRenderData::Create(
 	FGeometryCollectionBuiltMeshData MeshBuildData;
 	if (BuildMeshDataFromGeometryCollection(InCollection, MeshBuildData))
 	{
-	if (bInEnableNanite)
 	{
-		CreateNaniteData(MoveTemp(MeshBuildData), *RenderData.Get());
-	}
-	else
-	{
-		// Could always create mesh data if we want to be able to enable/disable nanite at runtime in cooked build.
-		CreateMeshData(InCollection, MeshBuildData, bInUseFullPrecisionUVs, *RenderData.Get());
+		if (bInEnableNanite)
+		{
+			CreateNaniteData(MoveTemp(MeshBuildData), *RenderData.Get());
+		}
+		else
+		{
+			// Could always create mesh data if we want to be able to enable/disable nanite at runtime in cooked build.
+			CreateMeshData(InCollection, MeshBuildData, bInUseFullPrecisionUVs, *RenderData.Get());
+		}
 	}
 	}
 
