add insights scope to FClusterUnionManager::ForceRecreateClusterUnionSharedGeometry

diff --git a/Engine/Source/Runtime/Experimental/Chaos/Private/Chaos/ClusterUnionManager.cpp b/Engine/Source/Runtime/Experimental/Chaos/Private/Chaos/ClusterUnionManager.cpp
--- a/Engine/Source/Runtime/Experimental/Chaos/Private/Chaos/ClusterUnionManager.cpp
+++ b/Engine/Source/Runtime/Experimental/Chaos/Private/Chaos/ClusterUnionManager.cpp
@@ -111,6 +111,7 @@ namespace Chaos
 	DECLARE_CYCLE_STAT(TEXT("FClusterUnionManager::ForceRecreateClusterUnionSharedGeometry"), STAT_ForceRecreateClusterUnionSharedGeometry, STATGROUP_Chaos);
 	TSharedPtr<FImplicitObject, ESPMode::ThreadSafe> FClusterUnionManager::ForceRecreateClusterUnionSharedGeometry(const FClusterUnion& Union)
 	{
+		SCOPE_CYCLE_COUNTER(STAT_ForceRecreateClusterUnionSharedGeometry);
 		if (Union.ChildParticles.IsEmpty() || !Union.InternalCluster)
 		{
 			return MakeShared<FImplicitObjectUnionClustered>();
