I/O - save crypto keys to generated config file to enable replays

diff --git a/Engine/Source/Runtime/Experimental/IoStoreOnDemand/Private/IoStoreOnDemand.cpp b/Engine/Source/Runtime/Experimental/IoStoreOnDemand/Private/IoStoreOnDemand.cpp
--- a/Engine/Source/Runtime/Experimental/IoStoreOnDemand/Private/IoStoreOnDemand.cpp
+++ b/Engine/Source/Runtime/Experimental/IoStoreOnDemand/Private/IoStoreOnDemand.cpp
@@ -2,10 +2,12 @@
 
 #include "IO/IoStoreOnDemand.h"
 #include "OnDemandIoDispatcherBackend.h"
+#include "EncryptionKeyManager.h"
 
 #include "FileCache.h"
 #include "HAL/PlatformMisc.h"
 #include "HAL/FileManager.h"
+#include "Misc/Base64.h"
 #include "Misc/CommandLine.h"
 #include "Misc/ConfigCacheIni.h"
 #include "Modules/ModuleManager.h"
@@ -245,7 +247,9 @@ public:
 	virtual void StartupModule() override;
 	virtual void ShutdownModule() override;
 
+private:
 	static uint64 ParseSizeParam(const TCHAR* Param);
+	static bool ParseEncryptionKeyParam(const FString& Param, FGuid& OutGuid, FAES::FAESKey& OutKey);
 };
 
 IMPLEMENT_MODULE(FIoStoreOnDemandModule, IoStoreOnDemand);
@@ -278,6 +282,27 @@ uint64 FIoStoreOnDemandModule::ParseSizeParam(const TCHAR* Param)
 	return 0;
 }
 
+bool FIoStoreOnDemandModule::ParseEncryptionKeyParam(const FString& Param, FGuid& OutKeyGuid, FAES::FAESKey& OutKey)
+{
+	TArray<FString> Tokens;
+	Param.ParseIntoArray(Tokens, TEXT(":"), true);
+
+	if (Tokens.Num() == 2)
+	{
+		TArray<uint8> KeyBytes;
+		if (FGuid::Parse(Tokens[0], OutKeyGuid) && FBase64::Decode(Tokens[1], KeyBytes))
+		{
+			if (OutKeyGuid != FGuid() && KeyBytes.Num() == FAES::FAESKey::KeySize)
+			{
+				FMemory::Memcpy(OutKey.Key, KeyBytes.GetData(), FAES::FAESKey::KeySize);
+				return true;
+			}
+		}
+	}
+
+	return false;
+}
+
 void FIoStoreOnDemandModule::StartupModule()
 {
 #if !WITH_EDITOR
@@ -297,7 +322,20 @@ void FIoStoreOnDemandModule::StartupModule()
 			}
 		}
 	}
-	
+
+	{
+		FString EncryptionKey;
+		if (FParse::Value(FCommandLine::Get(), TEXT("OnDemandEncryptionKey="), EncryptionKey))
+		{
+			FGuid KeyGuid;
+			FAES::FAESKey Key;
+			if (ParseEncryptionKeyParam(EncryptionKey, KeyGuid, Key))
+			{
+				UE::FEncryptionKeyManager::Get().AddKey(KeyGuid, Key);
+			}
+		}
+	}
+
 	if (!Endpoint.IsValid())
 	{
 		Endpoint = UE::FOnDemandEndpoint();
@@ -328,6 +366,17 @@ void FIoStoreOnDemandModule::StartupModule()
 			{
 				Endpoint.TocPath.RightChopInline(1);
 			}
+
+			FString EncryptionKey;
+			if (Config.GetString(TEXT("Endpoint"), TEXT("EncryptionKey"), EncryptionKey))
+			{
+				FGuid KeyGuid;
+				FAES::FAESKey Key;
+				if (ParseEncryptionKeyParam(EncryptionKey, KeyGuid, Key))
+				{
+					UE::FEncryptionKeyManager::Get().AddKey(KeyGuid, Key);
+				}
+			}
 		}
 	}
 
