Iris
* Added scope and renamed some variables to clarify them in FReplicationPrioritization

#rnx

diff --git a/Engine/Source/Runtime/Experimental/Iris/Core/Private/Iris/ReplicationSystem/Prioritization/ReplicationPrioritization.cpp b/Engine/Source/Runtime/Experimental/Iris/Core/Private/Iris/ReplicationSystem/Prioritization/ReplicationPrioritization.cpp
--- a/Engine/Source/Runtime/Experimental/Iris/Core/Private/Iris/ReplicationSystem/Prioritization/ReplicationPrioritization.cpp
+++ b/Engine/Source/Runtime/Experimental/Iris/Core/Private/Iris/ReplicationSystem/Prioritization/ReplicationPrioritization.cpp
@@ -481,10 +481,10 @@ UNetObjectPrioritizer* FReplicationPrioritization::GetPrioritizer(const FName Pr
 /**
   * Prioritize objects as per connection's wishes.
   */
-void FReplicationPrioritization::Prioritize(const FNetBitArrayView& ReplicatingConnections, const FNetBitArrayView& DirtyObjects)
+void FReplicationPrioritization::Prioritize(const FNetBitArrayView& ReplicatingConnections, const FNetBitArrayView& DirtyObjectsThisFrame)
 {
 	UpdatePrioritiesForNewAndDeletedObjects();
-	NotifyPrioritizersOfDirtyObjects(DirtyObjects);
+	NotifyPrioritizersOfDirtyObjects(DirtyObjectsThisFrame);
 
 	if (!ReplicatingConnections.IsAnyBitSet())
 	{
@@ -739,20 +739,22 @@ void FReplicationPrioritization::PrioritizeForConnection(uint32 ConnId, FPriorit
 /**
  * Notify prioritizers of which objects have been updated since last frame.	 
  */
-void FReplicationPrioritization::NotifyPrioritizersOfDirtyObjects(const FNetBitArrayView& DirtyObjects)
+void FReplicationPrioritization::NotifyPrioritizersOfDirtyObjects(const FNetBitArrayView& DirtyObjectsThisFrame)
 {
+	IRIS_PROFILER_SCOPE(FReplicationPrioritization_NotifyPrioritizersOfDirtyObjects);
+
 	FUpdateDirtyObjectsBatchHelper BatchHelper(NetRefHandleManager, PrioritizerInfos.Num());
 
 	constexpr SIZE_T MaxBatchObjectCount = FUpdateDirtyObjectsBatchHelper::Constants::MaxObjectCountPerBatch;
 	uint32 ObjectIndices[MaxBatchObjectCount];
 
 	const uint32 BitCount = ~0U;
-	for (uint32 ObjectCount, StartIndex = 0; (ObjectCount = DirtyObjects.GetSetBitIndices(StartIndex, BitCount, ObjectIndices, MaxBatchObjectCount)) > 0; )
+	for (uint32 ObjectCount, StartIndex = 0; (ObjectCount = DirtyObjectsThisFrame.GetSetBitIndices(StartIndex, BitCount, ObjectIndices, MaxBatchObjectCount)) > 0; )
 	{
 		BatchNotifyPrioritizersOfDirtyObjects(BatchHelper, ObjectIndices, ObjectCount);
 
 		StartIndex = ObjectIndices[ObjectCount - 1] + 1U;
-		if ((StartIndex == DirtyObjects.GetNumBits()) | (ObjectCount < MaxBatchObjectCount))
+		if ((StartIndex == DirtyObjectsThisFrame.GetNumBits()) | (ObjectCount < MaxBatchObjectCount))
 		{
 			break;
 		}
diff --git a/Engine/Source/Runtime/Experimental/Iris/Core/Private/Iris/ReplicationSystem/Prioritization/ReplicationPrioritization.h b/Engine/Source/Runtime/Experimental/Iris/Core/Private/Iris/ReplicationSystem/Prioritization/ReplicationPrioritization.h
--- a/Engine/Source/Runtime/Experimental/Iris/Core/Private/Iris/ReplicationSystem/Prioritization/ReplicationPrioritization.h
+++ b/Engine/Source/Runtime/Experimental/Iris/Core/Private/Iris/ReplicationSystem/Prioritization/ReplicationPrioritization.h
@@ -37,7 +37,7 @@ public:
 
 	void Init(FReplicationPrioritizationInitParams& Params);
 
-	void Prioritize(const FNetBitArrayView& ConnectionsToSend, const FNetBitArrayView& DirtyObjects);
+	void Prioritize(const FNetBitArrayView& ConnectionsToSend, const FNetBitArrayView& DirtyObjectsThisFrame);
 
 	void SetStaticPriority(uint32 ObjectIndex, float Prio);
 	bool SetPrioritizer(uint32 ObjectIndex, FNetObjectPrioritizerHandle Prioritizer);
@@ -53,7 +53,7 @@ private:
 
 	void UpdatePrioritiesForNewAndDeletedObjects();
 	void PrioritizeForConnection(uint32 ConnId, FPrioritizerBatchHelper& BatchHelper, FNetBitArrayView Objects);
-	void NotifyPrioritizersOfDirtyObjects(const FNetBitArrayView& DirtyObjects);
+	void NotifyPrioritizersOfDirtyObjects(const FNetBitArrayView& DirtyObjectsThisFrame);
 	void BatchNotifyPrioritizersOfDirtyObjects(FUpdateDirtyObjectsBatchHelper& BatchHelper, uint32* ObjectIndices, uint32 ObjectCount);
 	void InitPrioritizers();
 
