geometry collections should not try to replicate empty internal clusters

diff --git a/Engine/Source/Runtime/Experimental/GeometryCollectionEngine/Private/GeometryCollection/GeometryCollectionComponent.cpp b/Engine/Source/Runtime/Experimental/GeometryCollectionEngine/Private/GeometryCollection/GeometryCollectionComponent.cpp
--- a/Engine/Source/Runtime/Experimental/GeometryCollectionEngine/Private/GeometryCollection/GeometryCollectionComponent.cpp
+++ b/Engine/Source/Runtime/Experimental/GeometryCollectionEngine/Private/GeometryCollection/GeometryCollectionComponent.cpp
@@ -1877,7 +1877,9 @@ void UGeometryCollectionComponent::UpdateRepData()
 		for (auto Itr = ClustersToRep->CreateIterator(); Itr; ++Itr)
 		{
 			FPBDRigidClusteredParticleHandle* Cluster = *Itr;
-			if (Cluster->Disabled())
+
+			const TArray<FPBDRigidParticleHandle*>& Children = RigidClustering.GetChildrenMap().FindRef(Cluster);
+			if (Cluster->Disabled() || (Cluster->InternalCluster() && Children.IsEmpty()))
 			{
 				Itr.RemoveCurrent();
 			}
@@ -1895,7 +1897,7 @@ void UGeometryCollectionComponent::UpdateRepData()
 				int32 TransformGroupIdx;
 				if(Cluster->InternalCluster())
 				{
-					const TArray<FPBDRigidParticleHandle*>& Children = RigidClustering.GetChildrenMap()[Cluster];
+					
 					ensureMsgf(Children.Num(), TEXT("Internal cluster yet we have no children?"));
 					TransformGroupIdx = PhysicsProxy->GetTransformGroupIndexFromHandle(Children[0]);
 				}
