properly cleanup internal clusters when children are removed from them regardless of what path that removal comes from

diff --git a/Engine/Source/Runtime/Experimental/Chaos/Private/Chaos/PBDRigidClustering.cpp b/Engine/Source/Runtime/Experimental/Chaos/Private/Chaos/PBDRigidClustering.cpp
--- a/Engine/Source/Runtime/Experimental/Chaos/Private/Chaos/PBDRigidClustering.cpp
+++ b/Engine/Source/Runtime/Experimental/Chaos/Private/Chaos/PBDRigidClustering.cpp
@@ -654,7 +654,7 @@ namespace Chaos
 		return Islands;
 	}
 
-	void FRigidClustering::RemoveChildFromParent(FPBDRigidParticleHandle* Child, const FPBDRigidClusteredParticleHandle* ClusteredParent)
+	void FRigidClustering::RemoveChildFromParent(FPBDRigidParticleHandle* Child, FPBDRigidClusteredParticleHandle* ClusteredParent)
 	{
 		if (ensure(Child != nullptr && ClusteredParent != nullptr))
 		{
@@ -681,13 +681,24 @@ namespace Chaos
 			Child->SetW(Child->W() + ClusteredParent->W());
 			Child->SetPreV(Child->PreV() + ClusteredParent->PreV());
 			Child->SetPreW(Child->PreW() + ClusteredParent->PreW());
+
+			// We also need to do cluster book-keeping on the parent.
+			// If the parent is an internal cluster and has become empty, we need to mark this particle as ready to be destroyed.
+			// The only exception is for cluster unions which are managed separately.
+			if (FRigidHandleArray* Children = MChildren.Find(ClusteredParent))
+			{
+				ClusteredParent->SetClusterId(ClusterId{ ClusteredParent->Parent(), Children->Num()});
+				if (ClusteredParent->InternalCluster() && Children->IsEmpty() && ClusteredParent->PhysicsProxy() && ClusteredParent->PhysicsProxy()->GetType() == EPhysicsProxyType::GeometryCollectionType)
+				{
+					// We shouldn't ever need to do an AddUnique here since when we remove a child from a parent, the parent should only ever turn empty once.
+					EmptyInternalClustersPerProxy.FindOrAdd(ClusteredParent->PhysicsProxy()).Add(ClusteredParent);
+				}
+			}
 		}
 	}
 
-	void FRigidClustering::RemoveChildFromParentAndChildrenArray(FPBDRigidParticleHandle* Child, const FPBDRigidClusteredParticleHandle* ClusteredParent)
+	void FRigidClustering::RemoveChildFromParentAndChildrenArray(FPBDRigidParticleHandle* Child, FPBDRigidClusteredParticleHandle* ClusteredParent)
 	{
-		RemoveChildFromParent(Child, ClusteredParent);
-
 		// Also need to remove it from the children array.
 		if (FRigidHandleArray* Children = MChildren.Find(ClusteredParent))
 		{
@@ -697,6 +708,8 @@ namespace Chaos
 				Children->RemoveAtSwap(Index, 1, false);
 			}
 		}
+
+		RemoveChildFromParent(Child, ClusteredParent);
 	}
 
 	TArray<FPBDRigidParticleHandle*> FRigidClustering::CreateClustersFromNewIslands(
@@ -1016,7 +1029,7 @@ namespace Chaos
 				{
 					RemoveNodeConnections(Child);
 				}
-				HandleConnectivityOnReleaseClusterParticle(ClusteredParticle, bCreateNewClusters);
+				ActivatedChildren.Append(HandleConnectivityOnReleaseClusterParticle(ClusteredParticle, bCreateNewClusters));
 			}
 
 			for (FPBDRigidParticleHandle* Child : ActivatedChildren)
@@ -1383,10 +1396,25 @@ namespace Chaos
 
 			} // end if MCollisionImpulseArrayDirty
 		}
+
 		Timer.Stop();
 		UE_LOG(LogChaos, Verbose, TEXT("Cluster Break Update Time is %f"), Time);
 	}
 
+	DECLARE_CYCLE_STAT(TEXT("FRigidClustering::CleanupInternalClustersForProxy"), STAT_CleanupInternalClustersForProxy, STATGROUP_Chaos);
+	void FRigidClustering::CleanupInternalClustersForProxies(TArrayView<IPhysicsProxyBase*> Proxies)
+	{
+		SCOPE_CYCLE_COUNTER(STAT_CleanupInternalClustersForProxy);
+		for (IPhysicsProxyBase* Proxy : Proxies)
+		{
+			for (FPBDRigidClusteredParticleHandle* Particle : EmptyInternalClustersPerProxy.FindRef(Proxy))
+			{
+				DestroyClusterParticle(Particle);
+			}
+			EmptyInternalClustersPerProxy.Remove(Proxy);
+		}
+	}
+
 	DECLARE_CYCLE_STAT(TEXT("BreakingModel_AllParticles"), STAT_BreakingModel_AllParticles, STATGROUP_Chaos);
 	DECLARE_CYCLE_STAT(TEXT("BreakingModel"), STAT_BreakingModel, STATGROUP_Chaos);
 	void FRigidClustering::BreakingModel()
diff --git a/Engine/Source/Runtime/Experimental/Chaos/Private/PBDRigidsSolver.cpp b/Engine/Source/Runtime/Experimental/Chaos/Private/PBDRigidsSolver.cpp
--- a/Engine/Source/Runtime/Experimental/Chaos/Private/PBDRigidsSolver.cpp
+++ b/Engine/Source/Runtime/Experimental/Chaos/Private/PBDRigidsSolver.cpp
@@ -1075,15 +1075,19 @@ namespace Chaos
 		}
 
 		MarshallingManager.GetCurrentPullData_Internal()->DirtyGeometryCollections.Reset();
-		for (auto Proxy : PendingDestroyGeometryCollectionPhysicsProxy)
-		{
-			// Removing the geometry collection from the solver a bit delayed. This lets the cluster union do its cleanup first before
-			// the geometry collection if they're all being destroyed at the same time.
-			Proxy->OnRemoveFromSolver(this);
 
-			delete Proxy;
+		if (!PendingDestroyGeometryCollectionPhysicsProxy.IsEmpty())
+		{
+			GetEvolution()->GetRigidClustering().CleanupInternalClustersForProxies(TArrayView<IPhysicsProxyBase*>{reinterpret_cast<IPhysicsProxyBase**>(&PendingDestroyGeometryCollectionPhysicsProxy[0]), PendingDestroyGeometryCollectionPhysicsProxy.Num() });
+			for (auto Proxy : PendingDestroyGeometryCollectionPhysicsProxy)
+			{
+				// Removing the geometry collection from the solver a bit delayed. This lets the cluster union do its cleanup first before
+				// the geometry collection if they're all being destroyed at the same time.
+				Proxy->OnRemoveFromSolver(this);
+				delete Proxy;
+			}
+			PendingDestroyGeometryCollectionPhysicsProxy.Reset();
 		}
-		PendingDestroyGeometryCollectionPhysicsProxy.Reset();
 	}
 
 	void FPBDRigidsSolver::PrepareAdvanceBy(const FReal DeltaTime)
diff --git a/Engine/Source/Runtime/Experimental/Chaos/Public/Chaos/PBDRigidClustering.h b/Engine/Source/Runtime/Experimental/Chaos/Public/Chaos/PBDRigidClustering.h
--- a/Engine/Source/Runtime/Experimental/Chaos/Public/Chaos/PBDRigidClustering.h
+++ b/Engine/Source/Runtime/Experimental/Chaos/Public/Chaos/PBDRigidClustering.h
@@ -408,6 +408,12 @@ public:
 
 	void CreateNodeConnection(FPBDRigidClusteredParticleHandle* A, FPBDRigidClusteredParticleHandle* B);
 
+
+	/**
+	 * CleanupInternalClustersForProxies
+	 *	For a given set of physics proxies, cleanup any tracked internal clusters that we've marked as being empty.
+	 */
+	void CleanupInternalClustersForProxies(TArrayView<IPhysicsProxyBase*> Proxies);
  protected:
 
 	void ComputeStrainFromCollision(const FPBDCollisionConstraints& CollisionRule, const FReal Dt);
@@ -454,8 +460,8 @@ public:
 		const Chaos::FPBDRigidClusteredParticleHandle* Parent,
 		const FClusterCreationParameters& Parameters = FClusterCreationParameters());
 	
-	void RemoveChildFromParent(FPBDRigidParticleHandle* Child, const FPBDRigidClusteredParticleHandle* ClusteredParent);
-	void RemoveChildFromParentAndChildrenArray(FPBDRigidParticleHandle* Child, const FPBDRigidClusteredParticleHandle* ClusteredParent);
+	void RemoveChildFromParent(FPBDRigidParticleHandle* Child, FPBDRigidClusteredParticleHandle* ClusteredParent);
+	void RemoveChildFromParentAndChildrenArray(FPBDRigidParticleHandle* Child, FPBDRigidClusteredParticleHandle* ClusteredParent);
 
 	// When a body has broken due to contact resolution, record an entry in a set
 	// for the collision and the particle who's momentum should be restored.
@@ -520,6 +526,7 @@ private:
 	TArray<FCrumblingData> MAllClusterCrumblings;
 
 	TSet<FPBDRigidClusteredParticleHandle*> CrumbledSinceLastUpdate;
+	TMap<IPhysicsProxyBase*, TArray<FPBDRigidClusteredParticleHandle*>> EmptyInternalClustersPerProxy;
 
 	// Pairs of collision constraints and rigid particle handles of particles which collided with
 	// rigid clusters which broken. Some portion of the momentum change due to the constraint
