[Chaos VD] Initial implementation of Delta Recording.

- Now we only record dirty particles and only do a full capture every 10 seconds (configurable via CVar)
- Fixed an issue where "Solver Frames Start" events where incorrectly triggered, corrupting the recording when the solver gets out of sync with the game thread.
- Solver Frames now represent Solver Advance tasks, which are one Evolution step in code. "Solver steps" in the tool are now steps within a evolution step (Like Integrate, ApplyKinematicTargets,etc). The "Solver Steps" UI in the tool was updated to handle this change
- Now we use the geometry hash as an ID while recording. We were already doing this on the tool side during trace analysis, so this brings it into parity.
- Added a new Particle Destroyed trace event, needed to know when to destroy particles in the tool now that most frames are just deltas.
- Particles views are now traced in parallel
- Added a buffer in the CVD Thread Local Storage to be used for serialization. It is only accessible via a wrapper that ensures the buffer keeps allocated a minimum amount of memory, and is cleared up on scope exit but keeping that minimum memory allocated.
- Removed the use of FChaosArchiveContext for Geometry Serialization. This means we could potentially serialize more data but makes the code simpler and easier to parallelize.

diff --git a/Engine/Plugins/Experimental/ChaosVD/Source/ChaosVD/Private/ChaosVDParticleActor.cpp b/Engine/Plugins/Experimental/ChaosVD/Source/ChaosVD/Private/ChaosVDParticleActor.cpp
--- a/Engine/Plugins/Experimental/ChaosVD/Source/ChaosVD/Private/ChaosVDParticleActor.cpp
+++ b/Engine/Plugins/Experimental/ChaosVD/Source/ChaosVD/Private/ChaosVDParticleActor.cpp
@@ -5,27 +5,59 @@
 #include "ChaosVDGeometryBuilder.h"
 #include "ChaosVDScene.h"
 #include "ChaosVDRecording.h"
+#include "Components/MeshComponent.h"
 #include "Components/SceneComponent.h"
 #include "Components/StaticMeshComponent.h"
 #include "Engine/CollisionProfile.h"
 #include "Engine/StaticMesh.h"
 
+static FAutoConsoleVariable CVarChaosVDHideVolumeAndBrushesHack(
+	TEXT("p.Chaos.VD.Tool.HideVolumeAndBrushesHack"),
+	true,
+	TEXT("If true, it will hide any geometry if its name contains Volume or Brush"));
+
 AChaosVDParticleActor::AChaosVDParticleActor(const FObjectInitializer& ObjectInitializer) : Super(ObjectInitializer)
 {
 	RootComponent = CreateDefaultSubobject<USceneComponent>(TEXT("SceneComponent0"));
 }
 
-void AChaosVDParticleActor::UpdateFromRecordedData(const FChaosVDParticleDebugData& RecordedData, const Chaos::FRigidTransform3& SimulationTransform)
+void AChaosVDParticleActor::UpdateFromRecordedData(const FChaosVDParticleDebugData& InRecordedData, const Chaos::FRigidTransform3& SimulationTransform)
 {
-	SetActorLocationAndRotation(SimulationTransform.TransformPosition(RecordedData.Position), SimulationTransform.GetRotation() * RecordedData.Rotation, false);
+	SetActorLocationAndRotation(SimulationTransform.TransformPosition(InRecordedData.Position), SimulationTransform.GetRotation() * InRecordedData.Rotation, false);
+
+	if (RecordedDebugData.ImplicitObjectHash != InRecordedData.ImplicitObjectHash)
+	{
+		if (const TSharedPtr<FChaosVDScene>& ScenePtr = OwningScene.Pin())
+		{
+			if (const TSharedPtr<const Chaos::FImplicitObject>* Geometry = ScenePtr->GetUpdatedGeometry(InRecordedData.ImplicitObjectHash))
+			{
+				UpdateGeometryData(*Geometry, EChaosVDActorGeometryUpdateFlags::ForceUpdate);
+			}
+		}
+	}
 
 	// TODO: We should store a ptr to the data and in our custom details panel draw it, but the current data format is just a test
 	// So we should do this once it is final
-	RecordedDebugData = RecordedData;
+	RecordedDebugData = InRecordedData;
 }
 
-void AChaosVDParticleActor::UpdateGeometryData(const TSharedPtr<const Chaos::FImplicitObject>& ImplicitObject)
+void AChaosVDParticleActor::UpdateGeometryData(const TSharedPtr<const Chaos::FImplicitObject>& ImplicitObject, EChaosVDActorGeometryUpdateFlags OptionFlag)
 {
+	if (EnumHasAnyFlags(OptionFlag, EChaosVDActorGeometryUpdateFlags::ForceUpdate))
+	{
+		bIsGeometryDataGenerationStarted = false;
+
+		for (TWeakObjectPtr<UMeshComponent>& MeshComponent : MeshComponents)
+		{
+			if (MeshComponent.IsValid())
+			{
+				MeshComponent->DestroyComponent();	
+			}
+		}
+
+		MeshComponents.Reset();
+	}
+
 	if (bIsGeometryDataGenerationStarted)
 	{
 		return;
@@ -50,10 +82,26 @@ void AChaosVDParticleActor::UpdateGeometryData(const TSharedPtr<const Chaos::FIm
 				GeometryGenerator->CreateMeshComponentsFromImplicit<UStaticMesh, UInstancedStaticMeshComponent>(ImplicitObject.Get(), this, OutGeneratedMeshComponents, Transform);
 			}
 
-			
-
 			if (OutGeneratedMeshComponents.Num() > 0)
 			{
+				MeshComponents.Append(OutGeneratedMeshComponents);
+
+				if (CVarChaosVDHideVolumeAndBrushesHack->GetBool())
+				{
+					// This is a temp hack (and is not performant) we need until we have a proper way to filer out trigger volumes/brushes at will
+					// Without this most maps will be covered in boxes
+					if (RecordedDebugData.DebugName.Contains(TEXT("Brush")) || RecordedDebugData.DebugName.Contains(TEXT("Volume")))
+					{
+						for (TWeakObjectPtr<UMeshComponent> MeshComponent : OutGeneratedMeshComponents)
+						{
+							if (MeshComponent.IsValid())
+							{
+								MeshComponent->SetVisibility(false);
+							}
+						}
+					}
+				}
+
 				bIsGeometryDataGenerationStarted = true;
 			}
 		}
@@ -66,9 +114,9 @@ void AChaosVDParticleActor::SetScene(const TSharedPtr<FChaosVDScene>& InScene)
 
 	if (const TSharedPtr<FChaosVDScene>& ScenePtr = OwningScene.Pin())
 	{
-		GeometryUpdatedDelegate = ScenePtr->OnNewGeometryAvailable().AddWeakLambda(this, [this](const TSharedPtr<const Chaos::FImplicitObject>& ImplicitObject, const int32 ID)
+		GeometryUpdatedDelegate = ScenePtr->OnNewGeometryAvailable().AddWeakLambda(this, [this](const TSharedPtr<const Chaos::FImplicitObject>& ImplicitObject, const uint32 ID)
 		{
-			if (RecordedDebugData.ImplicitObjectID == ID)
+			if (RecordedDebugData.ImplicitObjectHash == ID)
 			{
 				UpdateGeometryData(ImplicitObject);
 			}
diff --git a/Engine/Plugins/Experimental/ChaosVD/Source/ChaosVD/Private/ChaosVDPlaybackController.cpp b/Engine/Plugins/Experimental/ChaosVD/Source/ChaosVD/Private/ChaosVDPlaybackController.cpp
--- a/Engine/Plugins/Experimental/ChaosVD/Source/ChaosVD/Private/ChaosVDPlaybackController.cpp
+++ b/Engine/Plugins/Experimental/ChaosVD/Source/ChaosVD/Private/ChaosVDPlaybackController.cpp
@@ -56,7 +56,7 @@ bool FChaosVDPlaybackController::LoadChaosVDRecordingFromTraceSession(const FStr
 		GoToTrackFrame(IChaosVDPlaybackControllerInstigator::InvalidGuid, EChaosVDTrackType::Solver, SolversIterator->Key, FrameNumber, StepNumber);
 	}
 
-	LoadedRecording->OnGeometryDataLoaded().AddLambda([this](const TSharedPtr<const Chaos::FImplicitObject>& NewGeometry, const int32 GeometryID)
+	LoadedRecording->OnGeometryDataLoaded().AddLambda([this](const TSharedPtr<const Chaos::FImplicitObject>& NewGeometry, const uint32 GeometryID)
 	{
 		if (const TSharedPtr<FChaosVDScene> ScenePtr = SceneToControl.Pin())
 		{
@@ -193,7 +193,7 @@ void FChaosVDPlaybackController::GoToTrackFrame(FGuid InstigatorID, EChaosVDTrac
 	}
 }
 
-int32 FChaosVDPlaybackController::GetTrackStepsAtFrame(EChaosVDTrackType TrackType, const int32 InTrackID, const int32 FrameNumber) const
+int32 FChaosVDPlaybackController::GetTrackStepsNumberAtFrame(EChaosVDTrackType TrackType, const int32 InTrackID, const int32 FrameNumber) const
 {
 	if (!LoadedRecording.IsValid())
 	{
@@ -224,6 +224,37 @@ int32 FChaosVDPlaybackController::GetTrackStepsAtFrame(EChaosVDTrackType TrackTy
 	}
 }
 
+const FChaosVDStepsContainer* FChaosVDPlaybackController::GetTrackStepsDataAtFrame(EChaosVDTrackType TrackType, int32 InTrackID, int32 FrameNumber) const
+{
+	if (!LoadedRecording.IsValid())
+	{
+		return nullptr;
+	}
+	
+	switch (TrackType)
+	{
+	case EChaosVDTrackType::Game:
+		// Game Tracks do not have steps
+		return nullptr;
+		break;
+	case EChaosVDTrackType::Solver:
+		{
+			if (const FChaosVDSolverFrameData* FrameData = LoadedRecording->GetSolverFrameData(InTrackID, FrameNumber))
+			{
+				return &FrameData->SolverSteps;
+			}
+			else
+			{
+				return nullptr;
+			}
+			break;
+		}
+	default:
+		return nullptr;
+		break;
+	}
+}
+
 int32 FChaosVDPlaybackController::GetTrackFramesNumber(EChaosVDTrackType TrackType, const int32 InTrackID) const
 {
 	if (!LoadedRecording.IsValid())
@@ -253,37 +284,43 @@ int32 FChaosVDPlaybackController::GetTrackFramesNumber(EChaosVDTrackType TrackTy
 	}
 }
 
-int32 FChaosVDPlaybackController::ConvertCurrentFrameToOtherTrackFrame(FChaosVDTrackInfo FromTrack, FChaosVDTrackInfo ToTrack)
+int32 FChaosVDPlaybackController::ConvertCurrentFrameToOtherTrackFrame(const FChaosVDTrackInfo* FromTrack, const FChaosVDTrackInfo* ToTrack)
 {
 	// Each track is on a different "space time", because it's source data ticked at different rates when was recorded, and some start/end at different points on time
 	// But all the recorded frame data on all of them use Platform Cycles as timestamps
 	// This method wraps specialized methods in the recording object to convert between these spaces. For example Game frame 1500 could be frame 5 on a specific solver
 	// And Frame 5 of that solver could be frame 30 on another solver.
 
-	const bool bBothTracksHaveSameID = FromTrack.TrackID == ToTrack.TrackID;
-	const bool bBothTracksHaveSameType = FromTrack.TrackType == ToTrack.TrackType;
+	if (FromTrack == nullptr || ToTrack == nullptr)
+	{
+		ensureMsgf(false, TEXT("One of provided track infos is not valid"));
+		return INDEX_NONE;
+	}
+
+	const bool bBothTracksHaveSameID = FromTrack->TrackID == ToTrack->TrackID;
+	const bool bBothTracksHaveSameType = FromTrack->TrackType == ToTrack->TrackType;
 	if (bBothTracksHaveSameType && bBothTracksHaveSameID)
 	{
-		return FromTrack.CurrentFrame;
+		return FromTrack->CurrentFrame;
 	}
 
-	switch (FromTrack.TrackType)
+	switch (FromTrack->TrackType)
 	{
 	case EChaosVDTrackType::Game:
 		{
 			// Convert from Game Frame to Solver Frame
-			return LoadedRecording->GetLowestSolverFrameNumberGameFrame(FromTrack.TrackID, FromTrack.CurrentFrame);
+			return LoadedRecording->GetLowestSolverFrameNumberGameFrame(FromTrack->TrackID, FromTrack->CurrentFrame);
 			break;
 		}
 		
 	case EChaosVDTrackType::Solver:
 		{
-			if (ToTrack.TrackType == EChaosVDTrackType::Solver)
+			if (ToTrack->TrackType == EChaosVDTrackType::Solver)
 			{
 				ensureMsgf(false, TEXT("Frame conversion between solver tracks is not supported yet"));
 				return INDEX_NONE;
 			}
-			return LoadedRecording->GetLowestGameFrameAtSolverFrameNumber(FromTrack.TrackID, FromTrack.CurrentFrame);
+			return LoadedRecording->GetLowestGameFrameAtSolverFrameNumber(FromTrack->TrackID, FromTrack->CurrentFrame);
 			break;
 		}
 	default:
diff --git a/Engine/Plugins/Experimental/ChaosVD/Source/ChaosVD/Private/ChaosVDScene.cpp b/Engine/Plugins/Experimental/ChaosVD/Source/ChaosVD/Private/ChaosVDScene.cpp
--- a/Engine/Plugins/Experimental/ChaosVD/Source/ChaosVD/Private/ChaosVDScene.cpp
+++ b/Engine/Plugins/Experimental/ChaosVD/Source/ChaosVD/Private/ChaosVDScene.cpp
@@ -65,16 +65,18 @@ void FChaosVDScene::AddReferencedObjects(FReferenceCollector& Collector)
 
 void FChaosVDScene::UpdateFromRecordedStepData(const int32 SolverID, const FString& SolverName, const FChaosVDStepData& InRecordedStepData, const FChaosVDSolverFrameData& InFrameData)
 {
-	TSet<int32> ParticlesIDsInRecordedStepData;
-	ParticlesIDsInRecordedStepData.Reserve(InRecordedStepData.RecordedParticles.Num());
-
 	FChaosVDParticlesByIDMap& SolverParticlesByID = ParticlesBySolverID.FindChecked(SolverID);
 
 	// Go over existing Particle VD Instances and update them or create them if needed 
 	for (const FChaosVDParticleDebugData& Particle : InRecordedStepData.RecordedParticles)
 	{
 		const int32 ParticleVDInstanceID = GetIDForRecordedParticleData(Particle);
-		ParticlesIDsInRecordedStepData.Add(ParticleVDInstanceID);
+
+		if (InRecordedStepData.ParticlesDestroyedIDs.Contains(ParticleVDInstanceID))
+		{
+			// Do not process the particle if it was destroyed in the same step
+			continue;
+		}
 
 		if (AChaosVDParticleActor** ExistingParticleVDInstancePtrPtr = SolverParticlesByID.Find(ParticleVDInstanceID))
 		{
@@ -113,7 +115,7 @@ void FChaosVDScene::UpdateFromRecordedStepData(const int32 SolverID, const FStri
 	int32 AmountRemoved = 0;
 	for (FChaosVDParticlesByIDMap::TIterator RemoveIterator = SolverParticlesByID.CreateIterator(); RemoveIterator; ++RemoveIterator)
 	{
-		if (!ParticlesIDsInRecordedStepData.Contains(RemoveIterator.Key()))
+		if (InRecordedStepData.ParticlesDestroyedIDs.Contains(RemoveIterator.Key()))
 		{
 			if (AChaosVDParticleActor* ActorToRemove = RemoveIterator.Value())
 			{
@@ -129,7 +131,7 @@ void FChaosVDScene::UpdateFromRecordedStepData(const int32 SolverID, const FStri
 	OnSceneUpdated().Broadcast();
 }
 
-void FChaosVDScene::HandleNewGeometryData(const TSharedPtr<const Chaos::FImplicitObject>& GeometryData, const int32 GeometryID) const
+void FChaosVDScene::HandleNewGeometryData(const TSharedPtr<const Chaos::FImplicitObject>& GeometryData, const uint32 GeometryID) const
 {
 	NewGeometryAvailableDelegate.Broadcast(GeometryData, GeometryID);
 }
@@ -192,6 +194,19 @@ void FChaosVDScene::CleanUpScene()
 	ParticlesBySolverID.Reset();
 }
 
+const TSharedPtr<const Chaos::FImplicitObject>* FChaosVDScene::GetUpdatedGeometry(int32 GeometryID) const
+{
+	if (ensure(LoadedRecording.IsValid()))
+	{
+		if (const TSharedPtr<const Chaos::FImplicitObject>* Geometry = LoadedRecording->GetGeometryDataMap().Find(GeometryID))
+		{
+			return Geometry;
+		}
+	}
+
+	return nullptr;
+}
+
 AChaosVDParticleActor* FChaosVDScene::SpawnParticleFromRecordedData(const FChaosVDParticleDebugData& InParticleData, const FChaosVDSolverFrameData& InFrameData)
 {
 	using namespace Chaos;
@@ -212,7 +227,7 @@ AChaosVDParticleActor* FChaosVDScene::SpawnParticleFromRecordedData(const FChaos
 		
 		if (ensure(LoadedRecording.IsValid()))
 		{
-			if (const TSharedPtr<const Chaos::FImplicitObject>* Geometry = LoadedRecording->GetGeometryDataMap().Find(InParticleData.ImplicitObjectID))
+			if (const TSharedPtr<const Chaos::FImplicitObject>* Geometry = LoadedRecording->GetGeometryDataMap().Find(InParticleData.ImplicitObjectHash))
 			{
 				NewActor->UpdateGeometryData(*Geometry);
 			}
diff --git a/Engine/Plugins/Experimental/ChaosVD/Source/ChaosVD/Private/Trace/ChaosVDTraceAnalyzer.cpp b/Engine/Plugins/Experimental/ChaosVD/Source/ChaosVD/Private/Trace/ChaosVDTraceAnalyzer.cpp
--- a/Engine/Plugins/Experimental/ChaosVD/Source/ChaosVD/Private/Trace/ChaosVDTraceAnalyzer.cpp
+++ b/Engine/Plugins/Experimental/ChaosVD/Source/ChaosVD/Private/Trace/ChaosVDTraceAnalyzer.cpp
@@ -11,6 +11,7 @@ void FChaosVDTraceAnalyzer::OnAnalysisBegin(const FOnAnalysisContext& Context)
 {
 	FInterfaceBuilder& Builder = Context.InterfaceBuilder;
 	Builder.RouteEvent(RouteId_ChaosVDParticle, "ChaosVDLogger", "ChaosVDParticle");
+	Builder.RouteEvent(RouteId_ChaosVDParticleDestroyed, "ChaosVDLogger", "ChaosVDParticleDestroyed");
 	
 	Builder.RouteEvent(RouteId_ChaosVDSolverFrameStart, "ChaosVDLogger", "ChaosVDSolverFrameStart");
 	Builder.RouteEvent(RouteId_ChaosVDSolverFrameEnd, "ChaosVDLogger", "ChaosVDSolverFrameEnd");
@@ -76,6 +77,7 @@ bool FChaosVDTraceAnalyzer::OnEvent(uint16 RouteId, EStyle Style, const FOnEvent
 
 			NewFrameData.SolverID = EventData.GetValue<int32>("SolverID");
 			NewFrameData.FrameCycle = EventData.GetValue<uint64>("Cycle");
+			NewFrameData.bIsKeyFrame = EventData.GetValue<bool>("IsKeyFrame");
 
 			FWideStringView DebugNameView;
 			EventData.GetString("DebugName", DebugNameView);
@@ -83,7 +85,6 @@ bool FChaosVDTraceAnalyzer::OnEvent(uint16 RouteId, EStyle Style, const FOnEvent
 
 			// Add an empty frame. It will be filled out by the solver trace events
 			ChaosVDTraceProvider->AddSolverFrame(NewFrameData.SolverID, MoveTemp(NewFrameData));
-	
 			break;
 		}
 	case RouteId_ChaosVDSolverFrameEnd:
@@ -98,7 +99,11 @@ bool FChaosVDTraceAnalyzer::OnEvent(uint16 RouteId, EStyle Style, const FOnEvent
 			if (FChaosVDSolverFrameData* FrameData  = ChaosVDTraceProvider->GetLastSolverFrame(SolverID))
 			{
 				// Add an empty step. It will be filled out by the particle (and later on other objects/elements) events
-				FrameData->SolverSteps.AddDefaulted();
+				FChaosVDStepData& StepData = FrameData->SolverSteps.AddDefaulted_GetRef();
+
+				FWideStringView DebugNameView;
+				EventData.GetString("StepName", DebugNameView);
+				StepData.StepName = DebugNameView;
 			}
 	
 			break;
@@ -117,12 +122,26 @@ bool FChaosVDTraceAnalyzer::OnEvent(uint16 RouteId, EStyle Style, const FOnEvent
 			// This can be null if the recording started Mid-Frame. In this case we just discard the data for now
 			if (FChaosVDSolverFrameData* FrameData = ChaosVDTraceProvider->GetLastSolverFrame(SolverID))
 			{
-				if (FrameData->SolverSteps.Num() > 0)
+				if (ensureMsgf(FrameData->SolverSteps.Num() > 0, TEXT("A particle was traced without a valid sstep scope")))
 				{
 					FrameData->SolverSteps.Last().RecordedParticles.Add(MoveTemp(ParticleData));
 				}
 			}
 
+			break;
+		}
+	case RouteId_ChaosVDParticleDestroyed:
+		{
+			const int32 SolverID = EventData.GetValue<int32>("SolverID");
+
+			if (FChaosVDSolverFrameData* FrameData = ChaosVDTraceProvider->GetLastSolverFrame(SolverID))
+			{
+				if (FrameData->SolverSteps.Num() > 0)
+				{
+					FrameData->SolverSteps.Last().ParticlesDestroyedIDs.Add(EventData.GetValue<int32>("ParticleID"));
+				}
+			}
+
 			break;
 		}
 	case RouteId_ChaosVDBinaryDataStart:
@@ -195,6 +214,7 @@ void FChaosVDTraceAnalyzer::ReadParticleDataFromEvent(const FEventData& InEventD
 	OutParticleData.ParticleState = static_cast<EChaosVDParticleState>(InEventData.GetValue<int8>("ObjectState"));
 	OutParticleData.ParticleIndex = InEventData.GetValue<int32>("ParticleID");
 	OutParticleData.ImplicitObjectID = InEventData.GetValue<int32>("ImplicitObjectID");
+	OutParticleData.ImplicitObjectHash = InEventData.GetValue<uint32>("ImplicitObjectHash");
 
 	FWideStringView DebugNameView;
 	InEventData.GetString("DebugName", DebugNameView);
@@ -218,4 +238,3 @@ void FChaosVDTraceAnalyzer::ReadParticleDataFromEvent(const FEventData& InEventD
 	OutParticleData.AngularVelocity.Y = InEventData.GetValue<float>("AngularVelocityY");
 	OutParticleData.AngularVelocity.Z = InEventData.GetValue<float>("AngularVelocityZ");
 }
-
diff --git a/Engine/Plugins/Experimental/ChaosVD/Source/ChaosVD/Private/Trace/ChaosVDTraceProvider.cpp b/Engine/Plugins/Experimental/ChaosVD/Source/ChaosVD/Private/Trace/ChaosVDTraceProvider.cpp
--- a/Engine/Plugins/Experimental/ChaosVD/Source/ChaosVD/Private/Trace/ChaosVDTraceProvider.cpp
+++ b/Engine/Plugins/Experimental/ChaosVD/Source/ChaosVD/Private/Trace/ChaosVDTraceProvider.cpp
@@ -4,6 +4,7 @@
 
 #include "Chaos/ImplicitObject.h"
 #include "ChaosVDRecording.h"
+#include "ChaosVisualDebugger/ChaosVisualDebuggerTrace.h"
 #include "Compression/OodleDataCompressionUtil.h"
 #include "Serialization/MemoryReader.h"
 
@@ -128,29 +129,22 @@ void FChaosVDTraceProvider::SetBinaryDataReadyToUse(const int32 DataID)
 				RawData = &UnprocessedData->RawData;
 			}
 
+#if WITH_CHAOS_VISUAL_DEBUGGER
 			// TODO: Create a system to register external "Data handlers" with the logic on how serialize each type
 			// This should not be done here but as this is the only type we have is ok for now
-			if (UnprocessedData->TypeName == TEXT("FImplicitObject"))
-			{
+			if (UnprocessedData->TypeName == TEXT("FChaosVDImplicitObjectDataWrapper"))
+			{	
 				Chaos::TSerializablePtr<Chaos::FImplicitObject> Geometry;
 
 				FMemoryReader MemeReader(*RawData);
 				Chaos::FChaosArchive Ar(MemeReader);
 
-				if (ChaosContext.IsValid())
-				{
-					Ar.SetContext(MoveTemp(ChaosContext));
-					Ar << Geometry;
-					ChaosContext = Ar.StealContext();
-				}
-				else
-				{
-					Ar << Geometry;
-					ChaosContext = Ar.StealContext();
-				}
-
-				InternalRecording->AddImplicitObject(ChaosContext->ObjToTag[Geometry.Get()], Geometry.Get());
+				FChaosVDImplicitObjectDataWrapper WrappedGeometryData;
+				WrappedGeometryData.Serialize(Ar);
+
+				InternalRecording->AddImplicitObject(WrappedGeometryData.Hash, WrappedGeometryData.ImplicitObject.Get());
 			}
+#endif
 		}
 	}
 	else
diff --git a/Engine/Plugins/Experimental/ChaosVD/Source/ChaosVD/Private/Widgets/SChaosVDPlaybackViewport.cpp b/Engine/Plugins/Experimental/ChaosVD/Source/ChaosVD/Private/Widgets/SChaosVDPlaybackViewport.cpp
--- a/Engine/Plugins/Experimental/ChaosVD/Source/ChaosVD/Private/Widgets/SChaosVDPlaybackViewport.cpp
+++ b/Engine/Plugins/Experimental/ChaosVD/Source/ChaosVD/Private/Widgets/SChaosVDPlaybackViewport.cpp
@@ -129,7 +129,12 @@ void SChaosVDPlaybackViewport::HandleControllerTrackFrameUpdated(TWeakPtr<FChaos
 		return;
 	}
 
-	GameFramesTimelineWidget->SetCurrentTimelineFrame(UpdatedTrackInfo->CurrentFrame, EChaosVDSetTimelineFrameFlags::None);
+	if (TSharedPtr<FChaosVDPlaybackController> ControllerSharedPtr = InController.Pin())
+	{
+		// The frame number we receive could be from a Solver track, so make sure it is converted to the correct game track frame number 
+		int32 GameTrackFrame = ControllerSharedPtr->ConvertCurrentFrameToOtherTrackFrame(UpdatedTrackInfo, ControllerSharedPtr->GetTrackInfo(EChaosVDTrackType::Game, FChaosVDPlaybackController::GameTrackID));
+		GameFramesTimelineWidget->SetCurrentTimelineFrame(GameTrackFrame, EChaosVDSetTimelineFrameFlags::None);
+	}
 }
 
 void SChaosVDPlaybackViewport::OnPlaybackSceneUpdated()
diff --git a/Engine/Plugins/Experimental/ChaosVD/Source/ChaosVD/Private/Widgets/SChaosVDSolverPlaybackControls.cpp b/Engine/Plugins/Experimental/ChaosVD/Source/ChaosVD/Private/Widgets/SChaosVDSolverPlaybackControls.cpp
--- a/Engine/Plugins/Experimental/ChaosVD/Source/ChaosVD/Private/Widgets/SChaosVDSolverPlaybackControls.cpp
+++ b/Engine/Plugins/Experimental/ChaosVD/Source/ChaosVD/Private/Widgets/SChaosVDSolverPlaybackControls.cpp
@@ -33,7 +33,7 @@ void SChaosVDSolverPlaybackControls::Construct(const FArguments& InArgs, int32 I
 				[
 					SNew(STextBlock)
 					.Justification(ETextJustify::Center)
-					.Text(LOCTEXT("PlaybackViewportWidgetPhysicsFramesLabel", "Physics Frames" ))
+					.Text(LOCTEXT("PlaybackViewportWidgetPhysicsFramesLabel", "Solver Frames" ))
 				]
 				+SVerticalBox::Slot()
 				[
@@ -53,7 +53,7 @@ void SChaosVDSolverPlaybackControls::Construct(const FArguments& InArgs, int32 I
 				[
 					SNew(STextBlock)
 					.Justification(ETextJustify::Center)
-					.Text(LOCTEXT("PlaybackViewportWidgetStepsLabel", "Solver Steps" ))
+					.Text_Lambda([this]()->FText{ return FText::Format(LOCTEXT("PlaybackViewportWidgetStepsLabel","Step {0}"), FText::AsCultureInvariant(CurrentStepName));})
 				]
 				+SVerticalBox::Slot()
 				[
@@ -61,7 +61,7 @@ void SChaosVDSolverPlaybackControls::Construct(const FArguments& InArgs, int32 I
 					.HidePlayStopButtons(true)
 					.OnFrameChanged_Raw(this, &SChaosVDSolverPlaybackControls::OnStepSelectionUpdated)
 					.MaxFrames(0)
-				]	
+				]
 			]
 		]
 	];
@@ -88,7 +88,13 @@ void SChaosVDSolverPlaybackControls::HandlePlaybackControllerDataUpdated(TWeakPt
 	if (ControllerSharedPtr.IsValid() && ControllerSharedPtr->IsRecordingLoaded())
 	{	
 		const int32 AvailableFrames = ControllerSharedPtr->GetTrackFramesNumber(EChaosVDTrackType::Solver, SolverID);
-		const int32 AvailableSteps = ControllerSharedPtr->GetTrackStepsAtFrame(EChaosVDTrackType::Solver, SolverID, ControllerSharedPtr->GetTrackCurrentFrame(EChaosVDTrackType::Solver, SolverID));
+
+		int32 AvailableSteps = INDEX_NONE;
+		if (const FChaosVDStepsContainer* StepData = ControllerSharedPtr->GetTrackStepsDataAtFrame(EChaosVDTrackType::Solver, SolverID, ControllerSharedPtr->GetTrackCurrentFrame(EChaosVDTrackType::Solver, SolverID)))
+		{
+			AvailableSteps = StepData->Num() > 0 ? StepData->Num() : INDEX_NONE;
+			CurrentStepName =  StepData->Num() > 0 ? (*StepData)[0].StepName : TEXT("NONE");
+		}
 
 		// Max is inclusive and we use this to request as the index on the recorded frames/steps arrays so we need to -1 to the available frames/steps
 		FramesTimelineWidget->UpdateMinMaxValue(0, AvailableFrames != INDEX_NONE ? AvailableFrames -1 : 0);
@@ -106,6 +112,38 @@ void SChaosVDSolverPlaybackControls::HandlePlaybackControllerDataUpdated(TWeakPt
 	}
 }
 
+void SChaosVDSolverPlaybackControls::UpdateStepsWidgetForFrame(const FChaosVDPlaybackController& InCurrentPlaybackController, int32 FrameNumber, int32 StepNumber, EChaosVDStepsWidgetUpdateFlags OptionsFlags)
+{
+	if (const FChaosVDStepsContainer* StepsData = InCurrentPlaybackController.GetTrackStepsDataAtFrame(EChaosVDTrackType::Solver, SolverID, FrameNumber))
+	{
+		constexpr TCHAR const* UnknownStepName = TEXT("Unknown");
+
+		if (StepsData->Num() > 0 && StepNumber < StepsData->Num())
+		{
+			const FChaosVDStepData& StepData = (*StepsData)[StepNumber];
+
+			if (EnumHasAnyFlags(OptionsFlags, EChaosVDStepsWidgetUpdateFlags::UpdateText))
+			{
+				CurrentStepName = StepData.StepName;
+			}
+
+			if (EnumHasAnyFlags(OptionsFlags, EChaosVDStepsWidgetUpdateFlags::SetTimelineStep))
+			{
+				int32 AvailableSteps = StepsData->Num() > 0 ? StepsData->Num() : INDEX_NONE;
+				StepsTimelineWidget->UpdateMinMaxValue(0,AvailableSteps != INDEX_NONE ? AvailableSteps -1 : 0);
+
+				// On Frame updates, always use Step 0
+				StepsTimelineWidget->SetCurrentTimelineFrame(StepNumber, EChaosVDSetTimelineFrameFlags::None);
+			}
+		}
+		else
+		{
+			StepsTimelineWidget->UpdateMinMaxValue(0,0);
+			CurrentStepName = UnknownStepName;
+		}
+	}
+}
+
 void SChaosVDSolverPlaybackControls::HandleControllerTrackFrameUpdated(TWeakPtr<FChaosVDPlaybackController> InController, const FChaosVDTrackInfo* UpdatedTrackInfo, FGuid InstigatorGuid)
 {
 	if (InstigatorGuid == GetInstigatorID())
@@ -120,35 +158,33 @@ void SChaosVDSolverPlaybackControls::HandleControllerTrackFrameUpdated(TWeakPtr<
 		{
 			FramesTimelineWidget->SetCurrentTimelineFrame(SolverTrackInfo->CurrentFrame, EChaosVDSetTimelineFrameFlags::None);
 
-			const int32 AvailableSteps = CurrentPlaybackControllerPtr->GetTrackStepsAtFrame(EChaosVDTrackType::Solver, SolverID, SolverTrackInfo->CurrentFrame);
-			StepsTimelineWidget->UpdateMinMaxValue(0,AvailableSteps != INDEX_NONE ? AvailableSteps -1 : 0);
-
-			// On Frame updates, always use Step 0
-			StepsTimelineWidget->SetCurrentTimelineFrame(0, EChaosVDSetTimelineFrameFlags::None);
+			constexpr int32 StepNumber = 0;
+			UpdateStepsWidgetForFrame(*CurrentPlaybackControllerPtr.Get(), SolverTrackInfo->CurrentFrame, StepNumber);
 		}
 	}
 }
 
-void SChaosVDSolverPlaybackControls::OnFrameSelectionUpdated(int32 NewFrameIndex) const
+void SChaosVDSolverPlaybackControls::OnFrameSelectionUpdated(int32 NewFrameIndex)
 {
 	if (const TSharedPtr<FChaosVDPlaybackController> PlaybackControllerPtr = PlaybackController.Pin())
 	{
-		const int32 AvailableSteps = PlaybackControllerPtr->GetTrackStepsAtFrame(EChaosVDTrackType::Solver, SolverID, PlaybackControllerPtr->GetTrackCurrentFrame(EChaosVDTrackType::Solver, SolverID));
-		StepsTimelineWidget->UpdateMinMaxValue(0,AvailableSteps != INDEX_NONE ? AvailableSteps - 1 : 0);
-
-		// On Frame updates, always use Step 0
 		constexpr int32 StepNumber = 0;
-		StepsTimelineWidget->SetCurrentTimelineFrame(StepNumber, EChaosVDSetTimelineFrameFlags::None);
+		UpdateStepsWidgetForFrame(*PlaybackControllerPtr.Get(), NewFrameIndex, StepNumber);
+
 		PlaybackControllerPtr->GoToTrackFrame(GetInstigatorID(), EChaosVDTrackType::Solver, SolverID, NewFrameIndex, StepNumber);
 	}
 }
 
-void SChaosVDSolverPlaybackControls::OnStepSelectionUpdated(int32 NewStepIndex) const
+void SChaosVDSolverPlaybackControls::OnStepSelectionUpdated(int32 NewStepIndex)
 {
 	if (const TSharedPtr<FChaosVDPlaybackController> PlaybackControllerPtr = PlaybackController.Pin())
 	{
 		// On Steps updates. Always use the current Frame
 		int32 CurrentFrame = PlaybackControllerPtr->GetTrackCurrentFrame(EChaosVDTrackType::Solver, SolverID);
+
+		// Only Update the text as if we are here it means manually set the step number already
+		UpdateStepsWidgetForFrame(*PlaybackControllerPtr.Get(), CurrentFrame, NewStepIndex, EChaosVDStepsWidgetUpdateFlags::UpdateText);
+
 		PlaybackControllerPtr->GoToTrackFrame(GetInstigatorID(), EChaosVDTrackType::Solver, SolverID, CurrentFrame, NewStepIndex);
 	}
 }
diff --git a/Engine/Plugins/Experimental/ChaosVD/Source/ChaosVD/Public/ChaosVDParticleActor.h b/Engine/Plugins/Experimental/ChaosVD/Source/ChaosVD/Public/ChaosVDParticleActor.h
--- a/Engine/Plugins/Experimental/ChaosVD/Source/ChaosVD/Public/ChaosVDParticleActor.h
+++ b/Engine/Plugins/Experimental/ChaosVD/Source/ChaosVD/Public/ChaosVDParticleActor.h
@@ -9,10 +9,19 @@
 
 class FChaosVDScene;
 struct FChaosVDParticleDebugData;
+class UMeshComponent;
 class USceneComponent;
 class UStaticMeshComponent;
 class UStaticMesh;
 
+/** Options flags to control how geometry is updated in a ChaosVDActor */
+enum class EChaosVDActorGeometryUpdateFlags : int32
+{
+	None = 0,
+	ForceUpdate = 1 << 0
+};
+ENUM_CLASS_FLAGS(EChaosVDActorGeometryUpdateFlags)
+
 /** Actor used to represent a Chaos Particle in the Visual Debugger's world */
 UCLASS()
 class AChaosVDParticleActor : public AActor
@@ -21,11 +30,11 @@ class AChaosVDParticleActor : public AActor
 public:
 	AChaosVDParticleActor(const FObjectInitializer& ObjectInitializer);
 
-	void UpdateFromRecordedData(const FChaosVDParticleDebugData& RecordedData, const Chaos::FRigidTransform3& SimulationTransform);
+	void UpdateFromRecordedData(const FChaosVDParticleDebugData& InRecordedData, const Chaos::FRigidTransform3& SimulationTransform);
 	
 	const FChaosVDParticleDebugData& GetDebugData() const { return RecordedDebugData; }
 
-	void UpdateGeometryData(const TSharedPtr<const Chaos::FImplicitObject>& ImplicitObject);
+	void UpdateGeometryData(const TSharedPtr<const Chaos::FImplicitObject>& ImplicitObject, EChaosVDActorGeometryUpdateFlags OptionFlags = EChaosVDActorGeometryUpdateFlags::None);
 
 	void SetScene(const TSharedPtr<FChaosVDScene>& InScene);
 
@@ -40,5 +49,7 @@ protected:
 
 	TWeakPtr<FChaosVDScene> OwningScene;
 
+	TArray<TWeakObjectPtr<UMeshComponent>> MeshComponents;
+
 	FDelegateHandle GeometryUpdatedDelegate;
 };
diff --git a/Engine/Plugins/Experimental/ChaosVD/Source/ChaosVD/Public/ChaosVDPlaybackController.h b/Engine/Plugins/Experimental/ChaosVD/Source/ChaosVD/Public/ChaosVDPlaybackController.h
--- a/Engine/Plugins/Experimental/ChaosVD/Source/ChaosVD/Public/ChaosVDPlaybackController.h
+++ b/Engine/Plugins/Experimental/ChaosVD/Source/ChaosVD/Public/ChaosVDPlaybackController.h
@@ -88,7 +88,16 @@ public:
 	 * @param FrameNumber Frame number to evaluate
 	 * @return Number of available steps
 	 */
-	int32 GetTrackStepsAtFrame(EChaosVDTrackType TrackType, int32 InTrackID, int32 FrameNumber) const;
+	int32 GetTrackStepsNumberAtFrame(EChaosVDTrackType TrackType, int32 InTrackID, int32 FrameNumber) const;
+
+	/**
+	 * Gets the available steps container in a track at the specified frame
+	 * @param TrackType Type of the track to evaluate
+	 * @param InTrackID ID of the track to evaluate
+	 * @param FrameNumber Frame number to evaluate
+	 * @return Ptr to the steps data container
+	 */
+	const FChaosVDStepsContainer* GetTrackStepsDataAtFrame(EChaosVDTrackType TrackType, int32 InTrackID, int32 FrameNumber) const;
 
 	/**
 	 * Gets the number of available frames for the specified track
@@ -119,7 +128,7 @@ public:
 	 * @param ToTrack Track info we want to use to convert the frame to
 	 * @return Converted Frame Number
 	 */
-	int32 ConvertCurrentFrameToOtherTrackFrame(FChaosVDTrackInfo FromTrack, FChaosVDTrackInfo ToTrack);
+	int32 ConvertCurrentFrameToOtherTrackFrame(const FChaosVDTrackInfo* FromTrack, const FChaosVDTrackInfo* ToTrack);
 
 	/**
 	 * Gets all the ids of the tracks, of the specified type, that are available available on the loaded recording
diff --git a/Engine/Plugins/Experimental/ChaosVD/Source/ChaosVD/Public/ChaosVDScene.h b/Engine/Plugins/Experimental/ChaosVD/Source/ChaosVD/Public/ChaosVDScene.h
--- a/Engine/Plugins/Experimental/ChaosVD/Source/ChaosVD/Public/ChaosVDScene.h
+++ b/Engine/Plugins/Experimental/ChaosVD/Source/ChaosVD/Public/ChaosVDScene.h
@@ -39,7 +39,7 @@ public:
 	/** Updates, Adds and Remove actors to match the provided Step Data */
 	void UpdateFromRecordedStepData(const int32 SolverID, const FString& SolverName, const FChaosVDStepData& InRecordedStepData, const FChaosVDSolverFrameData& InFrameData);
 
-	void HandleNewGeometryData(const TSharedPtr<const Chaos::FImplicitObject>&, const int32 GeometryID) const;
+	void HandleNewGeometryData(const TSharedPtr<const Chaos::FImplicitObject>&, const uint32 GeometryID) const;
 
 	void HandleEnterNewGameFrame(int32 FrameNumber, const TArray<int32>& AvailableSolversIds);
 
@@ -55,6 +55,8 @@ public:
 
 	FChaosVDGeometryDataLoaded& OnNewGeometryAvailable(){ return NewGeometryAvailableDelegate; }
 
+	const TSharedPtr<const Chaos::FImplicitObject>* GetUpdatedGeometry(int32 GeometryID) const;
+
 	TSharedPtr<FChaosVDRecording> LoadedRecording;
 
 private:
diff --git a/Engine/Plugins/Experimental/ChaosVD/Source/ChaosVD/Public/Trace/ChaosVDTraceAnalyzer.h b/Engine/Plugins/Experimental/ChaosVD/Source/ChaosVD/Public/Trace/ChaosVDTraceAnalyzer.h
--- a/Engine/Plugins/Experimental/ChaosVD/Source/ChaosVD/Public/Trace/ChaosVDTraceAnalyzer.h
+++ b/Engine/Plugins/Experimental/ChaosVD/Source/ChaosVD/Public/Trace/ChaosVDTraceAnalyzer.h
@@ -46,7 +46,8 @@ private:
 		RouteId_ChaosVDBinaryDataEnd,
 		RouteId_ChaosVDSolverSimulationSpace,
 		RouteId_BeginFrame,
-		RouteId_EndFrame
+		RouteId_EndFrame,
+		RouteId_ChaosVDParticleDestroyed,
 	};
 
 	void ReadParticleDataFromEvent(const FEventData& InEventData, FChaosVDParticleDebugData& OutParticleData);
diff --git a/Engine/Plugins/Experimental/ChaosVD/Source/ChaosVD/Public/Trace/ChaosVDTraceProvider.h b/Engine/Plugins/Experimental/ChaosVD/Source/ChaosVD/Public/Trace/ChaosVDTraceProvider.h
--- a/Engine/Plugins/Experimental/ChaosVD/Source/ChaosVD/Public/Trace/ChaosVDTraceProvider.h
+++ b/Engine/Plugins/Experimental/ChaosVD/Source/ChaosVD/Public/Trace/ChaosVDTraceProvider.h
@@ -17,7 +17,7 @@ struct FChaosVDRecording;
 
 struct FChaosVDBinaryDataContainer
 {
-	FChaosVDBinaryDataContainer(const int32 InDataID)
+	explicit FChaosVDBinaryDataContainer(const int32 InDataID)
 		: DataID(InDataID)
 	{
 	}
@@ -77,6 +77,4 @@ private:
 	TMap<int32, TSharedPtr<FChaosVDBinaryDataContainer>> UnprocessedDataByID;
 	
 	FOnBinaryDataReady BinaryDataReadyDelegate;
-
-	TUniquePtr<Chaos::FChaosArchiveContext> ChaosContext;
 };
diff --git a/Engine/Plugins/Experimental/ChaosVD/Source/ChaosVD/Public/Widgets/SChaosVDSolverPlaybackControls.h b/Engine/Plugins/Experimental/ChaosVD/Source/ChaosVD/Public/Widgets/SChaosVDSolverPlaybackControls.h
--- a/Engine/Plugins/Experimental/ChaosVD/Source/ChaosVD/Public/Widgets/SChaosVDSolverPlaybackControls.h
+++ b/Engine/Plugins/Experimental/ChaosVD/Source/ChaosVD/Public/Widgets/SChaosVDSolverPlaybackControls.h
@@ -11,6 +11,18 @@ struct FChaosVDTrackInfo;
 class FChaosVDPlaybackController;
 class SChaosVDTimelineWidget;
 
+/**
+ * Options flags to control how the Step timeline widgets should be updated
+ */
+enum class EChaosVDStepsWidgetUpdateFlags
+{
+	UpdateText =  1 << 0,
+	SetTimelineStep = 1 << 1,
+
+	Default = UpdateText | SetTimelineStep
+};
+ENUM_CLASS_FLAGS(EChaosVDStepsWidgetUpdateFlags)
+
 /** Widget that Generates playback controls for solvers
  * Which are two timelines, one for physics frames and other for steps
  */
@@ -24,13 +36,15 @@ public:
 
 private:
 
-	void OnFrameSelectionUpdated(int32 NewFrameIndex) const;
-	void OnStepSelectionUpdated(int32 NewStepIndex) const;
+	void OnFrameSelectionUpdated(int32 NewFrameIndex);
+	void OnStepSelectionUpdated(int32 NewStepIndex);
 
 	virtual void HandlePlaybackControllerDataUpdated(TWeakPtr<FChaosVDPlaybackController> InController) override;
+	void UpdateStepsWidgetForFrame(const FChaosVDPlaybackController& InCurrentPlaybackController, int32 FrameNumber, int32 StepNumber, EChaosVDStepsWidgetUpdateFlags OptionsFlags = EChaosVDStepsWidgetUpdateFlags::Default);
 	virtual void HandleControllerTrackFrameUpdated(TWeakPtr<FChaosVDPlaybackController> InController, const FChaosVDTrackInfo* UpdatedTrackInfo, FGuid InstigatorGuid) override;
 
 	int32 SolverID = INDEX_NONE;
+	FString CurrentStepName;
 	TSharedPtr<SChaosVDTimelineWidget> FramesTimelineWidget;
 	TSharedPtr<SChaosVDTimelineWidget> StepsTimelineWidget;
 };
diff --git a/Engine/Source/Runtime/Engine/Private/PhysicsEngine/ImmediatePhysics/ImmediatePhysicsChaos/ImmediatePhysicsSimulation_Chaos.cpp b/Engine/Source/Runtime/Engine/Private/PhysicsEngine/ImmediatePhysics/ImmediatePhysicsChaos/ImmediatePhysicsSimulation_Chaos.cpp
--- a/Engine/Source/Runtime/Engine/Private/PhysicsEngine/ImmediatePhysics/ImmediatePhysicsChaos/ImmediatePhysicsSimulation_Chaos.cpp
+++ b/Engine/Source/Runtime/Engine/Private/PhysicsEngine/ImmediatePhysics/ImmediatePhysicsChaos/ImmediatePhysicsSimulation_Chaos.cpp
@@ -379,6 +379,10 @@ namespace ImmediatePhysics_Chaos
 	{
 		using namespace Chaos;
 
+		// We trace events like particles destruction, so we need to ensure there is a valid context
+		// during solver destruction so these events can be traced properly.
+		CVD_SCOPE_CONTEXT(CVDContextData)
+
 		// NOTE: Particles now hold a list of all the constraints that reference them, but when
 		// we delete a particle, we do not notify the constraints. When we destroy constarints
 		// it tries to remove itself from the particle's list, so we must destroy the
diff --git a/Engine/Source/Runtime/Experimental/Chaos/Private/Chaos/Evolution/PBDMinEvolution.cpp b/Engine/Source/Runtime/Experimental/Chaos/Private/Chaos/Evolution/PBDMinEvolution.cpp
--- a/Engine/Source/Runtime/Experimental/Chaos/Private/Chaos/Evolution/PBDMinEvolution.cpp
+++ b/Engine/Source/Runtime/Experimental/Chaos/Private/Chaos/Evolution/PBDMinEvolution.cpp
@@ -122,7 +122,8 @@ namespace Chaos
 	void FPBDMinEvolution::AdvanceOneTimeStep(const FReal Dt, const FReal StepFraction)
 	{
 		SCOPE_CYCLE_COUNTER(STAT_MinEvolution_AdvanceOneTimeStep);
-		CVD_SCOPE_TRACE_SOLVER_STEP();
+
+		CVD_SCOPE_TRACE_SOLVER_STEP(TEXT("Evolution Advance"));
 
 		Integrate(Dt);
 
@@ -143,7 +144,7 @@ namespace Chaos
 			ScatterOutput(Dt);
 		}
 
-		CVD_TRACE_PARTICLES(Particles.GetParticleHandles());
+		CVD_TRACE_PARTICLES_SOA(Particles);
 	}
 
 	// A opportunity for systems to allocate buffers for the duration of the tick, if they have enough info to do so
diff --git a/Engine/Source/Runtime/Experimental/Chaos/Private/Chaos/Framework/PhysicsSolverBase.cpp b/Engine/Source/Runtime/Experimental/Chaos/Private/Chaos/Framework/PhysicsSolverBase.cpp
--- a/Engine/Source/Runtime/Experimental/Chaos/Private/Chaos/Framework/PhysicsSolverBase.cpp
+++ b/Engine/Source/Runtime/Experimental/Chaos/Private/Chaos/Framework/PhysicsSolverBase.cpp
@@ -148,6 +148,7 @@ namespace Chaos
 
 	void FPhysicsSolverAdvanceTask::AdvanceSolver()
 	{
+		CVD_SCOPE_TRACE_SOLVER_FRAME(FPhysicsSolverBase, Solver);
 		LLM_SCOPE(ELLMTag::ChaosUpdate);
 		SCOPE_CYCLE_COUNTER(STAT_ChaosTick);
 		CSV_SCOPED_TIMING_STAT_EXCLUSIVE(Physics);
@@ -364,9 +365,6 @@ namespace Chaos
 
 	FGraphEventRef FPhysicsSolverBase::AdvanceAndDispatch_External(FReal InDt)
 	{
-
-		CVD_TRACE_SOLVER_START_FRAME(FPhysicsSolverBase, *this);
-
 		const bool bSubstepping = MMaxSubSteps > 1;
 		SetSolverSubstep_External(bSubstepping);
 		const FReal DtWithPause = bPaused_External ? 0.0f : InDt;
diff --git a/Engine/Source/Runtime/Experimental/Chaos/Private/Chaos/PBDRigidsEvolutionGBF.cpp b/Engine/Source/Runtime/Experimental/Chaos/Private/Chaos/PBDRigidsEvolutionGBF.cpp
--- a/Engine/Source/Runtime/Experimental/Chaos/Private/Chaos/PBDRigidsEvolutionGBF.cpp
+++ b/Engine/Source/Runtime/Experimental/Chaos/Private/Chaos/PBDRigidsEvolutionGBF.cpp
@@ -249,9 +249,7 @@ void FPBDRigidsEvolutionGBF::Advance(const FReal Dt,const FReal MaxStepDt,const
 }
 
 void FPBDRigidsEvolutionGBF::AdvanceOneTimeStep(const FReal Dt,const FSubStepInfo& SubStepInfo)
-{
-	CVD_SCOPE_TRACE_SOLVER_STEP();
-		
+{	
 	PrepareTick();
 
 	AdvanceOneTimeStepImpl(Dt, SubStepInfo);
@@ -339,10 +337,16 @@ void FPBDRigidsEvolutionGBF::AdvanceOneTimeStepImpl(const FReal Dt, const FSubSt
 	}
 #endif
 
+	{
+		CVD_SCOPE_TRACE_SOLVER_STEP(TEXT("Evolution Start"))
+		CVD_TRACE_PARTICLES_SOA(Particles);
+	}
+
 	// Update the collision solver type (used to support runtime comparisons of solver types for debugging/testing)
 	UpdateCollisionSolverType();
 
 	{
+		CVD_SCOPE_TRACE_SOLVER_STEP(TEXT("Integrate"))
 		SCOPE_CYCLE_COUNTER(STAT_Evolution_Integrate);
 		CSV_SCOPED_TIMING_STAT(PhysicsVerbose, StepSolver_Integrate);
 		Integrate(Particles.GetActiveParticlesView(), Dt);
@@ -355,6 +359,7 @@ void FPBDRigidsEvolutionGBF::AdvanceOneTimeStepImpl(const FReal Dt, const FSubSt
 	}
 
 	{
+		CVD_SCOPE_TRACE_SOLVER_STEP(TEXT("ApplyKinematicTargets"))
 		SCOPE_CYCLE_COUNTER(STAT_Evolution_KinematicTargets);
 		CSV_SCOPED_TIMING_STAT(PhysicsVerbose, StepSolver_KinematicTargets);
 		ApplyKinematicTargets(Dt, SubStepInfo.PseudoFraction);
@@ -561,7 +566,11 @@ void FPBDRigidsEvolutionGBF::AdvanceOneTimeStepImpl(const FReal Dt, const FSubSt
 		GetCollisionConstraints().DetectProbeCollisions(Dt);
 	}
 
-	CVD_TRACE_PARTICLES(GetParticleHandles());
+	{
+		CVD_SCOPE_TRACE_SOLVER_STEP(TEXT("Evolution End"))
+		CVD_TRACE_PARTICLES_SOA(Particles);
+	}
+	
 
 #if !UE_BUILD_SHIPPING
 	if(SerializeEvolution)
diff --git a/Engine/Source/Runtime/Experimental/Chaos/Private/ChaosModule.cpp b/Engine/Source/Runtime/Experimental/Chaos/Private/ChaosModule.cpp
--- a/Engine/Source/Runtime/Experimental/Chaos/Private/ChaosModule.cpp
+++ b/Engine/Source/Runtime/Experimental/Chaos/Private/ChaosModule.cpp
@@ -2,6 +2,7 @@
 
 #include "ChaosModule.h"
 #include "CoreMinimal.h"
+#include "ChaosVisualDebugger/ChaosVisualDebuggerTrace.h"
 #include "Modules/ModuleManager.h"
 #include "Misc/CommandLine.h"
 
@@ -23,4 +24,15 @@ void FChaosEngineModule::StartupModule()
 	{
 		Chaos::GSingleThreadedPhysics = 1;
 	}
+
+#if WITH_CHAOS_VISUAL_DEBUGGER
+	FChaosVisualDebuggerTrace::RegisterEventHandlers();
+#endif
+}
+
+void FChaosEngineModule::ShutdownModule()
+{
+#if WITH_CHAOS_VISUAL_DEBUGGER
+	FChaosVisualDebuggerTrace::UnregisterEventHandlers();
+#endif
 }
diff --git a/Engine/Source/Runtime/Experimental/Chaos/Private/ChaosVisualDebugger/ChaosVisualDebuggerTrace.cpp b/Engine/Source/Runtime/Experimental/Chaos/Private/ChaosVisualDebugger/ChaosVisualDebuggerTrace.cpp
--- a/Engine/Source/Runtime/Experimental/Chaos/Private/ChaosVisualDebugger/ChaosVisualDebuggerTrace.cpp
+++ b/Engine/Source/Runtime/Experimental/Chaos/Private/ChaosVisualDebugger/ChaosVisualDebuggerTrace.cpp
@@ -1,12 +1,13 @@
 // Copyright Epic Games, Inc. All Rights Reserved.
-
 #include "ChaosVisualDebugger/ChaosVisualDebuggerTrace.h"
 
+#include "Chaos/PBDRigidsSOAs.h"
+#include "Tasks/Task.h"
+
 #if CHAOS_VISUAL_DEBUGGER_ENABLED
 #include "Chaos/ParticleHandle.h"
 #include "Chaos/Framework/PhysicsSolverBase.h"
 #include "Chaos/ImplicitObject.h"
-#include "ChaosVDRuntimeModule.h"
 #include "Compression/OodleDataCompressionUtil.h"
 #include "HAL/CriticalSection.h"
 #include "Misc/ScopeRWLock.h"
@@ -19,12 +20,14 @@ UE_TRACE_EVENT_DEFINE(ChaosVDLogger, ChaosVDSolverFrameEnd)
 
 UE_TRACE_CHANNEL_DEFINE(ChaosVDChannel);
 UE_TRACE_EVENT_DEFINE(ChaosVDLogger, ChaosVDParticle)
+UE_TRACE_EVENT_DEFINE(ChaosVDLogger, ChaosVDParticleDestroyed)
 UE_TRACE_EVENT_DEFINE(ChaosVDLogger, ChaosVDSolverStepStart)
 UE_TRACE_EVENT_DEFINE(ChaosVDLogger, ChaosVDSolverStepEnd)
 UE_TRACE_EVENT_DEFINE(ChaosVDLogger, ChaosVDBinaryDataStart)
 UE_TRACE_EVENT_DEFINE(ChaosVDLogger, ChaosVDBinaryDataContent)
 UE_TRACE_EVENT_DEFINE(ChaosVDLogger, ChaosVDBinaryDataEnd)
 UE_TRACE_EVENT_DEFINE(ChaosVDLogger, ChaosVDSolverSimulationSpace)
+UE_TRACE_EVENT_DEFINE(ChaosVDLogger, ChaosVDDummyEvent)
 
 static FAutoConsoleVariable CVarChaosVDCompressBinaryData(
 	TEXT("p.Chaos.VD.CompressBinaryData"),
@@ -36,17 +39,73 @@ static FAutoConsoleVariable CVarChaosVDCompressionMode(
 	2,
 	TEXT("Oodle compression mode to use, 4 is by default which equsals to ECompressionLevel::VeryFast"));
 
+/** Struct where we keep track of the geometry we are tracing */
 struct FChaosVDGeometryTraceContext
 {
-	FRWLock RWLock;
-	TUniquePtr<Chaos::FChaosArchiveContext> ChaosContext;
+	FRWLock TracedGeometrySetLock;
+	FRWLock CachedGeometryHashesLock;
+	TSet<uint32> GeometryTracedIDs;
+
+	uint32 GetGeometryHashForImplicit(const Chaos::FImplicitObject* Implicit)
+	{
+		if (!ensure(Implicit != nullptr))
+		{
+			return 0;
+		}
+
+		{
+			FReadScopeLock ReadLock(CachedGeometryHashesLock);
+			if (uint32* FoundHashPtr = CachedGeometryHashes.Find((void*) Implicit))
+			{
+				return *FoundHashPtr;
+			}
+		}
+
+		{
+			uint32 Hash = Implicit->GetTypeHash();
+
+			FWriteScopeLock WriteLock(CachedGeometryHashesLock);
+			CachedGeometryHashes.Add((void*)Implicit, Hash);
+			return Hash;
+		}
+	}
+
+	void RemoveCachedGeometryHash(const Chaos::FImplicitObject* Implicit)
+	{
+		if (Implicit == nullptr)
+		{
+			return;
+		}
+
+		FWriteScopeLock WriteLock(TracedGeometrySetLock);
+		CachedGeometryHashes.Remove((void*)Implicit);
+	}
+
+	//TODO: Remove this when/if we have stable geometry hashes that can be serialized
+	// Right now it is too costly calculate them each time at runtime
+	TMap<void*, uint32> CachedGeometryHashes;
 };
 
 static FChaosVDGeometryTraceContext GeometryTracerObject = FChaosVDGeometryTraceContext();
 
+FDelegateHandle FChaosVisualDebuggerTrace::RecordingStartedDelegateHandle = FDelegateHandle();
+FDelegateHandle FChaosVisualDebuggerTrace::RecordingStoppedDelegateHandle = FDelegateHandle();
+FDelegateHandle FChaosVisualDebuggerTrace::RecordingFullCaptureRequestedHandle = FDelegateHandle();
+
+FRWLock FChaosVisualDebuggerTrace::DeltaRecordingStatesLock = FRWLock();
+TSet<int32> FChaosVisualDebuggerTrace::SolverIDsForDeltaRecording = TSet<int32>();
+TSet<int32> FChaosVisualDebuggerTrace::RequestedFullCaptureSolverIDs = TSet<int32>();
+FThreadSafeBool FChaosVisualDebuggerTrace::bIsTracing = false;
+
+void FChaosVDImplicitObjectDataWrapper::Serialize(Chaos::FChaosArchive& Ar)
+{
+	Ar << Hash;
+	Ar << ImplicitObject;
+}
+
 void FChaosVisualDebuggerTrace::TraceParticle(const Chaos::FGeometryParticleHandle* ParticleHandle)
 {
-	if (!UE_TRACE_CHANNELEXPR_IS_ENABLED(ChaosVDChannel))
+	if (!IsTracing())
 	{
 		return;
 	}
@@ -62,7 +121,7 @@ void FChaosVisualDebuggerTrace::TraceParticle(const Chaos::FGeometryParticleHand
 
 void FChaosVisualDebuggerTrace::TraceParticle(Chaos::FGeometryParticleHandle* ParticleHandle, const FChaosVDContext& ContextData)
 {
-	if (!UE_TRACE_CHANNELEXPR_IS_ENABLED(ChaosVDChannel))
+	if (!IsTracing())
 	{
 		return;
 	}
@@ -82,8 +141,11 @@ void FChaosVisualDebuggerTrace::TraceParticle(Chaos::FGeometryParticleHandle* Pa
 		ParticleNameView = FStringView(*DebugNamePtr.Get());
 	}
 #endif
-	const int32 GeometryID = TraceImplicitObject(ParticleHandle->Geometry());
+
+	const uint32 GeometryHash = GeometryTracerObject.GetGeometryHashForImplicit(ParticleHandle->Geometry().Get());
 	
+	TraceImplicitObject({ GeometryHash, ParticleHandle->Geometry() });
+
 	UE_TRACE_LOG(ChaosVDLogger, ChaosVDParticle, ChaosVDChannel)
 		<< ChaosVDParticle.SolverID(ContextData.Id)
 		<< ChaosVDParticle.Cycle(FPlatformTime::Cycles64())
@@ -97,15 +159,60 @@ void FChaosVisualDebuggerTrace::TraceParticle(Chaos::FGeometryParticleHandle* Pa
 		<< CVD_TRACE_VECTOR_ON_EVENT(ChaosVDParticle, Velocity, Chaos::FConstGenericParticleHandle(ParticleHandle)->V())
 		<< CVD_TRACE_VECTOR_ON_EVENT(ChaosVDParticle, AngularVelocity, Chaos::FConstGenericParticleHandle(ParticleHandle)->W())
 
-		<< ChaosVDParticle.ImplicitObjectID(GeometryID)
+		<< ChaosVDParticle.ObjectState(static_cast<int8>(ParticleHandle->ObjectState()))
+		<< ChaosVDParticle.ImplicitObjectHash(GeometryHash);
+}
+
+void FChaosVisualDebuggerTrace::TraceParticles(const Chaos::TGeometryParticleHandles<Chaos::FReal, 3>& ParticleHandles)
+{
+	if (!IsTracing())
+	{
+		return;
+	}
+
+	const FChaosVDContext* CVDContextData = FChaosVDThreadContext::Get().GetCurrentContext();
+	if (!ensure(CVDContextData))
+	{
+		return;
+	}
+
+	ParallelFor(ParticleHandles.Size(),[&ParticleHandles, CopyContext = *CVDContextData](int32 ParticleIndex)
+	{
+		CVD_SCOPE_CONTEXT(CopyContext)
+		TraceParticle(ParticleHandles.Handle(ParticleIndex).Get(), CopyContext);
+	});
+}
 
-		<< ChaosVDParticle.ObjectState(static_cast<int8>(ParticleHandle->ObjectState()));
+void FChaosVisualDebuggerTrace::TraceParticleDestroyed(const Chaos::FGeometryParticleHandle* ParticleHandle)
+{
+	if (!IsTracing())
+	{
+		return;
+	}
+
+	if (!ParticleHandle)
+	{
+		UE_LOG(LogChaos, Warning, TEXT("Tried to Trace a null particle %hs"), __FUNCTION__);
+		return;
+	}
+
+	GeometryTracerObject.RemoveCachedGeometryHash(ParticleHandle->Geometry().Get());
 	
+	const FChaosVDContext* CVDContextData = FChaosVDThreadContext::Get().GetCurrentContext();
+	if (!ensure(CVDContextData))
+	{
+		return;
+	}
+
+	UE_TRACE_LOG(ChaosVDLogger, ChaosVDParticleDestroyed, ChaosVDChannel)
+		<< ChaosVDParticleDestroyed.SolverID(CVDContextData->Id)
+		<< ChaosVDParticleDestroyed.Cycle(FPlatformTime::Cycles64())
+		<< ChaosVDParticleDestroyed.ParticleID(ParticleHandle->UniqueIdx().Idx);
 }
 
-void FChaosVisualDebuggerTrace::TraceParticles(const Chaos::TGeometryParticleHandles<Chaos::FReal, 3>& ParticleHandles)
+void FChaosVisualDebuggerTrace::TraceParticlesSoA(const Chaos::FPBDRigidsSOAs& ParticlesSoA)
 {
-	if (!UE_TRACE_CHANNELEXPR_IS_ENABLED(ChaosVDChannel))
+	if (!IsTracing())
 	{
 		return;
 	}
@@ -116,16 +223,46 @@ void FChaosVisualDebuggerTrace::TraceParticles(const Chaos::TGeometryParticleHan
 		return;
 	}
 
-	for (uint32 ParticleIndex = 0; ParticleIndex < ParticleHandles.Size(); ParticleIndex++)
+	// If this solver is not being delta recorded, Trace all the particles
+	if (ShouldPerformFullCapture(CVDContextData->Id))
+	{
+		TraceParticlesView(ParticlesSoA.GetAllParticlesView());
+		return;
+	}
+
+	TraceParticlesView(ParticlesSoA.GetDirtyParticlesView());
+}
+
+void FChaosVisualDebuggerTrace::SetupForFullCaptureIfNeeded(int32 SolverID, bool& bOutFullCaptureRequested)
+{
+	DeltaRecordingStatesLock.ReadLock();
+	bOutFullCaptureRequested = RequestedFullCaptureSolverIDs.Contains(SolverID);
+	DeltaRecordingStatesLock.ReadUnlock();
+
+	if (bOutFullCaptureRequested)
 	{
-		// TODO: We should only trace the Particles that changed. probably using the Dirty Flags?
-		// Geometry data "uniqueness" will be handled by the trace helper
-		TraceParticle(ParticleHandles.Handle(ParticleIndex).Get(), *CVDContextData);
+		FWriteScopeLock WriteLock(DeltaRecordingStatesLock);
+		SolverIDsForDeltaRecording.Remove(SolverID);
+		RequestedFullCaptureSolverIDs.Remove(SolverID);
 	}
 }
 
+bool FChaosVisualDebuggerTrace::ShouldPerformFullCapture(int32 SolverID)
+{
+	FReadScopeLock ReadLock(DeltaRecordingStatesLock);
+	int32* FoundSolverID = SolverIDsForDeltaRecording.Find(SolverID);
+
+	// If the solver ID is on the SolverIDsForDeltaRecording set, it means we should NOT perform a full capture
+	return FoundSolverID == nullptr;
+}
+
 void FChaosVisualDebuggerTrace::TraceSolverFrameStart(const FChaosVDContext& ContextData, const FString& InDebugName)
 {
+	if (!IsTracing())
+	{
+		return;
+	}
+
 	if (!ensure(ContextData.Id != INDEX_NONE))
 	{
 		return;
@@ -133,14 +270,24 @@ void FChaosVisualDebuggerTrace::TraceSolverFrameStart(const FChaosVDContext& Con
 
 	FChaosVDThreadContext::Get().PushContext(ContextData);
 
+	// Check if we need to do a full capture for this solver, and setup accordingly
+	bool bOutIsFullCaptureRequested;
+	SetupForFullCaptureIfNeeded(ContextData.Id, bOutIsFullCaptureRequested);
+
 	UE_TRACE_LOG(ChaosVDLogger, ChaosVDSolverFrameStart, ChaosVDChannel)
 		<< ChaosVDSolverFrameStart.SolverID(ContextData.Id)
 		<< ChaosVDSolverFrameStart.Cycle(FPlatformTime::Cycles64())
-		<< ChaosVDSolverFrameStart.DebugName(*InDebugName, InDebugName.Len());
+		<< ChaosVDSolverFrameStart.DebugName(*InDebugName, InDebugName.Len())
+		<< ChaosVDSolverFrameStart.IsKeyFrame(bOutIsFullCaptureRequested);
 }
 
 void FChaosVisualDebuggerTrace::TraceSolverFrameEnd(const FChaosVDContext& ContextData)
 {
+	if (!IsTracing())
+	{
+		return;
+	}
+
 	FChaosVDThreadContext::Get().PopContext();
 
 	if (!ensure(ContextData.Id != INDEX_NONE))
@@ -148,14 +295,22 @@ void FChaosVisualDebuggerTrace::TraceSolverFrameEnd(const FChaosVDContext& Conte
 		return;
 	}
 
+	{
+		FWriteScopeLock WriteLock(DeltaRecordingStatesLock);
+		if (!SolverIDsForDeltaRecording.Contains(ContextData.Id))
+		{
+			SolverIDsForDeltaRecording.Add(ContextData.Id);
+		}
+	}
+
 	UE_TRACE_LOG(ChaosVDLogger, ChaosVDSolverFrameEnd, ChaosVDChannel)
 		<< ChaosVDSolverFrameEnd.SolverID(ContextData.Id)
 		<< ChaosVDSolverFrameEnd.Cycle(FPlatformTime::Cycles64());
 }
 
-void FChaosVisualDebuggerTrace::TraceSolverStepStart()
+void FChaosVisualDebuggerTrace::TraceSolverStepStart(FStringView StepName)
 {
-	if (!UE_TRACE_CHANNELEXPR_IS_ENABLED(ChaosVDChannel))
+	if (!IsTracing())
 	{
 		return;
 	}
@@ -168,12 +323,13 @@ void FChaosVisualDebuggerTrace::TraceSolverStepStart()
 
 	UE_TRACE_LOG(ChaosVDLogger, ChaosVDSolverStepStart, ChaosVDChannel)
 		<< ChaosVDSolverStepStart.Cycle(FPlatformTime::Cycles64())
-		<< ChaosVDSolverStepStart.SolverID(CVDContextData->Id);
+		<< ChaosVDSolverStepStart.SolverID(CVDContextData->Id)
+		<< ChaosVDSolverStepStart.StepName(StepName.GetData(), StepName.Len());
 }
 
 void FChaosVisualDebuggerTrace::TraceSolverStepEnd()
 {
-	if (!UE_TRACE_CHANNELEXPR_IS_ENABLED(ChaosVDChannel))
+	if (!IsTracing())
 	{
 		return;
 	}
@@ -191,7 +347,7 @@ void FChaosVisualDebuggerTrace::TraceSolverStepEnd()
 
 void FChaosVisualDebuggerTrace::TraceSolverSimulationSpace(const Chaos::FRigidTransform3& Transform)
 {
-	if (!UE_TRACE_CHANNELEXPR_IS_ENABLED(ChaosVDChannel))
+	if (!IsTracing())
 	{
 		return;
 	}
@@ -211,7 +367,7 @@ void FChaosVisualDebuggerTrace::TraceSolverSimulationSpace(const Chaos::FRigidTr
 
 void FChaosVisualDebuggerTrace::TraceBinaryData(const TArray<uint8>& InData, const FString& TypeName)
 {
-	if (!UE_TRACE_CHANNELEXPR_IS_ENABLED(ChaosVDChannel))
+	if (!IsTracing())
 	{
 		return;
 	}
@@ -270,68 +426,123 @@ void FChaosVisualDebuggerTrace::TraceBinaryData(const TArray<uint8>& InData, con
 	ensure(RemainingSize == 0);
 }
 
-int32 FChaosVisualDebuggerTrace::TraceImplicitObject(Chaos::TSerializablePtr<Chaos::FImplicitObject> Geometry)
+void FChaosVisualDebuggerTrace::TraceImplicitObject(FChaosVDImplicitObjectDataWrapper WrappedGeometryData)
 {
-	if (!UE_TRACE_CHANNELEXPR_IS_ENABLED(ChaosVDChannel))
+	if (!IsTracing())
 	{
-		return INDEX_NONE;
+		return;
 	}
 
-	//TODO: Find a better place to do this
-	if (!FChaosVDRuntimeModule::OnRecordingStop().IsBound())
+	uint32 GeometryID = WrappedGeometryData.Hash;
 	{
-		FChaosVDRuntimeModule::OnRecordingStop().BindStatic(&FChaosVisualDebuggerTrace::ResetGeometryTracerContext);
+		FReadScopeLock ReadLock(GeometryTracerObject.TracedGeometrySetLock);
+		if (GeometryTracerObject.GeometryTracedIDs.Find(GeometryID))
+		{
+			return;
+		}
 	}
-	
-	if (!FChaosVDRuntimeModule::OnRecordingStarted().IsBound())
+
 	{
-		FChaosVDRuntimeModule::OnRecordingStarted().BindStatic(&FChaosVisualDebuggerTrace::ResetGeometryTracerContext);
+		FWriteScopeLock WriteLock(GeometryTracerObject.TracedGeometrySetLock);
+		GeometryTracerObject.GeometryTracedIDs.Add(GeometryID);
 	}
 
+	FChaosVDScopedTLSBufferAccessor TLSDataBuffer;
+
+	FMemoryWriter MemWriterAr(TLSDataBuffer.BufferRef);
+	Chaos::FChaosArchive Ar(MemWriterAr);
+
+	Ar.SetShouldSkipUpdateCustomVersion(true);
+
+	WrappedGeometryData.Serialize(Ar);
+
+	TraceBinaryData(TLSDataBuffer.BufferRef, TEXT("FChaosVDImplicitObjectDataWrapper"));
+}	
+
+void FChaosVisualDebuggerTrace::RegisterEventHandlers()
+{
 	{
-		FReadScopeLock Lock(GeometryTracerObject.RWLock);
-		if (GeometryTracerObject.ChaosContext.IsValid())
+		FWriteScopeLock WriteLock(DeltaRecordingStatesLock);
+
+		
+		if (!RecordingStartedDelegateHandle.IsValid())
 		{
-			int32 SerializedObjectPtrTag = GeometryTracerObject.ChaosContext.Get()->GetObjectTag(Geometry.Get());
-			if (SerializedObjectPtrTag != INDEX_NONE)
-			{
-				// Geometry Data is already serialized, just return the tag
-				return SerializedObjectPtrTag;
-			}
+			RecordingStartedDelegateHandle = FChaosVDRuntimeModule::RegisterRecordingStartedCallback(FChaosVDRecordingStateChangedDelegate::FDelegate::CreateStatic(&FChaosVisualDebuggerTrace::HandleRecordingStart));
 		}
-	}
 
-	// TODO: Change this so it is not allocated each time
-	// We need to take into account we could be serializing on multiple threads
-	constexpr uint32 MaxTraceChunkSize = TNumericLimits<uint16>::Max();
-	TArray<uint8> RawData;
-	RawData.Reserve(MaxTraceChunkSize);
-	
-	int32 SerializedObjectPtrTag = INDEX_NONE;
+		if (!RecordingStoppedDelegateHandle.IsValid())
+		{
+			RecordingStoppedDelegateHandle = FChaosVDRuntimeModule::RegisterRecordingStopCallback(FChaosVDRecordingStateChangedDelegate::FDelegate::CreateStatic(&FChaosVisualDebuggerTrace::HandleRecordingStop));
+		}
 
-	{
-		FMemoryWriter MemWriterAr(RawData);
-		Chaos::FChaosArchive Ar(MemWriterAr);
-		
-		FWriteScopeLock WriteLock(GeometryTracerObject.RWLock);
-		if (GeometryTracerObject.ChaosContext.IsValid())
+		if (!RecordingFullCaptureRequestedHandle.IsValid())
 		{
-			Ar.SetContext(MoveTemp(GeometryTracerObject.ChaosContext));
+			RecordingFullCaptureRequestedHandle = FChaosVDRuntimeModule::RegisterFullCaptureRequestedCallback(FChaosVDCaptureRequestDelegate::FDelegate::CreateStatic(&FChaosVisualDebuggerTrace::PerformFullCapture));
 		}
+	}
+}
+
+void FChaosVisualDebuggerTrace::UnregisterEventHandlers()
+{
+	FWriteScopeLock WriteLock(DeltaRecordingStatesLock);
+	if (RecordingStartedDelegateHandle.IsValid())
+	{
+		FChaosVDRuntimeModule::RemoveRecordingStartedCallback(RecordingStartedDelegateHandle);
+	}
 
-		Ar << Geometry;
-		GeometryTracerObject.ChaosContext = Ar.StealContext();
-		SerializedObjectPtrTag = GeometryTracerObject.ChaosContext.Get()->GetObjectTag(Geometry.Get());
+	if (RecordingStoppedDelegateHandle.IsValid())
+	{
+		FChaosVDRuntimeModule::RemoveRecordingStopCallback(RecordingStoppedDelegateHandle);
 	}
 
-	TraceBinaryData(RawData, TEXT("FImplicitObject"));
+	if (RecordingFullCaptureRequestedHandle.IsValid())
+	{
+		FChaosVDRuntimeModule::RemoveFullCaptureRequestedCallback(RecordingFullCaptureRequestedHandle);
+	}
 
-	return SerializedObjectPtrTag;
+	bIsTracing = false;
 }
 
-void FChaosVisualDebuggerTrace::ResetGeometryTracerContext()
+void FChaosVisualDebuggerTrace::Reset()
 {
-	GeometryTracerObject.ChaosContext.Release();
+	{
+		FWriteScopeLock WriteLock(DeltaRecordingStatesLock);
+		RequestedFullCaptureSolverIDs.Reset();
+		SolverIDsForDeltaRecording.Reset();
+	}
+
+	{
+		FWriteScopeLock GeometryWriteLock(GeometryTracerObject.TracedGeometrySetLock);
+		GeometryTracerObject.GeometryTracedIDs.Reset();
+		GeometryTracerObject.CachedGeometryHashes.Reset();
+	}
 }
 
-#endif
+void FChaosVisualDebuggerTrace::HandleRecordingStop()
+{
+	bIsTracing = false;
+	Reset();
+}
+
+void FChaosVisualDebuggerTrace::HandleRecordingStart()
+{
+	Reset();
+	bIsTracing = true;
+}
+
+void FChaosVisualDebuggerTrace::PerformFullCapture(EChaosVDFullCaptureFlags CaptureOptions)
+{
+	if (EnumHasAnyFlags(CaptureOptions, EChaosVDFullCaptureFlags::Particles))
+	{
+		FWriteScopeLock WriteLock(DeltaRecordingStatesLock);
+		RequestedFullCaptureSolverIDs.Append(SolverIDsForDeltaRecording);
+	}
+
+	if (EnumHasAnyFlags(CaptureOptions, EChaosVDFullCaptureFlags::Geometry))
+	{
+		FWriteScopeLock GeometryWriteLock(GeometryTracerObject.TracedGeometrySetLock);
+		GeometryTracerObject.GeometryTracedIDs.Reset();
+	}
+}
+
+#endif
\ No newline at end of file
diff --git a/Engine/Source/Runtime/Experimental/Chaos/Public/Chaos/Framework/PhysicsSolverBase.h b/Engine/Source/Runtime/Experimental/Chaos/Public/Chaos/Framework/PhysicsSolverBase.h
--- a/Engine/Source/Runtime/Experimental/Chaos/Public/Chaos/Framework/PhysicsSolverBase.h
+++ b/Engine/Source/Runtime/Experimental/Chaos/Public/Chaos/Framework/PhysicsSolverBase.h
@@ -151,6 +151,11 @@ namespace Chaos
 
 #if WITH_CHAOS_VISUAL_DEBUGGER
 		FChaosVDContext CVDContext;
+	public:
+		FChaosVDContext& GetChaosVDContextData()
+		{
+			return CVDContext;
+		};
 #endif
 	};
 
diff --git a/Engine/Source/Runtime/Experimental/Chaos/Public/Chaos/PBDRigidsEvolution.h b/Engine/Source/Runtime/Experimental/Chaos/Public/Chaos/PBDRigidsEvolution.h
--- a/Engine/Source/Runtime/Experimental/Chaos/Public/Chaos/PBDRigidsEvolution.h
+++ b/Engine/Source/Runtime/Experimental/Chaos/Public/Chaos/PBDRigidsEvolution.h
@@ -19,6 +19,7 @@
 #include "Chaos/PendingSpatialData.h"
 #include "ProfilingDebugging/CsvProfiler.h"
 #include "RewindData.h"
+#include "ChaosVisualDebugger/ChaosVDContextProvider.h"
 
 extern int32 ChaosRigidsEvolutionApplyAllowEarlyOutCVar;
 extern int32 ChaosRigidsEvolutionApplyPushoutAllowEarlyOutCVar;
@@ -741,12 +742,17 @@ public:
 
 		const bool IsLastStep = (FMath::IsNearlyEqual(StepFraction, (FReal)1, (FReal)UE_KINDA_SMALL_NUMBER));
 
+		FChaosVDContextWrapper CVDContext;
+		CVD_GET_WRAPPED_CURRENT_CONTEXT(CVDContext);
+
 		// NOTE: ApplyKinematicTargetForParticle is run in a parallel-for. We only write to particle state
 		const auto& ApplyParticleKinematicTarget = 
-		[Dt, StepFraction, IsLastStep]
+		[Dt, StepFraction, IsLastStep, CVDContext]
 		(FTransientPBDRigidParticleHandle& Particle, const int32 ParticleIndex)
 		-> void
 		{
+			CVD_SCOPE_CONTEXT(CVDContext.Context);
+
 			TKinematicTarget<FReal, 3>& KinematicTarget = Particle.KinematicTarget();
 			const FVec3 CurrentX = Particle.X();
 			const FRotation3 CurrentR = Particle.R();
@@ -840,6 +846,8 @@ public:
 					Particle.UpdateWorldSpaceStateSwept(FRigidTransform3(Particle.P(), Particle.Q()), FVec3(0), -Particle.V() * Dt);
 				}
 			}
+
+			CVD_TRACE_PARTICLE(Particle.Handle())
 		};
 
 		// Apply kinematic targets in parallel
diff --git a/Engine/Source/Runtime/Experimental/Chaos/Public/Chaos/PBDRigidsEvolutionGBF.h b/Engine/Source/Runtime/Experimental/Chaos/Public/Chaos/PBDRigidsEvolutionGBF.h
--- a/Engine/Source/Runtime/Experimental/Chaos/Public/Chaos/PBDRigidsEvolutionGBF.h
+++ b/Engine/Source/Runtime/Experimental/Chaos/Public/Chaos/PBDRigidsEvolutionGBF.h
@@ -190,8 +190,14 @@ namespace Chaos
 			const FReal MaxVelocityBoundsExpansion = GetCollisionConstraints().GetDetectorSettings().MaxVelocityBoundsExpansion;
 			const FReal MaxAngularSpeedSq = CVars::HackMaxAngularVelocity * CVars::HackMaxAngularVelocity;
 			const FReal MaxSpeedSq = CVars::HackMaxVelocity * CVars::HackMaxVelocity;
+			
+			FChaosVDContextWrapper CVDContext;
+			CVD_GET_WRAPPED_CURRENT_CONTEXT(CVDContext);
+			
 			InParticles.ParallelFor([&](auto& GeomParticle, int32 Index) 
 			{
+				CVD_SCOPE_CONTEXT(CVDContext.Context);
+
 				//question: can we enforce this at the API layer? Right now islands contain non dynamic which makes this hard
 				auto PBDParticle = GeomParticle.CastToRigidParticle();
 				if (PBDParticle && PBDParticle->ObjectState() == EObjectStateType::Dynamic)
@@ -316,6 +322,8 @@ namespace Chaos
 							Particle.UpdateWorldSpaceStateSwept(FRigidTransform3(Particle.P(), Particle.Q()), FVec3(BoundsThickness), VelocityBoundsDelta);
 						}
 					}
+
+					CVD_TRACE_PARTICLE(PBDParticle->Handle())
 				}
 			});
 
diff --git a/Engine/Source/Runtime/Experimental/Chaos/Public/Chaos/PBDRigidsSOAs.h b/Engine/Source/Runtime/Experimental/Chaos/Public/Chaos/PBDRigidsSOAs.h
--- a/Engine/Source/Runtime/Experimental/Chaos/Public/Chaos/PBDRigidsSOAs.h
+++ b/Engine/Source/Runtime/Experimental/Chaos/Public/Chaos/PBDRigidsSOAs.h
@@ -2,6 +2,7 @@
 #pragma once
 
 #include "Chaos/ParticleHandle.h"
+#include "ChaosVisualDebugger/ChaosVisualDebuggerTrace.h"
 #include "UObject/ExternalPhysicsCustomObjectVersion.h"
 
 namespace Chaos
@@ -366,6 +367,8 @@ public:
 			}
 		}
 
+		CVD_TRACE_PARTICLE_DESTROYED(Particle);
+
 		auto PBDRigid = Particle->CastToRigidParticle();
 		if(PBDRigid)
 		{
diff --git a/Engine/Source/Runtime/Experimental/Chaos/Public/ChaosModule.h b/Engine/Source/Runtime/Experimental/Chaos/Public/ChaosModule.h
--- a/Engine/Source/Runtime/Experimental/Chaos/Public/ChaosModule.h
+++ b/Engine/Source/Runtime/Experimental/Chaos/Public/ChaosModule.h
@@ -9,5 +9,5 @@ class CHAOS_API FChaosEngineModule : public IModuleInterface
 public:
 
 	virtual void StartupModule() override;
-	virtual void ShutdownModule() override {};
+	virtual void ShutdownModule() override;
 };
diff --git a/Engine/Source/Runtime/Experimental/Chaos/Public/ChaosVisualDebugger/ChaosVDContextProvider.h b/Engine/Source/Runtime/Experimental/Chaos/Public/ChaosVisualDebugger/ChaosVDContextProvider.h
--- a/Engine/Source/Runtime/Experimental/Chaos/Public/ChaosVisualDebugger/ChaosVDContextProvider.h
+++ b/Engine/Source/Runtime/Experimental/Chaos/Public/ChaosVisualDebugger/ChaosVDContextProvider.h
@@ -8,6 +8,8 @@
 
 #if WITH_CHAOS_VISUAL_DEBUGGER
 
+#include "ChaosVDRuntimeModule.h"
+
 /** Chaos Visual Debugger data used to context for logging or debugging purposes */
 struct CHAOS_API FChaosVDContext
 {
@@ -20,6 +22,14 @@ class CHAOS_API FChaosVDThreadContext : public TThreadSingleton<FChaosVDThreadCo
 {
 public:
 
+	virtual ~FChaosVDThreadContext() override
+	{
+		if (RecordingStoppedHandle.IsValid())
+		{
+			FChaosVDRuntimeModule::RemoveRecordingStopCallback(RecordingStoppedHandle);
+		}
+	}
+
 	/** Copies the Current CVD context data into the provided struct
 	 * @return true if the copy was successful
 	 */
@@ -39,6 +49,27 @@ public:
 
 protected:
 	TArray<FChaosVDContext, TInlineAllocator<16>> LocalContextStack;
+
+	TArray<uint8>& GetTLSDataBufferRef()
+	{
+		checkf(BufferAccessCounter == 0, TEXT("The CVD buffer is already in use!"));
+
+		if (!RecordingStoppedHandle.IsValid())
+		{
+			// If we access the buffer, make sure we subscribe to the recording stop callback
+			// So we free any memory we keep allocated
+			RecordingStoppedHandle = FChaosVDRuntimeModule::RegisterRecordingStopCallback(FChaosVDRecordingStateChangedDelegate::FDelegate::CreateLambda([this](){ CVDDataBuffer.Empty(); }));
+		}
+		
+		return CVDDataBuffer;
+	}
+
+	FDelegateHandle RecordingStoppedHandle;
+
+	int32 BufferAccessCounter = 0;
+	TArray<uint8> CVDDataBuffer;
+
+	friend struct FChaosVDScopedTLSBufferAccessor;
 };
 
 /** Utility Class that will push the provided CVD Context Data to the local thread storage
@@ -57,9 +88,45 @@ struct CHAOS_API FChaosVDScopeContext
 	}
 };
 
+/** Wrapper class that provides access to the data buffer that should be used to trace CVD (Chaos Visual Debugger) data.
+ * It ensures the buffer is cleared on scope exit but keeping a minimum amount of memory allocated
+ */
+struct CHAOS_API FChaosVDScopedTLSBufferAccessor
+{
+	FChaosVDScopedTLSBufferAccessor() : BufferRef(FChaosVDThreadContext::Get().GetTLSDataBufferRef())
+	{
+		FChaosVDThreadContext::Get().BufferAccessCounter++;
+
+		// This is close to the max amount of data we can Trace as a single trace event
+		// Some times we will go over and some times this be more than we need
+		// but keeping this amount allocated on average reduces the allocation cost when tracing.
+		// TODO: There might be a better value for this, we need to do more testing
+		constexpr uint32 MaxTraceChunkSize = TNumericLimits<uint16>::Max();
+		BufferRef.Reserve(MaxTraceChunkSize);
+	}
+
+	~FChaosVDScopedTLSBufferAccessor()
+	{
+		// Clear the buffer on exit but maintain allocated a minimum amount of memory
+		constexpr uint32 MaxTraceChunkSize = TNumericLimits<uint16>::Max();
+		BufferRef.Reset(MaxTraceChunkSize);
+
+		FChaosVDThreadContext::Get().BufferAccessCounter--;
+	}
+	
+	TArray<uint8>& BufferRef;
+};
+
 #ifndef CVD_GET_CURRENT_CONTEXT
 	#define CVD_GET_CURRENT_CONTEXT(OutContext) \
-		ensure(FChaosVDThreadContext::Get().GetCurrentContext(OutContext))
+		FChaosVDThreadContext::Get().GetCurrentContext(OutContext);
+#endif
+
+#ifndef CVD_GET_WRAPPED_CURRENT_CONTEXT
+	#define CVD_GET_WRAPPED_CURRENT_CONTEXT(OutWrappedContext) \
+		FChaosVDContext CurrentContext; \
+		FChaosVDThreadContext::Get().GetCurrentContext(CurrentContext); \
+		OutWrappedContext = FChaosVDContextWrapper(CurrentContext);
 #endif
 
 #ifndef CVD_SCOPE_CONTEXT
@@ -77,5 +144,24 @@ struct CHAOS_API FChaosVDScopeContext
 	#define CVD_SCOPE_CONTEXT(InContext)
 #endif
 
+#ifndef CVD_GET_WRAPPED_CURRENT_CONTEXT
+	#define CVD_GET_WRAPPED_CURRENT_CONTEXT(OutWrappedContext)
+#endif
+
 #endif // WITH_CHAOS_VISUAL_DEBUGGER
 
+struct CHAOS_API FChaosVDContextWrapper
+{
+	FChaosVDContextWrapper()
+	{
+	}
+	
+#if WITH_CHAOS_VISUAL_DEBUGGER
+	FChaosVDContextWrapper(const FChaosVDContext& InContext)
+	: Context(InContext)
+	{
+	}
+
+	FChaosVDContext Context;
+#endif
+};
diff --git a/Engine/Source/Runtime/Experimental/Chaos/Public/ChaosVisualDebugger/ChaosVisualDebuggerTrace.h b/Engine/Source/Runtime/Experimental/Chaos/Public/ChaosVisualDebugger/ChaosVisualDebuggerTrace.h
--- a/Engine/Source/Runtime/Experimental/Chaos/Public/ChaosVisualDebugger/ChaosVisualDebuggerTrace.h
+++ b/Engine/Source/Runtime/Experimental/Chaos/Public/ChaosVisualDebugger/ChaosVisualDebuggerTrace.h
@@ -4,10 +4,12 @@
 #include "Chaos/Core.h"
 #include "Chaos/ParticleHandleFwd.h"
 #include "Chaos/Serializable.h"
-#include "Chaos/Vector.h"
 #include "Containers/Set.h"
 
 #include "Chaos/ChaosArchive.h"
+#include "ChaosVDContextProvider.h"
+#include "Chaos/ParticleIterator.h"
+#include "HAL/ThreadSafeBool.h"
 
 #ifndef CHAOS_VISUAL_DEBUGGER_ENABLED
 	#define CHAOS_VISUAL_DEBUGGER_ENABLED (WITH_CHAOS_VISUAL_DEBUGGER && UE_TRACE_ENABLED)
@@ -36,7 +38,7 @@
 	#endif
 
 	#ifndef CVD_TRACE_SOLVER_STEP_START
-		#define CVD_TRACE_SOLVER_STEP_START()
+		#define CVD_TRACE_SOLVER_STEP_START(StepName)
 	#endif
 
 	#ifndef CVD_TRACE_SOLVER_STEP_END
@@ -44,7 +46,7 @@
 	#endif
 
 	#ifndef CVD_SCOPE_TRACE_SOLVER_STEP
-		#define CVD_SCOPE_TRACE_SOLVER_STEP()
+		#define CVD_SCOPE_TRACE_SOLVER_STEP(StepName)
 	#endif
 
 	#ifndef CVD_TRACE_BINARY_DATA
@@ -54,8 +56,17 @@
 #ifndef CVD_TRACE_SOLVER_SIMULATION_SPACE
 		#define CVD_TRACE_SOLVER_SIMULATION_SPACE(InSimulationSpace)
 #endif
+
+#ifndef CVD_TRACE_PARTICLES_SOA
+	#define CVD_TRACE_PARTICLES_SOA(ParticleSoA)
+#endif
+
+#ifndef CVD_TRACE_PARTICLE_DESTROYED
+	#define CVD_TRACE_PARTICLE_DESTROYED(DestroyedParticleHandle)
+#endif
 #else
 
+#include "ChaosVDRuntimeModule.h"
 #include "Trace/Trace.h"
 #include "Trace/Trace.inl"
 
@@ -95,6 +106,7 @@ UE_TRACE_EVENT_BEGIN_EXTERN(ChaosVDLogger, ChaosVDSolverFrameStart)
 	UE_TRACE_EVENT_FIELD(int32, SolverID)
 	UE_TRACE_EVENT_FIELD(uint64, Cycle)
 	UE_TRACE_EVENT_FIELD(UE::Trace::WideString, DebugName)
+	UE_TRACE_EVENT_FIELD(bool, IsKeyFrame)
 UE_TRACE_EVENT_END()
 
 UE_TRACE_EVENT_BEGIN_EXTERN(ChaosVDLogger, ChaosVDSolverFrameEnd)
@@ -114,12 +126,25 @@ UE_TRACE_EVENT_BEGIN_EXTERN(ChaosVDLogger, ChaosVDParticle)
 	UE_TRACE_EVENT_FIELD(int8, ObjectState)
 	UE_TRACE_EVENT_FIELD(UE::Trace::WideString, DebugName)
 	UE_TRACE_EVENT_FIELD(int32, ImplicitObjectID)
+	UE_TRACE_EVENT_FIELD(uint32, ImplicitObjectHash)
+UE_TRACE_EVENT_END()
+
+UE_TRACE_EVENT_BEGIN_EXTERN(ChaosVDLogger, ChaosVDParticleCreated)
+	UE_TRACE_EVENT_FIELD(int32, SolverID)
+	UE_TRACE_EVENT_FIELD(uint64, Cycle)
+	UE_TRACE_EVENT_FIELD(int32, ParticleID)
+UE_TRACE_EVENT_END()
+
+UE_TRACE_EVENT_BEGIN_EXTERN(ChaosVDLogger, ChaosVDParticleDestroyed)
+	UE_TRACE_EVENT_FIELD(int32, SolverID)
+	UE_TRACE_EVENT_FIELD(uint64, Cycle)
+	UE_TRACE_EVENT_FIELD(int32, ParticleID)
 UE_TRACE_EVENT_END()
 
 UE_TRACE_EVENT_BEGIN_EXTERN(ChaosVDLogger, ChaosVDSolverStepStart)
 	UE_TRACE_EVENT_FIELD(int32, SolverID)
 	UE_TRACE_EVENT_FIELD(uint64, Cycle)
-	UE_TRACE_EVENT_FIELD(uint32, StepNumber)
+	UE_TRACE_EVENT_FIELD(UE::Trace::WideString, StepName)
 UE_TRACE_EVENT_END()
 
 UE_TRACE_EVENT_BEGIN_EXTERN(ChaosVDLogger, ChaosVDSolverStepEnd)
@@ -155,6 +180,10 @@ UE_TRACE_EVENT_BEGIN_EXTERN(ChaosVDLogger, ChaosVDSolverSimulationSpace)
 	CVD_DEFINE_TRACE_ROTATOR(Chaos::FReal, Rotation)
 UE_TRACE_EVENT_END()
 
+UE_TRACE_EVENT_BEGIN_EXTERN(ChaosVDLogger, ChaosVDDummyEvent)
+	UE_TRACE_EVENT_FIELD(int32, SolverID)
+UE_TRACE_EVENT_END()
+
 #ifndef CVD_TRACE_PARTICLE
 	#define CVD_TRACE_PARTICLE(ParticleHandle) \
 		FChaosVisualDebuggerTrace::TraceParticle(ParticleHandle);
@@ -162,7 +191,12 @@ UE_TRACE_EVENT_END()
 
 #ifndef CVD_TRACE_PARTICLES
 	#define CVD_TRACE_PARTICLES(ParticleHandles) \
-		FChaosVisualDebuggerTrace::TraceParticles(ParticleHandles)
+		FChaosVisualDebuggerTrace::TraceParticles(ParticleHandles);
+#endif
+
+#ifndef CVD_TRACE_PARTICLES_SOA
+	#define CVD_TRACE_PARTICLES_SOA(ParticleSoA) \
+	FChaosVisualDebuggerTrace::TraceParticlesSoA(ParticleSoA);
 #endif
 
 #ifndef CVD_TRACE_SOLVER_START_FRAME
@@ -185,8 +219,8 @@ UE_TRACE_EVENT_END()
 #endif
 
 #ifndef CVD_TRACE_SOLVER_STEP_START
-	#define CVD_TRACE_SOLVER_STEP_START() \
-		FChaosVisualDebuggerTrace::TraceSolverStepStart();
+	#define CVD_TRACE_SOLVER_STEP_START(StepName) \
+		FChaosVisualDebuggerTrace::TraceSolverStepStart(StepName);
 #endif
 
 #ifndef CVD_TRACE_SOLVER_STEP_END
@@ -195,8 +229,8 @@ UE_TRACE_EVENT_END()
 #endif
 
 #ifndef CVD_SCOPE_TRACE_SOLVER_STEP
-	#define CVD_SCOPE_TRACE_SOLVER_STEP() \
-		FChaosVDScopeSolverStep ScopeSolverStep;
+	#define CVD_SCOPE_TRACE_SOLVER_STEP(StepName) \
+		FChaosVDScopeSolverStep ScopeSolverStep(StepName);
 #endif
 
 #ifndef CVD_TRACE_BINARY_DATA
@@ -209,44 +243,181 @@ UE_TRACE_EVENT_END()
 	FChaosVisualDebuggerTrace::TraceSolverSimulationSpace(InSimulationSpace);
 #endif
 
+#ifndef CVD_TRACE_PARTICLE_DESTROYED
+	#define CVD_TRACE_PARTICLE_DESTROYED(DestroyedParticleHandle) \
+	FChaosVisualDebuggerTrace::TraceParticleDestroyed(DestroyedParticleHandle);
+#endif
+
+
 struct FChaosVDContext;
 
 namespace Chaos
 {
+	class FPBDRigidsSOAs;
 	class FImplicitObject;
 	class FPhysicsSolverBase;
 	template <typename T, int d>
 	class TGeometryParticleHandles;
 }
 
+/** Data Wrapper for used to Trace Implicit objects */
+struct CHAOS_API FChaosVDImplicitObjectDataWrapper
+{
+	uint32 Hash;
+	Chaos::TSerializablePtr<Chaos::FImplicitObject> ImplicitObject;
+
+	void Serialize(Chaos::FChaosArchive& Ar);
+};
+
 /** Class containing  all the Tracing logic to record data for the Chaos Visual Debugger tool */
 class FChaosVisualDebuggerTrace
 {
 public:
-
+	/**
+	 * Traces data from a Particle Handle. The CVD context currently pushed into will be used to tie this particle data to a specific solver frame and step
+	 * @param ParticleHandle Handle to process and Trace
+	 */
 	static CHAOS_API void TraceParticle(const Chaos::FGeometryParticleHandle* ParticleHandle);
+	
+	/**
+	 * Traces data from a Particle Handle using the provided CVD Context
+	 * @param ParticleHandle Handle to process and Trace
+	 * @param ContextData Context to be used to tied this Trace event to a specific solver frame and step
+	 */
 	static CHAOS_API void TraceParticle(Chaos::FGeometryParticleHandle* ParticleHandle, const FChaosVDContext& ContextData);
+
+	/**
+	 * Traces data from a collection of Particle Handles using. The CVD context currently pushed into will be used to tie this particle data to a specific solver frame and step.
+	 * It does not handle Full and Delta Recording automatically
+	 * @param ParticleHandles Handles collection to process and Trace
+	 */
 	static CHAOS_API void TraceParticles(const Chaos::TGeometryParticleHandles<Chaos::FReal, 3>& ParticleHandles);
+	
+	/**
+	 * Traces the destruction event for the provided particle handled so it can be reproduces in the CVD tool
+	 * @param ParticleHandle Handle that is being destroyed
+	 */
+	static CHAOS_API void TraceParticleDestroyed(const Chaos::FGeometryParticleHandle* ParticleHandle);
+
+	/**
+	 * Traces data from particles on the provided FPBDRigidsSOAs. It traces only the DirtyParticles view unless a full capture was requested
+	 * @param ParticlesSoA Particles SoA to evaluate and trace
+	 */
+	static CHAOS_API void TraceParticlesSoA(const Chaos::FPBDRigidsSOAs& ParticlesSoA);
+
+	/** Traces the provided particle view in parallel */
+	template<typename ParticleType>
+	static void TraceParticlesView(const Chaos::TParticleView<ParticleType>& ParticlesView);
+
+	/** Traces the start of a solver frame and it pushes its context data to the CVD TLS context stack */
 	static CHAOS_API void TraceSolverFrameStart(const FChaosVDContext& ContextData, const FString& InDebugName);
+	
+	/** Traces the end of a solver frame and removes its context data to the CVD TLS context stack */
 	static CHAOS_API void TraceSolverFrameEnd(const FChaosVDContext& ContextData);
-	static CHAOS_API void TraceSolverStepStart();
+
+	/** Traces the start of a solver step
+	 * @param StepName Name of the step. It will be used in the CVD Tool's UI
+	 */
+	static CHAOS_API void TraceSolverStepStart(FStringView StepName);
+	/** Traces the end of a solver step */
 	static CHAOS_API void TraceSolverStepEnd();
+
+	/** Traces the provider transform as simulation space of the solver that is currently on the CVD Context Stack
+	 * @param Transform Simulation space Transform
+	 */
 	static CHAOS_API void TraceSolverSimulationSpace(const Chaos::FRigidTransform3& Transform);
-	
+
+	/**
+	 * Traces a binary blob of data outside of any solver frame solver step scope.
+	 * @param InData Data to trace
+	 * @param TypeName Type name the data represents. It is used during Trace Analysis serialize it back (this is not automatic)
+	 */
 	static CHAOS_API void TraceBinaryData(const TArray<uint8>& InData, const FString& TypeName);
-	
-	static CHAOS_API int32 TraceImplicitObject(Chaos::TSerializablePtr<Chaos::FImplicitObject> Geometry);
 
+	/**
+	 * Serializes the implicit object contained in the wrapper and trace its it as binary data
+	 * The trace event is not tied to a particular Solver Frame/Step
+	 *  @param WrappedGeometryData Wrapper containing a ptr to the implicit and its ID
+	 */
+	static CHAOS_API void TraceImplicitObject(FChaosVDImplicitObjectDataWrapper WrappedGeometryData);
+
+	/** Returns true if the provided solver ID needs a Full Capture */
+	static bool ShouldPerformFullCapture(int32 SolverID);
+
+	/**
+	 * Gets the CVD Context data form an object that has such data. Usually Solvers
+	 * @tparam T type of the object with CVD context
+	 * @param ObjectWithContext A reference to where to get the CVD Context
+	 * @param OutCVDContext A reference to the CVD contexts in the provided object
+	 */
 	template<typename T>
 	static void GetCVDContext(T& ObjectWithContext, FChaosVDContext& OutCVDContext);
 
+	/**
+	 * Returns the debug name string of the provided object
+	 */
 	template<typename T>
-	static FString GetDebugName(T& ObjectWithContext);
+	static FString GetDebugName(T& ObjectWithDebugName);
+
+	/** Returns true if a CVD trace is running */
+	static bool IsTracing() { return bIsTracing;}
 
 private:
-	static void ResetGeometryTracerContext();
+
+	/** Binds to the static events triggered by the ChaosVD Runtime module */
+	static void RegisterEventHandlers();
+	/** Unbinds to the static events triggered by the ChaosVD Runtime module */
+	static void UnregisterEventHandlers();
+	
+	/** Resets the state of the CVD Tracer */
+	static void Reset();
+
+	static void HandleRecordingStop();
+	static void HandleRecordingStart();
+
+	/** Sets up the tracer to perform a full capture in the next solver frame */
+	static void PerformFullCapture(EChaosVDFullCaptureFlags CaptureOptions);
+
+	/** Setups the current Solver frame for a full capture if needed */
+	static void SetupForFullCaptureIfNeeded(int32 SolverID, bool& bOutFullCaptureRequested);
+
+	static FDelegateHandle RecordingStartedDelegateHandle;
+	static FDelegateHandle RecordingStoppedDelegateHandle;
+	static FDelegateHandle RecordingFullCaptureRequestedHandle;
+
+	static TSet<int32> SolverIDsForDeltaRecording;
+	static TSet<int32> RequestedFullCaptureSolverIDs;
+
+	static FThreadSafeBool bIsTracing;
+
+	static FRWLock DeltaRecordingStatesLock;
+
+	friend class FChaosEngineModule;
 };
 
+template <typename ParticleType>
+void FChaosVisualDebuggerTrace::TraceParticlesView(const Chaos::TParticleView<ParticleType>& ParticlesView)
+{
+	if (!IsTracing())
+	{
+		return;
+	}
+
+	const FChaosVDContext* CVDContextData = FChaosVDThreadContext::Get().GetCurrentContext();
+	if (!ensure(CVDContextData))
+	{
+		return;
+	}
+
+	FChaosVDContext CopyContext = *FChaosVDThreadContext::Get().GetCurrentContext();
+	
+	ParticlesView.ParallelFor([CopyContext](auto& Particle, int32 Index)
+	{
+		CVD_SCOPE_CONTEXT(CopyContext);
+		TraceParticle(Particle.Handle());
+	});
+}
+
 template <typename T>
 void FChaosVisualDebuggerTrace::GetCVDContext(T& ObjectWithContext, FChaosVDContext& OutCVDContext)
 {
@@ -265,9 +436,9 @@ FString FChaosVisualDebuggerTrace::GetDebugName(T& ObjectWithDebugName)
 
 struct FChaosVDScopeSolverStep
 {
-	FChaosVDScopeSolverStep()
+	FChaosVDScopeSolverStep(FStringView StepName)
 	{
-		CVD_TRACE_SOLVER_STEP_START();
+		CVD_TRACE_SOLVER_STEP_START(StepName);
 	}
 
 	~FChaosVDScopeSolverStep()
diff --git a/Engine/Source/Runtime/Experimental/ChaosVDData/Private/ChaosVDRecording.cpp b/Engine/Source/Runtime/Experimental/ChaosVDData/Private/ChaosVDRecording.cpp
--- a/Engine/Source/Runtime/Experimental/ChaosVDData/Private/ChaosVDRecording.cpp
+++ b/Engine/Source/Runtime/Experimental/ChaosVDData/Private/ChaosVDRecording.cpp
@@ -174,17 +174,17 @@ void FChaosVDRecording::GetAvailableSolverIDsAtGameFrameNumber(int32 FrameNumber
 	}
 }
 
-void FChaosVDRecording::AddImplicitObject(const int32 ID, const TSharedPtr<Chaos::FImplicitObject>& InImplicitObject)
+void FChaosVDRecording::AddImplicitObject(const uint32 ID, const TSharedPtr<Chaos::FImplicitObject>& InImplicitObject)
 {
-	if (ensure(!ImplicitObjects.Contains(ID)))
+	if (!ImplicitObjects.Contains(ID))
 	{
 		AddImplicitObject_Internal(ID, InImplicitObject);
 	}
 }
 
-void FChaosVDRecording::AddImplicitObject(const int32 ID, const Chaos::FImplicitObject* InImplicitObject)
+void FChaosVDRecording::AddImplicitObject(const uint32 ID, const Chaos::FImplicitObject* InImplicitObject)
 {
-	if (ensure(!ImplicitObjects.Contains(ID)))
+	if (!ImplicitObjects.Contains(ID))
 	{
 		// Only take ownership after we know we will add it to the map
 		const TSharedPtr<const Chaos::FImplicitObject> SharedImplicit(InImplicitObject);
@@ -192,7 +192,7 @@ void FChaosVDRecording::AddImplicitObject(const int32 ID, const Chaos::FImplicit
 	}
 }
 
-void FChaosVDRecording::AddImplicitObject_Internal(const int32 ID, const TSharedPtr<const Chaos::FImplicitObject>& InImplicitObject)
+void FChaosVDRecording::AddImplicitObject_Internal(uint32 ID, const TSharedPtr<const Chaos::FImplicitObject>& InImplicitObject)
 {
 	ImplicitObjects.Add(ID, InImplicitObject);
 	GeometryDataLoaded.Broadcast(InImplicitObject, ID);
diff --git a/Engine/Source/Runtime/Experimental/ChaosVDData/Public/ChaosVDRecording.h b/Engine/Source/Runtime/Experimental/ChaosVDData/Public/ChaosVDRecording.h
--- a/Engine/Source/Runtime/Experimental/ChaosVDData/Public/ChaosVDRecording.h
+++ b/Engine/Source/Runtime/Experimental/ChaosVDData/Public/ChaosVDRecording.h
@@ -3,7 +3,6 @@
 #pragma once
 
 #include "Chaos/Core.h"
-#include "Chaos/Serializable.h"
 #include "Chaos/ShapeInstanceFwd.h"
 #include "UObject/ObjectMacros.h"
 #include "Containers/UnrealString.h"
@@ -15,7 +14,7 @@ namespace Chaos
 }
 
 DECLARE_MULTICAST_DELEGATE(FChaosVDRecordingUpdated)
-DECLARE_MULTICAST_DELEGATE_TwoParams(FChaosVDGeometryDataLoaded, const TSharedPtr<const Chaos::FImplicitObject>&, const int32 GeometryID)
+DECLARE_MULTICAST_DELEGATE_TwoParams(FChaosVDGeometryDataLoaded, const TSharedPtr<const Chaos::FImplicitObject>&, const uint32 GeometryID)
 
 UENUM()
 enum class EChaosVDParticleType : uint8
@@ -71,20 +70,28 @@ struct FChaosVDParticleDebugData
 	EChaosVDParticleState ParticleState;
 
 	int32 ImplicitObjectID = INDEX_NONE;
+
+	UPROPERTY(VisibleAnywhere, Category= "Chaos Visual Debugger Data")
+	uint32 ImplicitObjectHash = 0;
 };
 
 struct FChaosVDStepData
 {
+	FString StepName;
 	TArray<FChaosVDParticleDebugData> RecordedParticles;
+	TSet<int32> ParticlesDestroyedIDs;
 };
 
+typedef TArray<FChaosVDStepData, TInlineAllocator<16>> FChaosVDStepsContainer;
+
 struct CHAOSVDDATA_API FChaosVDSolverFrameData
 {
 	FString DebugName;
-	int32 SolverID;
-	uint64 FrameCycle;
+	int32 SolverID = INDEX_NONE;
+	uint64 FrameCycle = 0;
 	Chaos::FRigidTransform3 SimulationTransform;
-	TArray<FChaosVDStepData, TInlineAllocator<16>> SolverSteps;
+	bool bIsKeyFrame = false;
+	FChaosVDStepsContainer SolverSteps;
 };
 
 struct FChaosVDGameFrameData
@@ -214,10 +221,10 @@ struct CHAOSVDDATA_API FChaosVDRecording
 	void GetAvailableSolverIDsAtGameFrameNumber(int32 FrameNumber, TArray<int32>& OutSolversID);
 
 	/** Returns a reference to the GeometryID-ImplicitObject map of this recording */
-	const TMap<int32, TSharedPtr<const Chaos::FImplicitObject>>& GetGeometryDataMap() const { return ImplicitObjects; };
+	const TMap<uint32, TSharedPtr<const Chaos::FImplicitObject>>& GetGeometryDataMap() const { return ImplicitObjects; };
 
 	/** Adds a shared Implicit Object to the recording */
-	void AddImplicitObject(const int32 ID, const TSharedPtr<Chaos::FImplicitObject>& InImplicitObject);
+	void AddImplicitObject(const uint32 ID, const TSharedPtr<Chaos::FImplicitObject>& InImplicitObject);
 
 	/** Session name of the trace session used to re-build this recording */
 	FString SessionName;
@@ -225,9 +232,9 @@ struct CHAOSVDDATA_API FChaosVDRecording
 protected:
 
 	/** Adds an Implicit Object to the recording and takes ownership of it */
-	void AddImplicitObject(const int32 ID, const Chaos::FImplicitObject* InImplicitObject);
+	void AddImplicitObject(const uint32 ID, const Chaos::FImplicitObject* InImplicitObject);
 	
-	void AddImplicitObject_Internal(const int32 ID, const TSharedPtr<const Chaos::FImplicitObject>& InImplicitObject);
+	void AddImplicitObject_Internal(const uint32 ID, const TSharedPtr<const Chaos::FImplicitObject>& InImplicitObject);
 	
 	TMap<int32, TArray<FChaosVDSolverFrameData>> RecordedFramesDataPerSolver;
 	TArray<FChaosVDGameFrameData> GameFrames;
@@ -235,7 +242,7 @@ protected:
 	FChaosVDGeometryDataLoaded GeometryDataLoaded;
 
 	/** Id to Ptr map of all shared geometry data required to visualize */
-	TMap<int32, TSharedPtr<const Chaos::FImplicitObject>> ImplicitObjects;
+	TMap<uint32, TSharedPtr<const Chaos::FImplicitObject>> ImplicitObjects;
 
 	friend class FChaosVDTraceProvider;
 };
diff --git a/Engine/Source/Runtime/Experimental/ChaosVisualDebugger/Private/ChaosVDRuntimeModule.cpp b/Engine/Source/Runtime/Experimental/ChaosVisualDebugger/Private/ChaosVDRuntimeModule.cpp
--- a/Engine/Source/Runtime/Experimental/ChaosVisualDebugger/Private/ChaosVDRuntimeModule.cpp
+++ b/Engine/Source/Runtime/Experimental/ChaosVisualDebugger/Private/ChaosVDRuntimeModule.cpp
@@ -31,8 +31,15 @@ FAutoConsoleCommand StopVDStartRecordingCommand(
 	})
 );
 
-FRecordingStateChangedDelegate FChaosVDRuntimeModule::RecordingStartedDelegate = FRecordingStateChangedDelegate();
-FRecordingStateChangedDelegate FChaosVDRuntimeModule::RecordingStopDelegate = FRecordingStateChangedDelegate();
+static FAutoConsoleVariable CVarChaosVDGTimeBetweenFullCaptures(
+	TEXT("p.Chaos.VD.TimeBetweenFullCaptures"),
+	10,
+	TEXT("Time interval in seconds after which a full capture (not only delta changes) should be recorded"));
+
+FChaosVDRecordingStateChangedDelegate FChaosVDRuntimeModule::RecordingStartedDelegate = FChaosVDRecordingStateChangedDelegate();
+FChaosVDRecordingStateChangedDelegate FChaosVDRuntimeModule::RecordingStopDelegate = FChaosVDRecordingStateChangedDelegate();
+FChaosVDCaptureRequestDelegate FChaosVDRuntimeModule::PerformFullCaptureDelegate = FChaosVDCaptureRequestDelegate();
+FRWLock FChaosVDRuntimeModule::DelegatesRWLock = FRWLock();
 
 FChaosVDRuntimeModule& FChaosVDRuntimeModule::Get()
 {
@@ -77,6 +84,14 @@ void FChaosVDRuntimeModule::GenerateRecordingFileName(FString& OutFileName)
 	} while (IFileManager::Get().FileExists(*OutFileName));
 }
 
+bool FChaosVDRuntimeModule::RequestFullCapture(float DeltaTime)
+{
+	// Full capture intervals are clamped to be no lower than 1 sec
+	FReadScopeLock ReadLock (DelegatesRWLock);
+	PerformFullCaptureDelegate.Broadcast(EChaosVDFullCaptureFlags::Particles);
+	return true;
+}
+
 void FChaosVDRuntimeModule::StartRecording(const TArray<FString>& Args)
 {
 	if (bIsRecording)
@@ -84,7 +99,20 @@ void FChaosVDRuntimeModule::StartRecording(const TArray<FString>& Args)
 		return;
 	}
 
-	OnRecordingStarted().ExecuteIfBound();
+	{
+		FReadScopeLock ReadLock(DelegatesRWLock);
+		RecordingStartedDelegate.Broadcast();
+	}
+
+	constexpr int32 MinAllowedTimeInSecondsBetweenCaptures = 1;
+	int32 ConfiguredTimeBetweenCaptures = CVarChaosVDGTimeBetweenFullCaptures->GetInt();
+
+	ensureAlwaysMsgf(ConfiguredTimeBetweenCaptures > MinAllowedTimeInSecondsBetweenCaptures,
+		TEXT("The minimum allowed time interval between full captures is [%d] seconds, but [%d] seconds were configured. Clamping to [%d] seconds"),
+		MinAllowedTimeInSecondsBetweenCaptures, ConfiguredTimeBetweenCaptures, MinAllowedTimeInSecondsBetweenCaptures);
+
+	FullCaptureRequesterHandle = FTSTicker::GetCoreTicker().AddTicker(FTickerDelegate::CreateRaw(this, &FChaosVDRuntimeModule::RequestFullCapture),
+		FMath::Clamp(ConfiguredTimeBetweenCaptures, MinAllowedTimeInSecondsBetweenCaptures, TNumericLimits<int32>::Max()));
 
 #if UE_TRACE_ENABLED
 
@@ -147,16 +175,27 @@ void FChaosVDRuntimeModule::StopRecording()
 	StopTrace();
 #endif
 
-	OnRecordingStop().ExecuteIfBound();
+	if (FullCaptureRequesterHandle.IsValid())
+	{
+		FTSTicker::GetCoreTicker().RemoveTicker(FullCaptureRequesterHandle);
+	}
+
+	{
+		FReadScopeLock ReadLock(DelegatesRWLock);
+		RecordingStopDelegate.Broadcast();
+	}
 
 	bIsRecording = false;
 }
 
 void FChaosVDRuntimeModule::HandleTraceStopRequest(FTraceAuxiliary::EConnectionType TraceType, const FString& TraceDestination)
 {
-	if (!ensure(bRequestedStop))
+	if (bIsRecording)
 	{
-		UE_LOG(LogChaosVDRuntime, Warning, TEXT("Trace Recording has been stopped unexpectedly"));
+		if (!ensure(bRequestedStop))
+		{
+			UE_LOG(LogChaosVDRuntime, Warning, TEXT("Trace Recording has been stopped unexpectedly"));
+		}
 	}
 
 	bRequestedStop = false;
diff --git a/Engine/Source/Runtime/Experimental/ChaosVisualDebugger/Public/ChaosVDRuntimeModule.h b/Engine/Source/Runtime/Experimental/ChaosVisualDebugger/Public/ChaosVDRuntimeModule.h
--- a/Engine/Source/Runtime/Experimental/ChaosVisualDebugger/Public/ChaosVDRuntimeModule.h
+++ b/Engine/Source/Runtime/Experimental/ChaosVisualDebugger/Public/ChaosVDRuntimeModule.h
@@ -5,10 +5,21 @@
 #include "Templates/SharedPointer.h"
 #include "Modules/ModuleInterface.h"
 #include "ProfilingDebugging/TraceAuxiliary.h"
+#include "Containers/Ticker.h"
+#include "Misc/ScopeRWLock.h"
 
 struct FChaosVDRecording;
 
-DECLARE_DELEGATE(FRecordingStateChangedDelegate)
+/* Option flags that controls what should be recorded when doing a full capture **/
+enum EChaosVDFullCaptureFlags : int32
+{
+	Geometry = 1 << 0,
+	Particles = 1 << 1,
+};
+ENUM_CLASS_FLAGS(EChaosVDFullCaptureFlags)
+
+DECLARE_MULTICAST_DELEGATE(FChaosVDRecordingStateChangedDelegate)
+DECLARE_MULTICAST_DELEGATE_OneParam(FChaosVDCaptureRequestDelegate, EChaosVDFullCaptureFlags)
 
 class CHAOSVDRUNTIME_API FChaosVDRuntimeModule : public IModuleInterface
 {
@@ -32,8 +43,41 @@ public:
 	/** Returns a unique ID used to be used to identify CVD (Chaos Visual Debugger) data */
 	int32 GenerateUniqueID();
 
-	static FRecordingStateChangedDelegate& OnRecordingStarted() { return RecordingStartedDelegate; };
-	static FRecordingStateChangedDelegate& OnRecordingStop() { return RecordingStopDelegate;};
+	static FDelegateHandle RegisterRecordingStartedCallback(const FChaosVDRecordingStateChangedDelegate::FDelegate& InCallback)
+	{
+		FWriteScopeLock WriteLock(DelegatesRWLock);
+		return RecordingStartedDelegate.Add(InCallback);
+	}
+
+	static FDelegateHandle RegisterRecordingStopCallback(const FChaosVDRecordingStateChangedDelegate::FDelegate& InCallback)
+	{
+		FWriteScopeLock WriteLock(DelegatesRWLock);
+		return RecordingStopDelegate.Add(InCallback);
+	}
+	
+	static FDelegateHandle RegisterFullCaptureRequestedCallback(const FChaosVDCaptureRequestDelegate::FDelegate& InCallback)
+	{
+		FWriteScopeLock WriteLock(DelegatesRWLock);
+		return PerformFullCaptureDelegate.Add(InCallback);
+	}
+
+	static bool RemoveRecordingStartedCallback(const FDelegateHandle& InDelegateToRemove)
+	{
+		FWriteScopeLock WriteLock(DelegatesRWLock);
+		return RecordingStartedDelegate.Remove(InDelegateToRemove);
+	}
+
+	static bool RemoveRecordingStopCallback(const FDelegateHandle& InDelegateToRemove)
+	{
+		FWriteScopeLock WriteLock(DelegatesRWLock);
+		return RecordingStopDelegate.Remove(InDelegateToRemove);
+	}
+	
+	static bool RemoveFullCaptureRequestedCallback(const FDelegateHandle& InDelegateToRemove)
+	{
+		FWriteScopeLock WriteLock(DelegatesRWLock);
+		return PerformFullCaptureDelegate.Remove(InDelegateToRemove);
+	}
 
 private:
 
@@ -42,6 +86,9 @@ private:
 	/** Finds a valid file name for a new file - Used to generate the file name for the Trace recording */
 	void GenerateRecordingFileName(FString& OutFileName);
 
+	/** Queues a full Capture of the simulation on the next frame */
+	bool RequestFullCapture(float DeltaTime);
+
 	/** Used to handle stop requests to the active trace session that were not done by us
 	 * That is a possible scenario because Trace is shared by other In-Editor tools
 	 */
@@ -50,8 +97,13 @@ private:
 	bool bIsRecording = false;
 	bool bRequestedStop = false;
 
-	static FRecordingStateChangedDelegate RecordingStartedDelegate;
-	static FRecordingStateChangedDelegate RecordingStopDelegate;
+	FTSTicker::FDelegateHandle FullCaptureRequesterHandle;
+
+	static FChaosVDRecordingStateChangedDelegate RecordingStartedDelegate;
+	static FChaosVDRecordingStateChangedDelegate RecordingStopDelegate;
+	static FChaosVDCaptureRequestDelegate PerformFullCaptureDelegate;
 
 	FThreadSafeCounter LastGeneratedID;
+
+	static FRWLock DelegatesRWLock;
 };
