Geometry collection : fix evaluation of strain when applying breaking velocities after a crumble cluster operation

diff --git a/Engine/Source/Runtime/Experimental/Chaos/Private/PhysicsProxy/GeometryCollectionPhysicsProxy.cpp b/Engine/Source/Runtime/Experimental/Chaos/Private/PhysicsProxy/GeometryCollectionPhysicsProxy.cpp
--- a/Engine/Source/Runtime/Experimental/Chaos/Private/PhysicsProxy/GeometryCollectionPhysicsProxy.cpp
+++ b/Engine/Source/Runtime/Experimental/Chaos/Private/PhysicsProxy/GeometryCollectionPhysicsProxy.cpp
@@ -2359,7 +2359,9 @@ static void ApplyToBreakingChildren_Internal(Chaos::FRigidClustering& Clustering
 			{
 				if (Chaos::FPBDRigidClusteredParticleHandle* ClusteredChildHandle = ChildHandle->CastToClustered())
 				{
-					if (ClusteredChildHandle->GetExternalStrain() > ClusteredChildHandle->GetInternalStrains())
+					// todo(chaos) : this does not account for the various damage models, we should eventually call an evaluate function to avoid replicating logic from the clustering code
+					// also we cannot account for collision impulses because they are set after the physics callbacks are evaluated
+					if (ClusteredChildHandle->GetExternalStrain() >= ClusteredChildHandle->GetInternalStrains())
 					{
 						Action(ClusteredChildHandle);
 					}
