Geometry collection : add cvar for material base damage threshold computation mode
- this should be temporary to test and find the most intuitive method of computing this

diff --git a/Engine/Source/Runtime/Experimental/Chaos/Private/PhysicsProxy/GeometryCollectionPhysicsProxy.cpp b/Engine/Source/Runtime/Experimental/Chaos/Private/PhysicsProxy/GeometryCollectionPhysicsProxy.cpp
--- a/Engine/Source/Runtime/Experimental/Chaos/Private/PhysicsProxy/GeometryCollectionPhysicsProxy.cpp
+++ b/Engine/Source/Runtime/Experimental/Chaos/Private/PhysicsProxy/GeometryCollectionPhysicsProxy.cpp
@@ -100,6 +100,12 @@ FAutoConsoleVariableRef CVarGeometryCollectionAlwaysGenerateConnectionGraph(
 	bGeometryCollectionAlwaysGenerateConnectionGraph,
 	TEXT("When enabled, always  generate the cluster's connection graph instead of using the rest collection stored one - Note: this should only be used for troubleshooting.[def: false]"));
 
+int32 GeometryCollectionAreaBasedDamageThresholdMode = 0;
+FAutoConsoleVariableRef CVarGeometryCollectionAreaBasedDamageThresholdMode(
+	TEXT("p.GeometryCollection.AreaBasedDamageThresholdMode"),
+	GeometryCollectionAreaBasedDamageThresholdMode,
+	TEXT("Area based damage threshold computation mode (0: sum of areas | 1: max of areas | 2: min of areas | 3: average of areas) [def: 0]"));
+
 DEFINE_LOG_CATEGORY_STATIC(UGCC_LOG, Error, All);
 
 static const FSharedSimulationSizeSpecificData& GetSizeSpecificData(const TArray<FSharedSimulationSizeSpecificData>& SizeSpecificData, const FGeometryCollection& RestCollection, const int32 TransformIndex, const FBox& BoundingBox);
@@ -1419,19 +1425,41 @@ float FGeometryCollectionPhysicsProxy::ComputeMaterialBasedDamageThreshold_Inter
 float FGeometryCollectionPhysicsProxy::ComputeMaterialBasedDamageThreshold_Internal(Chaos::FPBDRigidClusteredParticleHandle& ClusteredParticle) const
 {
 	float DamageThreshold = TNumericLimits<float>::Max();
-	float TotalConnectionArea = 0;
-	for (Chaos::FConnectivityEdge& Connection : ClusteredParticle.ConnectivityEdges())
+	float ComputedConnectionArea = 0;
+
+	// store in a local variable outisde of the loop , to make sure the compiler knows we are knot going to change it while in the loop
+	const int32 ComputationMode = GeometryCollectionAreaBasedDamageThresholdMode;
+
+	const Chaos::FConnectivityEdgeArray& ConnectivityEdges = ClusteredParticle.ConnectivityEdges();
+	for (const Chaos::FConnectivityEdge& Connection : ConnectivityEdges)
 	{
-		// strain actually store the surface area 
-		TotalConnectionArea += Connection.GetArea();
+		switch (ComputationMode)
+		{
+		case 0: // sum of areas
+		case 3: // average of areas ( we'll divide by the number of connection after the loop )
+			ComputedConnectionArea += Connection.GetArea();
+			break;
+		case 1: // max of areas
+			ComputedConnectionArea = FMath::Max(ComputedConnectionArea, Connection.GetArea());
+			break;
+		case 2: // min of areas
+			ComputedConnectionArea = FMath::Min(ComputedConnectionArea, Connection.GetArea());
+			break;
+		}
 	}
-	if (TotalConnectionArea > 0)
+	// if average we need to didvide by count
+	if (ComputationMode == 3 && ConnectivityEdges.Num() > 0)
+	{
+		ComputedConnectionArea /= (float)ConnectivityEdges.Num();
+	}
+
+	if (ComputedConnectionArea > 0)
 	{
 		// This model compute the total surface are from the connections and compute the maximum force 
 		if (const Chaos::FChaosPhysicsMaterial* ChaosMaterial = Parameters.PhysicalMaterialHandle.Get())
 		{
 			// force unit is Kg.cm/s2 here
-			const float ForceThreshold = ChaosMaterial->Strength.TensileStrength * TotalConnectionArea;
+			const float ForceThreshold = ChaosMaterial->Strength.TensileStrength * ComputedConnectionArea;
 			DamageThreshold = ForceThreshold;
 		}
 	}
