IoStoreOnDemand: Add more http and cache stats for the backend.
Implement Csv Profiler stats.

diff --git a/Engine/Source/Runtime/Experimental/IoStoreOnDemand/Private/OnDemandIoDispatcherBackend.cpp b/Engine/Source/Runtime/Experimental/IoStoreOnDemand/Private/OnDemandIoDispatcherBackend.cpp
--- a/Engine/Source/Runtime/Experimental/IoStoreOnDemand/Private/OnDemandIoDispatcherBackend.cpp
+++ b/Engine/Source/Runtime/Experimental/IoStoreOnDemand/Private/OnDemandIoDispatcherBackend.cpp
@@ -26,6 +26,7 @@
 #include "Misc/ScopeRWLock.h"
 #include "Modules/ModuleManager.h"
 #include "ProfilingDebugging/CountersTrace.h"
+#include "ProfilingDebugging/CsvProfiler.h"
 #include "Serialization/CompactBinary.h"
 #include "Serialization/CompactBinarySerialization.h"
 #include "Serialization/JsonReader.h"
@@ -55,16 +56,235 @@ namespace UE::IO::Private
 using namespace UE::Tasks;
 
 ///////////////////////////////////////////////////////////////////////////////
-TRACE_DECLARE_INT_COUNTER(TotalIoRequestCount, TEXT("OnDemandIoBackend/IoRequestCount"));
-TRACE_DECLARE_INT_COUNTER(TotalIoRequestFailCount, TEXT("OnDemandIoBackend/IoRequestFailCount"));
-TRACE_DECLARE_INT_COUNTER(TotalChunkRequestCount, TEXT("OnDemandIoBackend/ChunkRequestCount"));
-TRACE_DECLARE_INT_COUNTER(CacheHitCount, TEXT("OnDemandIoBackend/Cache/HitCount"));
-TRACE_DECLARE_INT_COUNTER(CachePutCount, TEXT("OnDemandIoBackend/Cache/PutCount"));
-TRACE_DECLARE_INT_COUNTER(CachePutRejectCount, TEXT("OnDemandIoBackend/Cache/RejectCount"));
-TRACE_DECLARE_INT_COUNTER(PendingHttpRequestCount, TEXT("OnDemandIoBackend/HTTP/PendingCount"));
-TRACE_DECLARE_INT_COUNTER(InflightHttpRequestCount, TEXT("OnDemandIoBackend/HTTP/InflightCount"));
-//TRACE_DECLARE_FLOAT_COUNTER(TotalRequestSeconds, TEXT("OnDemandIoBackend/Duration/TotalSeconds"));
+// TRACE STATS
+// iorequest stats
+TRACE_DECLARE_INT_COUNTER(IoRequestsMade, TEXT("OnDemandIoBackend/IoRequestsMade"));
+TRACE_DECLARE_INT_COUNTER(IoRequestsCompleted, TEXT("OnDemandIoBackend/IoRequestsCompleted"));
+TRACE_DECLARE_MEMORY_COUNTER(IoRequestsCompletedSize, TEXT("OnDemandIoBackend/Size/IoRequestsCompletedSize"));
+TRACE_DECLARE_INT_COUNTER(IoRequestsCancelled, TEXT("OnDemandIoBackend/IoRequestsCancelled"));
+TRACE_DECLARE_INT_COUNTER(IoRequestsFailed, TEXT("OnDemandIoBackend/IoRequestsFailed"));
+// chunkrequest stats
+TRACE_DECLARE_INT_COUNTER(ReadRequestsCreated, TEXT("OnDemandIoBackend/ReadRequestsCreated"));
+TRACE_DECLARE_INT_COUNTER(ReadRequestsRemoved, TEXT("OnDemandIoBackend/ReadRequestsRemoved"));
+// cache stats
+TRACE_DECLARE_INT_COUNTER(CacheHits, TEXT("OnDemandIoBackend/CacheHits"));
+TRACE_DECLARE_MEMORY_COUNTER(CacheHitsSize, TEXT("OnDemandIoBackend/Size/CacheHitsSize"));
+TRACE_DECLARE_INT_COUNTER(CachePuts, TEXT("OnDemandIoBackend/CachePuts"));
+TRACE_DECLARE_MEMORY_COUNTER(CachePutsSize, TEXT("OnDemandIoBackend/Size/CachePutsSize"));
+TRACE_DECLARE_INT_COUNTER(CacheRejects, TEXT("OnDemandIoBackend/CacheRejects"));
+TRACE_DECLARE_MEMORY_COUNTER(CacheRejectsSize, TEXT("OnDemandIoBackend/Size/CacheRejectsSize"));
+// http stats
+TRACE_DECLARE_INT_COUNTER(HttpRequestsCompleted, TEXT("OnDemandIoBackend/HttpRequestsCompleted"));
+TRACE_DECLARE_INT_COUNTER(HttpRequestsFailed, TEXT("OnDemandIoBackend/HttpRequestsFailed"));
+TRACE_DECLARE_INT_COUNTER(HttpRequestsPending, TEXT("OnDemandIoBackend/HttpRequestsPending"));
+TRACE_DECLARE_INT_COUNTER(HttpRequestsInflight, TEXT("OnDemandIoBackend/HttpRequestsInflight"));
+TRACE_DECLARE_MEMORY_COUNTER(HttpRequestsCompletedSize, TEXT("OnDemandIoBackend/Size/HttpRequestsCompletedSize"));
+///////////////////////////////////////////////////////////////////////////////
+// CSV STATS
+CSV_DEFINE_CATEGORY(IoStoreOnDemandBackend, true);
+// iorequest stats
+CSV_DEFINE_STAT(IoStoreOnDemandBackend, FrameIoRequestsMade);
+CSV_DEFINE_STAT(IoStoreOnDemandBackend, FrameIoRequestsCompleted);
+CSV_DEFINE_STAT(IoStoreOnDemandBackend, FrameIoRequestsCompletedSize);
+CSV_DEFINE_STAT(IoStoreOnDemandBackend, FrameIoRequestsCancelled);
+CSV_DEFINE_STAT(IoStoreOnDemandBackend, FrameIoRequestsFailed);
+// chunkrequest stats
+CSV_DEFINE_STAT(IoStoreOnDemandBackend, FrameReadRequestsCreated);
+CSV_DEFINE_STAT(IoStoreOnDemandBackend, FrameReadRequestsRemoved);
+// cache stats
+CSV_DEFINE_STAT(IoStoreOnDemandBackend, FrameCacheHits);
+CSV_DEFINE_STAT(IoStoreOnDemandBackend, FrameCacheHitsSize);
+CSV_DEFINE_STAT(IoStoreOnDemandBackend, FrameCachePuts);
+CSV_DEFINE_STAT(IoStoreOnDemandBackend, FrameCachePutsSize);
+CSV_DEFINE_STAT(IoStoreOnDemandBackend, FrameCacheRejects);
+CSV_DEFINE_STAT(IoStoreOnDemandBackend, FrameCacheRejectsSize);
+// http stats
+CSV_DEFINE_STAT(IoStoreOnDemandBackend, FrameHttpRequestsCompleted);
+CSV_DEFINE_STAT(IoStoreOnDemandBackend, FrameHttpRequestsFailed);
+CSV_DEFINE_STAT(IoStoreOnDemandBackend, FrameHttpRequestsPending);
+CSV_DEFINE_STAT(IoStoreOnDemandBackend, FrameHttpRequestsInflight);
+CSV_DEFINE_STAT(IoStoreOnDemandBackend, FrameHttpRequestsCompletedSize);
+
+#if (COUNTERSTRACE_ENABLED || CSV_PROFILER)
+class FOnDemandIoBackendStats
+{
+public:
+	void OnIoRequestEnqueue(FIoRequestImpl* Request);
+	void OnIoRequestComplete(FIoRequestImpl* Request);
+	void OnIoRequestCancel(FIoRequestImpl* Request);
+	void OnIoRequestFail(FIoRequestImpl* Request);
+	void OnChunkRequestCreate();
+	void OnChunkRequestRelease();
+	void OnCacheHit(uint64 InSize);
+	void OnCachePut(uint64 InSize);
+	void OnCacheReject(uint64 InSize);
+	void OnHttpRequestEnqueue();
+	void OnHttpRequestDequeue();
+	void OnHttpRequestComplete(uint64 InSize);
+	void OnHttpRequestFail();
+
+private:
+	int32 IoRequestsMade = 0;
+	std::atomic<int32> IoRequestsCompleted = 0;
+	int32 IoRequestsCancelled = 0;
+	std::atomic<int32> IoRequestsFailed = 0;
+	int32 ReadRequestsCreated = 0;
+	std::atomic<int32> ReadRequestsRemoved = 0;
+	std::atomic<int32> CacheHits = 0;
+	std::atomic<int32> CachePuts = 0;
+	std::atomic<int32> CacheRejects = 0;
+	int32 HttpRequestsCompleted = 0;
+	int32 HttpRequestsFailed = 0;
+	std::atomic<int32> HttpRequestsPending = 0;
+	int32 HttpRequestsInflight = 0;
+
+	std::atomic<uint64> IoRequestsCompletedSize = 0;
+	std::atomic<uint64> CacheHitsSize = 0;
+	std::atomic<uint64> CachePutsSize = 0;
+	std::atomic<uint64> CacheRejectsSize = 0;
+	uint64 HttpRequestsCompletedSize = 0;
+
+	float BytesToApproxMB(uint64 Bytes)
+	{
+		return float(double(Bytes) / 1024.0 / 1024.0);
+	}
+
+	float BytesToApproxKB(uint64 Bytes)
+	{
+		return float(double(Bytes) / 1024.0);
+	}
+
+};
+void FOnDemandIoBackendStats::OnIoRequestEnqueue(FIoRequestImpl* Request)
+{
+	++IoRequestsMade;
+	TRACE_COUNTER_SET(IoRequestsMade, IoRequestsMade);
+	CSV_CUSTOM_STAT_DEFINED(FrameIoRequestsMade, IoRequestsMade, ECsvCustomStatOp::Set);
+}
+void FOnDemandIoBackendStats::OnIoRequestComplete(FIoRequestImpl* Request)
+{
+	int32 Count = IoRequestsCompleted.fetch_add(1, std::memory_order_relaxed);
+	TRACE_COUNTER_SET(IoRequestsCompleted, Count);
+	CSV_CUSTOM_STAT_DEFINED(FrameIoRequestsCompleted, Count, ECsvCustomStatOp::Set);
+
+	uint64 Size = IoRequestsCompletedSize.fetch_add(Request->GetBuffer().DataSize());
+	TRACE_COUNTER_SET(IoRequestsCompletedSize, Size);
+	CSV_CUSTOM_STAT_DEFINED(FrameIoRequestsCompletedSize, BytesToApproxKB(Size), ECsvCustomStatOp::Set);
+}
+void FOnDemandIoBackendStats::OnIoRequestCancel(FIoRequestImpl* Request)
+{
+	++IoRequestsCancelled;
+	TRACE_COUNTER_SET(IoRequestsCancelled, IoRequestsCancelled);
+	CSV_CUSTOM_STAT_DEFINED(FrameIoRequestsCancelled, IoRequestsCancelled, ECsvCustomStatOp::Set);
+}
+void FOnDemandIoBackendStats::OnIoRequestFail(FIoRequestImpl* Request)
+{
+	int32 Count = IoRequestsFailed.fetch_add(1, std::memory_order_relaxed);
+	TRACE_COUNTER_SET(IoRequestsFailed, Count);
+	CSV_CUSTOM_STAT_DEFINED(FrameIoRequestsFailed, Count, ECsvCustomStatOp::Set);
+}
+void FOnDemandIoBackendStats::OnChunkRequestCreate()
+{
+	++ReadRequestsCreated;
+	TRACE_COUNTER_SET(ReadRequestsCreated, ReadRequestsCreated);
+	CSV_CUSTOM_STAT_DEFINED(FrameReadRequestsCreated, ReadRequestsCreated, ECsvCustomStatOp::Set);
+}
+void FOnDemandIoBackendStats::OnChunkRequestRelease()
+{
+	int32 Count = ReadRequestsRemoved.fetch_add(1, std::memory_order_relaxed);
+	TRACE_COUNTER_SET(ReadRequestsRemoved, Count);
+	CSV_CUSTOM_STAT_DEFINED(FrameReadRequestsRemoved, Count, ECsvCustomStatOp::Set);
+}
+void FOnDemandIoBackendStats::OnCacheHit(uint64 InSize)
+{
+	int32 Count = CacheHits.fetch_add(1, std::memory_order_relaxed);
+	TRACE_COUNTER_SET(CacheHits, Count);
+	CSV_CUSTOM_STAT_DEFINED(FrameCacheHits, Count, ECsvCustomStatOp::Set);
 
+	uint64 Size = CacheHitsSize.fetch_add(InSize, std::memory_order_relaxed);
+	TRACE_COUNTER_SET(CacheHitsSize, Size);
+	CSV_CUSTOM_STAT_DEFINED(FrameCacheHitsSize, BytesToApproxKB(Size), ECsvCustomStatOp::Set);
+}
+void FOnDemandIoBackendStats::OnCachePut(uint64 InSize)
+{
+	int32 Count = CachePuts.fetch_add(1, std::memory_order_relaxed);
+	TRACE_COUNTER_SET(CachePuts, Count);
+	CSV_CUSTOM_STAT_DEFINED(FrameCachePuts, Count, ECsvCustomStatOp::Set);
+
+	uint64 Size = CachePutsSize.fetch_add(InSize, std::memory_order_relaxed);
+	TRACE_COUNTER_SET(CachePutsSize, Size);
+	CSV_CUSTOM_STAT_DEFINED(FrameCachePutsSize, BytesToApproxKB(Size), ECsvCustomStatOp::Set);
+}
+void FOnDemandIoBackendStats::OnCacheReject(uint64 InSize)
+{
+	int32 Count = CacheRejects.fetch_add(1, std::memory_order_relaxed);
+	TRACE_COUNTER_SET(CacheRejects, Count);
+	CSV_CUSTOM_STAT_DEFINED(FrameCacheRejects, Count, ECsvCustomStatOp::Set);
+
+	uint64 Size = CacheRejectsSize.fetch_add(InSize, std::memory_order_relaxed);
+	TRACE_COUNTER_SET(CacheRejectsSize, Size);
+	CSV_CUSTOM_STAT_DEFINED(FrameCacheRejectsSize, BytesToApproxKB(Size), ECsvCustomStatOp::Set);
+}
+void FOnDemandIoBackendStats::OnHttpRequestEnqueue()
+{
+	int32 Pending = HttpRequestsPending.fetch_add(1, std::memory_order_relaxed);
+	TRACE_COUNTER_SET(HttpRequestsPending, Pending);
+	CSV_CUSTOM_STAT_DEFINED(FrameHttpRequestsPending, Pending, ECsvCustomStatOp::Set);
+}
+void FOnDemandIoBackendStats::OnHttpRequestDequeue()
+{
+	int32 Pending = HttpRequestsPending.fetch_sub(1, std::memory_order_relaxed);
+	TRACE_COUNTER_SET(HttpRequestsPending, Pending);
+	CSV_CUSTOM_STAT_DEFINED(FrameHttpRequestsPending, Pending, ECsvCustomStatOp::Set);
+
+	int32 Inflight = ++HttpRequestsInflight;
+	TRACE_COUNTER_SET(HttpRequestsInflight, Inflight);
+	CSV_CUSTOM_STAT_DEFINED(FrameHttpRequestsPending, Inflight, ECsvCustomStatOp::Set);
+}
+void FOnDemandIoBackendStats::OnHttpRequestComplete(uint64 InSize)
+{
+	int32 Inflight = --HttpRequestsInflight;
+	TRACE_COUNTER_SET(HttpRequestsInflight, Inflight);
+	CSV_CUSTOM_STAT_DEFINED(FrameHttpRequestsPending, Inflight, ECsvCustomStatOp::Set);
+
+	++HttpRequestsCompleted;
+	TRACE_COUNTER_SET(HttpRequestsCompleted, HttpRequestsCompleted);
+	CSV_CUSTOM_STAT_DEFINED(FrameHttpRequestsCompleted, HttpRequestsCompleted, ECsvCustomStatOp::Set);
+
+	HttpRequestsCompletedSize += InSize;
+	TRACE_COUNTER_SET(HttpRequestsCompletedSize, HttpRequestsCompletedSize);
+	CSV_CUSTOM_STAT_DEFINED(FrameHttpRequestsCompletedSize, BytesToApproxKB(HttpRequestsCompletedSize), ECsvCustomStatOp::Set);
+}
+void FOnDemandIoBackendStats::OnHttpRequestFail()
+{
+	int32 Inflight = --HttpRequestsInflight;
+	TRACE_COUNTER_SET(HttpRequestsInflight, Inflight);
+	CSV_CUSTOM_STAT_DEFINED(FrameHttpRequestsPending, Inflight, ECsvCustomStatOp::Set);
+
+	++HttpRequestsFailed;
+	TRACE_COUNTER_SET(HttpRequestsFailed, HttpRequestsFailed);
+	CSV_CUSTOM_STAT_DEFINED(FrameHttpRequestsFailed, HttpRequestsFailed, ECsvCustomStatOp::Set);
+
+}
+#else
+class FOnDemandIoBackendStats
+{
+public:
+	void OnIoRequestEnqueue(FIoRequestImpl* Request) {}
+	void OnIoRequestComplete(FIoRequestImpl* Request) {}
+	void OnIoRequestCancel(FIoRequestImpl* Request) {}
+	void OnIoRequestFail(FIoRequestImpl* Request) {}
+	void OnChunkRequestCreate() {}
+	void OnChunkRequestRelease() {}
+	void OnCacheHit(uint64 InSize) {}
+	void OnCachePut(uint64 InSize) {}
+	void OnCacheReject(uint64 InSize) {}
+	void OnHttpRequestEnqueue() {}
+	void OnHttpRequestDequeue() {}
+	void OnHttpRequestComplete(uint64 InSize) {}
+	void OnHttpRequestFail() {}
+};
+#endif
 ///////////////////////////////////////////////////////////////////////////////
 FIoHash GetChunkKey(const FIoHash& ChunkHash, const FIoOffsetAndLength& Range)
 {
@@ -909,6 +1129,7 @@ private:
 	FIoRequestQueue CompletedRequests;
 	FChunkRequestQueue HttpRequests;
 	TUniquePtr<FHttpClient> HttpClient;
+	FOnDemandIoBackendStats Stats;
 	FRWLock Lock;
 	std::atomic_bool bStopRequested{false};
 };
@@ -984,12 +1205,12 @@ void FOnDemandIoBackend::CompleteRequest(FChunkRequest* ChunkRequest)
 		
 		if (!bDecoded)
 		{
-			TRACE_COUNTER_INCREMENT(TotalIoRequestFailCount);
+			Stats.OnIoRequestFail(Request);
 			bCanCache = false;
 			Request->SetFailed();
 		}
 
-		TRACE_COUNTER_INCREMENT(TotalIoRequestCount);
+		Stats.OnIoRequestComplete(Request);
 		CompletedRequests.Enqueue(Request);
 		BackendContext->WakeUpDispatcherThreadDelegate.Execute();
 	}
@@ -998,17 +1219,20 @@ void FOnDemandIoBackend::CompleteRequest(FChunkRequest* ChunkRequest)
 	{
 		if (FIoStatus Status = Cache->Put(ChunkRequest->Params.ChunkKey, Chunk); Status.IsOk())
 		{
-			TRACE_COUNTER_INCREMENT(CachePutCount);
+			Stats.OnCachePut(Chunk.GetSize());
 		}
 		else
 		{
-			TRACE_COUNTER_INCREMENT(CachePutRejectCount);
+			Stats.OnCacheReject(Chunk.GetSize());
 		}
 	}
+	else if (ChunkRequest->bCached)
+	{
+		Stats.OnCacheHit(Chunk.GetSize());
+	}
 
+	Stats.OnChunkRequestRelease();
 	ChunkRequests.Release(ChunkRequest);
-
-	TRACE_COUNTER_INCREMENT(TotalChunkRequestCount);
 }
 
 bool FOnDemandIoBackend::Resolve(FIoRequestImpl* Request)
@@ -1021,6 +1245,7 @@ bool FOnDemandIoBackend::Resolve(FIoRequestImpl* Request)
 		return false;
 	}
 
+	Stats.OnIoRequestEnqueue(Request);
 	FChunkRequestParams RequestParams = FChunkRequestParams::Create(Request, ChunkInfo);
 	FChunkRequest* ChunkRequest = ChunkRequests.Create(Request, RequestParams);
 
@@ -1029,6 +1254,8 @@ bool FOnDemandIoBackend::Resolve(FIoRequestImpl* Request)
 		// The chunk for the request is already inflight 
 		return true;
 	}
+
+	Stats.OnChunkRequestCreate();
 		
 	auto CompleteOrEnqueueHttpRequest = [this, ChunkRequest]()
 	{
@@ -1036,7 +1263,6 @@ bool FOnDemandIoBackend::Resolve(FIoRequestImpl* Request)
 		{
 			if (TIoStatusOr<FIoBuffer> Status = ChunkRequest->CacheTask.GetResult(); Status.IsOk())
 			{
-				TRACE_COUNTER_INCREMENT(CacheHitCount);
 				ChunkRequest->Chunk = Status.ConsumeValueOrDie();
 				ChunkRequest->bCached = true;
 				return CompleteRequest(ChunkRequest);
@@ -1048,7 +1274,7 @@ bool FOnDemandIoBackend::Resolve(FIoRequestImpl* Request)
 			return CompleteRequest(ChunkRequest);
 		}
 
-		TRACE_COUNTER_INCREMENT(PendingHttpRequestCount);
+		Stats.OnHttpRequestEnqueue();
 		HttpRequests.Enqueue(ChunkRequest);
 		TickBackendEvent->Trigger();
 	};
@@ -1284,8 +1510,7 @@ uint32 FOnDemandIoBackend::Run()
 				NextChunkRequest = ChunkRequest->NextRequest;
 				ChunkRequest->NextRequest = nullptr;
 
-				TRACE_COUNTER_DECREMENT(PendingHttpRequestCount);
-				TRACE_COUNTER_INCREMENT(InflightHttpRequestCount);
+				Stats.OnHttpRequestDequeue();
 
 				TAnsiStringBuilder<256> Url;
 				ChunkRequest->Params.GetUrl(Url);
@@ -1294,13 +1519,14 @@ uint32 FOnDemandIoBackend::Run()
 				HttpClient->Get(Url.ToView(), ChunkRequest->Params.ChunkRange,
 					[this, ChunkRequest, &NumConcurrentRequests](TIoStatusOr<FIoBuffer> Status)
 					{
-						TRACE_COUNTER_DECREMENT(InflightHttpRequestCount);
 						NumConcurrentRequests--;
 
 						if (Status.Status().GetErrorCode() == EIoErrorCode::ReadError)
 						{
 							if (++ChunkRequest->HttpRetryCount <= GIoDispatcherMaxHttpRetryCount)
 							{
+								Stats.OnHttpRequestFail();
+								Stats.OnHttpRequestEnqueue();
 								return HttpRequests.Enqueue(ChunkRequest); 
 							}
 						}
@@ -1308,6 +1534,11 @@ uint32 FOnDemandIoBackend::Run()
 						if (Status.IsOk())
 						{
 							ChunkRequest->Chunk = Status.ConsumeValueOrDie();
+							Stats.OnHttpRequestComplete(ChunkRequest->Chunk.DataSize());
+						}
+						else
+						{
+							Stats.OnHttpRequestFail();
 						}
 
 						Launch(UE_SOURCE_LOCATION, [this, ChunkRequest]()
