Move IrisProtocol logs to verbose

#rnx

diff --git a/Engine/Source/Runtime/Experimental/Iris/Core/Private/Iris/ReplicationSystem/ReplicationProtocolManager.cpp b/Engine/Source/Runtime/Experimental/Iris/Core/Private/Iris/ReplicationSystem/ReplicationProtocolManager.cpp
--- a/Engine/Source/Runtime/Experimental/Iris/Core/Private/Iris/ReplicationSystem/ReplicationProtocolManager.cpp
+++ b/Engine/Source/Runtime/Experimental/Iris/Core/Private/Iris/ReplicationSystem/ReplicationProtocolManager.cpp
@@ -367,7 +367,7 @@ const FReplicationProtocol* FReplicationProtocolManager::CreateReplicationProtoc
 	RegisteredProtocols.AddUnique(ProtocolId, Info);
 	ProtocolToInfoMap.Add(Protocol, Info);
 
-	UE_LOG_PROTOCOLMANAGER(Log, TEXT("FReplicationProtocolManager::CreateReplicationProtocol Created new protocol %s with ProtocolId:0x%" UINT64_x_FMT), ToCStr(Protocol->DebugName), ProtocolId);
+	UE_LOG_PROTOCOLMANAGER(Verbose, TEXT("FReplicationProtocolManager::CreateReplicationProtocol Created new protocol %s with ProtocolId:0x%" UINT64_x_FMT), ToCStr(Protocol->DebugName), ProtocolId);
 #if UE_NET_ENABLE_PROTOCOLMANAGER_LOG
 	if (bIrisLogReplicationProtocols)
 	{
@@ -397,7 +397,7 @@ void FReplicationProtocolManager::InternalDestroyReplicationProtocol(const FRepl
 {
 	if (Protocol)
 	{
-		UE_LOG_PROTOCOLMANAGER(Log, TEXT("FReplicationProtocolManager::InternalDestroyReplicationProtocol Destroyed protocol %s with ProtocolId:0x%" UINT64_x_FMT), ToCStr(Protocol->DebugName), Protocol->ProtocolIdentifier);
+		UE_LOG_PROTOCOLMANAGER(Verbose, TEXT("FReplicationProtocolManager::InternalDestroyReplicationProtocol Destroyed protocol %s with ProtocolId:0x%" UINT64_x_FMT), ToCStr(Protocol->DebugName), Protocol->ProtocolIdentifier);
 
 		// Remove tracked descriptors
 		for (const FReplicationStateDescriptor* Descriptor : MakeArrayView(Protocol->ReplicationStateDescriptors, Protocol->ReplicationStateCount))
