Iris
* Bridge poll functor is now a TFunction

#rnx

diff --git a/Engine/Source/Runtime/Experimental/Iris/Core/Private/Iris/ReplicationSystem/Polling/ObjectPoller.cpp b/Engine/Source/Runtime/Experimental/Iris/Core/Private/Iris/ReplicationSystem/Polling/ObjectPoller.cpp
--- a/Engine/Source/Runtime/Experimental/Iris/Core/Private/Iris/ReplicationSystem/Polling/ObjectPoller.cpp
+++ b/Engine/Source/Runtime/Experimental/Iris/Core/Private/Iris/ReplicationSystem/Polling/ObjectPoller.cpp
@@ -88,7 +88,7 @@ void FObjectPoller::ForcePollObject(FInternalNetRefIndex ObjectIndex)
 	if (ObjectReplicationBridge->PreUpdateInstanceFunction && EnumHasAnyFlags(ObjectData.InstanceProtocol->InstanceTraits, EReplicationInstanceProtocolTraits::NeedsPreSendUpdate))
 	{
 		IRIS_PROFILER_SCOPE_VERBOSE(PreReplicationUpdate);
-		(*ObjectReplicationBridge->PreUpdateInstanceFunction)(ObjectData.RefHandle, ReplicatedInstances[ObjectIndex], ObjectReplicationBridge);
+		ObjectReplicationBridge->PreUpdateInstanceFunction(ObjectData.RefHandle, ReplicatedInstances[ObjectIndex], ObjectReplicationBridge);
 		++PollStats.PreUpdatedObjectCount;
 	}
 
@@ -136,7 +136,7 @@ void FObjectPoller::PushModelPollObject(FInternalNetRefIndex ObjectIndex)
 	{
 		IRIS_PROFILER_SCOPE_VERBOSE(PreReplicationUpdate);
 
-		(*ObjectReplicationBridge->PreUpdateInstanceFunction)(ObjectData.RefHandle, ReplicatedInstances[ObjectIndex], ObjectReplicationBridge);
+		ObjectReplicationBridge->PreUpdateInstanceFunction(ObjectData.RefHandle, ReplicatedInstances[ObjectIndex], ObjectReplicationBridge);
 		++PollStats.PreUpdatedObjectCount;
 
 		// Pre update may dirty push based properties. Detect it.
diff --git a/Engine/Source/Runtime/Experimental/Iris/Core/Public/Iris/ReplicationSystem/ObjectReplicationBridge.h b/Engine/Source/Runtime/Experimental/Iris/Core/Public/Iris/ReplicationSystem/ObjectReplicationBridge.h
--- a/Engine/Source/Runtime/Experimental/Iris/Core/Public/Iris/ReplicationSystem/ObjectReplicationBridge.h
+++ b/Engine/Source/Runtime/Experimental/Iris/Core/Public/Iris/ReplicationSystem/ObjectReplicationBridge.h
@@ -46,6 +46,7 @@ public:
 
 	struct FCreateNetRefHandleParams
 	{
+		//$IRIS TODO: These need documentation
 		bool bCanReceive = false;
 		bool bNeedsPreUpdate = false;
 		bool bNeedsWorldLocationUpdate = false;
@@ -193,7 +194,7 @@ protected:
 	IRISCORE_API UObject* GetObjectFromReferenceHandle(FNetRefHandle RefHandle) const;
 
 	/** Helper method that calls provided PreUpdateFunction and polls state data for all replicated instances with the NeedsPoll trait. */
-	using FInstancePreUpdateFunction = void(*)(FNetRefHandle, UObject*, const UReplicationBridge*);
+	using FInstancePreUpdateFunction = TFunction<void(FNetRefHandle, UObject*, const UReplicationBridge*)>;
 
 	/** Set the function that we should call before copying state data. */
 	IRISCORE_API void SetInstancePreUpdateFunction(FInstancePreUpdateFunction InPreUpdateFunction);
