Iris Fixed issue with deferred tearoff not immediately removing subobjects from scope

diff --git a/Engine/Source/Runtime/Experimental/Iris/Core/Private/Iris/ReplicationSystem/ReplicationBridge.cpp b/Engine/Source/Runtime/Experimental/Iris/Core/Private/Iris/ReplicationSystem/ReplicationBridge.cpp
--- a/Engine/Source/Runtime/Experimental/Iris/Core/Private/Iris/ReplicationSystem/ReplicationBridge.cpp
+++ b/Engine/Source/Runtime/Experimental/Iris/Core/Private/Iris/ReplicationSystem/ReplicationBridge.cpp
@@ -611,10 +611,13 @@ void UReplicationBridge::TearOffHandlesPendingTearOff()
 
 void UReplicationBridge::UpdateHandlesPendingTearOff()
 {
+	using namespace UE::Net;
+	using namespace UE::Net::Private;
+
 	TArray<FTearOffInfo, TInlineAllocator<32>> ObjectsStillPendingTearOff;
 	for (FTearOffInfo Info : MakeArrayView(HandlesPendingTearOff))
 	{
-		if (uint32 ObjectInternalIndex = NetRefHandleManager->GetInternalIndex(Info.Handle))
+		if (FInternalNetRefIndex ObjectInternalIndex = NetRefHandleManager->GetInternalIndex(Info.Handle))
 		{
 			// Immediate tear-off or object that no longer are referenced by any connections are destroyed
 			if (NetRefHandleManager->GetNetObjectRefCount(ObjectInternalIndex) == 0U || Info.bIsImmediate)
@@ -626,8 +629,12 @@ void UReplicationBridge::UpdateHandlesPendingTearOff()
 				// If the object is still in scope remove it from scope as objects being torn-off should not be added to new connections
 				if (NetRefHandleManager->IsScopableIndex(ObjectInternalIndex))
 				{
-					// Mark object as no longer scopable, and that we should not propagate changed states
+					// Mark object and subobjects as no longer scopable, and that we should not propagate changed states
 					NetRefHandleManager->RemoveFromScope(ObjectInternalIndex);
+					for (FInternalNetRefIndex SubObjectIndex : NetRefHandleManager->GetSubObjects(ObjectInternalIndex))
+					{
+						NetRefHandleManager->RemoveFromScope(SubObjectIndex);
+					}
 				}
 			
 				// Keep object in the pending tear-off list until the object is no longer referenced by any ReplicationWriter
