remove unnecessary immediate handling of cluster union deferred updates to better batch them for perf

diff --git a/Engine/Source/Runtime/Experimental/Chaos/Private/Chaos/ClusterUnionManager.cpp b/Engine/Source/Runtime/Experimental/Chaos/Private/Chaos/ClusterUnionManager.cpp
--- a/Engine/Source/Runtime/Experimental/Chaos/Private/Chaos/ClusterUnionManager.cpp
+++ b/Engine/Source/Runtime/Experimental/Chaos/Private/Chaos/ClusterUnionManager.cpp
@@ -204,7 +204,7 @@ namespace Chaos
 						HandleAddOperation(OpMap.Key, Op.Value, Op.Key == EClusterUnionOperation::AddReleased);
 						break;
 					case EClusterUnionOperation::Remove:
-						HandleRemoveOperation(OpMap.Key, Op.Value, EClusterUnionOperationTiming::Immediate);
+						HandleRemoveOperation(OpMap.Key, Op.Value, EClusterUnionOperationTiming::Defer);
 						break;
 					case EClusterUnionOperation::UpdateChildToParent:
 						HandleUpdateChildToParentOperation(OpMap.Key, Op.Value);
@@ -351,7 +351,7 @@ namespace Chaos
 			Cluster->InternalCluster->SetIsAnchored(true);
 		}
 
-		UpdateAllClusterUnionProperties(*Cluster, bIsNewCluster ? EUpdateClusterUnionPropertiesFlags::All : EUpdateClusterUnionPropertiesFlags::IncrementalGenerateConnectionGraph);
+		RequestDeferredClusterPropertiesUpdate(ClusterIndex, bIsNewCluster ? EUpdateClusterUnionPropertiesFlags::All : EUpdateClusterUnionPropertiesFlags::IncrementalGenerateConnectionGraph);
 
 		if (OldProxy)
 		{
@@ -366,27 +366,22 @@ namespace Chaos
 			}
 
 			MEvolution.SetPhysicsMaterial(Cluster->InternalCluster, MEvolution.GetPhysicsMaterial(FinalParticlesToAdd[0]));
-			Cluster->bNeedsXRInitialization = false;
 		}
-
-		if (Cluster->InternalCluster->Disabled())
-		{
-			MEvolution.EnableParticle(Cluster->InternalCluster);
-		}
-
-		MEvolution.DirtyParticle(*Cluster->InternalCluster);
-		MEvolution.GetParticles().MarkTransientDirtyParticle(Cluster->InternalCluster);
 	}
 
 	DECLARE_CYCLE_STAT(TEXT("FClusterUnionManager::DeferredClusterUnionUpdate"), STAT_DeferredClusterUnionUpdate, STATGROUP_Chaos);
-	void FClusterUnionManager::DeferredClusterUnionUpdate(FClusterUnion& Union)
+	void FClusterUnionManager::DeferredClusterUnionUpdate(FClusterUnion& Union, EUpdateClusterUnionPropertiesFlags Flags)
 	{
 		SCOPE_CYCLE_COUNTER(STAT_DeferredClusterUnionUpdate);
 
-		UpdateAllClusterUnionProperties(Union, EUpdateClusterUnionPropertiesFlags::IncrementalGenerateConnectionGraph);
+		UpdateAllClusterUnionProperties(Union, Flags);
 
 		if (!Union.ChildParticles.IsEmpty())
 		{
+			if (Union.InternalCluster->Disabled())
+			{
+				MEvolution.EnableParticle(Union.InternalCluster);
+			}
 			MEvolution.DirtyParticle(*Union.InternalCluster);
 		}
 		else
@@ -407,11 +402,11 @@ namespace Chaos
 			return;
 		}
 
-		for (FClusterUnionIndex Index : DeferredClusterUnionsForUpdateProperties)
+		for (const TPair<FClusterUnionIndex, EUpdateClusterUnionPropertiesFlags>& Kvp : DeferredClusterUnionsForUpdateProperties)
 		{
-			if (FClusterUnion* Union = FindClusterUnion(Index))
+			if (FClusterUnion* Union = FindClusterUnion(Kvp.Key))
 			{
-				DeferredClusterUnionUpdate(*Union);
+				DeferredClusterUnionUpdate(*Union, Kvp.Value);
 			}
 		}
 
@@ -476,13 +471,14 @@ namespace Chaos
 		// This gets changed because we go through an internal initialization route when we update the cluster union particle's properties.
 		Cluster->InternalCluster->SetPhysicsProxy(OldProxy);
 
+		constexpr EUpdateClusterUnionPropertiesFlags RemoveUpdateFlags = EUpdateClusterUnionPropertiesFlags::IncrementalGenerateConnectionGraph;
 		switch (UpdateClusterPropertiesTiming)
 		{
 		case EClusterUnionOperationTiming::Immediate:
-			DeferredClusterUnionUpdate(*Cluster);
+			DeferredClusterUnionUpdate(*Cluster, RemoveUpdateFlags);
 			break;
 		case EClusterUnionOperationTiming::Defer:
-			DeferredClusterUnionsForUpdateProperties.Add(ClusterIndex);
+			RequestDeferredClusterPropertiesUpdate(ClusterIndex, RemoveUpdateFlags);
 			break;
 		}		
 	}
@@ -513,14 +509,19 @@ namespace Chaos
 		// Update cluster properties.
 		FMatrix33 ClusterInertia(0);
 
-		const bool bRecomputeMassOrientation = EnumHasAnyFlags(Flags, EUpdateClusterUnionPropertiesFlags::RecomputeMassOrientation);
+		const bool bRecomputeMassOrientation = EnumHasAnyFlags(Flags, EUpdateClusterUnionPropertiesFlags::RecomputeMassOrientation) && ClusterUnion.bNeedsXRInitialization;
 
 		// TODO: These functions are generally just re-building the cluster from scratch. Need to figure out a way
 		// to get these functions to update the already existing cluster instead.
 		TSet<FPBDRigidParticleHandle*> FullChildrenSet(ClusterUnion.ChildParticles);
 
 		const FRigidTransform3 ForceMassOrientation{ ClusterUnion.InternalCluster->X(), ClusterUnion.InternalCluster->R() };
-		UpdateClusterMassProperties(ClusterUnion.InternalCluster, FullChildrenSet, ClusterInertia, (bRecomputeMassOrientation && ClusterUnion.bNeedsXRInitialization) ? nullptr : &ForceMassOrientation);
+		UpdateClusterMassProperties(ClusterUnion.InternalCluster, FullChildrenSet, ClusterInertia, bRecomputeMassOrientation ? nullptr : &ForceMassOrientation);
+		if (bRecomputeMassOrientation)
+		{
+			ClusterUnion.bNeedsXRInitialization = false;
+		}
+
 		UpdateKinematicProperties(ClusterUnion.InternalCluster, MClustering.GetChildrenMap(), MEvolution);
 
 		// We must reset collisions etc on this particle since geometry will be changed
@@ -795,14 +796,17 @@ namespace Chaos
 			}
 		}
 
-		UpdateAllClusterUnionProperties(*ClusterUnion, EUpdateClusterUnionPropertiesFlags::IncrementalGenerateConnectionGraph);
+		RequestDeferredClusterPropertiesUpdate(ClusterIndex, EUpdateClusterUnionPropertiesFlags::IncrementalGenerateConnectionGraph);
 
 		for (TPair<FPBDRigidClusteredParticleHandle*, bool>& Pair : DirtyParticleLockStates)
 		{
 			Pair.Key->SetChildToParentLocked(Pair.Value);
 		}
+	}
 
-		MEvolution.GetParticles().MarkTransientDirtyParticle(ClusterUnion->InternalCluster);
-		MEvolution.DirtyParticle(*ClusterUnion->InternalCluster);
+	void FClusterUnionManager::RequestDeferredClusterPropertiesUpdate(FClusterUnionIndex ClusterIndex, EUpdateClusterUnionPropertiesFlags Flags)
+	{
+		EUpdateClusterUnionPropertiesFlags& ExistingFlags = DeferredClusterUnionsForUpdateProperties.FindOrAdd(ClusterIndex);
+		EnumAddFlags(ExistingFlags, Flags);
 	}
 }
\ No newline at end of file
diff --git a/Engine/Source/Runtime/Experimental/Chaos/Public/Chaos/ClusterUnionManager.h b/Engine/Source/Runtime/Experimental/Chaos/Public/Chaos/ClusterUnionManager.h
--- a/Engine/Source/Runtime/Experimental/Chaos/Public/Chaos/ClusterUnionManager.h
+++ b/Engine/Source/Runtime/Experimental/Chaos/Public/Chaos/ClusterUnionManager.h
@@ -189,7 +189,8 @@ namespace Chaos
 		TMap<FClusterUnionIndex, FClusterUnion> ClusterUnions;
 
 		// The set of cluster unions that we want to defer updating their cluster union properties for.
-		TSet<FClusterUnionIndex> DeferredClusterUnionsForUpdateProperties;
+		void RequestDeferredClusterPropertiesUpdate(FClusterUnionIndex ClusterIndex, EUpdateClusterUnionPropertiesFlags Flags);
+		TMap<FClusterUnionIndex, EUpdateClusterUnionPropertiesFlags> DeferredClusterUnionsForUpdateProperties;
 		
 		//
 		// There are two ways we can pick a new union index:
@@ -217,7 +218,7 @@ namespace Chaos
 		TSharedPtr<FImplicitObject, ESPMode::ThreadSafe> ForceRecreateClusterUnionSharedGeometry(const FClusterUnion& Union);
 
 		// Handles updating the cluster union.
-		void DeferredClusterUnionUpdate(FClusterUnion& Union);
+		void DeferredClusterUnionUpdate(FClusterUnion& Union, EUpdateClusterUnionPropertiesFlags Flags);
 
 		// Flush the cluster union's incremental connectivity operations
 		void FlushIncrementalConnectivityGraphOperations(FClusterUnion& ClusterUnion);
diff --git a/Engine/Source/Runtime/Experimental/GeometryCollectionEngine/Private/GeometryCollection/GeometryCollectionComponent.cpp b/Engine/Source/Runtime/Experimental/GeometryCollectionEngine/Private/GeometryCollection/GeometryCollectionComponent.cpp
--- a/Engine/Source/Runtime/Experimental/GeometryCollectionEngine/Private/GeometryCollection/GeometryCollectionComponent.cpp
+++ b/Engine/Source/Runtime/Experimental/GeometryCollectionEngine/Private/GeometryCollection/GeometryCollectionComponent.cpp
@@ -2170,8 +2170,6 @@ bool UGeometryCollectionComponent::ProcessRepData(const float DeltaTime, const f
 		}
 	}
 
-	RigidClustering.GetClusterUnionManager().HandleDeferredClusterUnionUpdateProperties();
-
 	// Keep track of whether we did some "work" on this frame so we can turn off the async tick after
 	// multiple frames of not doing anything.
 	bool bProcessed = false;
