[ContactModification] New functions for FContactPairModification

- Check if a contact point is an edge contact point.
- Disable a manifold point. Doing this revealed an out of range index bug within UpdateCollisionConstraintFromSolver where the SolverManifoldPointIndex was used before we knew if the manifold point was disabled or not. Since only enabled points are copied to the solver this would run into an out of range index when looping through all points if any were disabled.
- New function to check if a contact point is disabled.

diff --git a/Engine/Source/Runtime/Experimental/Chaos/Private/Chaos/ContactModification.cpp b/Engine/Source/Runtime/Experimental/Chaos/Private/Chaos/ContactModification.cpp
--- a/Engine/Source/Runtime/Experimental/Chaos/Private/Chaos/ContactModification.cpp
+++ b/Engine/Source/Runtime/Experimental/Chaos/Private/Chaos/ContactModification.cpp
@@ -466,6 +466,25 @@ namespace Chaos
 		return { Constraint->GetShape0(), Constraint->GetShape1() };
 	}
 
+	bool FContactPairModifier::IsEdgeContactPoint(int32 ContactPointIdx) const
+	{
+		const FManifoldPoint& ManifoldPoint = Constraint->GetManifoldPoint(ContactPointIdx);
+
+		return ManifoldPoint.ContactPoint.ContactType == EContactPointType::EdgeEdge;
+	}
+
+	bool FContactPairModifier::IsContactPointDisabled(int32 ContactPointIdx) const
+	{
+		const FManifoldPoint& ManifoldPoint = Constraint->GetManifoldPoint(ContactPointIdx);
+
+		return ManifoldPoint.Flags.bDisabled;
+	}
+
+	void FContactPairModifier::SetContactPointDisabled(int32 ContactPointIdx) const
+	{
+		Constraint->DisableManifoldPoint(ContactPointIdx);
+	}
+
 	void FContactPairModifierIterator::SeekValidContact()
 	{
 		// Not valid to call from end.
diff --git a/Engine/Source/Runtime/Experimental/Chaos/Public/Chaos/ContactModification.h b/Engine/Source/Runtime/Experimental/Chaos/Public/Chaos/ContactModification.h
--- a/Engine/Source/Runtime/Experimental/Chaos/Public/Chaos/ContactModification.h
+++ b/Engine/Source/Runtime/Experimental/Chaos/Public/Chaos/ContactModification.h
@@ -220,6 +220,21 @@ namespace Chaos
 		*/
 		TVec2<const FPerShapeData*> GetShapePair() const;
 
+		/*
+		* Check to see if a contact point index is an edge contact.
+		*/
+		bool IsEdgeContactPoint(int32 ContactPointIdx) const;
+
+		/*
+		* Check to see if a contact point index is disabled.
+		*/
+		bool IsContactPointDisabled(int32 ContactPointIdx) const;
+
+		/*
+		* Set a contact point disabled.
+		*/
+		void SetContactPointDisabled(int32 ContactPointIdx) const;
+
 	private:
 		/**
 		 * @brief Update cached shape transforms in the constraint after modifying particle positions
