disable particles when children array becomes empty - add additional debug messages into geometry collection to track why ensure happens

diff --git a/Engine/Source/Runtime/Experimental/Chaos/Private/Chaos/PBDRigidClustering.cpp b/Engine/Source/Runtime/Experimental/Chaos/Private/Chaos/PBDRigidClustering.cpp
--- a/Engine/Source/Runtime/Experimental/Chaos/Private/Chaos/PBDRigidClustering.cpp
+++ b/Engine/Source/Runtime/Experimental/Chaos/Private/Chaos/PBDRigidClustering.cpp
@@ -690,6 +690,9 @@ namespace Chaos
 				ClusteredParent->SetClusterId(ClusterId{ ClusteredParent->Parent(), Children->Num()});
 				if (ClusteredParent->InternalCluster() && Children->IsEmpty() && ClusteredParent->PhysicsProxy() && ClusteredParent->PhysicsProxy()->GetType() == EPhysicsProxyType::GeometryCollectionType)
 				{
+					// It's safe to disable the particle until we get to the point where we want to destroy the particle.
+					MEvolution.DisableParticle(ClusteredParent);
+
 					// We shouldn't ever need to do an AddUnique here since when we remove a child from a parent, the parent should only ever turn empty once.
 					EmptyInternalClustersPerProxy.FindOrAdd(ClusteredParent->PhysicsProxy()).Add(ClusteredParent);
 				}
diff --git a/Engine/Source/Runtime/Experimental/GeometryCollectionEngine/Private/GeometryCollection/GeometryCollectionComponent.cpp b/Engine/Source/Runtime/Experimental/GeometryCollectionEngine/Private/GeometryCollection/GeometryCollectionComponent.cpp
--- a/Engine/Source/Runtime/Experimental/GeometryCollectionEngine/Private/GeometryCollection/GeometryCollectionComponent.cpp
+++ b/Engine/Source/Runtime/Experimental/GeometryCollectionEngine/Private/GeometryCollection/GeometryCollectionComponent.cpp
@@ -1879,7 +1879,9 @@ void UGeometryCollectionComponent::UpdateRepData()
 			FPBDRigidClusteredParticleHandle* Cluster = *Itr;
 
 			const TArray<FPBDRigidParticleHandle*>& Children = RigidClustering.GetChildrenMap().FindRef(Cluster);
-			if (Cluster->Disabled() || (Cluster->InternalCluster() && Children.IsEmpty()))
+			const bool bIsInternalCluster = Cluster->InternalCluster();
+			const bool bChildrenEmpty = Children.IsEmpty();
+			if (Cluster->Disabled() || (bIsInternalCluster && bChildrenEmpty))
 			{
 				Itr.RemoveCurrent();
 			}
@@ -1898,7 +1900,7 @@ void UGeometryCollectionComponent::UpdateRepData()
 				if(Cluster->InternalCluster())
 				{
 					
-					ensureMsgf(Children.Num(), TEXT("Internal cluster yet we have no children?"));
+					ensureMsgf(Children.Num(), TEXT("Internal cluster yet we have no children? [Num %d vs Cached Empty %d]"), Children.Num(), bChildrenEmpty);
 					TransformGroupIdx = PhysicsProxy->GetTransformGroupIndexFromHandle(Children[0]);
 				}
 				else
