Iris
- Allow LastResortPropertySerializer to store default state.

diff --git a/Engine/Source/Runtime/Experimental/Iris/Core/Private/Iris/ReplicationState/ReplicationStateDescriptorBuilder.cpp b/Engine/Source/Runtime/Experimental/Iris/Core/Private/Iris/ReplicationState/ReplicationStateDescriptorBuilder.cpp
--- a/Engine/Source/Runtime/Experimental/Iris/Core/Private/Iris/ReplicationState/ReplicationStateDescriptorBuilder.cpp
+++ b/Engine/Source/Runtime/Experimental/Iris/Core/Private/Iris/ReplicationState/ReplicationStateDescriptorBuilder.cpp
@@ -27,6 +27,8 @@
 #include "Misc/EnumClassFlags.h"
 #include "Templates/AlignmentTemplates.h"
 #include "UObject/UnrealType.h"
+#include "UObject/StrongObjectPtr.h"
+#include "Iris/Serialization/IrisObjectReferencePackageMap.h"
 
 #ifndef UE_NET_TEST_FAKE_REP_TAGS
 #	define UE_NET_TEST_FAKE_REP_TAGS 0
@@ -69,6 +71,17 @@ static FAutoConsoleVariableRef CVarIrisUseSupportsStructNetSerializerList(TEXT("
 static bool bWarnAboutStructPropertiesWithSuspectedNotReplicatedProperties = false;
 static FAutoConsoleVariableRef CVarIrisWarnAboutStructPropertiesWithSuspectedNotReplicatedProperties(TEXT("net.Iris.bWarnAboutStructPropertiesWithSuspectedNotReplicatedProperties"), bWarnAboutStructPropertiesWithSuspectedNotReplicatedProperties, TEXT("Try to detect if a struct replicated as a property might contain unannotated members, disabled by default."));
 
+static TStrongObjectPtr<UIrisObjectReferencePackageMap> s_IrisObjectReferencePackageMap = nullptr;
+
+static UIrisObjectReferencePackageMap* GetOrCreateIrisObjectReferencePackageMap()
+{
+	if (!s_IrisObjectReferencePackageMap.IsValid())
+	{
+		s_IrisObjectReferencePackageMap = TStrongObjectPtr<UIrisObjectReferencePackageMap>(NewObject<UIrisObjectReferencePackageMap>());
+	}
+	return s_IrisObjectReferencePackageMap.Get();
+}
+
 enum class EMemberPropertyTraits : uint32
 {
 	None								= 0U,
@@ -762,6 +775,7 @@ void FPropertyReplicationStateDescriptorBuilder::AllocateAndInitializeDefaultInt
 	// on both server and client
 	FNetSerializationContext Context;
 	FInternalNetSerializationContext InternalContext;
+	InternalContext.PackageMap = GetOrCreateIrisObjectReferencePackageMap();
 	Context.SetInternalContext(&InternalContext);
 	Context.SetIsInitializingDefaultState(true);
 
@@ -806,6 +820,7 @@ bool FPropertyReplicationStateDescriptorBuilder::CalculateDefaultStateChecksum(c
 	FNetBitStreamWriter Writer;
 	FNetSerializationContext Context(&Writer);
 	FInternalNetSerializationContext InternalContext;
+	InternalContext.PackageMap = GetOrCreateIrisObjectReferencePackageMap();
 	Context.SetInternalContext(&InternalContext);
 
 	// First try to serialize to small inline buffer
diff --git a/Engine/Source/Runtime/Experimental/Iris/Core/Private/Iris/Serialization/LastResortPropertyNetSerializer.cpp b/Engine/Source/Runtime/Experimental/Iris/Core/Private/Iris/Serialization/LastResortPropertyNetSerializer.cpp
--- a/Engine/Source/Runtime/Experimental/Iris/Core/Private/Iris/Serialization/LastResortPropertyNetSerializer.cpp
+++ b/Engine/Source/Runtime/Experimental/Iris/Core/Private/Iris/Serialization/LastResortPropertyNetSerializer.cpp
@@ -168,14 +168,6 @@ void FLastResortPropertyNetSerializer::Quantize(FNetSerializationContext& Contex
 	// Capture references
 	Private::FInternalNetSerializationContext* InternalContext = Context.GetInternalContext();
 	UIrisObjectReferencePackageMap* PackageMap = InternalContext ? InternalContext->PackageMap : nullptr;
-
-	// We do not store any defaults for properties using the LastResortPropertyNetSerializer.
-	if (Context.IsInitializingDefaultState())
-	{
-		Value.ObjectReferenceStorage.AdjustSize(Context, 0);
-		AdjustStorageSize(Context, Value, 0);
-		return;
-	}
 	
 	UIrisObjectReferencePackageMap::FObjectReferenceArray ObjectReferences;
 	if (PackageMap)
