Iris
-Reduced overhead of RepNotifies when using Iris replication to only store state for members that has RepNotifies

diff --git a/Engine/Source/Runtime/Experimental/Iris/Core/Private/Iris/ReplicationState/PropertyReplicationState.cpp b/Engine/Source/Runtime/Experimental/Iris/Core/Private/Iris/ReplicationState/PropertyReplicationState.cpp
--- a/Engine/Source/Runtime/Experimental/Iris/Core/Private/Iris/ReplicationState/PropertyReplicationState.cpp
+++ b/Engine/Source/Runtime/Experimental/Iris/Core/Private/Iris/ReplicationState/PropertyReplicationState.cpp
@@ -277,6 +277,36 @@ bool FPropertyReplicationState::PollPropertyReplicationState(const void* RESTRIC
 	return IsDirty();
 }
 
+bool FPropertyReplicationState::PollPropertyReplicationStateForRepNotifies(const void* RESTRICT SrcStateData)
+{
+	if (IsValid())
+	{
+		const FReplicationStateDescriptor* Descriptor = ReplicationStateDescriptor;
+		const uint8* SrcBuffer = reinterpret_cast<const uint8*>(SrcStateData);
+
+		IRIS_PROFILER_PROTOCOL_NAME(ReplicationStateDescriptor->DebugName->Name);
+
+		const FReplicationStateMemberDescriptor* MemberDescriptors = Descriptor->MemberDescriptors;
+		const FProperty** MemberProperties = Descriptor->MemberProperties;
+		const FReplicationStateMemberPropertyDescriptor* MemberPropertyDescriptors = Descriptor->MemberPropertyDescriptors;
+		const uint32 MemberCount = Descriptor->MemberCount;
+
+		for (uint32 MemberIt = 0; MemberIt < MemberCount; ++MemberIt)
+		{
+			const FReplicationStateMemberDescriptor& MemberDescriptor = MemberDescriptors[MemberIt];
+			const FReplicationStateMemberPropertyDescriptor& MemberPropertyDescriptor = MemberPropertyDescriptors[MemberIt];
+
+			if (MemberPropertyDescriptor.RepNotifyFunction)
+			{
+				const FProperty* Property = MemberProperties[MemberIt];
+				SetPropertyValue(MemberIt, SrcBuffer + Property->GetOffset_ForGC() + Property->ElementSize*MemberPropertyDescriptor.ArrayIndex);
+			}
+		}
+	}
+
+	return IsDirty();
+}
+
 void FPropertyReplicationState::PushPropertyReplicationState(void* RESTRICT DstData, bool bInPushAll) const
 {
 	// $IRIS TODO: Rewrite this to iterate over change mask instead of iterating over all members and querying the mask
@@ -379,12 +409,12 @@ void FPropertyReplicationState::CallRepNotifies(void* RESTRICT DstData, const FC
 				bool bShouldCallRepNotify = false;
 				if (Params.bOnlyCallIfDiffersFromLocal)
 				{
-					// We try to be backwards compatible and respect RepNotify_Always/RepNotify_Changed unless it is the intial state where we only will call the repnotify of the received value differs from the local one.
+					// We try to be backwards compatible and respect RepNotify_Always/RepNotify_Changed unless it is the initial state where we only will call the repnotify of the received value differs from the local one.
 					bShouldCallRepNotify = Params.bIsInit ? !Private::InternalCompareMember(Descriptor, MemberIt, ValuePtr, PrevValuePtr) : EnumHasAnyFlags(Descriptor->MemberTraitsDescriptors[MemberIt].Traits, EReplicationStateMemberTraits::HasRepNotifyAlways) || !Private::InternalCompareMember(Descriptor, MemberIt, ValuePtr, PrevValuePtr);
 				}
 				else
 				{
-					// Trust data from server and call RepNotify without doing additonal compare unless it is the inital state.
+					// Trust data from server and call RepNotify without doing additonal compare unless it is the initial state.
 					bShouldCallRepNotify = !Params.bIsInit || !Private::InternalCompareMember(Descriptor, MemberIt, ValuePtr, PrevValuePtr);
 				}
 
diff --git a/Engine/Source/Runtime/Experimental/Iris/Core/Private/Iris/ReplicationSystem/PropertyReplicationFragment.cpp b/Engine/Source/Runtime/Experimental/Iris/Core/Private/Iris/ReplicationSystem/PropertyReplicationFragment.cpp
--- a/Engine/Source/Runtime/Experimental/Iris/Core/Private/Iris/ReplicationSystem/PropertyReplicationFragment.cpp
+++ b/Engine/Source/Runtime/Experimental/Iris/Core/Private/Iris/ReplicationSystem/PropertyReplicationFragment.cpp
@@ -39,7 +39,7 @@ FPropertyReplicationFragment::FPropertyReplicationFragment(EReplicationFragmentT
 				PrevReplicationState = MakeUnique<FPropertyReplicationState>(InDescriptor);
 
 				// Poll to get instance default for our previous state
-				PrevReplicationState->PollPropertyReplicationState(InOwner);
+				PrevReplicationState->PollPropertyReplicationStateForRepNotifies(InOwner);
 
 				Traits |= EReplicationFragmentTraits::KeepPreviousState;
 			}
@@ -126,8 +126,7 @@ void FPropertyReplicationFragment::ApplyReplicatedState(FReplicationStateApplyCo
 	// If we do not rely on received data to issue rep notifies we need to store a copy of the local state before we apply the new received state.
 	if (!bUsePrevReceivedStateForOnReps && PrevReplicationState)
 	{
-		// This could use a partial poll instead only polling RepNotify properties.
-		PrevReplicationState->PollPropertyReplicationState(Owner);
+		PrevReplicationState->PollPropertyReplicationStateForRepNotifies(Owner);
 	}
 
 	// Create a wrapping property replication state, cheap as we are simply injecting the already constructed state
diff --git a/Engine/Source/Runtime/Experimental/Iris/Core/Public/Iris/ReplicationState/PropertyReplicationState.h b/Engine/Source/Runtime/Experimental/Iris/Core/Public/Iris/ReplicationState/PropertyReplicationState.h
--- a/Engine/Source/Runtime/Experimental/Iris/Core/Public/Iris/ReplicationState/PropertyReplicationState.h
+++ b/Engine/Source/Runtime/Experimental/Iris/Core/Public/Iris/ReplicationState/PropertyReplicationState.h
@@ -88,6 +88,11 @@ public:
 	*/
 	IRISCORE_API bool PollPropertyReplicationState(const void* RESTRICT SrcData);
 
+	/**
+	 * Poll src data from properties which we want to store in order to determine if we need to call RepNotifies
+	*/
+	IRISCORE_API bool PollPropertyReplicationStateForRepNotifies(const void* RESTRICT SrcData);
+
 	/** Push received state data to properties in DstData buffer, Note: DstData is a UClass/UStruct containing properties
 		Compare and update representation in DstStateBuffer and update ChangeMask 
 	*/
