Chaos - fix SleepEnabled cvar (was being ignored)

diff --git a/Engine/Source/Runtime/Experimental/Chaos/Private/Chaos/Island/IslandManager.cpp b/Engine/Source/Runtime/Experimental/Chaos/Private/Chaos/Island/IslandManager.cpp
--- a/Engine/Source/Runtime/Experimental/Chaos/Private/Chaos/Island/IslandManager.cpp
+++ b/Engine/Source/Runtime/Experimental/Chaos/Private/Chaos/Island/IslandManager.cpp
@@ -1913,6 +1913,11 @@ namespace Chaos::Private
 
 	void FPBDIslandManager::ProcessSleep()
 	{
+		if (!CVars::bChaosSolverSleepEnabled)
+		{
+			return;
+		}
+
 		// Isloated particles are not kept in any island and need to be handled separately.
 		ProcessParticlesSleep();
 
