Chaos - use stack alloctaor for mesh collision

diff --git a/Engine/Source/Runtime/Experimental/Chaos/Private/Chaos/Collision/ConvexTriangleContactPoint.cpp b/Engine/Source/Runtime/Experimental/Chaos/Private/Chaos/Collision/ConvexTriangleContactPoint.cpp
--- a/Engine/Source/Runtime/Experimental/Chaos/Private/Chaos/Collision/ConvexTriangleContactPoint.cpp
+++ b/Engine/Source/Runtime/Experimental/Chaos/Private/Chaos/Collision/ConvexTriangleContactPoint.cpp
@@ -7,6 +7,7 @@
 #include "Chaos/DebugDrawQueue.h"
 #include "Chaos/SAT.h"
 #include "Chaos/Triangle.h"
+#include "Misc/MemStack.h"
 
 //PRAGMA_DISABLE_OPTIMIZATION
 
@@ -142,8 +143,8 @@ namespace Chaos
 		//
 		
 		// We store the convex vertex distances to the triangle face for use in the edge-edge culling
-		// @todo(chaos): scratch or stack allocator
-		TArray<FReal> ConvexVertexDs;
+		FMemMark Mark(FMemStack::Get());
+		TArray<FReal, TMemStackAllocator<alignof(FReal)>> ConvexVertexDs;
 		ConvexVertexDs.SetNum(Convex.NumVertices());
 		TArrayView<FReal> ConvexVertexDsView = MakeArrayView(ConvexVertexDs);
 
