Chaos - fix sleep when particles are close together and one falls asleep and is immediately awoken by a new inactive contact with the neighbour.

This was caused by a flaw in the calculation of CullDistance that resulted in it increasing slightly (by G.Dt.Dt - normally around 1cm) when one particle goes to sleep.

diff --git a/Engine/Source/Runtime/Experimental/Chaos/Private/Chaos/Collision/ParticlePairMidPhase.cpp b/Engine/Source/Runtime/Experimental/Chaos/Private/Chaos/Collision/ParticlePairMidPhase.cpp
--- a/Engine/Source/Runtime/Experimental/Chaos/Private/Chaos/Collision/ParticlePairMidPhase.cpp
+++ b/Engine/Source/Runtime/Experimental/Chaos/Private/Chaos/Collision/ParticlePairMidPhase.cpp
@@ -1647,7 +1647,14 @@ namespace Chaos
 			const bool bUseSweep = Flags.bIsCCD && ShouldEnableCCD(Dt);
 
 			// Extend cull distance based on velocity
-			const FReal VMaxDt = (FConstGenericParticleHandle(GetParticle0())->V() - FConstGenericParticleHandle(GetParticle1())->V()).GetAbsMax() * Dt;
+			// NOTE: We use PreV here which is the velocity after collisions from the previous tick because we want the 
+			// velocity without gravity from this tick applied. This is mainly so that we get the same CullDistance from 
+			// one tick to the next, even if one of the particles goes to sleep, and therefore its velocity is now zero 
+			// because gravity is no longer applied. Also see FPBDIslandManager::PropagateIslandSleep for other issues 
+			// related to velocity and sleeping...
+			const FReal VMax0 = FConstGenericParticleHandle(GetParticle0())->PreV().GetAbsMax();
+			const FReal VMax1 = FConstGenericParticleHandle(GetParticle1())->PreV().GetAbsMax();
+			const FReal VMaxDt = FMath::Max(VMax0, VMax1) * Dt;
 			if (!bUseSweep)
 			{
 				// Normal (non sweep) mode: we increase CullDistance based on velocity up to a limit
