Chaos - fix BoundingVolumeHierarchy mid point calculation. Add more comments.

diff --git a/Engine/Source/Runtime/Experimental/Chaos/Private/Chaos/BoundingVolumeHierarchy.cpp b/Engine/Source/Runtime/Experimental/Chaos/Private/Chaos/BoundingVolumeHierarchy.cpp
--- a/Engine/Source/Runtime/Experimental/Chaos/Private/Chaos/BoundingVolumeHierarchy.cpp
+++ b/Engine/Source/Runtime/Experimental/Chaos/Private/Chaos/BoundingVolumeHierarchy.cpp
@@ -407,6 +407,7 @@ namespace Chaos
 			Counts.Pos[i] += (Box.Min()[i] > MidPoint[i] || Box.Max()[i] > MidPoint[i]) ? 1 : 0;
 		}
 	}
+
 	template<class OBJECT_ARRAY, class LEAF_TYPE, class T, int d>
 	int32 TBoundingVolumeHierarchy<OBJECT_ARRAY, LEAF_TYPE, T, d>::GenerateNextLevel(const TVector<T, d>& GlobalMin, const TVector<T, d>& GlobalMax, const TArray<int32>& Objects, const int32 Axis, const int32 Level, const bool AllowMultipleSplitting)
 	{
@@ -415,32 +416,45 @@ namespace Chaos
 			return GenerateNextLevel(GlobalMin, GlobalMax, Objects, Level);
 		}
 
-		FSplitCount<d> Counts[2];
-		TArray<TBVHNode<T, d>> LocalElements;
-		TArray<TArray<int32>> LocalObjects;
-		TArray<LEAF_TYPE> LocalLeafs;
-		LocalElements.SetNum(2);
-		LocalObjects.SetNum(2);
-		LocalLeafs.SetNum(2);
+		TBVHNode<T, d> LocalElements[2];
+		TArray<int32> LocalObjects[2];
+		LEAF_TYPE LocalLeafs[2];
+
 		TAABB<T, d> GlobalBox(GlobalMin, GlobalMax);
 		const TVector<T, d> WorldCenter = GlobalBox.Center();
 
-		// @todo(chaos): I'm not sure what MinCenterSearch and MaxCenterSearch are supposed to be
-		// doing here, but using two different points and counting objects either side along each
-		// axis (see AccumulateNextLevelCount) will double-count objects. In the simplest case of
-		// a regular grid of objects, this results in always splitting along one axis rather than
-		// alternating between the three axes and that's probably not what was intended. Consider
-		// switching to KD Tree logic (sort along axis and pick median point for split position).
-		// OK, how super amazing was it that this comment block perfectly right-justified itself?
-		const TVector<T, d> MinCenterSearch = WorldCenter;//TAABB<T, d>(GlobalMin, WorldCenter).Center();
-		const TVector<T, d> MaxCenterSearch = WorldCenter;//TAABB<T, d>(WorldCenter, GlobalMax).Center();
+		// We are splitting the node at WorldCenter along axis Axis.
+		// The two child nodes will split at NextLevelCenter[0] and NextLevelCenter[1] along an axis we have yet to determine.
+		TAABB<T, d> NextLevelBounds[2];
+		TVector<T, d> NextLevelCenters[2];
+		FSplitCount<d> NextLevelCounts[2];
+		for (int32 i = 0; i < 2; ++i)
+		{
+			TVector<T, d> Min = GlobalBox.Min();
+			TVector<T, d> Max = GlobalBox.Max();
+			if (i == 0)
+			{
+				Max[Axis] = WorldCenter[Axis];
+			}
+			else
+			{
+				Min[Axis] = WorldCenter[Axis];
+			}
+
+			NextLevelBounds[i] = TAABB<T, d>(Min, Max);
+			NextLevelCenters[i] = T(0.5) * (Min + Max);
+		}
 
+		// Determine which side of the divide each body falls. While we are here, we accumulate the number of objects
+		// on each side of the child nodes would be for each of the 3 axis choices.
 		for(int32 i = 0; i < Objects.Num(); ++i)
 		{
 			check(Objects[i] >= 0 && Objects[i] < GetObjectCount(*MObjects));
 			const TAABB<T, d>& ObjectBox = Chaos::GetWorldSpaceBoundingBox(*MObjects, Objects[i], MWorldSpaceBoxes);
-			const TVector<T, d> ObjectCenter = ObjectBox.Center();
-			bool MinA = false, MaxA = false;
+
+			// This side of the divide do we fall on? It could be both sides if we span the split point
+			bool MinA = false;
+			bool MaxA = false;
 			if(ObjectBox.Min()[Axis] < WorldCenter[Axis])
 			{
 				MinA = true;
@@ -449,28 +463,27 @@ namespace Chaos
 			{
 				MaxA = true;
 			}
-			ensure(MinA || MaxA); // If this is not true we have a nan and the object gets ignored
+
+			// If this is not true we have a nan and the object gets ignored
+			ensure(MinA || MaxA);
+
+			// Add the object to the selected child node object list and determine where it 
+			// would fall in the next level's split, for all possible split axes
 			if(MinA)
 			{
 				LocalObjects[0].Add(Objects[i]);
-				AccumulateNextLevelCount(ObjectBox, MinCenterSearch, Counts[0]);
+				AccumulateNextLevelCount(ObjectBox, NextLevelCenters[0], NextLevelCounts[0]);
 			}
 			if(MaxA)
 			{
 				LocalObjects[1].Add(Objects[i]);
-				AccumulateNextLevelCount(ObjectBox, MaxCenterSearch, Counts[1]);
+				AccumulateNextLevelCount(ObjectBox, NextLevelCenters[1], NextLevelCounts[1]);
 			}
 		}
 		PhysicsParallelFor(2, [&](int32 i)
 		{
-			TVector<T, d> Min = GlobalBox.Min();
-			TVector<T, d> Max = GlobalBox.Max();
-			if(i == 0)
-				Max[Axis] = WorldCenter[Axis];
-			else
-				Min[Axis] = WorldCenter[Axis];
-			LocalElements[i].MMin = Min;
-			LocalElements[i].MMax = Max;
+			LocalElements[i].MMin = NextLevelBounds[i].Min();
+			LocalElements[i].MMax = NextLevelBounds[i].Max();
 			LocalElements[i].MAxis = -1;
 #if !UE_BUILD_SHIPPING && !UE_BUILD_TEST
 			LocalElements[i].LeafIndex = -1;
@@ -482,7 +495,7 @@ namespace Chaos
 				int32 MaxCulled = 0;
 				for(int32 LocalAxis = 0; LocalAxis < d; ++LocalAxis)
 				{
-					int32 CulledWorstCase = FMath::Min(Counts[i].Neg[LocalAxis], Counts[i].Pos[LocalAxis]);
+					int32 CulledWorstCase = FMath::Min(NextLevelCounts[i].Neg[LocalAxis], NextLevelCounts[i].Pos[LocalAxis]);
 					if(CulledWorstCase > MaxCulled)
 					{
 						MaxCulled = CulledWorstCase;
