Iris
- Repnotifies now defaults to using legacy compatibility mode to avoid issues with logic that does not handle getting repnotifies where the value might be the same. Created jira  https://jira.it.epicgames.com/browse/UE-187417 to optimize memory usage.

diff --git a/Engine/Source/Runtime/Experimental/Iris/Core/Private/Iris/ReplicationSystem/PropertyReplicationFragment.cpp b/Engine/Source/Runtime/Experimental/Iris/Core/Private/Iris/ReplicationSystem/PropertyReplicationFragment.cpp
--- a/Engine/Source/Runtime/Experimental/Iris/Core/Private/Iris/ReplicationSystem/PropertyReplicationFragment.cpp
+++ b/Engine/Source/Runtime/Experimental/Iris/Core/Private/Iris/ReplicationSystem/PropertyReplicationFragment.cpp
@@ -11,14 +11,13 @@
 namespace UE::Net
 {
 
-static bool bUsePrevReceivedStateForOnReps = true;
+static bool bUsePrevReceivedStateForOnReps = false;
 static FAutoConsoleVariableRef CVarUsePrevReceivedStateForOnReps(
 		TEXT("net.Iris.UsePrevReceivedStateForOnReps"),
 		bUsePrevReceivedStateForOnReps,
 		TEXT("If true OnReps will use the previous received state when doing onreps and not do any compares, if set to false we will copy the local state and do a compare before issuing onreps"
 		));
 
-
 FPropertyReplicationFragment::FPropertyReplicationFragment(EReplicationFragmentTraits InTraits, UObject* InOwner, const FReplicationStateDescriptor* InDescriptor)
 : FReplicationFragment(InTraits)
 , ReplicationStateDescriptor(InDescriptor)
