Iris - Store FName strings, when required, as they were. To prevent default state checksum mismatch the serialization is a no-op when serializing default state. Equality check will ignore casing.

diff --git a/Engine/Plugins/Runtime/ReplicationSystemTestPlugin/Source/Private/Tests/Serialization/TestNameNetSerializer.cpp b/Engine/Plugins/Runtime/ReplicationSystemTestPlugin/Source/Private/Tests/Serialization/TestNameNetSerializer.cpp
--- a/Engine/Plugins/Runtime/ReplicationSystemTestPlugin/Source/Private/Tests/Serialization/TestNameNetSerializer.cpp
+++ b/Engine/Plugins/Runtime/ReplicationSystemTestPlugin/Source/Private/Tests/Serialization/TestNameNetSerializer.cpp
@@ -116,16 +116,15 @@ void FTestNameNetSerializer::TestIsEqual()
 	for (int32 ValueIt = 0, ValueEndIt = TestNameCount; ValueIt != ValueEndIt; ++ValueIt)
 	{
 		CompareValues[1].Add(TestNames[(ValueIt + 1) % ValueEndIt]);
-		ExpectedResults[1].Add(TestNames[ValueIt].IsEqual(TestNames[(ValueIt + 1) % ValueEndIt], ENameCase::CaseSensitive, true));
+		ExpectedResults[1].Add(TestNames[ValueIt].IsEqual(TestNames[(ValueIt + 1) % ValueEndIt], ENameCase::IgnoreCase, true));
 	}
 
 	// Do two rounds of testing per config, one where we compare each value with itself and one where we compare against a value in range.
 	for (SIZE_T TestRoundIt : {0, 1})
 	{
 		// Do both quantized and regular compares
-		for (SIZE_T CompareIt : {0, 1})
+		for (bool bQuantizedCompare : {false, true})
 		{
-			bool bQuantizedCompare = CompareIt == 0;
 			const bool bSuccess = Super::TestIsEqual(TestNames, CompareValues[TestRoundIt].GetData(), ExpectedResults[TestRoundIt].GetData(), TestNameCount, SerializerConfig, bQuantizedCompare);
 			if (!bSuccess)
 			{
diff --git a/Engine/Source/Runtime/Experimental/Iris/Core/Private/Iris/ReplicationState/ReplicationStateDescriptorBuilder.cpp b/Engine/Source/Runtime/Experimental/Iris/Core/Private/Iris/ReplicationState/ReplicationStateDescriptorBuilder.cpp
--- a/Engine/Source/Runtime/Experimental/Iris/Core/Private/Iris/ReplicationState/ReplicationStateDescriptorBuilder.cpp
+++ b/Engine/Source/Runtime/Experimental/Iris/Core/Private/Iris/ReplicationState/ReplicationStateDescriptorBuilder.cpp
@@ -823,6 +823,9 @@ bool FPropertyReplicationStateDescriptorBuilder::CalculateDefaultStateChecksum(c
 	InternalContext.PackageMap = GetOrCreateIrisObjectReferencePackageMap();
 	Context.SetInternalContext(&InternalContext);
 
+	// Tell serializers we are serializing default state. It allows serializers to opt out of being part of the checksum by simply not serializing any data.
+	Context.SetIsInitializingDefaultState(true);
+
 	// First try to serialize to small inline buffer
 	Writer.InitBytes((uint8*)TempBuffer.GetData(), SmallBufferSize);
 	FReplicationStateOperations::Serialize(Context, Descriptor->DefaultStateBuffer, Descriptor);
diff --git a/Engine/Source/Runtime/Experimental/Iris/Core/Private/Iris/Serialization/StringNetSerializers.cpp b/Engine/Source/Runtime/Experimental/Iris/Core/Private/Iris/Serialization/StringNetSerializers.cpp
--- a/Engine/Source/Runtime/Experimental/Iris/Core/Private/Iris/Serialization/StringNetSerializers.cpp
+++ b/Engine/Source/Runtime/Experimental/Iris/Core/Private/Iris/Serialization/StringNetSerializers.cpp
@@ -94,6 +94,11 @@ const FStringNetSerializer::ConfigType FStringNetSerializer::DefaultConfig;
 // FNameNetSerializer
 void FNameNetSerializer::Serialize(FNetSerializationContext& Context, const FNetSerializeArgs& Args)
 {
+	if (Context.IsInitializingDefaultState())
+	{
+		return;
+	}
+
 	const QuantizedType& Value = *reinterpret_cast<QuantizedType*>(Args.Source);
 	
 	FNetBitStreamWriter* Writer = Context.GetBitStreamWriter();
@@ -119,6 +124,12 @@ void FNameNetSerializer::Deserialize(FNetSerializationContext& Context, const FN
 {
 	using namespace Private;
 
+	// Unexpected, but consistent with Serialize.
+	if (Context.IsInitializingDefaultState())
+	{
+		return;
+	}
+
 	QuantizedType& Target = *reinterpret_cast<QuantizedType*>(Args.Target);
 	const uint32 CurrentElementCount = Target.ElementCount;
 
@@ -180,8 +191,6 @@ void FNameNetSerializer::Quantize(FNetSerializationContext& Context, const FNetQ
 	const bool bIsString = (AsEName == nullptr || !ShouldReplicateAsInteger(*AsEName, SourceName));
 	if (bIsString)
 	{
-		const bool bRequireCaseInsensitivity = true;
-
 		const FNameEntry* DisplayNameEntry = SourceName.GetDisplayNameEntry();
 		TargetName.bIsString = 1;
 		TargetName.bIsEncoded = DisplayNameEntry->IsWide();
@@ -192,10 +201,6 @@ void FNameNetSerializer::Quantize(FNetSerializationContext& Context, const FNetQ
 		{
 			WIDECHAR TempWideBuffer[NAME_SIZE];
 			SourceName.GetPlainWIDEString(TempWideBuffer);
-			if (bRequireCaseInsensitivity)
-			{
-				FCString::Strupr(TempWideBuffer);
-			}
 
 			// Our codec uses up to 3 bytes per codepoint
 			const uint32 NameLength = DisplayNameEntry->GetNameLength() + 1U;
@@ -236,10 +241,6 @@ void FNameNetSerializer::Quantize(FNetSerializationContext& Context, const FNetQ
 			// At this time it's impossible to avoid the double copy unless we use a fixed excessive storage
 			ANSICHAR TempAnsiBuffer[NAME_SIZE];
 			SourceName.GetPlainANSIString(TempAnsiBuffer);
-			if (bRequireCaseInsensitivity)
-			{
-				FCStringAnsi::Strupr(TempAnsiBuffer);
-			}
 
 			FMemory::Memcpy(TargetName.ElementStorage, TempAnsiBuffer, NewElementCount*sizeof(ANSICHAR));
 		}
@@ -321,11 +322,35 @@ bool FNameNetSerializer::IsEqual(FNetSerializationContext& Context, const FNetIs
 		/**
 		 * Ideally we would have some process agnostic fast way to compare the FNames. There's currently no such thing.
 		 * We could store a hash as well in the quantized state but currently we don't. This method shouldn't be called
-		 * often anyway and the above checks should catch most differences. We want case sensitive checks.
+		 * often anyway and the above checks should catch most differences.
 		 */
-		const bool bIsStringEqual = !Value0.bIsString
-			|| FCStringAnsi::Strncmp(static_cast<const ANSICHAR*>(Value0.ElementStorage), static_cast<const ANSICHAR*>(Value1.ElementStorage), Value0.ElementCount) == 0;
-		return bIsStringEqual;
+		if (Value0.bIsString)
+		{
+			if (Value0.bIsEncoded)
+			{
+				SourceType SourceValue0;
+				SourceType SourceValue1;
+
+				FNetDequantizeArgs DequantizeArgs = {};
+				DequantizeArgs.NetSerializerConfig = Args.NetSerializerConfig;
+
+				DequantizeArgs.Source = Args.Source0;
+				DequantizeArgs.Target = NetSerializerValuePointer(&SourceValue0);
+				Dequantize(Context, DequantizeArgs);
+
+				DequantizeArgs.Source = Args.Source1;
+				DequantizeArgs.Target = NetSerializerValuePointer(&SourceValue1);
+				Dequantize(Context, DequantizeArgs);
+
+				return SourceValue0 == SourceValue1;
+			}
+			else
+			{
+				return FCStringAnsi::Strnicmp(static_cast<const ANSICHAR*>(Value0.ElementStorage), static_cast<const ANSICHAR*>(Value1.ElementStorage), Value0.ElementCount) == 0;;
+			}
+		}
+
+		return true;
 	}
 	else
 	{
