extra checks when handling add operations in FClusterUnionManager::FlushIncrementalConnectivityGraphOperations to handle the potential case where the cluster union changed and its geometry is no longer an implicit object union

diff --git a/Engine/Source/Runtime/Experimental/Chaos/Private/Chaos/ClusterUnionManager.cpp b/Engine/Source/Runtime/Experimental/Chaos/Private/Chaos/ClusterUnionManager.cpp
--- a/Engine/Source/Runtime/Experimental/Chaos/Private/Chaos/ClusterUnionManager.cpp
+++ b/Engine/Source/Runtime/Experimental/Chaos/Private/Chaos/ClusterUnionManager.cpp
@@ -587,79 +587,82 @@ namespace Chaos
 		{
 			if (Op.Value == EClusterUnionConnectivityOperation::Add)
 			{
-				// Use the acceleration structure of the cluster union itself to make finding overlaps easy.
-				const FImplicitObjectUnion& ShapeUnion = ClusterUnion.SharedGeometry->GetObjectChecked<FImplicitObjectUnion>();
-				const FRigidTransform3 FromTransform = GetParticleRigidFrameInClusterUnion(Op.Key, ClusterWorldTM);
-				FAABB3 ParticleLocalBounds = Op.Key->SharedGeometry()->BoundingBox().TransformedAABB(FromTransform);
-				ParticleLocalBounds.Thicken(kAddThickness);
-
-				ShapeUnion.VisitOverlappingLeafObjects(
-					ParticleLocalBounds,
-					[this, &ClusterUnion, &FromTransform, &Op, kAddThickness](const FImplicitObject* ToGeom, const FRigidTransform3& ToTransform, const int32 RootObjectIndex, const int32 ObjectIndex, const int32 LeafObjectIndex)
-					{
-						check(ToGeom != nullptr);
-						// RootObjectIndex is the index in the cluster union.
-						// NOTE: This will need to be re-thought if we ever decide to make the mapping between child shapes and child particles not 1-to-1 (e.g. if we ever attempt to simplify the cluster union shape).
-
-						// Ignore intersections against the same particle since we just added the particle (potentially) into the geometry.
-						if (Op.Key == ClusterUnion.ChildParticles[RootObjectIndex])
+				if (!ClusterUnion.ChildParticles.IsEmpty() && ClusterUnion.SharedGeometry->GetType() == ImplicitObjectType::Union)
+				{
+					// Use the acceleration structure of the cluster union itself to make finding overlaps easy.
+					const FImplicitObjectUnion& ShapeUnion = ClusterUnion.SharedGeometry->GetObjectChecked<FImplicitObjectUnion>();
+					const FRigidTransform3 FromTransform = GetParticleRigidFrameInClusterUnion(Op.Key, ClusterWorldTM);
+					FAABB3 ParticleLocalBounds = Op.Key->SharedGeometry()->BoundingBox().TransformedAABB(FromTransform);
+					ParticleLocalBounds.Thicken(kAddThickness);
+
+					ShapeUnion.VisitOverlappingLeafObjects(
+						ParticleLocalBounds,
+						[this, &ClusterUnion, &FromTransform, &Op, kAddThickness](const FImplicitObject* ToGeom, const FRigidTransform3& ToTransform, const int32 RootObjectIndex, const int32 ObjectIndex, const int32 LeafObjectIndex)
 						{
-							return;
-						}
-
-						FAABB3 ToAABB = ToGeom->CalculateTransformedBounds(ToTransform);
-						ToAABB.Thicken(kAddThickness);
+							check(ToGeom != nullptr);
+							// RootObjectIndex is the index in the cluster union.
+							// NOTE: This will need to be re-thought if we ever decide to make the mapping between child shapes and child particles not 1-to-1 (e.g. if we ever attempt to simplify the cluster union shape).
 
-						const bool bOverlap = Utilities::CastHelper(*ToGeom, ToTransform, 
-							[&Op, &FromTransform, &ToAABB, kAddThickness](const auto& ToGeomDowncast, const FRigidTransform3& FinalToTransform)
+							// Ignore intersections against the same particle since we just added the particle (potentially) into the geometry.
+							if (Op.Key == ClusterUnion.ChildParticles[RootObjectIndex])
 							{
-								const FShapesArray& AllFromShapes = Op.Key->ShapesArray();
-								for (int32 FromIndex = 0; FromIndex < AllFromShapes.Num(); ++FromIndex)
-								{
-									if (!AllFromShapes[FromIndex])
-									{
-										continue;
-									}
-
-									const FPerShapeData& FromShape = *AllFromShapes[FromIndex];
-									const FImplicitObject* FromGeom = FromShape.GetGeometry().Get();
-
-									if (!FromGeom)
-									{
-										continue;
-									}
+								return;
+							}
 
-									FAABB3 FromAABB = FromGeom->CalculateTransformedBounds(FromTransform);
-									FromAABB.Thicken(kAddThickness);
+							FAABB3 ToAABB = ToGeom->CalculateTransformedBounds(ToTransform);
+							ToAABB.Thicken(kAddThickness);
 
-									// First sanity check to see if the two shape AABB's intersect.
-									if (!FromAABB.Intersects(ToAABB))
-									{
-										continue;
-									}
-
-									// Now do a more accurate overlap check between the two shapes.
-									// Note that passing MTD here is critical as it gets us a more accurate check for some reason...
-									FMTDInfo MTDInfo;
-									if (OverlapQuery(*FromGeom, FromTransform, ToGeomDowncast, FinalToTransform, kAddThickness, &MTDInfo))
-									{
-										return true;
-									}
-									else
+							const bool bOverlap = Utilities::CastHelper(*ToGeom, ToTransform,
+								[&Op, &FromTransform, &ToAABB, kAddThickness](const auto& ToGeomDowncast, const FRigidTransform3& FinalToTransform)
+								{
+									const FShapesArray& AllFromShapes = Op.Key->ShapesArray();
+									for (int32 FromIndex = 0; FromIndex < AllFromShapes.Num(); ++FromIndex)
 									{
-										continue;
+										if (!AllFromShapes[FromIndex])
+										{
+											continue;
+										}
+
+										const FPerShapeData& FromShape = *AllFromShapes[FromIndex];
+										const FImplicitObject* FromGeom = FromShape.GetGeometry().Get();
+
+										if (!FromGeom)
+										{
+											continue;
+										}
+
+										FAABB3 FromAABB = FromGeom->CalculateTransformedBounds(FromTransform);
+										FromAABB.Thicken(kAddThickness);
+
+										// First sanity check to see if the two shape AABB's intersect.
+										if (!FromAABB.Intersects(ToAABB))
+										{
+											continue;
+										}
+
+										// Now do a more accurate overlap check between the two shapes.
+										// Note that passing MTD here is critical as it gets us a more accurate check for some reason...
+										FMTDInfo MTDInfo;
+										if (OverlapQuery(*FromGeom, FromTransform, ToGeomDowncast, FinalToTransform, kAddThickness, &MTDInfo))
+										{
+											return true;
+										}
+										else
+										{
+											continue;
+										}
 									}
-								}
 
-								return false;
-							});
+									return false;
+								});
 
-						if (bOverlap)
-						{
-							MClustering.CreateNodeConnection(Op.Key, ClusterUnion.ChildParticles[RootObjectIndex]);
+							if (bOverlap)
+							{
+								MClustering.CreateNodeConnection(Op.Key, ClusterUnion.ChildParticles[RootObjectIndex]);
+							}
 						}
-					}
-				);
+					);
+				}
 			}
 			else if (Op.Value == EClusterUnionConnectivityOperation::Remove)
 			{
