Clean up : Remove Break, Removal and Crumbling Events from CHAOSSOLVERENGINE_API, those events are now in ENGINE_API

diff --git a/Engine/Config/BaseEngine.ini b/Engine/Config/BaseEngine.ini
--- a/Engine/Config/BaseEngine.ini
+++ b/Engine/Config/BaseEngine.ini
@@ -367,6 +367,10 @@ DefaultWorldPartitionHLODLayer=/Engine/Maps/Templates/HLODs/HLODLayer_Instanced.
 +StructRedirects=(OldName="AnimNode_Slot",NewName="/Script/AnimGraphRuntime.AnimNode_Slot")
 +StructRedirects=(OldName="FormatTextArgument",NewName="/Script/Engine.FormatArgumentData")
 
++StructRedirects=(OldName="ChaosBreakEvent",NewName="/Script/Engine.ChaosEventType.BreakChaosEvent")
++StructRedirects=(OldName="ChaosRemovalEvent",NewName="/Script/Engine.ChaosEventType.RemovalChaosEvent")
++StructRedirects=(OldName="ChaosCrumblingEvent",NewName="/Script/Engine.ChaosEventType.CrumblingChaosEvent")
+
 +FunctionRedirects=(OldName="ConvertTransformToRelative",NewName="/Script/Engine.KismetMathLibrary.MakeRelativeTransform")
 +PropertyRedirects=(OldName="KismetMathLibrary.MakeRelativeTransform.LocalTransform",NewName="A")
 +PropertyRedirects=(OldName="KismetMathLibrary.MakeRelativeTransform.WorldTransform",NewName="RelativeTo")
diff --git a/Engine/Source/Runtime/Experimental/ChaosSolverEngine/Private/Chaos/ChaosGameplayEventDispatcher.cpp b/Engine/Source/Runtime/Experimental/ChaosSolverEngine/Private/Chaos/ChaosGameplayEventDispatcher.cpp
--- a/Engine/Source/Runtime/Experimental/ChaosSolverEngine/Private/Chaos/ChaosGameplayEventDispatcher.cpp
+++ b/Engine/Source/Runtime/Experimental/ChaosSolverEngine/Private/Chaos/ChaosGameplayEventDispatcher.cpp
@@ -561,7 +561,7 @@ void UChaosGameplayEventDispatcher::HandleBreakingEvents(const Chaos::FBreakingE
 				if ((BreakingData.EmitterFlag & Chaos::EventEmitterFlag::OwnDispatcher) && BreakingData.Proxy)
 				{
 					UPrimitiveComponent* const PrimComp = Scene.GetOwningComponent<UPrimitiveComponent>(BreakingData.Proxy);
-					if (PrimComp && BreakEventRegistrations.Contains(PrimComp))
+					if (PrimComp)
 					{
 						// queue them up so we can release the physics data before trigging BP events
 						FBreakChaosEvent& BreakEvent = PendingBreakEvents.Emplace_GetRef(BreakingData);
@@ -664,19 +664,16 @@ void UChaosGameplayEventDispatcher::HandleCrumblingEvents(const Chaos::FCrumblin
 			{
 				if (UPrimitiveComponent* const PrimComp = Scene.GetOwningComponent<UPrimitiveComponent>(CrumblingDataItem.Proxy))
 				{
-					if (CrumblingEventRegistrations.Contains(PrimComp))
-					{
-						// queue them up so we can release the physics data before triggering BP events
-						FCrumblingChaosEvent& CrumblingEvent = PendingCrumblingEvent.AddZeroed_GetRef();
-						CrumblingEvent.Component = PrimComp;
-						CrumblingEvent.Location = CrumblingDataItem.Location;
-						CrumblingEvent.Orientation = CrumblingDataItem.Orientation;
-						CrumblingEvent.LinearVelocity = CrumblingDataItem.LinearVelocity;
-						CrumblingEvent.AngularVelocity = CrumblingDataItem.AngularVelocity;
-						CrumblingEvent.Mass = static_cast<float>(CrumblingDataItem.Mass);
-						CrumblingEvent.LocalBounds = FBox(CrumblingDataItem.LocalBounds.Min(), CrumblingDataItem.LocalBounds.Max());
-						CrumblingEvent.Children = CrumblingDataItem.Children;
-					}
+					// queue them up so we can release the physics data before triggering BP events
+					FCrumblingChaosEvent& CrumblingEvent = PendingCrumblingEvent.AddZeroed_GetRef();
+					CrumblingEvent.Component = PrimComp;
+					CrumblingEvent.Location = CrumblingDataItem.Location;
+					CrumblingEvent.Orientation = CrumblingDataItem.Orientation;
+					CrumblingEvent.LinearVelocity = CrumblingDataItem.LinearVelocity;
+					CrumblingEvent.AngularVelocity = CrumblingDataItem.AngularVelocity;
+					CrumblingEvent.Mass = static_cast<float>(CrumblingDataItem.Mass);
+					CrumblingEvent.LocalBounds = FBox(CrumblingDataItem.LocalBounds.Min(), CrumblingDataItem.LocalBounds.Max());
+					CrumblingEvent.Children = CrumblingDataItem.Children;
 				}
 			}
 		}
diff --git a/Engine/Source/Runtime/Experimental/ChaosSolverEngine/Public/Chaos/ChaosGameplayEventDispatcher.h b/Engine/Source/Runtime/Experimental/ChaosSolverEngine/Public/Chaos/ChaosGameplayEventDispatcher.h
--- a/Engine/Source/Runtime/Experimental/ChaosSolverEngine/Public/Chaos/ChaosGameplayEventDispatcher.h
+++ b/Engine/Source/Runtime/Experimental/ChaosSolverEngine/Public/Chaos/ChaosGameplayEventDispatcher.h
@@ -20,32 +20,9 @@ namespace Chaos
 	struct FBreakingData;
 }
 
-// FChaosBreakEvent has been moved from ChaosSolverEngine to Engine, this structure need to be kept for backward compatibility.
-USTRUCT(BlueprintType)
-struct CHAOSSOLVERENGINE_API FChaosBreakEvent : public FBreakChaosEvent
-{
-	GENERATED_BODY()
-	FChaosBreakEvent() : FBreakChaosEvent() {}
-	FChaosBreakEvent(const FBreakChaosEvent& BreakChaosEvent) : FBreakChaosEvent(BreakChaosEvent) {}
-	FChaosBreakEvent(const Chaos::FBreakingData& BreakingData) : FBreakChaosEvent(BreakingData) {}
-};
-
-USTRUCT(BlueprintType)
-struct CHAOSSOLVERENGINE_API FChaosRemovalEvent : public FRemovalChaosEvent
-{
-	GENERATED_BODY()
-	FChaosRemovalEvent() : FRemovalChaosEvent() {}
-	FChaosRemovalEvent(const FRemovalChaosEvent& RemovalChaosEvent) : FRemovalChaosEvent(RemovalChaosEvent) {}
-};
-
-USTRUCT(BlueprintType)
-struct CHAOSSOLVERENGINE_API FChaosCrumblingEvent : public FCrumblingChaosEvent
-{
-	GENERATED_BODY()
-	FChaosCrumblingEvent() : FCrumblingChaosEvent() {}
-	FChaosCrumblingEvent(const FCrumblingChaosEvent& CrumbleChaosEvent) : FCrumblingChaosEvent(CrumbleChaosEvent) {}
-};
-
+typedef FBreakChaosEvent FChaosBreakEvent;
+typedef FRemovalChaosEvent FChaosRemovalEvent;
+typedef FCrumblingChaosEvent FChaosCrumblingEvent;
 
 typedef TFunction<void(const FBreakChaosEvent&)> FOnBreakEventCallback;
 
diff --git a/Engine/Source/Runtime/Experimental/GeometryCollectionEngine/Private/GeometryCollection/GeometryCollectionComponent.cpp b/Engine/Source/Runtime/Experimental/GeometryCollectionEngine/Private/GeometryCollection/GeometryCollectionComponent.cpp
--- a/Engine/Source/Runtime/Experimental/GeometryCollectionEngine/Private/GeometryCollection/GeometryCollectionComponent.cpp
+++ b/Engine/Source/Runtime/Experimental/GeometryCollectionEngine/Private/GeometryCollection/GeometryCollectionComponent.cpp
@@ -1016,7 +1016,7 @@ void UGeometryCollectionComponent::DispatchBreakEvent(const FBreakChaosEvent& Ev
 	// bp
 	if (OnChaosBreakEvent.IsBound())
 	{
-		OnChaosBreakEvent.Broadcast(FChaosBreakEvent(Event));
+		OnChaosBreakEvent.Broadcast(Event);
 	}
 }
 
@@ -1028,7 +1028,7 @@ void UGeometryCollectionComponent::DispatchRemovalEvent(const FRemovalChaosEvent
 	// bp
 	if (OnChaosRemovalEvent.IsBound())
 	{
-		OnChaosRemovalEvent.Broadcast(FChaosRemovalEvent(Event));
+		OnChaosRemovalEvent.Broadcast(Event);
 	}
 }
 
@@ -1037,7 +1037,7 @@ void UGeometryCollectionComponent::DispatchCrumblingEvent(const FCrumblingChaosE
 	// bp
 	if (OnChaosCrumblingEvent.IsBound())
 	{
-		OnChaosCrumblingEvent.Broadcast(FChaosCrumblingEvent(Event));
+		OnChaosCrumblingEvent.Broadcast(Event);
 	}
 }
 
@@ -3172,7 +3172,7 @@ void UGeometryCollectionComponent::OnPostPhysicsSync()
 			{
 				if (OnRootBreakEvent.IsBound())
 				{
-					FChaosBreakEvent Event;
+					FBreakChaosEvent Event;
 					Event.Index = GetRootIndex();
 					Event.Component = this;
 					OnRootBreakEvent.Broadcast(Event);
diff --git a/Engine/Source/Runtime/Experimental/GeometryCollectionEngine/Public/GeometryCollection/GeometryCollectionComponent.h b/Engine/Source/Runtime/Experimental/GeometryCollectionEngine/Public/GeometryCollection/GeometryCollectionComponent.h
--- a/Engine/Source/Runtime/Experimental/GeometryCollectionEngine/Public/GeometryCollection/GeometryCollectionComponent.h
+++ b/Engine/Source/Runtime/Experimental/GeometryCollectionEngine/Public/GeometryCollection/GeometryCollectionComponent.h
@@ -56,12 +56,12 @@ enum class EObjectStateTypeEnum : uint8;
 namespace Chaos { enum class EObjectStateType: int8; }
 template<class InElementType> class TManagedArray;
 
-// To deprecate:  to use FBreakChaosEvent from Engine instead of FChaosBreakEvent from ChaosSolverEngine
-DECLARE_DYNAMIC_MULTICAST_DELEGATE_OneParam(FOnChaosBreakEvent, const FChaosBreakEvent&, BreakEvent);
 
-DECLARE_DYNAMIC_MULTICAST_DELEGATE_OneParam(FOnChaosRemovalEvent, const FChaosRemovalEvent&, RemovalEvent);
+DECLARE_DYNAMIC_MULTICAST_DELEGATE_OneParam(FOnChaosBreakEvent, const FBreakChaosEvent&, BreakEvent);
 
-DECLARE_DYNAMIC_MULTICAST_DELEGATE_OneParam(FOnChaosCrumblingEvent, const FChaosCrumblingEvent&, CrumbleEvent);
+DECLARE_DYNAMIC_MULTICAST_DELEGATE_OneParam(FOnChaosRemovalEvent, const FRemovalChaosEvent&, RemovalEvent);
+
+DECLARE_DYNAMIC_MULTICAST_DELEGATE_OneParam(FOnChaosCrumblingEvent, const FCrumblingChaosEvent&, CrumbleEvent);
 
 namespace GeometryCollection
 {
