fix cluster union update child to parent temporary lock to work even with deferred updates

diff --git a/Engine/Source/Runtime/Experimental/Chaos/Private/Chaos/ClusterUnionManager.cpp b/Engine/Source/Runtime/Experimental/Chaos/Private/Chaos/ClusterUnionManager.cpp
--- a/Engine/Source/Runtime/Experimental/Chaos/Private/Chaos/ClusterUnionManager.cpp
+++ b/Engine/Source/Runtime/Experimental/Chaos/Private/Chaos/ClusterUnionManager.cpp
@@ -575,6 +575,16 @@ namespace Chaos
 		{
 			FlushIncrementalConnectivityGraphOperations(ClusterUnion);
 		}
+
+		for (FPBDRigidClusteredParticleHandle* ChildParticle : PendingParticlesToUndoChildToParentLock)
+		{
+			if (ChildParticle)
+			{
+				ChildParticle->SetChildToParentLocked(false);
+			}
+		}
+
+		PendingParticlesToUndoChildToParentLock.Empty();
 	}
 
 	DECLARE_CYCLE_STAT(TEXT("FClusterUnionManager::FlushIncrementalConnectivityGraphOperations"), STAT_FlushIncrementalConnectivityGraphOperations, STATGROUP_Chaos);
@@ -764,9 +774,8 @@ namespace Chaos
 		// We need to keep track of all the particles we touched here. We need this because UpdateClusterUnionParticlesChildToParent
 		// is an *authoritative* update on the ChildToParent of the particle. However, UpdateAllClusterUnionProperties in certain cases
 		// may try to recompute the ChildToParent using the position of the particle. To counteract this, we will force the ChildToParent
-		// to be *temporarily* locked for the duration of UpdateAllClusterUnionProperties. However, unless bLock is true, we will restore the
+		// to be *temporarily* locked for the duration of the *next* UpdateAllClusterUnionProperties. However, unless bLock is true, we will restore the
 		// lock state of the particle to what it was previously.
-		TArray<TPair<FPBDRigidClusteredParticleHandle*, bool>> DirtyParticleLockStates;
 		for (FPBDRigidParticleHandle* Particle : Particles)
 		{
 			if (!ensure(Particle))
@@ -783,9 +792,9 @@ namespace Chaos
 					{
 						ChildHandle->SetChildToParent(Update->ChildToParent);
 
-						if (!Update->bLock)
+						if (!Update->bLock && !ChildHandle->IsChildToParentLocked())
 						{
-							DirtyParticleLockStates.Add({ ChildHandle, ChildHandle->IsChildToParentLocked() });
+							PendingParticlesToUndoChildToParentLock.Add(ChildHandle);
 						}
 						ChildHandle->SetChildToParentLocked(true);
 						PendingChildToParentUpdates.Remove(ChildHandle);
@@ -800,11 +809,6 @@ namespace Chaos
 		}
 
 		RequestDeferredClusterPropertiesUpdate(ClusterIndex, EUpdateClusterUnionPropertiesFlags::IncrementalGenerateConnectionGraph);
-
-		for (TPair<FPBDRigidClusteredParticleHandle*, bool>& Pair : DirtyParticleLockStates)
-		{
-			Pair.Key->SetChildToParentLocked(Pair.Value);
-		}
 	}
 
 	void FClusterUnionManager::RequestDeferredClusterPropertiesUpdate(FClusterUnionIndex ClusterIndex, EUpdateClusterUnionPropertiesFlags Flags)
diff --git a/Engine/Source/Runtime/Experimental/Chaos/Public/Chaos/ClusterUnionManager.h b/Engine/Source/Runtime/Experimental/Chaos/Public/Chaos/ClusterUnionManager.h
--- a/Engine/Source/Runtime/Experimental/Chaos/Public/Chaos/ClusterUnionManager.h
+++ b/Engine/Source/Runtime/Experimental/Chaos/Public/Chaos/ClusterUnionManager.h
@@ -175,6 +175,7 @@ namespace Chaos
 		TClusterIndexOpMap<FClusterUnionIndex> PendingClusterIndexOperations;
 		TClusterIndexOpMap<FClusterUnionExplicitIndex> PendingExplicitIndexOperations;
 		TMap<FPBDRigidClusteredParticleHandle*, FClusterUnionChildToParentUpdate> PendingChildToParentUpdates;
+		TSet<FPBDRigidClusteredParticleHandle*> PendingParticlesToUndoChildToParentLock;
 
 		template<typename TIndex>
 		void AddPendingOperation(TClusterIndexOpMap<TIndex>& OpMap, TIndex Index, EClusterUnionOperation Op, const TArray<FPBDRigidParticleHandle*>& Particles)
