fix issue with fracture dynamic collection not matching rest collection, causing an ensure on fracture reset: change defaults to update dynamic collection, and make sure the dynamic is also reset if the rest is updated so they do not go out of sync

diff --git a/Engine/Source/Runtime/Experimental/GeometryCollectionEngine/Private/GeometryCollection/GeometryCollectionComponent.cpp b/Engine/Source/Runtime/Experimental/GeometryCollectionEngine/Private/GeometryCollection/GeometryCollectionComponent.cpp
--- a/Engine/Source/Runtime/Experimental/GeometryCollectionEngine/Private/GeometryCollection/GeometryCollectionComponent.cpp
+++ b/Engine/Source/Runtime/Experimental/GeometryCollectionEngine/Private/GeometryCollection/GeometryCollectionComponent.cpp
@@ -3718,7 +3718,8 @@ FGeometryCollectionEdit::~FGeometryCollectionEdit()
 
 		auto ProcessComponent = [this](UGeometryCollectionComponent* ToProcess)
 		{
-			if (EnumHasAnyFlags(EditUpdate, GeometryCollection::EEditUpdate::Dynamic))
+			// Note: Reset dynamic collection for both rest and dynamic updates, since we don't want the dynamic to be out of sync with the rest
+			if (EnumHasAnyFlags(EditUpdate, GeometryCollection::EEditUpdate::Dynamic | GeometryCollection::EEditUpdate::Rest))
 			{
 				ToProcess->ResetDynamicCollection();
 			}
diff --git a/Engine/Source/Runtime/Experimental/GeometryCollectionEngine/Public/GeometryCollection/GeometryCollectionComponent.h b/Engine/Source/Runtime/Experimental/GeometryCollectionEngine/Public/GeometryCollection/GeometryCollectionComponent.h
--- a/Engine/Source/Runtime/Experimental/GeometryCollectionEngine/Public/GeometryCollection/GeometryCollectionComponent.h
+++ b/Engine/Source/Runtime/Experimental/GeometryCollectionEngine/Public/GeometryCollection/GeometryCollectionComponent.h
@@ -208,7 +208,7 @@ public:
 	 * @param bShapeIsUnchanged				Override indicating the overall shape of the geometry and clusters is unchanged, even if the rest collection changed.  Useful to e.g., not re-compute convex hulls when we don't need to.
 	 * @param bPropagateAcrossComponents	Propagate updates to all components with the same underlying Rest Collection
 	 */
-	FGeometryCollectionEdit(UGeometryCollectionComponent* InComponent, GeometryCollection::EEditUpdate EditUpdate = GeometryCollection::EEditUpdate::RestPhysics, bool bShapeIsUnchanged = false, bool bPropagateToAllMatchingComponents = true);
+	FGeometryCollectionEdit(UGeometryCollectionComponent* InComponent, GeometryCollection::EEditUpdate EditUpdate = GeometryCollection::EEditUpdate::RestPhysicsDynamic, bool bShapeIsUnchanged = false, bool bPropagateToAllMatchingComponents = true);
 	~FGeometryCollectionEdit();
 
 	UGeometryCollection* GetRestCollection();
@@ -678,7 +678,7 @@ public:
 
 	FORCEINLINE const UGeometryCollection* GetRestCollection() const { return RestCollection; }
 
-	FORCEINLINE FGeometryCollectionEdit EditRestCollection(GeometryCollection::EEditUpdate EditUpdate = GeometryCollection::EEditUpdate::RestPhysics, bool bShapeIsUnchanged = false) { return FGeometryCollectionEdit(this, EditUpdate, bShapeIsUnchanged); }
+	FORCEINLINE FGeometryCollectionEdit EditRestCollection(GeometryCollection::EEditUpdate EditUpdate = GeometryCollection::EEditUpdate::RestPhysicsDynamic, bool bShapeIsUnchanged = false) { return FGeometryCollectionEdit(this, EditUpdate, bShapeIsUnchanged); }
 #if WITH_EDITOR
 	FORCEINLINE FScopedColorEdit EditBoneSelection(bool bForceUpdate = false) { return FScopedColorEdit(this, bForceUpdate); }
 
