fix crash on dataflow validation when output connection is not registered

diff --git a/Engine/Source/Runtime/Experimental/Dataflow/Core/Private/Dataflow/DataflowNode.cpp b/Engine/Source/Runtime/Experimental/Dataflow/Core/Private/Dataflow/DataflowNode.cpp
--- a/Engine/Source/Runtime/Experimental/Dataflow/Core/Private/Dataflow/DataflowNode.cpp
+++ b/Engine/Source/Runtime/Experimental/Dataflow/Core/Private/Dataflow/DataflowNode.cpp
@@ -350,9 +350,8 @@ bool FDataflowNode::ValidateConnections()
 						UE_LOG(LogChaos, Warning, TEXT("Missing dataflow RegisterOutputConnection in constructor for (%s:%s)"), *GetName().ToString(),*PropName.ToString());
 						bHasValidConnections = false;
 					}
-
-					// Validate passthrough connections if they exist
-					if (const FString* PassthroughName = Property->FindMetaData(FDataflowNode::DataflowPassthrough))
+					// If OutputConnection is valid, validate passthrough connections if they exist
+					else if (const FString* PassthroughName = Property->FindMetaData(FDataflowNode::DataflowPassthrough))
 					{
 						void* PassthroughConnectionAddress = OutputConnection->GetPassthroughRealAddress();
 						if(PassthroughConnectionAddress == nullptr)
