AncientGame fails to compile, Fix breaking API
FBreakChaosEvent was created in Engine to replaced FChaosBreakEvent event which was in CHAOSSOLVER_API and so no accessible in Engine for circular dependency issue. But in order to backward compatibility with blueprint delegate, and code the only easy way was to keep the same name. This CL is changing back those names to FChaosBreakEvent.
Same apply for Removal and Crumbling events.

diff --git a/Engine/Config/BaseEngine.ini b/Engine/Config/BaseEngine.ini
--- a/Engine/Config/BaseEngine.ini
+++ b/Engine/Config/BaseEngine.ini
@@ -367,9 +367,9 @@ DefaultWorldPartitionHLODLayer=/Engine/Maps/Templates/HLODs/HLODLayer_Instanced.
 +StructRedirects=(OldName="AnimNode_Slot",NewName="/Script/AnimGraphRuntime.AnimNode_Slot")
 +StructRedirects=(OldName="FormatTextArgument",NewName="/Script/Engine.FormatArgumentData")
 
-+StructRedirects=(OldName="ChaosBreakEvent",NewName="/Script/Engine.ChaosEventType.BreakChaosEvent")
-+StructRedirects=(OldName="ChaosRemovalEvent",NewName="/Script/Engine.ChaosEventType.RemovalChaosEvent")
-+StructRedirects=(OldName="ChaosCrumblingEvent",NewName="/Script/Engine.ChaosEventType.CrumblingChaosEvent")
++StructRedirects=(OldName="ChaosBreakEvent",NewName="/Script/Engine.ChaosEventType.ChaosBreakEvent")
++StructRedirects=(OldName="ChaosRemovalEvent",NewName="/Script/Engine.ChaosEventType.ChaosRemovalEvent")
++StructRedirects=(OldName="ChaosCrumblingEvent",NewName="/Script/Engine.ChaosEventType.ChaosCrumblingEvent")
 
 +FunctionRedirects=(OldName="ConvertTransformToRelative",NewName="/Script/Engine.KismetMathLibrary.MakeRelativeTransform")
 +PropertyRedirects=(OldName="KismetMathLibrary.MakeRelativeTransform.LocalTransform",NewName="A")
diff --git a/Engine/Source/Runtime/Engine/Private/PhysicsEngine/Experimental/ChaosEventRelay.cpp b/Engine/Source/Runtime/Engine/Private/PhysicsEngine/Experimental/ChaosEventRelay.cpp
--- a/Engine/Source/Runtime/Engine/Private/PhysicsEngine/Experimental/ChaosEventRelay.cpp
+++ b/Engine/Source/Runtime/Engine/Private/PhysicsEngine/Experimental/ChaosEventRelay.cpp
@@ -15,7 +15,7 @@ void UChaosEventRelay::DispatchPhysicsCollisionEvents(const TArray<FCollisionCha
 	}
 }
 
-void UChaosEventRelay::DispatchPhysicsBreakEvents(const TArray<FBreakChaosEvent>& BreakEvents)
+void UChaosEventRelay::DispatchPhysicsBreakEvents(const TArray<FChaosBreakEvent>& BreakEvents)
 {
 	if (OnBreakEvent.IsBound())
 	{
@@ -23,7 +23,7 @@ void UChaosEventRelay::DispatchPhysicsBreakEvents(const TArray<FBreakChaosEvent>
 	}
 }
 
-void UChaosEventRelay::DispatchPhysicsRemovalEvents(const TArray<FRemovalChaosEvent>& RemovalEvents)
+void UChaosEventRelay::DispatchPhysicsRemovalEvents(const TArray<FChaosRemovalEvent>& RemovalEvents)
 {
 	if (OnRemovalEvent.IsBound())
 	{
@@ -31,7 +31,7 @@ void UChaosEventRelay::DispatchPhysicsRemovalEvents(const TArray<FRemovalChaosEv
 	}
 }
 
-void UChaosEventRelay::DispatchPhysicsCrumblingEvents(const TArray<FCrumblingChaosEvent>& CrumblingEvents)
+void UChaosEventRelay::DispatchPhysicsCrumblingEvents(const TArray<FChaosCrumblingEvent>& CrumblingEvents)
 {
 	if (OnCrumblingEvent.IsBound())
 	{
diff --git a/Engine/Source/Runtime/Engine/Private/PhysicsEngine/Experimental/ChaosEventType.cpp b/Engine/Source/Runtime/Engine/Private/PhysicsEngine/Experimental/ChaosEventType.cpp
--- a/Engine/Source/Runtime/Engine/Private/PhysicsEngine/Experimental/ChaosEventType.cpp
+++ b/Engine/Source/Runtime/Engine/Private/PhysicsEngine/Experimental/ChaosEventType.cpp
@@ -44,7 +44,7 @@ FCollisionChaosEvent::FCollisionChaosEvent(const Chaos::FCollidingData& Collisio
 
 }
 
-FBreakChaosEvent::FBreakChaosEvent()
+FChaosBreakEvent::FChaosBreakEvent()
 	: Component(nullptr)
 	, Location(FVector::ZeroVector)
 	, Velocity(FVector::ZeroVector)
@@ -56,7 +56,7 @@ FBreakChaosEvent::FBreakChaosEvent()
 {
 }
 
-FBreakChaosEvent::FBreakChaosEvent(const Chaos::FBreakingData& BreakingData)
+FChaosBreakEvent::FChaosBreakEvent(const Chaos::FBreakingData& BreakingData)
 	: Component(nullptr)
 	, Location(BreakingData.Location)
 	, Velocity(BreakingData.Velocity)
@@ -68,7 +68,7 @@ FBreakChaosEvent::FBreakChaosEvent(const Chaos::FBreakingData& BreakingData)
 {
 }
 
-FRemovalChaosEvent::FRemovalChaosEvent()
+FChaosRemovalEvent::FChaosRemovalEvent()
 	: Component(nullptr)
 	, Location(FVector::ZeroVector)
 	, Mass(0.0f)
diff --git a/Engine/Source/Runtime/Engine/Private/PhysicsEngine/Experimental/PhysScene_Chaos.cpp b/Engine/Source/Runtime/Engine/Private/PhysicsEngine/Experimental/PhysScene_Chaos.cpp
--- a/Engine/Source/Runtime/Engine/Private/PhysicsEngine/Experimental/PhysScene_Chaos.cpp
+++ b/Engine/Source/Runtime/Engine/Private/PhysicsEngine/Experimental/PhysScene_Chaos.cpp
@@ -1081,13 +1081,13 @@ void FPhysScene_Chaos::HandleBreakingEvents(const Chaos::FBreakingEventData& Eve
 	{
 		LastBreakEventDispatchTime = Event.BreakingData.TimeCreated;
 		Chaos::FBreakingDataArray const& BreakingDataArray = Event.BreakingData.AllBreakingsArray;
-		TArray<FBreakChaosEvent> PendingBreakEvents;
+		TArray<FChaosBreakEvent> PendingBreakEvents;
 		for (const Chaos::FBreakingData& BreakingData: BreakingDataArray)
 		{
 			if ((BreakingData.EmitterFlag & Chaos::EventEmitterFlag::GlobalDispatcher) && BreakingData.Proxy)
 			{
 				UPrimitiveComponent* Comp = GetOwningComponent<UPrimitiveComponent>(BreakingData.Proxy);
-				FBreakChaosEvent& BreakEvent = PendingBreakEvents.Emplace_GetRef(BreakingData);
+				FChaosBreakEvent& BreakEvent = PendingBreakEvents.Emplace_GetRef(BreakingData);
 				BreakEvent.Component = Comp;
 			}
 		}
@@ -1106,7 +1106,7 @@ void FPhysScene_Chaos::HandleRemovalEvents(const Chaos::FRemovalEventData& Event
 	if (RemovalDataTimestamp > LastRemovalEventDispatchTime)
 	{
 		LastRemovalEventDispatchTime = RemovalDataTimestamp;
-		TArray<FRemovalChaosEvent> PendingRemovalEvents;
+		TArray<FChaosRemovalEvent> PendingRemovalEvents;
 		Chaos::FRemovalDataArray const& RemovalData = Event.RemovalData.AllRemovalArray;
 		for (Chaos::FRemovalData const& RemovalDataItem : RemovalData)
 		{
@@ -1115,7 +1115,7 @@ void FPhysScene_Chaos::HandleRemovalEvents(const Chaos::FRemovalEventData& Event
 				UPrimitiveComponent* Comp = GetOwningComponent<UPrimitiveComponent>(RemovalDataItem.Proxy);
 				if (GlobalRemovalEventRegistrations.Contains(Comp))
 				{
-					FRemovalChaosEvent& RemovalEvent = PendingRemovalEvents.AddZeroed_GetRef();
+					FChaosRemovalEvent& RemovalEvent = PendingRemovalEvents.AddZeroed_GetRef();
 					RemovalEvent.Component = Comp;
 					RemovalEvent.Location = RemovalDataItem.Location;
 					RemovalEvent.Mass = RemovalDataItem.Mass;
@@ -1137,12 +1137,12 @@ void FPhysScene_Chaos::HandleCrumblingEvents(const Chaos::FCrumblingEventData& E
 	if (Event.CrumblingData.bHasGlobalEvent && CrumblingDataTimestamp > LastCrumblingEventDispatchTime)
 	{
 		LastCrumblingEventDispatchTime = CrumblingDataTimestamp;
-		TArray<FCrumblingChaosEvent> PendingCrumblingEvent;
+		TArray<FChaosCrumblingEvent> PendingCrumblingEvent;
 		for (const Chaos::FCrumblingData& CrumblingDataItem : Event.CrumblingData.AllCrumblingsArray)
 		{
 			if ((CrumblingDataItem.EmitterFlag & Chaos::EventEmitterFlag::GlobalDispatcher) && CrumblingDataItem.Proxy)
 			{
-				FCrumblingChaosEvent& CrumblingEvent = PendingCrumblingEvent.AddZeroed_GetRef();
+				FChaosCrumblingEvent& CrumblingEvent = PendingCrumblingEvent.AddZeroed_GetRef();
 				CrumblingEvent.Component = GetOwningComponent<UPrimitiveComponent>(CrumblingDataItem.Proxy);
 				CrumblingEvent.Location = CrumblingDataItem.Location;
 				CrumblingEvent.Orientation = CrumblingDataItem.Orientation;
diff --git a/Engine/Source/Runtime/Engine/Public/Physics/Experimental/ChaosEventRelay.h b/Engine/Source/Runtime/Engine/Public/Physics/Experimental/ChaosEventRelay.h
--- a/Engine/Source/Runtime/Engine/Public/Physics/Experimental/ChaosEventRelay.h
+++ b/Engine/Source/Runtime/Engine/Public/Physics/Experimental/ChaosEventRelay.h
@@ -12,9 +12,9 @@
 
 
 DECLARE_DYNAMIC_MULTICAST_DELEGATE_OneParam(FCollisionEventSignature, const TArray<FCollisionChaosEvent>&, CollisionEvents);
-DECLARE_DYNAMIC_MULTICAST_DELEGATE_OneParam(FBreakEventSignature, const TArray<FBreakChaosEvent>&, BreakEvents);
-DECLARE_DYNAMIC_MULTICAST_DELEGATE_OneParam(FRemovalEventSignature, const TArray<FRemovalChaosEvent>&, RemovalEvents);
-DECLARE_DYNAMIC_MULTICAST_DELEGATE_OneParam(FCrumblingEventSignature, const TArray<FCrumblingChaosEvent>&, CrumblingEvents);
+DECLARE_DYNAMIC_MULTICAST_DELEGATE_OneParam(FBreakEventSignature, const TArray<FChaosBreakEvent>&, BreakEvents);
+DECLARE_DYNAMIC_MULTICAST_DELEGATE_OneParam(FRemovalEventSignature, const TArray<FChaosRemovalEvent>&, RemovalEvents);
+DECLARE_DYNAMIC_MULTICAST_DELEGATE_OneParam(FCrumblingEventSignature, const TArray<FChaosCrumblingEvent>&, CrumblingEvents);
 
 
 /**
@@ -30,9 +30,9 @@ public:
 	UChaosEventRelay();
 
 	void DispatchPhysicsCollisionEvents(const TArray<FCollisionChaosEvent>& CollisionEvents);
-	void DispatchPhysicsBreakEvents(const TArray<FBreakChaosEvent>& BreakEvents);
-	void DispatchPhysicsRemovalEvents(const TArray<FRemovalChaosEvent>& RemovalEvents);
-	void DispatchPhysicsCrumblingEvents(const TArray<FCrumblingChaosEvent>& CrumblingEvents);
+	void DispatchPhysicsBreakEvents(const TArray<FChaosBreakEvent>& BreakEvents);
+	void DispatchPhysicsRemovalEvents(const TArray<FChaosRemovalEvent>& RemovalEvents);
+	void DispatchPhysicsCrumblingEvents(const TArray<FChaosCrumblingEvent>& CrumblingEvents);
 
 	
 	UPROPERTY(BlueprintAssignable, Category = "Physics Events")
diff --git a/Engine/Source/Runtime/Engine/Public/Physics/Experimental/ChaosEventType.h b/Engine/Source/Runtime/Engine/Public/Physics/Experimental/ChaosEventType.h
--- a/Engine/Source/Runtime/Engine/Public/Physics/Experimental/ChaosEventType.h
+++ b/Engine/Source/Runtime/Engine/Public/Physics/Experimental/ChaosEventType.h
@@ -80,14 +80,14 @@ public:
 
 
 USTRUCT(BlueprintType)
-struct ENGINE_API FBreakChaosEvent
+struct ENGINE_API FChaosBreakEvent
 {
 	GENERATED_BODY()
 
 public:
 
-	FBreakChaosEvent();
-	FBreakChaosEvent(const Chaos::FBreakingData& BreakingData);
+	FChaosBreakEvent();
+	FChaosBreakEvent(const Chaos::FBreakingData& BreakingData);
 
 	/** primitive component involved in the break event */
 	UPROPERTY(BlueprintReadOnly, Category = "Break Event")
@@ -123,13 +123,13 @@ public:
 };
 
 USTRUCT(BlueprintType)
-struct ENGINE_API FRemovalChaosEvent
+struct ENGINE_API FChaosRemovalEvent
 {
 	GENERATED_BODY()
 
 public:
 
-	FRemovalChaosEvent();
+	FChaosRemovalEvent();
 
 	UPROPERTY(BlueprintReadOnly, Category = "Removal Event")
 	TObjectPtr<UPrimitiveComponent> Component = nullptr;
@@ -143,12 +143,12 @@ public:
 
 
 USTRUCT(BlueprintType)
-struct ENGINE_API FCrumblingChaosEvent
+struct ENGINE_API FChaosCrumblingEvent
 {
 	GENERATED_BODY()
 
 public:
-	FCrumblingChaosEvent()
+	FChaosCrumblingEvent()
 		: Component(nullptr)
 		, Location(FVector::ZeroVector)
 		, Orientation(FQuat::Identity)
diff --git a/Engine/Source/Runtime/Experimental/ChaosSolverEngine/Private/Chaos/ChaosGameplayEventDispatcher.cpp b/Engine/Source/Runtime/Experimental/ChaosSolverEngine/Private/Chaos/ChaosGameplayEventDispatcher.cpp
--- a/Engine/Source/Runtime/Experimental/ChaosSolverEngine/Private/Chaos/ChaosGameplayEventDispatcher.cpp
+++ b/Engine/Source/Runtime/Experimental/ChaosSolverEngine/Private/Chaos/ChaosGameplayEventDispatcher.cpp
@@ -42,9 +42,9 @@ void UChaosGameplayEventDispatcher::OnUnregister()
 }
 
 // internal
-static void DispatchPendingBreakEvents(TArray<FBreakChaosEvent> const& Events, TMap<TObjectPtr<UPrimitiveComponent>, FBreakEventCallbackWrapper> const& Registrations)
+static void DispatchPendingBreakEvents(TArray<FChaosBreakEvent> const& Events, TMap<TObjectPtr<UPrimitiveComponent>, FBreakEventCallbackWrapper> const& Registrations)
 {
-	for (FBreakChaosEvent const& E : Events)
+	for (FChaosBreakEvent const& E : Events)
 	{
 		if (E.Component)
 		{
@@ -57,9 +57,9 @@ static void DispatchPendingBreakEvents(TArray<FBreakChaosEvent> const& Events, T
 	}
 }
 
-static void DispatchPendingRemovalEvents(TArray<FRemovalChaosEvent> const& Events, TMap<TObjectPtr<UPrimitiveComponent>, FRemovalEventCallbackWrapper> const& Registrations)
+static void DispatchPendingRemovalEvents(TArray<FChaosRemovalEvent> const& Events, TMap<TObjectPtr<UPrimitiveComponent>, FRemovalEventCallbackWrapper> const& Registrations)
 {
-	for (FRemovalChaosEvent const& E : Events)
+	for (FChaosRemovalEvent const& E : Events)
 	{
 		if (E.Component)
 		{
@@ -72,9 +72,9 @@ static void DispatchPendingRemovalEvents(TArray<FRemovalChaosEvent> const& Event
 	}
 }
 
-static void DispatchPendingCrumblingEvents(TArray<FCrumblingChaosEvent> const& Events, TMap<TObjectPtr<UPrimitiveComponent>, FCrumblingEventCallbackWrapper> const& Registrations)
+static void DispatchPendingCrumblingEvents(TArray<FChaosCrumblingEvent> const& Events, TMap<TObjectPtr<UPrimitiveComponent>, FCrumblingEventCallbackWrapper> const& Registrations)
 {
-	for (FCrumblingChaosEvent const& E : Events)
+	for (FChaosCrumblingEvent const& E : Events)
 	{
 		if (E.Component)
 		{
@@ -541,7 +541,7 @@ void UChaosGameplayEventDispatcher::HandleBreakingEvents(const Chaos::FBreakingE
 
 	// BREAK EVENTS
 
-	TArray<FBreakChaosEvent> PendingBreakEvents;
+	TArray<FChaosBreakEvent> PendingBreakEvents;
 
 	const float BreakingDataTimestamp = Event.BreakingData.TimeCreated;
 	if (BreakingDataTimestamp > LastBreakingDataTime)
@@ -564,7 +564,7 @@ void UChaosGameplayEventDispatcher::HandleBreakingEvents(const Chaos::FBreakingE
 					if (PrimComp)
 					{
 						// queue them up so we can release the physics data before trigging BP events
-						FBreakChaosEvent& BreakEvent = PendingBreakEvents.Emplace_GetRef(BreakingData);
+						FChaosBreakEvent& BreakEvent = PendingBreakEvents.Emplace_GetRef(BreakingData);
 						BreakEvent.Component = PrimComp;
 					}
 				}
@@ -605,7 +605,7 @@ void UChaosGameplayEventDispatcher::HandleRemovalEvents(const Chaos::FRemovalEve
 {
 	// REMOVAL EVENTS
 
-	TArray<FRemovalChaosEvent> PendingRemovalEvents;
+	TArray<FChaosRemovalEvent> PendingRemovalEvents;
 
 	const float RemovalDataTimestamp = Event.RemovalData.TimeCreated;
 	if (RemovalDataTimestamp > LastRemovalDataTime)
@@ -627,7 +627,7 @@ void UChaosGameplayEventDispatcher::HandleRemovalEvents(const Chaos::FRemovalEve
 					if (PrimComp && RemovalEventRegistrations.Contains(PrimComp))
 					{
 						// queue them up so we can release the physics data before trigging BP events
-						FRemovalChaosEvent& RemovalEvent = PendingRemovalEvents.AddZeroed_GetRef();
+						FChaosRemovalEvent& RemovalEvent = PendingRemovalEvents.AddZeroed_GetRef();
 						RemovalEvent.Component = PrimComp;
 						RemovalEvent.Location = RemovalDataItem.Location;
 						RemovalEvent.Mass = RemovalDataItem.Mass;
@@ -657,7 +657,7 @@ void UChaosGameplayEventDispatcher::HandleCrumblingEvents(const Chaos::FCrumblin
 		const FPhysScene& Scene = *(GetWorld()->GetPhysicsScene());
 
 		// let's assume crumbles are rare, so we will iterate breaks instead of registered components for now
-		TArray<FCrumblingChaosEvent> PendingCrumblingEvent;
+		TArray<FChaosCrumblingEvent> PendingCrumblingEvent;
 		for (const Chaos::FCrumblingData& CrumblingDataItem : Event.CrumblingData.AllCrumblingsArray)
 		{	
 			if ((CrumblingDataItem.EmitterFlag & Chaos::EventEmitterFlag::OwnDispatcher) && CrumblingDataItem.Proxy)
@@ -665,7 +665,7 @@ void UChaosGameplayEventDispatcher::HandleCrumblingEvents(const Chaos::FCrumblin
 				if (UPrimitiveComponent* const PrimComp = Scene.GetOwningComponent<UPrimitiveComponent>(CrumblingDataItem.Proxy))
 				{
 					// queue them up so we can release the physics data before triggering BP events
-					FCrumblingChaosEvent& CrumblingEvent = PendingCrumblingEvent.AddZeroed_GetRef();
+					FChaosCrumblingEvent& CrumblingEvent = PendingCrumblingEvent.AddZeroed_GetRef();
 					CrumblingEvent.Component = PrimComp;
 					CrumblingEvent.Location = CrumblingDataItem.Location;
 					CrumblingEvent.Orientation = CrumblingDataItem.Orientation;
diff --git a/Engine/Source/Runtime/Experimental/ChaosSolverEngine/Public/Chaos/ChaosGameplayEventDispatcher.h b/Engine/Source/Runtime/Experimental/ChaosSolverEngine/Public/Chaos/ChaosGameplayEventDispatcher.h
--- a/Engine/Source/Runtime/Experimental/ChaosSolverEngine/Public/Chaos/ChaosGameplayEventDispatcher.h
+++ b/Engine/Source/Runtime/Experimental/ChaosSolverEngine/Public/Chaos/ChaosGameplayEventDispatcher.h
@@ -20,11 +20,7 @@ namespace Chaos
 	struct FBreakingData;
 }
 
-typedef FBreakChaosEvent FChaosBreakEvent;
-typedef FRemovalChaosEvent FChaosRemovalEvent;
-typedef FCrumblingChaosEvent FChaosCrumblingEvent;
-
-typedef TFunction<void(const FBreakChaosEvent&)> FOnBreakEventCallback;
+typedef TFunction<void(const FChaosBreakEvent&)> FOnBreakEventCallback;
 
 /** UStruct wrapper so we can store the TFunction in a TMap */
 USTRUCT()
@@ -36,7 +32,7 @@ public:
 	FOnBreakEventCallback BreakEventCallback;
 };
 
-typedef TFunction<void(const FRemovalChaosEvent&)> FOnRemovalEventCallback;
+typedef TFunction<void(const FChaosRemovalEvent&)> FOnRemovalEventCallback;
 
 /** UStruct wrapper so we can store the TFunction in a TMap */
 USTRUCT()
@@ -48,7 +44,7 @@ public:
 	FOnRemovalEventCallback RemovalEventCallback;
 };
 
-typedef TFunction<void(const FCrumblingChaosEvent&)> FOnCrumblingEventCallback;
+typedef TFunction<void(const FChaosCrumblingEvent&)> FOnCrumblingEventCallback;
 
 /** UStruct wrapper so we can store the TFunction in a TMap */
 USTRUCT()
diff --git a/Engine/Source/Runtime/Experimental/GeometryCollectionEngine/Private/GeometryCollection/GeometryCollectionComponent.cpp b/Engine/Source/Runtime/Experimental/GeometryCollectionEngine/Private/GeometryCollection/GeometryCollectionComponent.cpp
--- a/Engine/Source/Runtime/Experimental/GeometryCollectionEngine/Private/GeometryCollection/GeometryCollectionComponent.cpp
+++ b/Engine/Source/Runtime/Experimental/GeometryCollectionEngine/Private/GeometryCollection/GeometryCollectionComponent.cpp
@@ -1008,7 +1008,7 @@ void UGeometryCollectionComponent::AddTorqueInRadians(FVector Torque, FName Bone
 	DispatchFieldCommand(Command);
 }
 
-void UGeometryCollectionComponent::DispatchBreakEvent(const FBreakChaosEvent& Event)
+void UGeometryCollectionComponent::DispatchBreakEvent(const FChaosBreakEvent& Event)
 {
 	// native
 	NotifyBreak(Event);
@@ -1020,7 +1020,7 @@ void UGeometryCollectionComponent::DispatchBreakEvent(const FBreakChaosEvent& Ev
 	}
 }
 
-void UGeometryCollectionComponent::DispatchRemovalEvent(const FRemovalChaosEvent& Event)
+void UGeometryCollectionComponent::DispatchRemovalEvent(const FChaosRemovalEvent& Event)
 {
 	// native
 	NotifyRemoval(Event);
@@ -1032,7 +1032,7 @@ void UGeometryCollectionComponent::DispatchRemovalEvent(const FRemovalChaosEvent
 	}
 }
 
-void UGeometryCollectionComponent::DispatchCrumblingEvent(const FCrumblingChaosEvent& Event)
+void UGeometryCollectionComponent::DispatchCrumblingEvent(const FChaosCrumblingEvent& Event)
 {
 	// bp
 	if (OnChaosCrumblingEvent.IsBound())
@@ -1473,7 +1473,7 @@ void UGeometryCollectionComponent::PostEditChangeProperty(FPropertyChangedEvent&
 }
 #endif
 
-static void DispatchGeometryCollectionBreakEvent(const FBreakChaosEvent& Event)
+static void DispatchGeometryCollectionBreakEvent(const FChaosBreakEvent& Event)
 {
 	if (UGeometryCollectionComponent* const GC = Cast<UGeometryCollectionComponent>(Event.Component))
 	{
@@ -1481,7 +1481,7 @@ static void DispatchGeometryCollectionBreakEvent(const FBreakChaosEvent& Event)
 	}
 }
 
-static void DispatchGeometryCollectionRemovalEvent(const FRemovalChaosEvent& Event)
+static void DispatchGeometryCollectionRemovalEvent(const FChaosRemovalEvent& Event)
 {
 	if (UGeometryCollectionComponent* const GC = Cast<UGeometryCollectionComponent>(Event.Component))
 	{
@@ -1489,7 +1489,7 @@ static void DispatchGeometryCollectionRemovalEvent(const FRemovalChaosEvent& Eve
 	}
 }
 
-static void DispatchGeometryCollectionCrumblingEvent(const FCrumblingChaosEvent& Event)
+static void DispatchGeometryCollectionCrumblingEvent(const FChaosCrumblingEvent& Event)
 {
 	if (UGeometryCollectionComponent* const GC = Cast<UGeometryCollectionComponent>(Event.Component))
 	{
@@ -3172,7 +3172,7 @@ void UGeometryCollectionComponent::OnPostPhysicsSync()
 			{
 				if (OnRootBreakEvent.IsBound())
 				{
-					FBreakChaosEvent Event;
+					FChaosBreakEvent Event;
 					Event.Index = GetRootIndex();
 					Event.Component = this;
 					OnRootBreakEvent.Broadcast(Event);
diff --git a/Engine/Source/Runtime/Experimental/GeometryCollectionEngine/Public/GeometryCollection/GeometryCollectionComponent.h b/Engine/Source/Runtime/Experimental/GeometryCollectionEngine/Public/GeometryCollection/GeometryCollectionComponent.h
--- a/Engine/Source/Runtime/Experimental/GeometryCollectionEngine/Public/GeometryCollection/GeometryCollectionComponent.h
+++ b/Engine/Source/Runtime/Experimental/GeometryCollectionEngine/Public/GeometryCollection/GeometryCollectionComponent.h
@@ -57,11 +57,11 @@ namespace Chaos { enum class EObjectStateType: int8; }
 template<class InElementType> class TManagedArray;
 
 
-DECLARE_DYNAMIC_MULTICAST_DELEGATE_OneParam(FOnChaosBreakEvent, const FBreakChaosEvent&, BreakEvent);
+DECLARE_DYNAMIC_MULTICAST_DELEGATE_OneParam(FOnChaosBreakEvent, const FChaosBreakEvent&, BreakEvent);
 
-DECLARE_DYNAMIC_MULTICAST_DELEGATE_OneParam(FOnChaosRemovalEvent, const FRemovalChaosEvent&, RemovalEvent);
+DECLARE_DYNAMIC_MULTICAST_DELEGATE_OneParam(FOnChaosRemovalEvent, const FChaosRemovalEvent&, RemovalEvent);
 
-DECLARE_DYNAMIC_MULTICAST_DELEGATE_OneParam(FOnChaosCrumblingEvent, const FCrumblingChaosEvent&, CrumbleEvent);
+DECLARE_DYNAMIC_MULTICAST_DELEGATE_OneParam(FOnChaosCrumblingEvent, const FChaosCrumblingEvent&, CrumbleEvent);
 
 namespace GeometryCollection
 {
@@ -1009,10 +1009,10 @@ public:
 	void SetNotifyGlobalCrumblings(bool bNewNotifyGlobalCrumblings, bool bGlobalNewCrumblingEventIncludesChildren);
 	
 	/** Overrideable native notification */
-	virtual void NotifyBreak(const FBreakChaosEvent& Event) {};
+	virtual void NotifyBreak(const FChaosBreakEvent& Event) {};
 
 	/** Overrideable native notification */
-	virtual void NotifyRemoval(const FRemovalChaosEvent& Event) {};
+	virtual void NotifyRemoval(const FChaosRemovalEvent& Event) {};
 
 	UPROPERTY(BlueprintAssignable, Category = "Chaos")
 	FOnChaosBreakEvent OnChaosBreakEvent;
@@ -1026,11 +1026,11 @@ public:
 	// todo(chaos) remove when no longer necessary
 	FOnChaosBreakEvent OnRootBreakEvent;
 
-	void DispatchBreakEvent(const FBreakChaosEvent& Event);
+	void DispatchBreakEvent(const FChaosBreakEvent& Event);
 
-	void DispatchRemovalEvent(const FRemovalChaosEvent& Event);
+	void DispatchRemovalEvent(const FChaosRemovalEvent& Event);
 
-	void DispatchCrumblingEvent(const FCrumblingChaosEvent& Event);
+	void DispatchCrumblingEvent(const FChaosCrumblingEvent& Event);
 
 	UPROPERTY(Transient, VisibleAnywhere, BlueprintReadWrite, Interp, Category = "Chaos")
 	float DesiredCacheTime;
