fix an issue where GC particles with intercluster edges added into a cluster union would not break due to strain/collision

diff --git a/Engine/Source/Runtime/Experimental/Chaos/Private/Chaos/PBDRigidClustering.cpp b/Engine/Source/Runtime/Experimental/Chaos/Private/Chaos/PBDRigidClustering.cpp
--- a/Engine/Source/Runtime/Experimental/Chaos/Private/Chaos/PBDRigidClustering.cpp
+++ b/Engine/Source/Runtime/Experimental/Chaos/Private/Chaos/PBDRigidClustering.cpp
@@ -1770,7 +1770,7 @@ namespace Chaos
 			}
 
 			auto ComputeStrainLambda = [this, &ContactHandle, &InvDt](
-				const FPBDRigidClusteredParticleHandle* Cluster,
+				FPBDRigidClusteredParticleHandle* Cluster,
 				FRealSingle& OutTotalImpulseAccumulator)
 			{
 				const FVec3 ContactWorldLocation = ContactHandle->GetContact().CalculateWorldContactLocation();
@@ -1818,7 +1818,17 @@ namespace Chaos
 						// technically this is no longer an impulse but we intend in the future to get rid of COllisionImpulse property and use external strain instead
 						CollisionImpulseToApply *= InvDt;
 					}
-					if (bUseDamagePropagation)
+
+					if (Cluster->ClusterIds().NumChildren == 0)
+					{
+						// Special case that can happen when we're looking to hit a cluster union with partial destruction as a result of intercluster edges.
+						// We might hit a leaf particle in a geometry collection which has 0 children. In that case, we actually just want to apply strain to the
+						// particle itself rather than trying to find its children.
+						Cluster->CollisionImpulses() += CollisionImpulseToApply;
+						UpdateTopLevelParticle(Cluster);
+						OutTotalImpulseAccumulator += CollisionImpulseToApply;
+					}
+					else if (bUseDamagePropagation)
 					{
 						// propagation based breaking model start from the closest particle and propagate through the connection graph
 						// propagation logic is dealt when evaluating the strain
