Chaos : GC performance : avoid large copies for geometry collection results
- Geometry collection result is now held as a ref counter pointer

diff --git a/Engine/Source/Runtime/Experimental/Chaos/Private/PhysicsProxy/GeometryCollectionPhysicsProxy.cpp b/Engine/Source/Runtime/Experimental/Chaos/Private/PhysicsProxy/GeometryCollectionPhysicsProxy.cpp
--- a/Engine/Source/Runtime/Experimental/Chaos/Private/PhysicsProxy/GeometryCollectionPhysicsProxy.cpp
+++ b/Engine/Source/Runtime/Experimental/Chaos/Private/PhysicsProxy/GeometryCollectionPhysicsProxy.cpp
@@ -2931,10 +2931,8 @@ void FGeometryCollectionPhysicsProxy::PrepareBufferData(Chaos::FDirtyGeometryCol
 	BufferData.SetProxy(*this);
 
 	IsObjectDynamic = false;
-	FGeometryCollectionResults& Results = BufferData.Results;
+	FGeometryCollectionResults& Results = BufferData.Results();
 	Results.SetSolverDt(SolverLastDt);
-
-	const int32 NumTransformGroupElements = ThreadCollection.NumElements(FGeometryCollection::TransformGroup);
 	Results.InitArrays(ThreadCollection);
 }
 
@@ -3003,7 +3001,7 @@ void FGeometryCollectionPhysicsProxy::BufferPhysicsResults_Internal(Chaos::FPBDR
 	PrepareBufferData(BufferData, PhysicsThreadCollection, CurrentSolver->GetLastDt());
 
 	IsObjectDynamic = false;
-	FGeometryCollectionResults& Results = BufferData.Results;
+	FGeometryCollectionResults& Results = BufferData.Results();
 
 	const FTransform& ActorToWorld = Parameters.WorldTransform;
 	const bool IsActorScaled = !ActorToWorld.GetScale3D().Equals(FVector::OneVector);
@@ -3381,14 +3379,14 @@ bool FGeometryCollectionPhysicsProxy::PullFromPhysicsState(const Chaos::FDirtyGe
 	 */
 
 	// earlier exit
-	const bool HasResults = (PullData.Results.GetNumEntries() > 0);
-	const bool HasNextResults = (NextPullData && NextPullData->Results.GetNumEntries() > 0);
+	const bool HasResults = (PullData.Results().GetNumEntries() > 0);
+	const bool HasNextResults = (NextPullData && NextPullData->Results().GetNumEntries() > 0);
 	if (!HasResults && !HasNextResults)
 	{
 		return false;
 	}
 
-	const FGeometryCollectionResults& Results = PullData.Results;
+	const FGeometryCollectionResults& Results = PullData.Results();
 
 	const int32 NumTransforms = GameThreadCollection.Transform.Num();
 	const bool bNeedInterpolation = (NextPullData != nullptr);
@@ -3490,7 +3488,7 @@ bool FGeometryCollectionPhysicsProxy::PullFromPhysicsState(const Chaos::FDirtyGe
 		// second : interpolate-able ones
 		if (bNeedInterpolation)
 		{
-			const FGeometryCollectionResults& NextResults = NextPullData->Results;
+			const FGeometryCollectionResults& NextResults = NextPullData->Results();
 
 			// for that case we cannot just go through the list of entries since Results and NextResults may have different number of entries that don't always match
 			// so we need to go through the transform indices and find the matching entries on both side 
diff --git a/Engine/Source/Runtime/Experimental/Chaos/Public/Chaos/PullPhysicsDataImp.h b/Engine/Source/Runtime/Experimental/Chaos/Public/Chaos/PullPhysicsDataImp.h
--- a/Engine/Source/Runtime/Experimental/Chaos/Public/Chaos/PullPhysicsDataImp.h
+++ b/Engine/Source/Runtime/Experimental/Chaos/Public/Chaos/PullPhysicsDataImp.h
@@ -66,7 +66,32 @@ struct FDirtyRigidParticleReplicationErrorData : public TBasePullData<FSinglePar
 
 struct FDirtyGeometryCollectionData : public TBasePullData<FGeometryCollectionPhysicsProxy, FProxyTimestampBase>
 {
-	FGeometryCollectionResults Results;
+public:
+	bool HasResults() const { return ResultPtr.IsValid(); }
+
+	FDirtyGeometryCollectionData()
+		: ResultPtr(new FGeometryCollectionResults)
+	{}
+
+	FDirtyGeometryCollectionData(const FDirtyGeometryCollectionData& Other) = default;
+	FDirtyGeometryCollectionData(FDirtyGeometryCollectionData&& Other) = default;
+	FDirtyGeometryCollectionData& operator=(const FDirtyGeometryCollectionData& Other) = default;
+	FDirtyGeometryCollectionData& operator=(FDirtyGeometryCollectionData&& Other) = default;
+
+	FGeometryCollectionResults& Results()
+	{
+		check(ResultPtr);
+		return *ResultPtr;
+	}
+
+	const FGeometryCollectionResults& Results() const 
+	{
+		check(ResultPtr);
+		return *ResultPtr;
+	}
+
+private:
+	TRefCountPtr<FGeometryCollectionResults> ResultPtr;
 };
 
 struct FDirtyClusterUnionParticleData
diff --git a/Engine/Source/Runtime/Experimental/Chaos/Public/GeometryCollectionProxyData.h b/Engine/Source/Runtime/Experimental/Chaos/Public/GeometryCollectionProxyData.h
--- a/Engine/Source/Runtime/Experimental/Chaos/Public/GeometryCollectionProxyData.h
+++ b/Engine/Source/Runtime/Experimental/Chaos/Public/GeometryCollectionProxyData.h
@@ -296,7 +296,7 @@ private:
  * Buffer structure for communicating simulation state between game and physics
  * threads.
  */
-class FGeometryCollectionResults
+class FGeometryCollectionResults: public FRefCountedObject
 {
 public:
 	FGeometryCollectionResults();
