Chaos - move BasicCollisionDetector and BasicBroadPhase (used by RBAN) into appropriately named files, plus other minor cleanup

diff --git a/Engine/Source/Programs/HeadlessChaos/Private/HeadlessChaosTestMinEvolution.cpp b/Engine/Source/Programs/HeadlessChaos/Private/HeadlessChaosTestMinEvolution.cpp
--- a/Engine/Source/Programs/HeadlessChaos/Private/HeadlessChaosTestMinEvolution.cpp
+++ b/Engine/Source/Programs/HeadlessChaos/Private/HeadlessChaosTestMinEvolution.cpp
@@ -2,8 +2,7 @@
 
 #include "HeadlessChaos.h"
 
-#include "Chaos/Collision/ParticlePairBroadPhase.h"
-#include "Chaos/Collision/ParticlePairCollisionDetector.h"
+#include "Chaos/Collision/BasicCollisionDetector.h"
 #include "Chaos/Evolution/PBDMinEvolution.h"
 #include "Chaos/ParticleHandle.h"
 #include "Chaos/PBDCollisionConstraints.h"
diff --git a/Engine/Source/Runtime/Engine/Private/PhysicsEngine/ImmediatePhysics/ImmediatePhysicsChaos/ImmediatePhysicsSimulation_Chaos.cpp b/Engine/Source/Runtime/Engine/Private/PhysicsEngine/ImmediatePhysics/ImmediatePhysicsChaos/ImmediatePhysicsSimulation_Chaos.cpp
--- a/Engine/Source/Runtime/Engine/Private/PhysicsEngine/ImmediatePhysics/ImmediatePhysicsChaos/ImmediatePhysicsSimulation_Chaos.cpp
+++ b/Engine/Source/Runtime/Engine/Private/PhysicsEngine/ImmediatePhysics/ImmediatePhysicsChaos/ImmediatePhysicsSimulation_Chaos.cpp
@@ -10,7 +10,7 @@
 #include "PhysicalMaterials/PhysicalMaterial.h"
 #include "PhysicsEngine/BodyInstance.h"
 
-#include "Chaos/Collision/ParticlePairCollisionDetector.h"
+#include "Chaos/Collision/BasicCollisionDetector.h"
 #include "Chaos/DebugDrawQueue.h"
 #include "Chaos/Evolution/PBDMinEvolution.h"
 #include "Chaos/Joint/ChaosJointLog.h"
diff --git a/Engine/Source/Runtime/Experimental/Chaos/Private/Chaos/CollisionOneShotManifolds.cpp b/Engine/Source/Runtime/Experimental/Chaos/Private/Chaos/CollisionOneShotManifolds.cpp
--- a/Engine/Source/Runtime/Experimental/Chaos/Private/Chaos/CollisionOneShotManifolds.cpp
+++ b/Engine/Source/Runtime/Experimental/Chaos/Private/Chaos/CollisionOneShotManifolds.cpp
@@ -4,7 +4,6 @@
 
 #include "Chaos/Box.h"
 #include "Chaos/CollisionResolution.h"
-#include "Chaos/Collision/ContactTriangles.h"
 #include "Chaos/Collision/PBDCollisionConstraint.h"
 #include "Chaos/Convex.h"
 #include "Chaos/Defines.h"
diff --git a/Engine/Source/Runtime/Experimental/Chaos/Private/Chaos/Evolution/PBDMinEvolution.cpp b/Engine/Source/Runtime/Experimental/Chaos/Private/Chaos/Evolution/PBDMinEvolution.cpp
--- a/Engine/Source/Runtime/Experimental/Chaos/Private/Chaos/Evolution/PBDMinEvolution.cpp
+++ b/Engine/Source/Runtime/Experimental/Chaos/Private/Chaos/Evolution/PBDMinEvolution.cpp
@@ -1,7 +1,7 @@
 // Copyright Epic Games, Inc. All Rights Reserved.
 
 #include "Chaos/Evolution/PBDMinEvolution.h"
-#include "Chaos/Collision/ParticlePairCollisionDetector.h"
+#include "Chaos/Collision/BasicCollisionDetector.h"
 #include "Chaos/Evolution/SolverConstraintContainer.h"
 #include "Chaos/PBDCollisionConstraints.h"
 #include "Chaos/PBDConstraintContainer.h"
diff --git a/Engine/Source/Runtime/Experimental/Chaos/Private/Chaos/HeightField.cpp b/Engine/Source/Runtime/Experimental/Chaos/Private/Chaos/HeightField.cpp
--- a/Engine/Source/Runtime/Experimental/Chaos/Private/Chaos/HeightField.cpp
+++ b/Engine/Source/Runtime/Experimental/Chaos/Private/Chaos/HeightField.cpp
@@ -3,7 +3,6 @@
 #include "Chaos/HeightField.h"
 #include "Chaos/Collision/ContactPoint.h"
 #include "Chaos/Collision/ContactPointsMiscShapes.h"
-#include "Chaos/Collision/ContactTriangles.h"
 #include "Chaos/Collision/TriangleOverlap.h"
 #include "Chaos/Collision/PBDCollisionConstraint.h"
 #include "Chaos/CollisionOneShotManifolds.h"
diff --git a/Engine/Source/Runtime/Experimental/Chaos/Public/Chaos/Collision/BasicBroadPhase.h b/Engine/Source/Runtime/Experimental/Chaos/Public/Chaos/Collision/BasicBroadPhase.h
new file mode 100644
--- /dev/null
+++ b/Engine/Source/Runtime/Experimental/Chaos/Public/Chaos/Collision/BasicBroadPhase.h
@@ -0,0 +1,105 @@
+// Copyright Epic Games, Inc. All Rights Reserved.
+#pragma once
+
+#include "Chaos/BoundingVolumeUtilities.h"
+#include "Chaos/Collision/ParticlePairMidPhase.h"
+#include "Chaos/ParticleHandle.h"
+#include "Chaos/PBDCollisionConstraints.h"
+#include "Chaos/PBDRigidsSOAs.h"
+#include "ChaosStats.h"
+
+namespace Chaos
+{
+	/**
+	 * Run through a list of particle pairs and pass them onto the collision detector if their AABBs overlap.
+	 * In addition, collide all particles in ParticlesA with all particles in ParticlesB.
+	 *
+	 * No spatial acceleration, and the order is assumed to be already optimized for cache efficiency.
+	 */
+	class CHAOS_API FBasicBroadPhase
+	{
+	public:
+		using FParticleHandle = TGeometryParticleHandle<FReal, 3>;
+		using FAABB = FAABB3;
+
+		FBasicBroadPhase(const TArray<FParticlePair>* InParticlePairs, const TArray<FParticleHandle*>* InParticlesA, const TArray<FParticleHandle*>* InParticlesB)
+			: ParticlePairs(InParticlePairs)
+			, ParticlesA(InParticlesA)
+			, ParticlesB(InParticlesB)
+		{
+		}
+
+		/**
+		 *
+		 */
+		void ProduceOverlaps(
+			FReal Dt,
+			Private::FCollisionConstraintAllocator* Allocator,
+			const FCollisionDetectorSettings& Settings)
+		{
+			SCOPE_CYCLE_COUNTER(STAT_Collisions_ParticlePairBroadPhase);
+
+			Allocator->SetMaxContexts(1);
+
+			FCollisionContext Context;
+			Context.SetSettings(Settings);
+			Context.Allocator = Allocator->GetContextAllocator(0);
+
+			if (ParticlePairs != nullptr)
+			{
+				for (const FParticlePair& ParticlePair : *ParticlePairs)
+				{
+					// Pair array is const and the particles are not, but we need to const_cast
+					FParticleHandle* ParticleA = const_cast<FParticleHandle*>(ParticlePair[0]);
+					FParticleHandle* ParticleB = const_cast<FParticleHandle*>(ParticlePair[1]);
+
+					if ((ParticleA != nullptr) && (ParticleB != nullptr))
+					{
+						ProduceOverlaps(Dt, Context, ParticleA, ParticleB, ParticleA);
+					}
+				}
+			}
+
+			if ((ParticlesA != nullptr) && (ParticlesB != nullptr))
+			{
+				for (FParticleHandle* ParticleA : *ParticlesA)
+				{
+					if (ParticleA != nullptr)
+					{
+						for (FParticleHandle* ParticleB : *ParticlesB)
+						{
+							if (ParticleB != nullptr)
+							{
+								ProduceOverlaps(Dt, Context, ParticleA, ParticleB, ParticleA);
+							}
+						}
+					}
+				}
+			}
+		}
+
+	private:
+		inline void ProduceOverlaps(
+			FReal Dt,
+			const FCollisionContext& Context,
+			FParticleHandle* ParticleA,
+			FParticleHandle* ParticleB,
+			FParticleHandle* SearchParticle)
+		{
+			const FAABB3& Box0 = ParticleA->WorldSpaceInflatedBounds();
+			const FAABB3& Box1 = ParticleB->WorldSpaceInflatedBounds();
+			if (Box0.Intersects(Box1))
+			{
+				FParticlePairMidPhase* MidPhase = Context.GetAllocator()->GetMidPhase(ParticleA, ParticleB, SearchParticle, Context);
+				if (MidPhase != nullptr)
+				{
+					MidPhase->GenerateCollisions(Context.GetSettings().BoundsExpansion, Dt, Context);
+				}
+			}
+		}
+
+		const TArray<FParticlePair>* ParticlePairs;
+		const TArray<FParticleHandle*>* ParticlesA;
+		const TArray<FParticleHandle*>* ParticlesB;
+	};
+}
diff --git a/Engine/Source/Runtime/Experimental/Chaos/Public/Chaos/Collision/BasicCollisionDetector.h b/Engine/Source/Runtime/Experimental/Chaos/Public/Chaos/Collision/BasicCollisionDetector.h
new file mode 100644
--- /dev/null
+++ b/Engine/Source/Runtime/Experimental/Chaos/Public/Chaos/Collision/BasicCollisionDetector.h
@@ -0,0 +1,48 @@
+// Copyright Epic Games, Inc. All Rights Reserved.
+#pragma once
+
+#include "Chaos/Collision/CollisionDetector.h"
+#include "Chaos/Collision/BasicBroadPhase.h"
+#include "Chaos/PBDCollisionConstraints.h"
+
+#include "Chaos/ChaosPerfTest.h"
+#include "ChaosStats.h"
+#include "ProfilingDebugging/CsvProfiler.h"
+
+namespace Chaos
+{
+	class CHAOS_API FBasicCollisionDetector : public FCollisionDetector
+	{
+	public:
+		FBasicCollisionDetector(FBasicBroadPhase& InBroadPhase, FPBDCollisionConstraints& InCollisionContainer)
+			: FCollisionDetector(InCollisionContainer)
+			, BroadPhase(InBroadPhase)
+		{
+		}
+
+		FBasicBroadPhase& GetBroadPhase() { return BroadPhase; }
+
+		virtual void DetectCollisions(const FReal Dt, FEvolutionResimCache* Unused) override
+		{
+			SCOPE_CYCLE_COUNTER(STAT_Collisions_Detect);
+			CHAOS_SCOPED_TIMER(DetectCollisions);
+			CSV_SCOPED_TIMING_STAT(Chaos, DetectCollisions);
+
+			if (!GetCollisionContainer().GetCollisionsEnabled())
+			{
+				return;
+			}
+
+			CollisionContainer.BeginDetectCollisions();
+
+			// Collision detection pipeline: BroadPhase -> MidPhase ->NarrowPhase -> Container
+			BroadPhase.ProduceOverlaps(Dt, &GetCollisionContainer().GetConstraintAllocator(), GetCollisionContainer().GetDetectorSettings());
+
+			CollisionContainer.EndDetectCollisions();
+		}
+
+	private:
+		FBasicBroadPhase& BroadPhase;
+	};
+
+}
\ No newline at end of file
diff --git a/Engine/Source/Runtime/Experimental/Chaos/Public/Chaos/Collision/CollisionVisitor.h b/Engine/Source/Runtime/Experimental/Chaos/Public/Chaos/Collision/CollisionVisitor.h
--- a/Engine/Source/Runtime/Experimental/Chaos/Public/Chaos/Collision/CollisionVisitor.h
+++ b/Engine/Source/Runtime/Experimental/Chaos/Public/Chaos/Collision/CollisionVisitor.h
@@ -1,7 +1,6 @@
 // Copyright Epic Games, Inc. All Rights Reserved.
 #pragma once
 
-#include "CoreMinimal.h"
 #include "Chaos/Core.h"
 
 namespace Chaos
diff --git a/Engine/Source/Runtime/Experimental/Chaos/Public/Chaos/Collision/ContactTriangles.h b/Engine/Source/Runtime/Experimental/Chaos/Public/Chaos/Collision/ContactTriangles.h
--- a/Engine/Source/Runtime/Experimental/Chaos/Public/Chaos/Collision/ContactTriangles.h
+++ b/Engine/Source/Runtime/Experimental/Chaos/Public/Chaos/Collision/ContactTriangles.h
@@ -118,40 +118,28 @@ namespace Chaos
 		}
 	};
 
-	/**
-	 * @brief A set of triangles assocuated with some contact points.
-	*/
-	class CHAOS_API FContactTriangles
-	{
-	public:
-
-		inline const int32 Num() const
-		{
-			return Triangles.Num();
-		}
-
-		inline const FContactTriangle& At(const int32 Index) const
-		{
-			return Triangles[Index];
-		}
-
-
-	private:
-		TArray<FContactTriangle> Triangles;
-	};
-
 	using FContactVertexID = int32;
 
 	struct FContactEdgeID
 	{
 		FContactEdgeID()
-			: VertexIDs{ INDEX_NONE < INDEX_NONE }
+			: VertexIDs{ INDEX_NONE, INDEX_NONE }
 		{
 		}
 
-		FContactEdgeID(const FContactVertexID VertextIndexA, const FContactVertexID VertexIndexB)
-			: VertexIDs{ FMath::Min(VertextIndexA, VertexIndexB), FMath::Max(VertextIndexA, VertexIndexB) }
+		FContactEdgeID(const FContactVertexID VertexIndexA, const FContactVertexID VertexIndexB)
 		{
+			// EdgeID is the same if we swap the vertex indices
+			if (VertexIndexA < VertexIndexB)
+			{
+				VertexIDs[0] = VertexIndexA;
+				VertexIDs[1] = VertexIndexB;
+			}
+			else
+			{
+				VertexIDs[0] = VertexIndexB;
+				VertexIDs[1] = VertexIndexA;
+			}
 		}
 
 		bool IsValid() const
diff --git a/Engine/Source/Runtime/Experimental/Chaos/Public/Chaos/Collision/ParticlePairBroadPhase.h b/Engine/Source/Runtime/Experimental/Chaos/Public/Chaos/Collision/ParticlePairBroadPhase.h
--- a/Engine/Source/Runtime/Experimental/Chaos/Public/Chaos/Collision/ParticlePairBroadPhase.h
+++ b/Engine/Source/Runtime/Experimental/Chaos/Public/Chaos/Collision/ParticlePairBroadPhase.h
@@ -1,105 +1,6 @@
 // Copyright Epic Games, Inc. All Rights Reserved.
 #pragma once
 
-#include "Chaos/BoundingVolumeUtilities.h"
-#include "Chaos/Collision/ParticlePairMidPhase.h"
-#include "Chaos/ParticleHandle.h"
-#include "Chaos/PBDCollisionConstraints.h"
-#include "Chaos/PBDRigidsSOAs.h"
-#include "ChaosStats.h"
-
-namespace Chaos
-{
-	/**
-	 * Run through a list of particle pairs and pass them onto the collision detector if their AABBs overlap.
-	 * In addition, collide all particles in ParticlesA with all particles in ParticlesB.
-	 *
-	 * No spatial acceleration, and the order is assumed to be already optimized for cache efficiency.
-	 */
-	class CHAOS_API FBasicBroadPhase
-	{
-	public:
-		using FParticleHandle = TGeometryParticleHandle<FReal, 3>;
-		using FAABB = FAABB3;
-
-		FBasicBroadPhase(const TArray<FParticlePair>* InParticlePairs, const TArray<FParticleHandle*>* InParticlesA, const TArray<FParticleHandle*>* InParticlesB)
-			: ParticlePairs(InParticlePairs)
-			, ParticlesA(InParticlesA)
-			, ParticlesB(InParticlesB)
-		{
-		}
-
-		/**
-		 *
-		 */
-		void ProduceOverlaps(
-			FReal Dt, 
-			Private::FCollisionConstraintAllocator* Allocator,
-			const FCollisionDetectorSettings& Settings)
-		{
-			SCOPE_CYCLE_COUNTER(STAT_Collisions_ParticlePairBroadPhase);
-
-			Allocator->SetMaxContexts(1);
-
-			FCollisionContext Context;
-			Context.SetSettings(Settings);
-			Context.Allocator = Allocator->GetContextAllocator(0);
-
-			if (ParticlePairs != nullptr)
-			{
-				for (const FParticlePair& ParticlePair : *ParticlePairs)
-				{
-					// Pair array is const and the particles are not, but we need to const_cast
-					FParticleHandle* ParticleA = const_cast<FParticleHandle*>(ParticlePair[0]);
-					FParticleHandle* ParticleB = const_cast<FParticleHandle*>(ParticlePair[1]);
-
-					if ((ParticleA != nullptr) && (ParticleB != nullptr))
-					{
-						ProduceOverlaps(Dt, Context, ParticleA, ParticleB, ParticleA);
-					}
-				}
-			}
-
-			if ((ParticlesA != nullptr) && (ParticlesB != nullptr))
-			{
-				for (FParticleHandle* ParticleA : *ParticlesA)
-				{
-					if (ParticleA != nullptr)
-					{
-						for (FParticleHandle* ParticleB : *ParticlesB)
-						{
-							if (ParticleB != nullptr)
-							{
-								ProduceOverlaps(Dt, Context, ParticleA, ParticleB, ParticleA);
-							}
-						}
-					}
-				}
-			}
-		}
-
-	private:
-		inline void ProduceOverlaps(
-			FReal Dt,
-			const FCollisionContext& Context,
-			FParticleHandle* ParticleA,
-			FParticleHandle* ParticleB,
-			FParticleHandle* SearchParticle)
-		{
-			const FAABB3& Box0 = ParticleA->WorldSpaceInflatedBounds();
-			const FAABB3& Box1 = ParticleB->WorldSpaceInflatedBounds();
-			if (Box0.Intersects(Box1))
-			{
-				FParticlePairMidPhase* MidPhase = Context.GetAllocator()->GetMidPhase(ParticleA, ParticleB, SearchParticle, Context);
-				if (MidPhase != nullptr)
-				{
-					MidPhase->GenerateCollisions(Context.GetSettings().BoundsExpansion, Dt, Context);
-				}
-			}
-		}
-
-		const TArray<FParticlePair>* ParticlePairs;
-		const TArray<FParticleHandle*>* ParticlesA;
-		const TArray<FParticleHandle*>* ParticlesB;
-	};
-}
+//
+// TO BE REMOVED
+//
\ No newline at end of file
diff --git a/Engine/Source/Runtime/Experimental/Chaos/Public/Chaos/Collision/ParticlePairCollisionDetector.h b/Engine/Source/Runtime/Experimental/Chaos/Public/Chaos/Collision/ParticlePairCollisionDetector.h
--- a/Engine/Source/Runtime/Experimental/Chaos/Public/Chaos/Collision/ParticlePairCollisionDetector.h
+++ b/Engine/Source/Runtime/Experimental/Chaos/Public/Chaos/Collision/ParticlePairCollisionDetector.h
@@ -1,48 +1,6 @@
 // Copyright Epic Games, Inc. All Rights Reserved.
 #pragma once
 
-#include "Chaos/Collision/CollisionDetector.h"
-#include "Chaos/Collision/ParticlePairBroadPhase.h"
-#include "Chaos/PBDCollisionConstraints.h"
-
-#include "Chaos/ChaosPerfTest.h"
-#include "ChaosStats.h"
-#include "ProfilingDebugging/CsvProfiler.h"
-
-namespace Chaos
-{
-	class CHAOS_API FBasicCollisionDetector : public FCollisionDetector
-	{
-	public:
-		FBasicCollisionDetector(FBasicBroadPhase& InBroadPhase, FPBDCollisionConstraints& InCollisionContainer)
-			: FCollisionDetector(InCollisionContainer)
-			, BroadPhase(InBroadPhase)
-		{
-		}
-
-		FBasicBroadPhase& GetBroadPhase() { return BroadPhase; }
-
-		virtual void DetectCollisions(const FReal Dt, FEvolutionResimCache* Unused) override
-		{
-			SCOPE_CYCLE_COUNTER(STAT_Collisions_Detect);
-			CHAOS_SCOPED_TIMER(DetectCollisions);
-			CSV_SCOPED_TIMING_STAT(Chaos, DetectCollisions);
-
-			if (!GetCollisionContainer().GetCollisionsEnabled())
-			{
-				return;
-			}
-
-			CollisionContainer.BeginDetectCollisions();
-
-			// Collision detection pipeline: BroadPhase -> MidPhase ->NarrowPhase -> Container
-			BroadPhase.ProduceOverlaps(Dt, &GetCollisionContainer().GetConstraintAllocator(), GetCollisionContainer().GetDetectorSettings());
-
-			CollisionContainer.EndDetectCollisions();
-		}
-
-	private:
-		FBasicBroadPhase& BroadPhase;
-	};
-
-}
\ No newline at end of file
+//
+//	TO BE REMOVED
+//
\ No newline at end of file
