cluster union : unify iterative and oneshot connectivity graph generation
- also fixes the one shot generation not working because of the lack of collision particles

diff --git a/Engine/Source/Runtime/Experimental/Chaos/Private/Chaos/ClusterUnionManager.cpp b/Engine/Source/Runtime/Experimental/Chaos/Private/Chaos/ClusterUnionManager.cpp
--- a/Engine/Source/Runtime/Experimental/Chaos/Private/Chaos/ClusterUnionManager.cpp
+++ b/Engine/Source/Runtime/Experimental/Chaos/Private/Chaos/ClusterUnionManager.cpp
@@ -31,6 +31,98 @@ namespace Chaos
 
 			return Frame;
 		}
+
+		void AddParticleToConnectionGraph(FRigidClustering& Clustering, FClusterUnion& ClusterUnion, FPBDRigidParticleHandle* Particle)
+		{
+			if (!ClusterUnion.ChildParticles.IsEmpty() && ClusterUnion.SharedGeometry->GetType() == ImplicitObjectType::Union)
+			{
+				constexpr FReal kAddThickness = 5.0;
+				const FRigidTransform3 ClusterWorldTM(ClusterUnion.InternalCluster->X(), ClusterUnion.InternalCluster->R());
+
+				// Use the acceleration structure of the cluster union itself to make finding overlaps easy.
+				const FImplicitObjectUnion& ShapeUnion = ClusterUnion.SharedGeometry->GetObjectChecked<FImplicitObjectUnion>();
+				const FRigidTransform3 FromTransform = GetParticleRigidFrameInClusterUnion(Particle, ClusterWorldTM);
+				FAABB3 ParticleLocalBounds = Particle->SharedGeometry()->BoundingBox().TransformedAABB(FromTransform);
+				ParticleLocalBounds.Thicken(kAddThickness);
+
+				ShapeUnion.VisitOverlappingLeafObjects(
+					ParticleLocalBounds,
+					[&Clustering, &ClusterUnion, &FromTransform, Particle, kAddThickness](const FImplicitObject* ToGeom, const FRigidTransform3& ToTransform, const int32 RootObjectIndex, const int32 ObjectIndex, const int32 LeafObjectIndex)
+					{
+						check(ToGeom != nullptr);
+						// RootObjectIndex is the index in the cluster union.
+						// NOTE: This will need to be re-thought if we ever decide to make the mapping between child shapes and child particles not 1-to-1 (e.g. if we ever attempt to simplify the cluster union shape).
+
+						// Ignore intersections against the same particle since we just added the particle (potentially) into the geometry.
+						if (Particle == ClusterUnion.ChildParticles[RootObjectIndex])
+						{
+							return;
+						}
+
+						FAABB3 ToAABB = ToGeom->CalculateTransformedBounds(ToTransform);
+						ToAABB.Thicken(kAddThickness);
+
+						const bool bOverlap = Utilities::CastHelper(*ToGeom, ToTransform,
+							[Particle, &FromTransform, &ToAABB, kAddThickness](const auto& ToGeomDowncast, const FRigidTransform3& FinalToTransform)
+							{
+								const FShapesArray& AllFromShapes = Particle->ShapesArray();
+								for (int32 FromIndex = 0; FromIndex < AllFromShapes.Num(); ++FromIndex)
+								{
+									if (!AllFromShapes[FromIndex])
+									{
+										continue;
+									}
+
+									const FPerShapeData& FromShape = *AllFromShapes[FromIndex];
+									const FImplicitObject* FromGeom = FromShape.GetGeometry().Get();
+
+									if (!FromGeom)
+									{
+										continue;
+									}
+
+									FAABB3 FromAABB = FromGeom->CalculateTransformedBounds(FromTransform);
+									FromAABB.Thicken(kAddThickness);
+
+									// First sanity check to see if the two shape AABB's intersect.
+									if (!FromAABB.Intersects(ToAABB))
+									{
+										continue;
+									}
+
+									// Now do a more accurate overlap check between the two shapes.
+									// Note that passing MTD here is critical as it gets us a more accurate check for some reason...
+									FMTDInfo MTDInfo;
+									if (OverlapQuery(*FromGeom, FromTransform, ToGeomDowncast, FinalToTransform, kAddThickness, &MTDInfo))
+									{
+										return true;
+									}
+									else
+									{
+										continue;
+									}
+								}
+
+								return false;
+							});
+
+						if (bOverlap)
+						{
+							Clustering.CreateNodeConnection(Particle, ClusterUnion.ChildParticles[RootObjectIndex]);
+						}
+					}
+				);
+			}
+		}
+
+		void GenerateConnectionGraph(FRigidClustering& Clustering, FClusterUnion& ClusterUnion)
+		{
+			Clustering.ClearConnectionGraph(ClusterUnion.InternalCluster);
+			for (FPBDRigidParticleHandle* ChildParticle : ClusterUnion.ChildParticles)
+			{
+				AddParticleToConnectionGraph(Clustering, ClusterUnion, ChildParticle);
+			}
+		}
 	}
 
 	FClusterUnionManager::FClusterUnionManager(FRigidClustering& InClustering, FPBDRigidsEvolutionGBF& InEvolution)
@@ -567,8 +659,7 @@ namespace Chaos
 
 		if(EnumHasAnyFlags(Flags, EUpdateClusterUnionPropertiesFlags::ForceGenerateConnectionGraph))
 		{
-			MClustering.ClearConnectionGraph(ClusterUnion.InternalCluster);
-			MClustering.GenerateConnectionGraph(ClusterUnion.InternalCluster, ClusterUnion.Parameters);
+			GenerateConnectionGraph(MClustering, ClusterUnion);
 			ClusterUnion.PendingConnectivityOperations.Empty();
 		}
 		else if (EnumHasAnyFlags(Flags, EUpdateClusterUnionPropertiesFlags::IncrementalGenerateConnectionGraph))
@@ -590,89 +681,12 @@ namespace Chaos
 	DECLARE_CYCLE_STAT(TEXT("FClusterUnionManager::FlushIncrementalConnectivityGraphOperations"), STAT_FlushIncrementalConnectivityGraphOperations, STATGROUP_Chaos);
 	void FClusterUnionManager::FlushIncrementalConnectivityGraphOperations(FClusterUnion& ClusterUnion)
 	{
-		constexpr FReal kAddThickness = 5.0;
 		SCOPE_CYCLE_COUNTER(STAT_FlushIncrementalConnectivityGraphOperations);
-		const FRigidTransform3 ClusterWorldTM(ClusterUnion.InternalCluster->X(), ClusterUnion.InternalCluster->R());
 		for (const TPair<FPBDRigidParticleHandle*, EClusterUnionConnectivityOperation>& Op : ClusterUnion.PendingConnectivityOperations)
 		{
 			if (Op.Value == EClusterUnionConnectivityOperation::Add)
 			{
-				if (!ClusterUnion.ChildParticles.IsEmpty() && ClusterUnion.SharedGeometry->GetType() == ImplicitObjectType::Union)
-				{
-					// Use the acceleration structure of the cluster union itself to make finding overlaps easy.
-					const FImplicitObjectUnion& ShapeUnion = ClusterUnion.SharedGeometry->GetObjectChecked<FImplicitObjectUnion>();
-					const FRigidTransform3 FromTransform = GetParticleRigidFrameInClusterUnion(Op.Key, ClusterWorldTM);
-					FAABB3 ParticleLocalBounds = Op.Key->SharedGeometry()->BoundingBox().TransformedAABB(FromTransform);
-					ParticleLocalBounds.Thicken(kAddThickness);
-
-					ShapeUnion.VisitOverlappingLeafObjects(
-						ParticleLocalBounds,
-						[this, &ClusterUnion, &FromTransform, &Op, kAddThickness](const FImplicitObject* ToGeom, const FRigidTransform3& ToTransform, const int32 RootObjectIndex, const int32 ObjectIndex, const int32 LeafObjectIndex)
-						{
-							check(ToGeom != nullptr);
-							// RootObjectIndex is the index in the cluster union.
-							// NOTE: This will need to be re-thought if we ever decide to make the mapping between child shapes and child particles not 1-to-1 (e.g. if we ever attempt to simplify the cluster union shape).
-
-							// Ignore intersections against the same particle since we just added the particle (potentially) into the geometry.
-							if (Op.Key == ClusterUnion.ChildParticles[RootObjectIndex])
-							{
-								return;
-							}
-
-							FAABB3 ToAABB = ToGeom->CalculateTransformedBounds(ToTransform);
-							ToAABB.Thicken(kAddThickness);
-
-							const bool bOverlap = Utilities::CastHelper(*ToGeom, ToTransform,
-								[&Op, &FromTransform, &ToAABB, kAddThickness](const auto& ToGeomDowncast, const FRigidTransform3& FinalToTransform)
-								{
-									const FShapesArray& AllFromShapes = Op.Key->ShapesArray();
-									for (int32 FromIndex = 0; FromIndex < AllFromShapes.Num(); ++FromIndex)
-									{
-										if (!AllFromShapes[FromIndex])
-										{
-											continue;
-										}
-
-										const FPerShapeData& FromShape = *AllFromShapes[FromIndex];
-										const FImplicitObject* FromGeom = FromShape.GetGeometry().Get();
-
-										if (!FromGeom)
-										{
-											continue;
-										}
-
-										FAABB3 FromAABB = FromGeom->CalculateTransformedBounds(FromTransform);
-										FromAABB.Thicken(kAddThickness);
-
-										// First sanity check to see if the two shape AABB's intersect.
-										if (!FromAABB.Intersects(ToAABB))
-										{
-											continue;
-										}
-
-										// Now do a more accurate overlap check between the two shapes.
-										// Note that passing MTD here is critical as it gets us a more accurate check for some reason...
-										FMTDInfo MTDInfo;
-										if (OverlapQuery(*FromGeom, FromTransform, ToGeomDowncast, FinalToTransform, kAddThickness, &MTDInfo))
-										{
-											return true;
-										}
-										else
-										{
-											continue;
-										}
-									}
-
-									return false;
-								});
-
-							if (bOverlap)
-							{
-								MClustering.CreateNodeConnection(Op.Key, ClusterUnion.ChildParticles[RootObjectIndex]);
-							}
-						}
-					);
-				}
+				AddParticleToConnectionGraph(MClustering, ClusterUnion, Op.Key);
 			}
 			else if (Op.Value == EClusterUnionConnectivityOperation::Remove)
 			{
