I/O - fixed race when flushing on demand content writer

diff --git a/Engine/Source/Runtime/Experimental/IoStoreOnDemand/Private/OnDemandIoStoreWriter.cpp b/Engine/Source/Runtime/Experimental/IoStoreOnDemand/Private/OnDemandIoStoreWriter.cpp
--- a/Engine/Source/Runtime/Experimental/IoStoreOnDemand/Private/OnDemandIoStoreWriter.cpp
+++ b/Engine/Source/Runtime/Experimental/IoStoreOnDemand/Private/OnDemandIoStoreWriter.cpp
@@ -83,6 +83,7 @@ private:
 	TMap<FString, TSharedPtr<FOnDemandContainer>> Containers;
 	FCriticalSection ContainersCS;
 	FEventRef WriteCompletedEvent;
+	FEventRef WriteQueueEmptyEvent;
 	UE::Tasks::FPipe TaskPipe;
 	std::atomic<uint32> PendingCount{0};
 	std::atomic<uint32> TotalCount{0};
@@ -92,6 +93,7 @@ private:
 FOnDemandIoStoreWriter::FOnDemandIoStoreWriter(const FIoStoreWriterSettings& Settings, const FString& InOutputDirectory, uint32 InMaxConcurrentWrites)
 	: WriterSettings(Settings)
 	, OutputDirectory(InOutputDirectory)
+	, WriteQueueEmptyEvent(EEventMode::ManualReset)
 	, TaskPipe(UE_SOURCE_LOCATION)
 	, MaxConcurrentWrites(InMaxConcurrentWrites)
 {
@@ -108,9 +110,10 @@ TSharedPtr<IIoStoreWriter> FOnDemandIoStoreWriter::CreateContainer(const FString
 
 void FOnDemandIoStoreWriter::Flush() 
 {
-	TaskPipe.WaitUntilEmpty();
+	WriteQueueEmptyEvent->Wait();
 
 	FScopeLock _(&ContainersCS);
+	TaskPipe.WaitUntilEmpty();
 
 	FOnDemandToc TocResource;
 	TocResource.Header.Magic = FOnDemandTocHeader::ExpectedMagic;
@@ -205,6 +208,8 @@ void FOnDemandIoStoreWriter::Append(
 	IIoStoreWriteRequest* Request,
 	const FIoWriteOptions& WriteOptions)
 {
+	WriteQueueEmptyEvent->Reset();
+
 	for(;;)
 	{
 		{
@@ -301,7 +306,10 @@ void FOnDemandIoStoreWriter::Append(
 								PendingWrite->ErrorText = TEXT("Invalid source buffer");
 							}
 
-							PendingCount--;
+							if (PendingCount.fetch_sub(1, std::memory_order_relaxed) == 1)
+							{
+								WriteQueueEmptyEvent->Trigger();
+							}
 							WriteCompletedEvent->Trigger();
 						}, EncodeChunkTask);
 
