I/O - write on demand config to correct path

diff --git a/Engine/Source/Runtime/Experimental/IoStoreOnDemand/Private/IoStoreOnDemand.cpp b/Engine/Source/Runtime/Experimental/IoStoreOnDemand/Private/IoStoreOnDemand.cpp
--- a/Engine/Source/Runtime/Experimental/IoStoreOnDemand/Private/IoStoreOnDemand.cpp
+++ b/Engine/Source/Runtime/Experimental/IoStoreOnDemand/Private/IoStoreOnDemand.cpp
@@ -367,12 +367,12 @@ void FIoStoreOnDemandModule::StartupModule()
 				Endpoint.TocPath.RightChopInline(1);
 			}
 
-			FString EncryptionKey;
-			if (Config.GetString(TEXT("Endpoint"), TEXT("EncryptionKey"), EncryptionKey))
+			FString ContentKey;
+			if (Config.GetString(TEXT("Endpoint"), TEXT("ContentKey"), ContentKey))
 			{
 				FGuid KeyGuid;
 				FAES::FAESKey Key;
-				if (ParseEncryptionKeyParam(EncryptionKey, KeyGuid, Key))
+				if (ParseEncryptionKeyParam(ContentKey, KeyGuid, Key))
 				{
 					UE::FEncryptionKeyManager::Get().AddKey(KeyGuid, Key);
 				}
