SingleParticlePhysicsProxy::PushToPhysicsStateImp -
Destroy Transient Constraints when collision is changed on a particle so that mid-phase can be recreated with new shapes if there are any.

diff --git a/Engine/Source/Runtime/Experimental/Chaos/Private/PhysicsProxy/SingleParticlePhysicsProxy.cpp b/Engine/Source/Runtime/Experimental/Chaos/Private/PhysicsProxy/SingleParticlePhysicsProxy.cpp
--- a/Engine/Source/Runtime/Experimental/Chaos/Private/PhysicsProxy/SingleParticlePhysicsProxy.cpp
+++ b/Engine/Source/Runtime/Experimental/Chaos/Private/PhysicsProxy/SingleParticlePhysicsProxy.cpp
@@ -177,6 +177,8 @@ void PushToPhysicsStateImp(const Chaos::FDirtyPropertiesManager& Manager, Chaos:
 
 			if(bHasCollision)
 			{
+				// destroy collision constraints so that mid-phase is recreated with newly added shapes if any
+				Evolution.DestroyTransientConstraints(Handle);
 				//make sure it's in acceleration structure
 				Evolution.DirtyParticle(*Handle);
 			}
