Improvements to character ground constraint
- Some changes to be able to edit the constraint data on the physics thread
- Add a constraint torque to keep the character upright
- Some fixes for sleeping

diff --git a/Engine/Source/Programs/HeadlessChaos/Private/Character/TestCharacterGroundConstraintSolver.cpp b/Engine/Source/Programs/HeadlessChaos/Private/Character/TestCharacterGroundConstraintSolver.cpp
--- a/Engine/Source/Programs/HeadlessChaos/Private/Character/TestCharacterGroundConstraintSolver.cpp
+++ b/Engine/Source/Programs/HeadlessChaos/Private/Character/TestCharacterGroundConstraintSolver.cpp
@@ -74,7 +74,8 @@ namespace ChaosTest
 		Data.GroundDistance = 10.0f;
 		Settings.TargetHeight = 15.0f;
 		Settings.RadialForceLimit = 0.0f;
-		Settings.TorqueLimit = 0.0f;
+		Settings.SwingTorqueLimit = 0.0f;
+		Settings.TwistTorqueLimit = 0.0f;
 
 		UpdateSingleBody(1, 0);
 
@@ -96,7 +97,8 @@ namespace ChaosTest
 		Data.GroundDistance = 10.0f;
 		Settings.TargetHeight = 15.0f;
 		Settings.RadialForceLimit = 0.0f;
-		Settings.TorqueLimit = 0.0f;
+		Settings.SwingTorqueLimit = 0.0f;
+		Settings.TwistTorqueLimit = 0.0f;
 
 		UpdateSingleBody(1, 0);
 
@@ -117,7 +119,8 @@ namespace ChaosTest
 		Data.GroundDistance = 16.0f;
 		Settings.TargetHeight = 15.0f;
 		Settings.RadialForceLimit = 0.0f;
-		Settings.TorqueLimit = 0.0f;
+		Settings.SwingTorqueLimit = 0.0f;
+		Settings.TwistTorqueLimit = 0.0f;
 
 		Solver.SetBodies(&CharacterBody);
 		Solver.GatherInput(Dt, Settings, Data);
@@ -142,7 +145,8 @@ namespace ChaosTest
 		Data.GroundDistance = 16.0f;
 		Settings.TargetHeight = 15.0f;
 		Settings.RadialForceLimit = 0.0f;
-		Settings.TorqueLimit = 0.0f;
+		Settings.SwingTorqueLimit = 0.0f;
+		Settings.TwistTorqueLimit = 0.0f;
 
 		Solver.SetBodies(&CharacterBody);
 		Solver.GatherInput(Dt, Settings, Data);
@@ -185,7 +189,8 @@ namespace ChaosTest
 		Data.GroundDistance = 10.0f;
 		Settings.TargetHeight = 15.0f;
 		Settings.RadialForceLimit = 0.0f;
-		Settings.TorqueLimit = 0.0f;
+		Settings.SwingTorqueLimit = 0.0f;
+		Settings.TwistTorqueLimit = 0.0f;
 		Data.GroundNormal = FVec3(1.0, 0.0, 1.0);
 		Data.GroundNormal.SafeNormalize();
 
@@ -214,7 +219,8 @@ namespace ChaosTest
 		Settings.RadialForceLimit = (1.0 / (CharacterBody.InvM() * Dt * Dt)) * 10.1;
 		Data.TargetDeltaPosition = FVec3(-5.0, 0.0, 0.0);
 		Data.TargetDeltaFacing = -0.4;
-		Settings.TorqueLimit = (1.0 / (CharacterBody.InvILocal().Z * Dt * Dt)) * 0.5;
+		Settings.SwingTorqueLimit = 0.0f;
+		Settings.TwistTorqueLimit = (1.0 / (CharacterBody.InvILocal().Z * Dt * Dt)) * 0.5;
 
 		UpdateSingleBody(1, 0);
 
@@ -247,7 +253,8 @@ namespace ChaosTest
 		Settings.RadialForceLimit = (1.0 / (CharacterBody.InvM() * Dt * Dt)) * 10.0;
 		Data.TargetDeltaPosition = FVec3(-11.0, 0.0, 0.0);
 		Data.TargetDeltaFacing = 0.6;
-		Settings.TorqueLimit = (1.0 / (CharacterBody.InvILocal().Z * Dt * Dt)) * 0.5;
+		Settings.SwingTorqueLimit = 0.0f;
+		Settings.TwistTorqueLimit = (1.0 / (CharacterBody.InvILocal().Z * Dt * Dt)) * 0.5;
 
 		UpdateSingleBody(5, 0);
 
@@ -265,7 +272,7 @@ namespace ChaosTest
 		FVec3 ExpectedImpulse = (1.0 / (CharacterBody.InvM() * Dt)) * FVec3(-10.0, 0.0, 5.0);
 		EXPECT_VECTOR_FLOAT_EQ(Solver.GetLinearImpulse(Dt), ExpectedImpulse);
 		EXPECT_FLOAT_EQ(FMath::Abs(Solver.GetLinearImpulse(Dt).X), Settings.RadialForceLimit * Dt);
-		EXPECT_FLOAT_EQ(FMath::Abs(Solver.GetAngularImpulse(Dt).Z), Settings.TorqueLimit * Dt);
+		EXPECT_FLOAT_EQ(FMath::Abs(Solver.GetAngularImpulse(Dt).Z), Settings.TwistTorqueLimit * Dt);
 	}
 
 	// No overlap with motion target
@@ -278,7 +285,8 @@ namespace ChaosTest
 		Settings.TargetHeight = 18.0f;
 		Settings.RadialForceLimit = (1.0 / (CharacterBody.InvM() * Dt * Dt)) * 10.1;
 		Data.TargetDeltaPosition = FVec3(-5.0, 0.0, 0.0);
-		Settings.TorqueLimit = 0.0f;
+		Settings.SwingTorqueLimit = 0.0f;
+		Settings.TwistTorqueLimit = 0.0f;
 		Settings.AssumedOnGroundHeight = 1.0f;
 
 		UpdateSingleBody(1, 0);
@@ -299,7 +307,8 @@ namespace ChaosTest
 		Settings.TargetHeight = 18.0f;
 		Settings.RadialForceLimit = (1.0 / (CharacterBody.InvM() * Dt * Dt)) * 10.1;
 		Data.TargetDeltaPosition = FVec3(-5.0, 0.0, 0.0);
-		Settings.TorqueLimit = 0.0f;
+		Settings.SwingTorqueLimit = 0.0f;
+		Settings.TwistTorqueLimit = 0.0f;
 		Settings.AssumedOnGroundHeight = 2.1f;
 
 		UpdateSingleBody(1, 0);
@@ -321,7 +330,8 @@ namespace ChaosTest
 		Settings.TargetHeight = 10.0f;
 		Settings.RadialForceLimit = (1.0 / (CharacterBody.InvM() * Dt * Dt)) * 100.0;
 		Data.TargetDeltaPosition = FVec3(10.0, 0.0, 0.0);
-		Settings.TorqueLimit = 0.0f;
+		Settings.SwingTorqueLimit = 0.0f;
+		Settings.TwistTorqueLimit = 0.0f;
 		Settings.AssumedOnGroundHeight = 1.0f;
 		Data.GroundNormal = FVec3(-1.0, 0.0, 1.0);
 		Data.GroundNormal.SafeNormalize();
@@ -350,7 +360,8 @@ namespace ChaosTest
 		Settings.TargetHeight = 10.0f;
 		Settings.RadialForceLimit = (1.0 / (CharacterBody.InvM() * Dt * Dt)) * 100.0;
 		Data.TargetDeltaPosition = FVec3(10.0, -5.0, 0.0);
-		Settings.TorqueLimit = 0.0f;
+		Settings.SwingTorqueLimit = 0.0f;
+		Settings.TwistTorqueLimit = 0.0f;
 		Settings.AssumedOnGroundHeight = 1.0f;
 
 		UpdateTwoBody(1, 0);
@@ -382,7 +393,8 @@ namespace ChaosTest
 		Settings.RadialForceLimit = (1.0 / (CharacterBody.InvM() * Dt * Dt)) * 100.0;
 		Data.TargetDeltaPosition = FVec3(0.0, 0.0, 0.0);
 		Data.TargetDeltaFacing = 0.0f;
-		Settings.TorqueLimit = (1.0 / (CharacterBody.InvILocal().Z * Dt * Dt)) * 100.0;
+		Settings.SwingTorqueLimit = 0.0f;
+		Settings.TwistTorqueLimit = (1.0 / (CharacterBody.InvILocal().Z * Dt * Dt)) * 100.0;
 		Settings.AssumedOnGroundHeight = 1.0f;
 
 		UpdateTwoBody(1, 0);
@@ -421,7 +433,8 @@ namespace ChaosTest
 		Data.GroundDistance = 10.0f;
 		Settings.TargetHeight = 15.0f;
 		Settings.RadialForceLimit = 0.0f;
-		Settings.TorqueLimit = 0.0f;
+		Settings.SwingTorqueLimit = 0.0f;
+		Settings.TwistTorqueLimit = 0.0f;
 
 		UpdateTwoBody(1, 0);
 
@@ -450,7 +463,8 @@ namespace ChaosTest
 		Data.GroundDistance = 20.0f;
 		Settings.TargetHeight = 10.0f;
 		Settings.RadialForceLimit = 0.0f;
-		Settings.TorqueLimit = 0.0f;
+		Settings.SwingTorqueLimit = 0.0f;
+		Settings.TwistTorqueLimit = 0.0f;
 
 		UpdateTwoBody(1, 0);
 
diff --git a/Engine/Source/Runtime/Experimental/Chaos/Private/Chaos/Character/CharacterGroundConstraintContainer.cpp b/Engine/Source/Runtime/Experimental/Chaos/Private/Chaos/Character/CharacterGroundConstraintContainer.cpp
--- a/Engine/Source/Runtime/Experimental/Chaos/Private/Chaos/Character/CharacterGroundConstraintContainer.cpp
+++ b/Engine/Source/Runtime/Experimental/Chaos/Private/Chaos/Character/CharacterGroundConstraintContainer.cpp
@@ -98,7 +98,8 @@ namespace Chaos
 		{
 			const bool bIsInGraph = Constraint->IsInConstraintGraph();
 			const bool bShouldBeInGraph = Constraint->IsEnabled() && CanEvaluate(Constraint);
-
+			const bool bShouldUpdateGraph = Constraint->GetGroundParticle() && Constraint->bGroundParticleChanged;
+			
 			if (bShouldBeInGraph && !bIsInGraph)
 			{
 				IslandManager.AddConstraint(ContainerId, Constraint, Constraint->GetConstrainedParticles());
@@ -107,6 +108,12 @@ namespace Chaos
 			{
 				IslandManager.RemoveConstraint(Constraint);
 			}
+			else if (bIsInGraph && bShouldUpdateGraph)
+			{
+				IslandManager.RemoveConstraint(Constraint);
+				IslandManager.AddConstraint(ContainerId, Constraint, Constraint->GetConstrainedParticles());
+				Constraint->bGroundParticleChanged = false;
+			}
 		}
 	}
 
diff --git a/Engine/Source/Runtime/Experimental/Chaos/Private/Chaos/Character/CharacterGroundConstraintSolver.cpp b/Engine/Source/Runtime/Experimental/Chaos/Private/Chaos/Character/CharacterGroundConstraintSolver.cpp
--- a/Engine/Source/Runtime/Experimental/Chaos/Private/Chaos/Character/CharacterGroundConstraintSolver.cpp
+++ b/Engine/Source/Runtime/Experimental/Chaos/Private/Chaos/Character/CharacterGroundConstraintSolver.cpp
@@ -45,6 +45,7 @@ namespace Chaos
 		void FCharacterGroundConstraintSolver::FImpulseData::Reset()
 		{
 			LinearPositionImpulse = FSolverVec3::ZeroVector;
+			AngularSwingImpulse = FSolverVec3::ZeroVector;
 			AngularImpulse = FSolverReal(0.0f);
 			LinearCorrectionImpulse = FSolverReal(0.0f);
 		}
@@ -53,16 +54,16 @@ namespace Chaos
 		/// FConstraintData
 
 		FCharacterGroundConstraintSolver::FConstraintData::FConstraintData()
-			: Normal(FSolverVec3(0.0, 0.0, 1.0)) // Need to initialize normal as it is used in GetLinearImpulse
-			, VerticalAxis(FSolverVec3(0.0, 0.0, 1.0)) // Need to vertical axis as it is used in GetAngularImpulse
+			: Normal(FSolverVec3(0.0f, 0.0f, 1.0f)) // Need to initialize normal as it is used in GetLinearImpulse
+			, VerticalAxis(FSolverVec3(0.0f, 0.0f, 1.0f)) // Need to vertical axis as it is used in GetAngularImpulse
 			, CharacterInvM(-1.0f)
-			, CharacterInvI(-1.0f)
+			, CharacterInvIZ(-1.0f)
 		{
 		}
 
 		bool FCharacterGroundConstraintSolver::FConstraintData::IsValid()
 		{
-			return (CharacterInvM > 0.0f) && (CharacterInvI > 0.0f);
+			return (CharacterInvM > 0.0f) && (CharacterInvIZ > 0.0f);
 		}
 
 		//////////////////////////////////////////////////////////////////////////
@@ -82,106 +83,147 @@ namespace Chaos
 		{
 			ImpulseData.Reset();
 
-			ConstraintData.Normal = FSolverVec3(Data.GroundNormal);
-			ConstraintData.VerticalAxis = FSolverVec3(Settings.VerticalAxis);
-			ConstraintData.AssumedOnGroundHeight = FSolverReal(Settings.AssumedOnGroundHeight);
+			//////////////////////////////////////////////////////////////////////////
+			// Normal constraint
 
-			// Convert force/torque limits into solver impulse limits
-			ConstraintData.AngularImpulseLimit = FSolverReal(Settings.TorqueLimit * Dt * Dt);
-			ConstraintData.RadialImpulseLimit = FSolverReal(Settings.RadialForceLimit * Dt * Dt);
+			const FVec3 Character_X = BodyData.CharacterBody.X();
+			const FVec3 Character_P = BodyData.CharacterBody.P();
+			const FRotation3 Character_R(BodyData.CharacterBody.R());
+			const FRotation3 Character_Q(BodyData.CharacterBody.Q());
 
-			FRotation3 RelQuatCharacter = BodyData.CharacterBody.R().Inverse() * BodyData.CharacterBody.Q();
-			FReal CharacterAngularDisplacement = FVec3::DotProduct(RelQuatCharacter.ToRotationVector(), Settings.VerticalAxis);
-			FVec3 CharacterLinearDisplacement = BodyData.CharacterBody.P() - BodyData.CharacterBody.X();
+			const FVec3 PC_Init = Character_X - FVec3(Settings.TargetHeight * Settings.VerticalAxis);
+			const FVec3 PG_Init = Character_X - FVec3(Data.GroundDistance * Settings.VerticalAxis);
+			const FRotation3 RelQuatCharacter = Character_Q * Character_R.Inverse();
+			const FReal CharacterAngularDisplacement = RelQuatCharacter.ToRotationVector().Dot(Settings.VerticalAxis);
+			const FVec3 CharacterLinearDisplacement = FVec3(Character_P - Character_X);
 
-			// Vertical ground constraint
-			FVec3 PC_Init = BodyData.CharacterBody.X() - Settings.TargetHeight * Settings.VerticalAxis;
-			FVec3 PG_Init = BodyData.CharacterBody.X() - Data.GroundDistance * Settings.VerticalAxis;
-			FVec3 InitialError = PC_Init - PG_Init;
+			FVec3 InitialError = FVec3(PC_Init - PG_Init);
 			FVec3 ProjectedError = InitialError + CharacterLinearDisplacement;
 
-			// Motion Target
-			FVec3 MotionTarget = BodyData.CharacterBody.X() + Data.TargetDeltaPosition;
-
-			// If the slope is too steep adjust the target to not point up the slope
-			FReal DP = FVec3::DotProduct(Settings.VerticalAxis, Data.GroundNormal);
-			if (DP <= Settings.CosMaxWalkableSlopeAngle)
-			{
-				FVec3 UpSlope = Settings.VerticalAxis - DP * Data.GroundNormal;
-				UpSlope.Normalize();
-				FReal UpMotion = FVec3::DotProduct(MotionTarget - BodyData.CharacterBody.X(), UpSlope);
-				if (UpMotion > FReal(0.0f))
-				{
-					MotionTarget -= UpMotion * UpSlope;
-				}
-			}
+			const FReal CharacterInvM = BodyData.CharacterBody.InvM();
+			FReal TwoBodyEffectiveMass = 0.0;
+			FReal GroundInvM = 0.0;
+			FMatrix33 GroundInvI(0.0);
+			FVec3 GroundNormal = Data.GroundNormal;
+			FVec3 GroundOffset = FVec3(ForceInitToZero);
 
-			FVec3 MotionTargetError = BodyData.CharacterBody.P() - MotionTarget;
-			FReal MotionTargetAngularError = CharacterAngularDisplacement - Data.TargetDeltaFacing;
+			FRotation3 RelQuatGround = FRotation3::Identity;
+			FVec3 GroundLinearDisplacement(ForceInitToZero);
+			FVec3 Ground_X(ForceInitToZero);
+			FVec3 Ground_P(ForceInitToZero);
 
-			if (BodyData.GroundBody.IsValid())
+			const bool bTwoBody = BodyData.GroundBody.IsValid();
+			if (bTwoBody)
 			{
-				FVec3 RC_Init = PC_Init - BodyData.GroundBody.X();
-				FVec3 RG_Init = PG_Init - BodyData.GroundBody.X();
+				Ground_X = BodyData.GroundBody.X();
+				Ground_P = BodyData.GroundBody.P();
+
+				const FRotation3 Ground_R(BodyData.GroundBody.R());
+				const FRotation3 Ground_Q(BodyData.GroundBody.Q());
+				RelQuatGround = Ground_Q * Ground_R.Inverse();
 
-				FRotation3 RelQuatGround = BodyData.GroundBody.R().Inverse() * BodyData.GroundBody.Q();
-				FReal GroundAngularDisplacement = FVec3::DotProduct(RelQuatGround.ToRotationVector(), Settings.VerticalAxis);
-				FVec3 GroundLinearDisplacement = BodyData.GroundBody.P() - BodyData.GroundBody.X() + RelQuatGround * RG_Init - RG_Init;
+				const FVec3 RC_Init = FVec3(PC_Init - Ground_X);
+				const FVec3 RG_Init = FVec3(PG_Init - Ground_X);
 
-				ConstraintData.GroundOffset = RelQuatGround * RC_Init;
+				GroundLinearDisplacement = FVec3(Ground_P - Ground_X) + RelQuatGround * RG_Init - RG_Init;
+				GroundOffset = RelQuatGround * RC_Init;
+				GroundNormal = RelQuatGround * GroundNormal;
 				ProjectedError -= GroundLinearDisplacement;
 
-				MotionTargetError -= GroundLinearDisplacement;
-				MotionTargetAngularError -= GroundAngularDisplacement;
+				GroundInvM = BodyData.GroundBody.InvM();
+				GroundInvI = Utilities::ComputeWorldSpaceInertia(Ground_Q, BodyData.GroundBody.InvILocal());
+				const FVec3 R_Cross_N = GroundOffset.Cross(GroundNormal);
+				TwoBodyEffectiveMass = 1.0f / (CharacterInvM + GroundInvM + R_Cross_N.Dot(GroundInvI * R_Cross_N));
 			}
 
-			bool CurrentlyOverlapping = FVec3::DotProduct(InitialError, Data.GroundNormal) < FReal(0.0);
-			FReal ProjectedOverlap = FVec3::DotProduct(ProjectedError, -Data.GroundNormal);
-			bool WillOverlap = ProjectedOverlap > FReal(0.0);
-			bool NeedsCorrection = CurrentlyOverlapping && WillOverlap;
+			// Vertical damping
+			const FReal ProjectedOverlap = -ProjectedError.Dot(GroundNormal);
+			const bool WillOverlap = ProjectedOverlap > 0.0;
 
 			if (WillOverlap)
 			{
-				FReal DampingFactor = FMath::Clamp(Settings.DampingFactor * Dt, 0.0, 1.0);
-				FVec3 Diff = DampingFactor * ProjectedOverlap * Data.GroundNormal;
+				FReal DampingFactor = FMath::Clamp(Settings.DampingFactor, 0.0f, 1.0f);
+				FVec3 Diff = DampingFactor * ProjectedOverlap * GroundNormal;
 				InitialError += Diff;
 				ProjectedError += Diff;
 			}
 
+			// Correction
+			const bool CurrentlyOverlapping = InitialError.Dot(GroundNormal) < 0.0f;
+			const bool NeedsCorrection = CurrentlyOverlapping && WillOverlap;
+			CorrectionSolveFunction = NeedsCorrection ? &SolveCorrectionSingleBody : &NoSolve;
+
+			// Write constraint data
+			ConstraintData.CharacterInvM = FSolverReal(CharacterInvM);
+			ConstraintData.GroundInvM = FSolverReal(GroundInvM);
+			ConstraintData.GroundInvI = FSolverMatrix33(GroundInvI);
+			ConstraintData.SingleBodyEffectiveMass = FSolverReal(1.0 / CharacterInvM);
+			ConstraintData.TwoBodyEffectiveMass = FSolverReal(TwoBodyEffectiveMass);
+			ConstraintData.Normal = FSolverVec3(GroundNormal);
+			ConstraintData.GroundOffset = FSolverVec3(GroundOffset);
 			ConstraintData.InitialError = FSolverVec3(InitialError);
 			ConstraintData.InitialProjectedError = FSolverVec3(ProjectedError);
-			ConstraintData.MotionTargetError = FSolverVec3(MotionTargetError);
-			ConstraintData.MotionTargetAngularError = FSolverReal(MotionTargetAngularError);
 
-			FSolverMatrix33 WorldSpaceCharacterInvI = FSolverMatrix33(Utilities::ComputeWorldSpaceInertia(BodyData.CharacterBody.Q(), BodyData.CharacterBody.InvILocal()));
-			ConstraintData.CharacterInvI = FSolverVec3::DotProduct(ConstraintData.VerticalAxis, WorldSpaceCharacterInvI * ConstraintData.VerticalAxis);
-			ConstraintData.CharacterInvM = BodyData.CharacterBody.InvM();
-			ConstraintData.SingleBodyEffectiveMass = FSolverReal(1.0f / BodyData.CharacterBody.InvM());
-			ConstraintData.EffectiveInertia = FSolverReal(1.0f / ConstraintData.CharacterInvI);
+			//////////////////////////////////////////////////////////////////////////
+			// Upright constraint
+
+			ConstraintData.CharacterInvI = FSolverMatrix33(Utilities::ComputeWorldSpaceInertia(Character_Q, BodyData.CharacterBody.InvILocal()));
+			ConstraintData.CharacterVerticalAxis = FSolverVec3(Character_Q.RotateVector(Settings.VerticalAxis));
+
+			//////////////////////////////////////////////////////////////////////////
+			// Motion Target
+
+			FVec3 TargetDeltaPos = Data.TargetDeltaPosition;
+			FReal TargetDeltaFacing = Data.TargetDeltaFacing;
 
-			if (BodyData.IsTwoBody())
+			if (BodyData.GroundBody.IsValid())
 			{
-				ConstraintData.GroundInvM = BodyData.GroundBody.InvM();
-				ConstraintData.GroundInvI = FSolverMatrix33(Utilities::ComputeWorldSpaceInertia(BodyData.GroundBody.Q(), BodyData.GroundBody.InvILocal()));
-				ConstraintData.TwoBodyEffectiveMass = FSolverReal(1.0f) / FSolverVec3::DotProduct(Data.GroundNormal, Utilities::ComputeJointFactorMatrix(ConstraintData.GroundOffset, ConstraintData.GroundInvI, ConstraintData.GroundInvM + ConstraintData.CharacterInvM) * Data.GroundNormal);
-				PositionSolveFunction = &SolvePositionTwoBody;
+				TargetDeltaPos = GroundLinearDisplacement + RelQuatGround * TargetDeltaPos;
+				TargetDeltaFacing += RelQuatGround.ToRotationVector().Dot(Settings.VerticalAxis);
 			}
-			else
+
+			// If the slope is too steep adjust the target to not point up the slope
+			FReal DP = FVec3::DotProduct(Settings.VerticalAxis, GroundNormal);
+			if (DP <= Data.CosMaxWalkableSlopeAngle)
 			{
-				PositionSolveFunction = &SolvePositionSingleBody;
+				FVec3 UpSlope = Settings.VerticalAxis - DP * GroundNormal;
+				UpSlope.Normalize();
+				FReal UpMotion = TargetDeltaPos.Dot(UpSlope);
+				if (UpMotion > FReal(0.0))
+				{
+					TargetDeltaPos -= UpMotion * FVec3(UpSlope);
+				}
 			}
-			// Note: It would be good to have a two body correction function, but that doesn't work so well
-			// if contacts and other constraints don't support correction
-			CorrectionSolveFunction = NeedsCorrection ? &SolveCorrectionSingleBody : &NoSolve;
+
+			const FVec3 MotionTarget = Character_X + TargetDeltaPos;
+			const FVec3 MotionTargetError = FVec3(Character_P - MotionTarget);
+			const FReal MotionTargetAngularError = CharacterAngularDisplacement - TargetDeltaFacing;
+
+			const FMatrix33 WorldSpaceCharacterInvI = Utilities::ComputeWorldSpaceInertia(BodyData.CharacterBody.Q(), BodyData.CharacterBody.InvILocal());
+			const FReal CharacterInvI = Settings.VerticalAxis.Dot(WorldSpaceCharacterInvI * Settings.VerticalAxis);
+			const FReal DtSq = Dt * Dt;
+
+			// Write constraint data
+			ConstraintData.AngularImpulseLimit = FSolverReal(Settings.TwistTorqueLimit * DtSq);
+			ConstraintData.AngularSwingImpulseLimit = FSolverReal(Settings.SwingTorqueLimit * DtSq);
+			ConstraintData.RadialImpulseLimit = FSolverReal(Settings.RadialForceLimit * DtSq);
+			ConstraintData.VerticalAxis = FSolverVec3(Settings.VerticalAxis);
+			ConstraintData.AssumedOnGroundHeight = FSolverReal(Settings.AssumedOnGroundHeight);
+			ConstraintData.CharacterInvIZ = FSolverReal(CharacterInvI);
+			ConstraintData.EffectiveInertia = FSolverReal(1.0 / ConstraintData.CharacterInvIZ);
+			ConstraintData.MotionTargetError = FSolverVec3(MotionTargetError);
+			ConstraintData.MotionTargetAngularError = FSolverReal(MotionTargetAngularError);
+
+			PositionSolveFunction = bTwoBody ? &SolvePositionTwoBody : SolvePositionSingleBody;
 		}
 
-		void FCharacterGroundConstraintSolver::ScatterOutput(const FReal Dt, FVec3& OutSolverAppliedForce, FVec3& OutSolverAppliedTorque)
+		void FCharacterGroundConstraintSolver::ScatterOutput(FReal Dt, FVec3& OutSolverAppliedForce, FVec3& OutSolverAppliedTorque)
 		{
 			if (Dt > UE_SMALL_NUMBER)
 			{
 				FReal FreqSq = 1.0 / (Dt * Dt);
 				OutSolverAppliedForce = FVec3(ImpulseData.LinearPositionImpulse + ImpulseData.LinearCorrectionImpulse * ConstraintData.Normal) * FreqSq;
-				OutSolverAppliedTorque = FVec3(ImpulseData.AngularImpulse * ConstraintData.VerticalAxis) * FreqSq;
+				OutSolverAppliedTorque = FVec3(ImpulseData.AngularImpulse * ConstraintData.VerticalAxis + ImpulseData.AngularSwingImpulse) * FreqSq;
 			}
 			else
 			{
diff --git a/Engine/Source/Runtime/Experimental/Chaos/Private/Chaos/Character/CharacterGroundConstraintSolver.h b/Engine/Source/Runtime/Experimental/Chaos/Private/Chaos/Character/CharacterGroundConstraintSolver.h
--- a/Engine/Source/Runtime/Experimental/Chaos/Private/Chaos/Character/CharacterGroundConstraintSolver.h
+++ b/Engine/Source/Runtime/Experimental/Chaos/Private/Chaos/Character/CharacterGroundConstraintSolver.h
@@ -64,6 +64,7 @@ namespace Chaos
 				void Reset();
 
 				FSolverVec3 LinearPositionImpulse;
+				FSolverVec3 AngularSwingImpulse;
 				FSolverReal AngularImpulse;
 				FSolverReal LinearCorrectionImpulse;
 			} ImpulseData;
@@ -74,23 +75,65 @@ namespace Chaos
 				FConstraintData();
 				bool IsValid();
 
-				FSolverMatrix33 GroundInvI;				/// World space ground body inverse inertia
-				FSolverVec3 GroundOffset;				/// Offset vector from ground body CoM to the constraint position
-				FSolverVec3 InitialError;				/// Constraint error pre-integration
-				FSolverVec3 InitialProjectedError;		/// Projected constraint error post-integration
-				FSolverVec3 Normal;						///	Ground plane normal direction
-				FSolverVec3 VerticalAxis;				/// World space vertical axis
-				FSolverVec3 MotionTargetError;			/// Projected constraint error for motion target constraint
-				FSolverReal MotionTargetAngularError;	/// Projected angular facing error for motion target constraint
-				FSolverReal CharacterInvM;				/// Character inverse mass
-				FSolverReal CharacterInvI;				/// Character inverse inertia about the vertical axis
-				FSolverReal GroundInvM;					/// Ground body inverse mass
-				FSolverReal TwoBodyEffectiveMass;		/// Effective mass for the two body system
-				FSolverReal SingleBodyEffectiveMass;	/// Effective mass for a single body constraint (the character mass)
-				FSolverReal EffectiveInertia;			/// Effective angular mass to rotate about the vertical axis
-				FSolverReal AngularImpulseLimit;		/// Angular impulse from the solver is clamped to this limit
-				FSolverReal RadialImpulseLimit;			/// Linear impulse from the solver is clamped to this limit
-				FSolverReal AssumedOnGroundHeight;		/// Height below which the character is assumed to be grounded
+				/// World space ground body inverse inertia
+				FSolverMatrix33 CharacterInvI;
+
+				/// World space ground body inverse inertia
+				FSolverMatrix33 GroundInvI;
+
+				/// Offset vector from ground body CoM to the constraint position
+				FSolverVec3 GroundOffset;
+
+				/// Constraint error pre-integration
+				FSolverVec3 InitialError;
+
+				/// Projected constraint error post-integration
+				FSolverVec3 InitialProjectedError;
+
+				///	Ground plane normal direction
+				FSolverVec3 Normal;
+
+				/// World space vertical axis
+				FSolverVec3 VerticalAxis;
+
+				/// Vertical axis rotated by the character initial rotation
+				FSolverVec3 CharacterVerticalAxis;
+
+				/// Projected constraint error for motion target constraint
+				FSolverVec3 MotionTargetError;
+
+				/// Projected angular facing error for motion target constraint
+				FSolverReal MotionTargetAngularError;
+
+				/// Character inverse mass
+				FSolverReal CharacterInvM;
+
+				/// Character inverse inertia about the vertical axis
+				FSolverReal CharacterInvIZ;
+
+				/// Ground body inverse mass
+				FSolverReal GroundInvM;
+
+				/// Effective mass for the two body system
+				FSolverReal TwoBodyEffectiveMass;
+
+				/// Effective mass for a single body constraint (the character mass)
+				FSolverReal SingleBodyEffectiveMass;
+
+				/// Effective angular mass to rotate about the vertical axis
+				FSolverReal EffectiveInertia;
+
+				/// Angular impulse from the solver is clamped to this limit
+				FSolverReal AngularImpulseLimit;
+
+				/// Angular swing impulse from the solver is clamped to this limit
+				FSolverReal AngularSwingImpulseLimit;
+
+				/// Radial linear impulse from the solver is clamped to this limit
+				FSolverReal RadialImpulseLimit;
+
+				/// Height below which the character is assumed to be grounded
+				FSolverReal AssumedOnGroundHeight;
 			} ConstraintData;
 
 			using FSolveFunctionType = void (*)(const FConstraintData& ConstraintData, FBodyData& BodyData, FImpulseData& ImpulseData);
@@ -175,79 +218,110 @@ namespace Chaos
 
 		FORCEINLINE void FCharacterGroundConstraintSolver::SolveCorrectionSingleBody(const FConstraintData& ConstraintData, FBodyData& BodyData, FImpulseData& ImpulseData)
 		{
-			FSolverReal Error = FSolverVec3::DotProduct(BodyData.CharacterBody.CP() + ConstraintData.InitialError, ConstraintData.Normal);
+			const FSolverReal Error = ConstraintData.Normal.Dot(BodyData.CharacterBody.CP() + ConstraintData.InitialError);
 			if (Error < FSolverReal(0.0f))
 			{
-				FSolverReal Delta = -ConstraintData.SingleBodyEffectiveMass * Error;
+				const FSolverReal Delta = -ConstraintData.SingleBodyEffectiveMass * Error;
 				ImpulseData.LinearCorrectionImpulse += Delta;
-				BodyData.CharacterBody.ApplyPositionCorrectionDelta(ConstraintData.CharacterInvM * Delta * ConstraintData.Normal);
+				BodyData.CharacterBody.ApplyPositionCorrectionDelta(FSolverVec3(ConstraintData.CharacterInvM * Delta * ConstraintData.Normal));
 			}
 		}
 
 		FORCEINLINE void FCharacterGroundConstraintSolver::SolvePositionSingleBody(const FConstraintData& ConstraintData, FBodyData& BodyData, FImpulseData& ImpulseData)
 		{
 			// Normal
-			FSolverReal Error = FSolverVec3::DotProduct(BodyData.CharacterBody.DP() + ConstraintData.InitialProjectedError, ConstraintData.Normal);
+			const FSolverReal Error = ConstraintData.Normal.Dot(BodyData.CharacterBody.DP() + ConstraintData.InitialProjectedError);
 			if (Error < FSolverReal(0.0f))
 			{
-				FSolverVec3 Delta = -ConstraintData.SingleBodyEffectiveMass * Error * ConstraintData.Normal;
+				const FSolverVec3 Delta = -ConstraintData.SingleBodyEffectiveMass * Error * ConstraintData.Normal;
 				ImpulseData.LinearPositionImpulse += Delta;
-				BodyData.CharacterBody.ApplyPositionDelta(ConstraintData.CharacterInvM * Delta);
+				BodyData.CharacterBody.ApplyPositionDelta(FSolverVec3(ConstraintData.CharacterInvM * Delta));
+			}
+
+			// Angular constraint
+			FSolverVec3 NewCharacterVerticalAxis = ConstraintData.CharacterVerticalAxis + BodyData.CharacterBody.DQ().Cross(ConstraintData.CharacterVerticalAxis);
+			NewCharacterVerticalAxis.Normalize();
+			const FSolverVec3 CrossProd = NewCharacterVerticalAxis.Cross(ConstraintData.VerticalAxis);
+			const FSolverReal SizeSq = CrossProd.SizeSquared();
+			if (SizeSq > UE_SMALL_NUMBER)
+			{
+				const FSolverVec3 AngAxis = CrossProd * FMath::InvSqrt(SizeSq);
+				const FSolverReal AngResistance = FSolverReal(1.0f) / (ConstraintData.CharacterInvI * AngAxis).Dot(AngAxis);
+				const FSolverVec3 NewSwingImpulse = ClampMagnitude(ImpulseData.AngularSwingImpulse + AngResistance * FMath::Asin(FMath::Sqrt(SizeSq)) * AngAxis, ConstraintData.AngularSwingImpulseLimit);
+				const FSolverVec3 Delta = NewSwingImpulse - ImpulseData.AngularSwingImpulse;
+				ImpulseData.AngularSwingImpulse = NewSwingImpulse;
+				BodyData.CharacterBody.ApplyRotationDelta(FSolverVec3(ConstraintData.CharacterInvI * Delta));
 			}
 
 			const FSolverReal NormalImpulse = FSolverVec3::DotProduct(ImpulseData.LinearPositionImpulse, ConstraintData.Normal);
-			if (((NormalImpulse + ImpulseData.LinearCorrectionImpulse) > FSolverReal(0.0f)) || Error < ConstraintData.AssumedOnGroundHeight)
+			const FSolverReal MinNormalImpulse(0.0f);
+			if (((NormalImpulse + ImpulseData.LinearCorrectionImpulse) > MinNormalImpulse) || Error < ConstraintData.AssumedOnGroundHeight)
 			{
 				// Target Position
-				FSolverVec3 MotionTargetError = ProjectOntoPlane(ConstraintData.MotionTargetError + BodyData.CharacterBody.DP(), ConstraintData.Normal);
-				FSolverVec3 InitialMotionTargetImpulse = ProjectOntoPlane(ImpulseData.LinearPositionImpulse, ConstraintData.Normal);
+				const FSolverVec3 MotionTargetError = ProjectOntoPlane(ConstraintData.MotionTargetError + BodyData.CharacterBody.DP(), ConstraintData.Normal);
+				const FSolverVec3 InitialMotionTargetImpulse = ProjectOntoPlane(ImpulseData.LinearPositionImpulse, ConstraintData.Normal);
 				FSolverVec3 NewMotionTargetImpulse = ClampMagnitude(InitialMotionTargetImpulse - ConstraintData.SingleBodyEffectiveMass * MotionTargetError, ConstraintData.RadialImpulseLimit);
-				FSolverVec3 Delta = NewMotionTargetImpulse - InitialMotionTargetImpulse;
+				const FSolverVec3 Delta = NewMotionTargetImpulse - InitialMotionTargetImpulse;
 				ImpulseData.LinearPositionImpulse += Delta;
-				BodyData.CharacterBody.ApplyPositionDelta(ConstraintData.CharacterInvM * Delta);
+				BodyData.CharacterBody.ApplyPositionDelta(FSolverVec3(ConstraintData.CharacterInvM * Delta));
 
 				// Target Rotation
-				FSolverReal MotionTargetAngularError = ConstraintData.MotionTargetAngularError + FSolverVec3::DotProduct(BodyData.CharacterBody.DQ(), ConstraintData.VerticalAxis);
-				FSolverReal NewAngularImpulse = ClampAbs(ImpulseData.AngularImpulse - ConstraintData.EffectiveInertia * MotionTargetAngularError, ConstraintData.AngularImpulseLimit);
-				FSolverReal AngularDelta = NewAngularImpulse - ImpulseData.AngularImpulse;
+				const FSolverReal MotionTargetAngularError = ConstraintData.MotionTargetAngularError + ConstraintData.VerticalAxis.Dot(BodyData.CharacterBody.DQ());
+				const FSolverReal NewAngularImpulse = ClampAbs(ImpulseData.AngularImpulse - ConstraintData.EffectiveInertia * MotionTargetAngularError, ConstraintData.AngularImpulseLimit);
+				const FSolverReal AngularDelta = NewAngularImpulse - ImpulseData.AngularImpulse;
 				ImpulseData.AngularImpulse += AngularDelta;
-				BodyData.CharacterBody.ApplyRotationDelta(ConstraintData.CharacterInvI * AngularDelta * ConstraintData.VerticalAxis);
+				BodyData.CharacterBody.ApplyRotationDelta(FSolverVec3(ConstraintData.CharacterInvI * AngularDelta * ConstraintData.VerticalAxis));
 			}
 		}
 
 		FORCEINLINE void FCharacterGroundConstraintSolver::SolvePositionTwoBody(const FConstraintData& ConstraintData, FBodyData& BodyData, FImpulseData& ImpulseData)
 		{
-			FSolverReal Error = FSolverVec3::DotProduct(ConstraintData.InitialProjectedError + BodyData.CharacterBody.DP() - BodyData.GroundBody.DP()
-				- FSolverVec3::CrossProduct(BodyData.GroundBody.DQ(), ConstraintData.GroundOffset), ConstraintData.Normal);
+			const FSolverReal Error = ConstraintData.Normal.Dot(ConstraintData.InitialProjectedError + BodyData.CharacterBody.DP()
+				- BodyData.GroundBody.DP() - BodyData.GroundBody.DQ().Cross(ConstraintData.GroundOffset));
 			if (Error < FSolverReal(0.0f))
 			{
-				FSolverVec3 Delta = -ConstraintData.TwoBodyEffectiveMass * Error * ConstraintData.Normal;
+				const FSolverVec3 Delta = -ConstraintData.TwoBodyEffectiveMass * Error * ConstraintData.Normal;
 				ImpulseData.LinearPositionImpulse += Delta;
-				BodyData.CharacterBody.ApplyPositionDelta(ConstraintData.CharacterInvM * Delta);
-				BodyData.GroundBody.ApplyPositionDelta(-ConstraintData.GroundInvM * Delta);
-				BodyData.GroundBody.ApplyRotationDelta(-ConstraintData.GroundInvI * FSolverVec3::CrossProduct(ConstraintData.GroundOffset, Delta));
+				BodyData.CharacterBody.ApplyPositionDelta(FSolverVec3(ConstraintData.CharacterInvM * Delta));
+				BodyData.GroundBody.ApplyPositionDelta(FSolverVec3(-ConstraintData.GroundInvM * Delta));
+				BodyData.GroundBody.ApplyRotationDelta(FSolverVec3(-ConstraintData.GroundInvI * ConstraintData.GroundOffset.Cross(Delta)));
+			}
+
+			// Angular constraint
+			FSolverVec3 NewCharacterVerticalAxis = ConstraintData.CharacterVerticalAxis + BodyData.CharacterBody.DQ().Cross(ConstraintData.CharacterVerticalAxis);
+			NewCharacterVerticalAxis.Normalize();
+			const FSolverVec3 CrossProd = NewCharacterVerticalAxis.Cross(ConstraintData.VerticalAxis);
+			const FSolverReal SizeSq = CrossProd.SizeSquared();
+			if (SizeSq > UE_SMALL_NUMBER)
+			{
+				const FSolverVec3 AngAxis = CrossProd * FMath::InvSqrt(SizeSq);
+				const FSolverReal AngResistance = FSolverReal(1.0f) / (ConstraintData.CharacterInvI * AngAxis).Dot(AngAxis);
+				const FSolverVec3 NewSwingImpulse = ClampMagnitude(ImpulseData.AngularSwingImpulse + AngResistance * FMath::Asin(FMath::Sqrt(SizeSq)) * AngAxis, ConstraintData.AngularSwingImpulseLimit);
+				const FSolverVec3 Delta = NewSwingImpulse - ImpulseData.AngularSwingImpulse;
+				ImpulseData.AngularSwingImpulse = NewSwingImpulse;
+				BodyData.CharacterBody.ApplyRotationDelta(FSolverVec3(ConstraintData.CharacterInvI * Delta));
 			}
 
 			const FSolverReal NormalImpulse = FSolverVec3::DotProduct(ImpulseData.LinearPositionImpulse, ConstraintData.Normal);
 			const FSolverReal MinNormalImpulse(0.0f);
-			if (((NormalImpulse + ImpulseData.LinearCorrectionImpulse) > MinNormalImpulse) || Error < FSolverReal(5.0f))
+			if (((NormalImpulse + ImpulseData.LinearCorrectionImpulse) > MinNormalImpulse) || Error < ConstraintData.AssumedOnGroundHeight)
 			{
 				// Target Position
 				FSolverVec3 MotionTargetError = ConstraintData.MotionTargetError + BodyData.CharacterBody.DP() - BodyData.GroundBody.DP();
-				MotionTargetError -= FSolverVec3::CrossProduct(BodyData.GroundBody.DQ(), MotionTargetError);
+				MotionTargetError -= BodyData.GroundBody.DQ().Cross(MotionTargetError);
 				MotionTargetError = ProjectOntoPlane(MotionTargetError, ConstraintData.Normal);
-				FSolverVec3 InitialMotionTargetImpulse = ProjectOntoPlane(ImpulseData.LinearPositionImpulse, ConstraintData.Normal);
-				FSolverVec3 NewMotionTargetImpulse = ClampMagnitude(InitialMotionTargetImpulse - ConstraintData.SingleBodyEffectiveMass* MotionTargetError, ConstraintData.RadialImpulseLimit);
-				FSolverVec3 Delta = NewMotionTargetImpulse - InitialMotionTargetImpulse;
+				const FSolverVec3 InitialMotionTargetImpulse = ProjectOntoPlane(ImpulseData.LinearPositionImpulse, ConstraintData.Normal);
+				FSolverVec3 NewMotionTargetImpulse = ClampMagnitude(InitialMotionTargetImpulse - ConstraintData.SingleBodyEffectiveMass * MotionTargetError, ConstraintData.RadialImpulseLimit);
+				const FSolverVec3 Delta = NewMotionTargetImpulse - InitialMotionTargetImpulse;
 				ImpulseData.LinearPositionImpulse += Delta;
-				BodyData.CharacterBody.ApplyPositionDelta(ConstraintData.CharacterInvM * Delta);
+				BodyData.CharacterBody.ApplyPositionDelta(FSolverVec3(ConstraintData.CharacterInvM * Delta));
 
 				// Target Rotation
-				FSolverReal MotionTargetAngularError = ConstraintData.MotionTargetAngularError + FSolverVec3::DotProduct(BodyData.CharacterBody.DQ() - BodyData.GroundBody.DQ(), ConstraintData.VerticalAxis);
-				FSolverReal NewAngularImpulse = ClampAbs(ImpulseData.AngularImpulse - ConstraintData.EffectiveInertia * MotionTargetAngularError, ConstraintData.AngularImpulseLimit);
-				FSolverReal AngularDelta = NewAngularImpulse - ImpulseData.AngularImpulse;
+				const FSolverReal MotionTargetAngularError = ConstraintData.MotionTargetAngularError + ConstraintData.VerticalAxis.Dot(BodyData.CharacterBody.DQ() - BodyData.GroundBody.DQ());
+				const FSolverReal NewAngularImpulse = ClampAbs(ImpulseData.AngularImpulse - ConstraintData.EffectiveInertia * MotionTargetAngularError, ConstraintData.AngularImpulseLimit);
+				const FSolverReal AngularDelta = NewAngularImpulse - ImpulseData.AngularImpulse;
 				ImpulseData.AngularImpulse += AngularDelta;
-				BodyData.CharacterBody.ApplyRotationDelta(ConstraintData.CharacterInvI * AngularDelta * ConstraintData.VerticalAxis);
+				BodyData.CharacterBody.ApplyRotationDelta(FSolverVec3(ConstraintData.CharacterInvI * AngularDelta * ConstraintData.VerticalAxis));
 			}
 		}
 
diff --git a/Engine/Source/Runtime/Experimental/Chaos/Private/Chaos/PullPhysicsDataImp.cpp b/Engine/Source/Runtime/Experimental/Chaos/Private/Chaos/PullPhysicsDataImp.cpp
--- a/Engine/Source/Runtime/Experimental/Chaos/Private/Chaos/PullPhysicsDataImp.cpp
+++ b/Engine/Source/Runtime/Experimental/Chaos/Private/Chaos/PullPhysicsDataImp.cpp
@@ -12,6 +12,7 @@ void FPullPhysicsData::Reset()
 	DirtyGeometryCollections.Reset();
 	DirtyClusterUnions.Reset();
 	DirtyJointConstraints.Reset();
+	DirtyCharacterGroundConstraints.Reset();
 }
 
 }
\ No newline at end of file
diff --git a/Engine/Source/Runtime/Experimental/Chaos/Private/PhysicsProxy/CharacterGroundConstraintProxy.cpp b/Engine/Source/Runtime/Experimental/Chaos/Private/PhysicsProxy/CharacterGroundConstraintProxy.cpp
--- a/Engine/Source/Runtime/Experimental/Chaos/Private/PhysicsProxy/CharacterGroundConstraintProxy.cpp
+++ b/Engine/Source/Runtime/Experimental/Chaos/Private/PhysicsProxy/CharacterGroundConstraintProxy.cpp
@@ -99,6 +99,7 @@ namespace Chaos
 				if (FGeometryParticleHandle* Handle = GetParticleHandleFromProxy(GroundParticleProxy->ParticleProxy))
 				{
 					Constraint_PT->GroundParticle = Handle;
+					Constraint_PT->bGroundParticleChanged = true;
 				}
 			}
 		}
@@ -148,6 +149,11 @@ namespace Chaos
 		{
 			Buffer.Force = Constraint_PT->GetSolverAppliedForce();
 			Buffer.Torque = Constraint_PT->GetSolverAppliedTorque();
+			Buffer.GroundNormal = Constraint_PT->GetData().GroundNormal;
+			Buffer.GroundDistance = Constraint_PT->GetData().GroundDistance;
+			Buffer.GroundParticle = Constraint_PT->GetGroundParticle();
+			Buffer.TargetDeltaPos = Constraint_PT->GetData().TargetDeltaPosition;
+			Buffer.TargetDeltaFacing = Constraint_PT->GetData().TargetDeltaFacing;
 		}
 	}
 
@@ -157,6 +163,30 @@ namespace Chaos
 		{
 			Constraint_GT->SolverAppliedForce = Buffer.Force;
 			Constraint_GT->SolverAppliedTorque = Buffer.Torque;
+
+			Constraint_GT->ConstraintData.Modify(false, Constraint_GT->DirtyFlags, Constraint_GT->Proxy, [&Buffer](FCharacterGroundConstraintDynamicData& Data) {
+				Data.GroundDistance = Buffer.GroundDistance;
+				Data.GroundNormal = Buffer.GroundNormal;
+				Data.TargetDeltaPosition = Buffer.TargetDeltaPos;
+				Data.TargetDeltaFacing = Buffer.TargetDeltaFacing;
+				});
+
+			if (Buffer.GroundParticle)
+			{
+				FSingleParticlePhysicsProxy* NewGroundProxy = (FSingleParticlePhysicsProxy*)(Buffer.GroundParticle->PhysicsProxy());
+				Constraint_GT->GroundProxy.Modify(false, Constraint_GT->DirtyFlags, Constraint_GT->Proxy, [NewGroundProxy](FParticleProxyProperty& Data)
+					{
+						Data.ParticleProxy = NewGroundProxy;
+					});
+			}
+			else
+			{
+				Constraint_GT->GroundProxy.Modify(false, Constraint_GT->DirtyFlags, Constraint_GT->Proxy, [](FParticleProxyProperty& Data)
+					{
+						Data.ParticleProxy = nullptr;
+					});
+			}
+			
 		}
 
 		return true;
diff --git a/Engine/Source/Runtime/Experimental/Chaos/Public/Chaos/Character/CharacterGroundConstraint.h b/Engine/Source/Runtime/Experimental/Chaos/Public/Chaos/Character/CharacterGroundConstraint.h
--- a/Engine/Source/Runtime/Experimental/Chaos/Public/Chaos/Character/CharacterGroundConstraint.h
+++ b/Engine/Source/Runtime/Experimental/Chaos/Public/Chaos/Character/CharacterGroundConstraint.h
@@ -47,6 +47,10 @@ namespace Chaos
 				{
 					Data.ParticleProxy = InGroundProxy;
 				});
+			if (InGroundProxy)
+			{
+				WakeGroundBody();
+			}
 		}
 
 		//////////////////////////////////////////////////////////////////////////
@@ -93,6 +97,7 @@ namespace Chaos
 			ConstraintData.Modify(true, DirtyFlags, Proxy, [&Value](FCharacterGroundConstraintDynamicData& Data) {
 				Data.GroundDistance = Value > FReal(0.0) ? Value : FReal(0.0);
 				});
+			WakeCharacterBody();
 		}
 		FReal GetGroundDistance() const { return ConstraintData.Read().GroundDistance; }
 
@@ -147,6 +152,14 @@ namespace Chaos
 				});
 		}
 		FReal GetMaxWalkableSlopeAngle() const { return FMath::Acos(ConstraintSettings.Read().CosMaxWalkableSlopeAngle); }
+
+		void SetCosMaxWalkableSlopeAngle(const FReal& Value)
+		{
+			ConstraintSettings.Modify(true, DirtyFlags, Proxy, [&Value](FCharacterGroundConstraintSettings& Data) {
+				FReal CosAngle = FMath::Clamp(Value, -1.0, 1.0);
+				Data.CosMaxWalkableSlopeAngle = CosAngle;
+				});
+		}
 		FReal GetCosMaxWalkableSlopeAngle() const { return ConstraintSettings.Read().CosMaxWalkableSlopeAngle; }
 
 		/// Set the target change in position for the character. This is used to move the character
@@ -163,6 +176,7 @@ namespace Chaos
 			ConstraintData.Modify(true, DirtyFlags, Proxy, [&Value](FCharacterGroundConstraintDynamicData& Data) {
 				Data.TargetDeltaPosition = Value;
 				});
+			WakeCharacterBody();
 		}
 		FVec3 GetTargetDeltaPosition() const { return ConstraintData.Read().TargetDeltaPosition; }
 
@@ -179,9 +193,21 @@ namespace Chaos
 			ConstraintData.Modify(true, DirtyFlags, Proxy, [&Value](FCharacterGroundConstraintDynamicData& Data) {
 				Data.TargetDeltaFacing = Value;
 				});
+			WakeCharacterBody();
 		}
 		FReal GetTargetDeltaFacing() const { return ConstraintData.Read().TargetDeltaFacing; }
 
+		///  Set both the target delta position and facing angle. This is more efficient than setting
+		/// them both individually
+		void SetMotionTarget(const FVector& TargetDeltaPosition, const FReal& TargetDeltaFacing)
+		{
+			ConstraintData.Modify(true, DirtyFlags, Proxy, [&TargetDeltaPosition, &TargetDeltaFacing](FCharacterGroundConstraintDynamicData& Data) {
+				Data.TargetDeltaPosition = TargetDeltaPosition;
+				Data.TargetDeltaFacing = TargetDeltaFacing;
+				});
+			WakeCharacterBody();
+		}
+
 		/// The maximum force that the character can apply in the plane defined by the ground normal
 		/// to move the character to its target position (or hold the character in place if no target
 		/// delta position is set).
@@ -193,16 +219,46 @@ namespace Chaos
 		}
 		FReal GetRadialForceLimit() const { return ConstraintSettings.Read().RadialForceLimit; }
 
+
 		/// The maximum torque that the character can apply about the vertical axis to rotate the
 		/// character to its target facing direction (or hold the character in place if no target
 		/// delta facing is set).
+		UE_DEPRECATED(5.3, "Split into swing/twist torque limits. Use SetTwistTorqueLimit instead for torque limit about vertical axis.")
 		void SetTorqueLimit(const FReal& Value)
+		{
+			SetTwistTorqueLimit(Value);
+		}
+		UE_DEPRECATED(5.3, "Split into swing/twist torque limits. Use GetTwistTorqueLimit instead for torque limit about vertical axis.")
+		FReal GetTorqueLimit() const { return GetTwistTorqueLimit(); }
+
+		/// The maximum torque that the character can apply about the vertical axis to rotate the
+		/// character to its target facing direction (or hold the character in place if no target
+		/// delta facing is set).
+		void SetTwistTorqueLimit(const FReal& Value)
+		{
+			ConstraintSettings.Modify(true, DirtyFlags, Proxy, [&Value](FCharacterGroundConstraintSettings& Data) {
+				Data.TwistTorqueLimit = Value > FReal(0.0) ? Value : FReal(0.0);
+				});
+		}
+		FReal GetTwistTorqueLimit() const { return ConstraintSettings.Read().TwistTorqueLimit; }
+
+		/// The maximum torque that the character can apply about the non vertical axes to keep
+		/// the character upright
+		void SetSwingTorqueLimit(const FReal& Value)
+		{
+			ConstraintSettings.Modify(true, DirtyFlags, Proxy, [&Value](FCharacterGroundConstraintSettings& Data) {
+				Data.SwingTorqueLimit = Value > FReal(0.0) ? Value : FReal(0.0);
+				});
+		}
+		FReal GetSwingTorqueLimit() const { return ConstraintSettings.Read().SwingTorqueLimit; }
+
+		void SetUserData(void* Value)
 		{
 			ConstraintSettings.Modify(true, DirtyFlags, Proxy, [&Value](FCharacterGroundConstraintSettings& Data) {
-				Data.TorqueLimit = Value > FReal(0.0) ? Value : FReal(0.0);
+				Data.UserData = Value;
 				});
 		}
-		FReal GetTorqueLimit() const { return ConstraintSettings.Read().TorqueLimit; }
+		void* GetUserData() const { return ConstraintSettings.Read().UserData; }
 
 		//////////////////////////////////////////////////////////////////////////
 		// Output
@@ -214,6 +270,30 @@ namespace Chaos
 		FVector GetSolverAppliedTorque() const { return SolverAppliedTorque; }
 
 	private:
+		void WakeCharacterBody()
+		{
+			if (FSingleParticlePhysicsProxy* Character = CharacterProxy.Read().ParticleProxy)
+			{
+				FRigidBodyHandle_External& Body = Character->GetGameThreadAPI();
+				if (Body.ObjectState() == EObjectStateType::Sleeping)
+				{
+					Body.SetObjectState(EObjectStateType::Dynamic);
+				}
+			}
+		}
+
+		void WakeGroundBody()
+		{
+			if (FSingleParticlePhysicsProxy* Ground = GroundProxy.Read().ParticleProxy)
+			{
+				FRigidBodyHandle_External& Body = Ground->GetGameThreadAPI();
+				if (Body.ObjectState() == EObjectStateType::Sleeping)
+				{
+					Body.SetObjectState(EObjectStateType::Dynamic);
+				}
+			}
+		}
+
 		//////////////////////////////////////////////////////////////////////////
 		// FConstraintBase implementation
 		virtual void SyncRemoteDataImp(FDirtyPropertiesManager& Manager, int32 DataIdx, FDirtyChaosProperties& RemoteData) override;
diff --git a/Engine/Source/Runtime/Experimental/Chaos/Public/Chaos/Character/CharacterGroundConstraintContainer.h b/Engine/Source/Runtime/Experimental/Chaos/Public/Chaos/Character/CharacterGroundConstraintContainer.h
--- a/Engine/Source/Runtime/Experimental/Chaos/Public/Chaos/Character/CharacterGroundConstraintContainer.h
+++ b/Engine/Source/Runtime/Experimental/Chaos/Public/Chaos/Character/CharacterGroundConstraintContainer.h
@@ -52,16 +52,30 @@ namespace Chaos
 		virtual void SetEnabled(bool InEnabled) override { bDisabled = !InEnabled; }
 		virtual bool IsEnabled() const { return !bDisabled; }
 
-		/// Currently settings and data are only modifiable on the game thread
-		/// via FCharacterGroundConstraint
+		/// Settings is only modifiable on the game thread via FCharacterGroundConstraint
+		/// but Data can be modified on the physics thread and is synced with the game thread
 		const FCharacterGroundConstraintSettings& GetSettings() const { return Settings; }
 		const FCharacterGroundConstraintDynamicData& GetData() const { return Data; }
 
+		void SetData(const FCharacterGroundConstraintDynamicData& InData)
+		{
+			Data = InData;
+		}
+
 		// Declared final so that TPBDConstraintGraphRuleImpl::AddToGraph() does not need to hit vtable
 		virtual FParticlePair GetConstrainedParticles() const override final { return { CharacterParticle, GroundParticle }; }
 		FGeometryParticleHandle* GetCharacterParticle() const { return CharacterParticle; }
 		FGeometryParticleHandle* GetGroundParticle() const { return GroundParticle; }
 
+		void SetGroundParticle(FGeometryParticleHandle* InGroundParticle)
+		{
+			if (InGroundParticle != GroundParticle)
+			{
+				GroundParticle = InGroundParticle;
+				bGroundParticleChanged = true;
+			}
+		}
+
 		// Get the force applied by the solver due to this constraint. Units are ML/T^2
 		FVec3 GetSolverAppliedForce() const { return SolverAppliedForce; }
 
@@ -82,6 +96,7 @@ namespace Chaos
 		FGeometryParticleHandle* GroundParticle;
 		bool bDisabled = false;
 		bool bEnabledDuringResim;
+		bool bGroundParticleChanged = false;
 		EResimType ResimType = EResimType::FullResim;
 		ESyncState SyncState = ESyncState::InSync;
 	};
diff --git a/Engine/Source/Runtime/Experimental/Chaos/Public/Chaos/Character/CharacterGroundConstraintSettings.h b/Engine/Source/Runtime/Experimental/Chaos/Public/Chaos/Character/CharacterGroundConstraintSettings.h
--- a/Engine/Source/Runtime/Experimental/Chaos/Public/Chaos/Character/CharacterGroundConstraintSettings.h
+++ b/Engine/Source/Runtime/Experimental/Chaos/Public/Chaos/Character/CharacterGroundConstraintSettings.h
@@ -1,64 +1,42 @@
 // Copyright Epic Games, Inc. All Rights Reserved.
 #pragma once
 
-#include "CoreMinimal.h"
-
 #include "Chaos/Core.h"
-#include "Chaos/Transform.h"
-#include "Chaos/Vector.h"
 
 namespace Chaos
 {
 	class CHAOS_API FCharacterGroundConstraintSettings
 	{
 	public:
-		FCharacterGroundConstraintSettings()
-			: VerticalAxis(FVec3(0, 0, 1))
-			, TargetHeight(0)
-			, RadialForceLimit(1000)
-			, TorqueLimit(1000)
-			, CosMaxWalkableSlopeAngle(0.633)
-			, DampingFactor(0.0)
-			, AssumedOnGroundHeight(5.0)
-			, UserData(nullptr)
-		{
-		}
-
 		bool operator==(const FCharacterGroundConstraintSettings& Other) const
 		{
 			return !FMemory::Memcmp(this, &Other, sizeof(*this));
 		}
 
-		FVec3 VerticalAxis;				/// World space up direction (default z axis)
-		FReal TargetHeight;				/// Desired distance from the character body to the ground
-		FReal RadialForceLimit;			/// How much force the character can apply parallel to the ground plane to reach the target position
-		FReal TorqueLimit;				///	How much torque the character can apply about the vertical axis to reach the target facing direction
-		FReal CosMaxWalkableSlopeAngle;	/// Cosine of the maximum angle in degrees that the character is allowed to walk on
-		FReal DampingFactor;			/// Applies a damping to the vertical ground constraint making it softer. Units: /T
-		FReal AssumedOnGroundHeight;	/// Below this height the character is assumed to be on the ground and can apply force/torque to reach the target position and facing 
-		void* UserData;					/// Reserved for user data
+		FVec3 VerticalAxis = FVec3(0.0, 0.0, 1.0);	/// World space up direction (default z axis)
+		FReal TargetHeight = 0.0;					/// Desired distance from the character body to the ground
+		FReal RadialForceLimit = 1500.0;			/// How much force the character can apply parallel to the ground plane to reach the target position
+		FReal TwistTorqueLimit = 1000.0;			///	How much torque the character can apply about the vertical axis to reach the target facing direction
+		FReal SwingTorqueLimit = 5000.0;			///	How much torque the character can apply about the other axes to remain upright
+		FReal CosMaxWalkableSlopeAngle = 0.633;		/// Cosine of the maximum angle in degrees that the character is allowed to walk on
+		FReal DampingFactor = 0.0;					/// Applies a damping to the vertical ground constraint making it softer. Units: /T
+		FReal AssumedOnGroundHeight = 2.0;			/// Below this height the character is assumed to be on the ground and can apply force/torque to reach the target position and facing
+		void* UserData = nullptr;
 	};
 
 	class CHAOS_API FCharacterGroundConstraintDynamicData
 	{
 	public:
-		FCharacterGroundConstraintDynamicData()
-			: GroundNormal(FVec3(0, 0, 1))
-			, TargetDeltaPosition(FVec3(0, 0, 0))
-			, TargetDeltaFacing(0)
-			, GroundDistance(0)
-		{
-		}
-
 		bool operator==(const FCharacterGroundConstraintDynamicData& Other) const
 		{
 			return !FMemory::Memcmp(this, &Other, sizeof(*this));
 		}
 
-		FVec3 GroundNormal;				/// World space ground normal
-		FVec3 TargetDeltaPosition;		///	Target linear movement vector. Will be projected onto ground plane
-		FReal TargetDeltaFacing;		/// Target rotation in radians about the vertical axis
-		FReal GroundDistance;			/// Distance from the character body to the ground
+		FVec3 GroundNormal = FVec3::ZAxisVector;		/// World space ground normal
+		FVec3 TargetDeltaPosition = FVec3::ZeroVector;	///	Target linear movement vector. Will be projected onto ground plane
+		FReal TargetDeltaFacing = 0.0;					/// Target rotation in radians about the vertical axis
+		FReal GroundDistance = 1.0e10f;					/// Distance from the character body to the ground
+		FReal CosMaxWalkableSlopeAngle = 0.633;			/// Override for max walkable slope angle
 	};
 
 } // namespace Chaos
\ No newline at end of file
diff --git a/Engine/Source/Runtime/Experimental/Chaos/Public/Chaos/PullPhysicsDataImp.h b/Engine/Source/Runtime/Experimental/Chaos/Public/Chaos/PullPhysicsDataImp.h
--- a/Engine/Source/Runtime/Experimental/Chaos/Public/Chaos/PullPhysicsDataImp.h
+++ b/Engine/Source/Runtime/Experimental/Chaos/Public/Chaos/PullPhysicsDataImp.h
@@ -132,8 +132,13 @@ struct FDirtyJointConstraintData : public TBasePullData<FJointConstraintPhysicsP
 
 struct FDirtyCharacterGroundConstraintData : public TBasePullData<FCharacterGroundConstraintProxy, FProxyTimestampBase>
 {
-	FVector Force = FVector(0);
-	FVector Torque = FVector(0);
+	FVector Force = FVector(0.0);
+	FVector Torque = FVector(0.0);
+	FVector GroundNormal = FVector(0.0, 0.0, 1.0);
+	FVector TargetDeltaPos = FVector(0.0);
+	FReal TargetDeltaFacing = 0.0;
+	FReal GroundDistance = 0.0;
+	FGeometryParticleHandle* GroundParticle = nullptr;
 };
 
 //A simulation frame's result of dirty particles. These are all the particles that were dirtied in this particular sim step
