ChaosFlesh : Add Gauss Seildel point triangle position targets

diff --git a/Engine/Source/Runtime/Experimental/Chaos/Private/Chaos/Deformable/ChaosDeformableSolver.cpp b/Engine/Source/Runtime/Experimental/Chaos/Private/Chaos/Deformable/ChaosDeformableSolver.cpp
--- a/Engine/Source/Runtime/Experimental/Chaos/Private/Chaos/Deformable/ChaosDeformableSolver.cpp
+++ b/Engine/Source/Runtime/Experimental/Chaos/Private/Chaos/Deformable/ChaosDeformableSolver.cpp
@@ -152,6 +152,7 @@ namespace Chaos::Softs
 			AllTetEMeshArray.Reset(new TArray<FSolverReal>());
 			AllTetNuMeshArray.Reset(new TArray<FSolverReal>());
 			AllTetAlphaJArray.Reset(new TArray<FSolverReal>());
+			GSWeakConstraints.Reset(new FGaussSeidelWeakConstraints<FSolverReal, FSolverParticles>({}, {}, {}, {}, {}, GDeformableXPBDWeakConstraintParams));
 		}
 
 		InitializeKinematicConstraint();
@@ -569,26 +570,44 @@ namespace Chaos::Softs
 				PositionTargetStiffness[i] = DataPackage.Stiffness;
 			}
 
-			int32 InitIndex = Evolution->AddConstraintInitRange(1, true);
-			int32 ConstraintIndex = Evolution->AddConstraintRuleRange(1, true);
+			if (Property.bUseGaussSeidelConstraints)
+			{
+				for (int i = PositionTargets.NumPositionTargets() - 1; i >= 0; i--)
+				{
+					for (int32 j = 0; j < PositionTargetIndices[i].Num(); j++)
+					{
+						PositionTargetIndices[i][j] += Range[0];
+					}
+					for (int32 j = 0; j < PositionTargetSecondIndices[i].Num(); j++)
+					{
+						PositionTargetSecondIndices[i][j] += Range[0];
+					}
+				}
+				GSWeakConstraints->AddExtraConstraints(PositionTargetIndices, PositionTargetWeights, PositionTargetStiffness, PositionTargetSecondIndices, PositionTargetSecondWeights);
+			} 
+			else
+			{
+				int32 InitIndex = Evolution->AddConstraintInitRange(1, true);
+				int32 ConstraintIndex = Evolution->AddConstraintRuleRange(1, true);
 
-			FXPBDWeakConstraints<FSolverReal, FSolverParticles>* WeakConstraint =
-				new FXPBDWeakConstraints<FSolverReal, FSolverParticles>(Evolution->Particles(),
-					PositionTargetIndices, PositionTargetWeights, PositionTargetStiffness, PositionTargetSecondIndices, PositionTargetSecondWeights, GDeformableXPBDWeakConstraintParams);
+				FXPBDWeakConstraints<FSolverReal, FSolverParticles>* WeakConstraint =
+					new FXPBDWeakConstraints<FSolverReal, FSolverParticles>(Evolution->Particles(),
+						PositionTargetIndices, PositionTargetWeights, PositionTargetStiffness, PositionTargetSecondIndices, PositionTargetSecondWeights, GDeformableXPBDWeakConstraintParams);
 
-			Evolution->ConstraintInits()[InitIndex] =
-				[WeakConstraint, this](FSolverParticles& InParticles, const FSolverReal Dt)
-			{
-				WeakConstraint->Init(InParticles, Dt);
-			};
+				Evolution->ConstraintInits()[InitIndex] =
+					[WeakConstraint, this](FSolverParticles& InParticles, const FSolverReal Dt)
+				{
+					WeakConstraint->Init(InParticles, Dt);
+				};
 
-			Evolution->ConstraintRules()[ConstraintIndex] =
-				[WeakConstraint](FSolverParticles& InParticles, const FSolverReal Dt)
-			{
-				WeakConstraint->ApplyInParallel(InParticles, Dt);
-			};
+				Evolution->ConstraintRules()[ConstraintIndex] =
+					[WeakConstraint](FSolverParticles& InParticles, const FSolverReal Dt)
+				{
+					WeakConstraint->ApplyInParallel(InParticles, Dt);
+				};
 
-			WeakConstraints.Add(TUniquePtr<FXPBDWeakConstraints<FSolverReal, FSolverParticles>>(WeakConstraint));
+				WeakConstraints.Add(TUniquePtr<FXPBDWeakConstraints<FSolverReal, FSolverParticles>>(WeakConstraint));
+			}
 		}
 	}
 
@@ -1168,6 +1187,22 @@ namespace Chaos::Softs
 			{
 				this->GSNeohookeanConstraints->Apply(InParticles, Dt);
 			};
+
+			if (Property.bEnablePositionTargets)
+			{
+				TArray<TArray<int32>> ParticlesPerColor;
+				GSWeakConstraints->ComputeInitialWCData(Evolution->Particles(), GSNeohookeanConstraints->GetMeshConstraints(), GSNeohookeanConstraints->GetIncidentElements(), GSNeohookeanConstraints->GetIncidentElementsLocal(), ParticlesPerColor);
+				GSNeohookeanConstraints->SetParticlesPerColor(MoveTemp(ParticlesPerColor));
+
+				GSNeohookeanConstraints->AddAdditionalRes = [this](const FSolverParticles& InParticles, const int32 p, const FSolverReal Dt, TVec3<FSolverReal>& res)
+				{
+					this->GSWeakConstraints->AddWCResidual(InParticles, p, Dt, res);
+				};
+				GSNeohookeanConstraints->AddAdditionalHessian = [this](const FSolverParticles& InParticles, const int32 p, const FSolverReal Dt, Chaos::PMatrix<FSolverReal, 3, 3>& hessian)
+				{
+					this->GSWeakConstraints->AddWCHessian(p, Dt, hessian);
+				};
+			}
 		} 
 		else
 		{
@@ -1186,6 +1221,63 @@ namespace Chaos::Softs
 			{
 				this->GSCorotatedConstraints->Apply(InParticles, Dt);
 			};
+
+			if (Property.bEnablePositionTargets)
+			{
+				TArray<TArray<int32>> ParticlesPerColor;
+				GSWeakConstraints->ComputeInitialWCData(Evolution->Particles(), GSCorotatedConstraints->GetMeshConstraints(), GSCorotatedConstraints->GetIncidentElements(), GSCorotatedConstraints->GetIncidentElementsLocal(), ParticlesPerColor);
+				GSCorotatedConstraints->SetParticlesPerColor(MoveTemp(ParticlesPerColor));
+
+				GSCorotatedConstraints->AddAdditionalRes = [this](const FSolverParticles& InParticles, const int32 p, const FSolverReal Dt, TVec3<FSolverReal>& res)
+				{
+					this->GSWeakConstraints->AddWCResidual(InParticles, p, Dt, res);
+				};
+				GSCorotatedConstraints->AddAdditionalHessian = [this](const FSolverParticles& InParticles, const int32 p, const FSolverReal Dt, Chaos::PMatrix<FSolverReal, 3, 3>& hessian)
+				{
+					this->GSWeakConstraints->AddWCHessian(p, Dt, hessian);
+				};
+			}
+		}
+		if (Property.bEnablePositionTargets)
+		{
+			int32 InitIndex1 = Evolution->AddConstraintInitRange(1, true);
+			Evolution->ConstraintInits()[InitIndex1] =
+				[this](FSolverParticles& InParticles, const FSolverReal Dt)
+			{
+				this->GSWeakConstraints->Init(InParticles, Dt);
+			};
+		}
+
+		if (Property.bDoSelfCollision)
+		{
+			int32 InitIndex = Evolution->AddConstraintInitRange(1, true);
+			Evolution->ConstraintInits()[InitIndex] =
+				[this](FSolverParticles& InParticles, const FSolverReal Dt)
+			{
+				//TODO(Yushan & Joey): Add the collision detection code in here:
+			};
+
+			int32 InitIndex1 = Evolution->AddConstraintInitRange(1, true);
+			if (Property.bUseGSNeohookean)
+			{
+				Evolution->ConstraintInits()[InitIndex1] =
+					[this](FSolverParticles& InParticles, const FSolverReal Dt)
+				{
+					TArray<TArray<int32>> ParticlesPerColor;
+					this->GSWeakConstraints->ComputeCollisionWCData(this->Evolution->Particles(), this->GSNeohookeanConstraints->GetMeshConstraints(), this->GSNeohookeanConstraints->GetIncidentElements(), this->GSNeohookeanConstraints->GetIncidentElementsLocal(), ParticlesPerColor);
+					this->GSNeohookeanConstraints->SetParticlesPerColor(MoveTemp(ParticlesPerColor));
+				};
+			}
+			else
+			{
+				Evolution->ConstraintInits()[InitIndex1] =
+					[this](FSolverParticles& InParticles, const FSolverReal Dt)
+				{
+					TArray<TArray<int32>> ParticlesPerColor;
+					this->GSWeakConstraints->ComputeCollisionWCData(this->Evolution->Particles(), this->GSCorotatedConstraints->GetMeshConstraints(), this->GSCorotatedConstraints->GetIncidentElements(), this->GSCorotatedConstraints->GetIncidentElementsLocal(), ParticlesPerColor);
+					this->GSCorotatedConstraints->SetParticlesPerColor(MoveTemp(ParticlesPerColor));
+				};
+			}
 		}
 	}
 
diff --git a/Engine/Source/Runtime/Experimental/Chaos/Public/Chaos/Deformable/ChaosDeformableSolver.h b/Engine/Source/Runtime/Experimental/Chaos/Public/Chaos/Deformable/ChaosDeformableSolver.h
--- a/Engine/Source/Runtime/Experimental/Chaos/Public/Chaos/Deformable/ChaosDeformableSolver.h
+++ b/Engine/Source/Runtime/Experimental/Chaos/Public/Chaos/Deformable/ChaosDeformableSolver.h
@@ -13,6 +13,7 @@
 #include "Chaos/XPBDCorotatedConstraints.h"
 #include "Chaos/Deformable/GaussSeidelCorotatedConstraints.h"
 #include "Chaos/Deformable/GaussSeidelNeohookeanConstraints.h"
+#include "Chaos/Deformable/GaussSeidelWeakConstraints.h"
 #include "Chaos/XPBDWeakConstraints.h"
 #include "Chaos/BlendedXPBDCorotatedConstraints.h"
 #include "Chaos/XPBDGridBasedCorotatedConstraints.h"
@@ -178,6 +179,7 @@ namespace Chaos::Softs
 		TArray<TUniquePtr<Softs::FXPBDCorotatedConstraints<Softs::FSolverReal, Softs::FSolverParticles>>> CorotatedConstraints;
 		TUniquePtr<Softs::FGaussSeidelCorotatedConstraints<Softs::FSolverReal, Softs::FSolverParticles>> GSCorotatedConstraints;
 		TUniquePtr<Softs::FGaussSeidelNeohookeanConstraints<Softs::FSolverReal, Softs::FSolverParticles>> GSNeohookeanConstraints;
+		TUniquePtr<Softs::FGaussSeidelWeakConstraints<Softs::FSolverReal, Softs::FSolverParticles>> GSWeakConstraints;
 		TArray<TUniquePtr<Softs::FXPBDWeakConstraints<Softs::FSolverReal, Softs::FSolverParticles>>> WeakConstraints;
 		TArray<TUniquePtr<Softs::FBlendedXPBDCorotatedConstraints<Softs::FSolverReal, Softs::FSolverParticles>>> BlendedCorotatedConstraints;
 		TUniquePtr<Softs::FXPBDGridBasedCorotatedConstraints<Softs::FSolverReal, Softs::FSolverParticles>> GridBasedCorotatedConstraint;
@@ -192,6 +194,10 @@ namespace Chaos::Softs
 		TUniquePtr <TArray<FSolverReal>> AllTetEMeshArray;
 		TUniquePtr <TArray<FSolverReal>> AllTetNuMeshArray;
 		TUniquePtr <TArray<FSolverReal>> AllTetAlphaJArray;
+		TUniquePtr <TArray<TArray<int32>>> AllIndices;
+		TUniquePtr <TArray<TArray<int32>>> AllSecondIndices;
+		TUniquePtr <TArray<FSolverReal>> AllWeights;
+		TUniquePtr <TArray<FSolverReal>> AllSecondWeights;
 
 		//typedef TMap<int32, TTuple<float, Chaos::Softs::FPAndInvM, FVector3f>> TransientConstraintBufferMap;
 		typedef TMap<int32, TTuple<float, float, FVector3f>> TransientConstraintBufferMap;
diff --git a/Engine/Source/Runtime/Experimental/Chaos/Public/Chaos/Deformable/GaussSeidelCorotatedConstraints.h b/Engine/Source/Runtime/Experimental/Chaos/Public/Chaos/Deformable/GaussSeidelCorotatedConstraints.h
--- a/Engine/Source/Runtime/Experimental/Chaos/Public/Chaos/Deformable/GaussSeidelCorotatedConstraints.h
+++ b/Engine/Source/Runtime/Experimental/Chaos/Public/Chaos/Deformable/GaussSeidelCorotatedConstraints.h
@@ -8,7 +8,7 @@
 #include "Chaos/GraphColoring.h"
 #include "Chaos/NewtonCorotatedCache.h"
 #include "Chaos/Framework/Parallel.h"
-//#include "Chaos/Deformable/GaussSeidelWeakConstraints.h"
+#include "Chaos/Deformable/GaussSeidelWeakConstraints.h"
 
 
 namespace Chaos::Softs
@@ -44,7 +44,7 @@ namespace Chaos::Softs
 			const T& NuMesh = (T).3,
 			const bool bRecordMetricIn = false
 		)
-			: Base(InParticles, InMesh, EMeshArray, NuMeshArray, MoveTemp(AlphaJMeshArray), InParams, NuMesh, bRecordMetricIn, false), IncidentElements(MoveTemp(IncidentElementsIn)), IncidentElementsLocal(IncidentElementsLocalIn),
+			: Base(InParticles, InMesh, EMeshArray, NuMeshArray, MoveTemp(AlphaJMeshArray), InParams, NuMesh, bRecordMetricIn, false), IncidentElements(IncidentElementsIn), IncidentElementsLocal(IncidentElementsLocalIn),
 			ParticleStartIndex(ParticleStartIndexIn), ParticleEndIndex(ParticleEndIndexIn), bDoQuasistatics(bDoQuasistaticsIn), bDoSOR(bDoSORIn), OmegaSOR(InOmegaSOR)
 		{
 			Base::LambdaArray.SetNum(0);
@@ -124,6 +124,11 @@ namespace Chaos::Softs
 
 		void Init() const override {}
 
+		TArray<TArray<int32>> GetIncidentElements() const { return IncidentElements; }
+		TArray<TArray<int32>> GetIncidentElementsLocal() const { return IncidentElementsLocal; }
+		TArray<TVector<int32, 4>> GetMeshConstraints() const { return MeshConstraints; }
+		void SetParticlesPerColor(TArray<TArray<int32>>&& InParticlesPerColor) {ParticlesPerColor = MoveTemp(InParticlesPerColor); }
+
 
 		void ApplySOR(ParticleType& Particles, const T Dt) const
 		{
@@ -292,6 +297,9 @@ namespace Chaos::Softs
 					 }
 				 }
 			 };
+
+			 AddAdditionalRes = [](const ParticleType& InParticles, const int32 p, const T Dt, TVec3<T>& res) {};
+			 AddAdditionalHessian = [](const ParticleType& InParticles, const int32 p, const T Dt, Chaos::PMatrix<T, 3, 3>& hessian) {};
 		 }
 
 	private:
@@ -301,9 +309,13 @@ namespace Chaos::Softs
 		{
 			Chaos::TVector<T, 3> ParticleResidual = ComputePerParticleResidual(p, Particle2Incident[p], Particles, Dt, !bDoQuasistatics);
 
+			AddAdditionalRes(Particles, p, Dt, ParticleResidual);
+
 			if (ParticleResidual.Size() > LocalNewtonTol)
 			{
 				Chaos::PMatrix<T, 3, 3> SimplifiedHessian = ComputePerParticleCorotatedHessianSimple(p, Particle2Incident[p], Particles, Dt, !bDoQuasistatics);
+				AddAdditionalHessian(Particles, p, Dt, SimplifiedHessian);
+				
 				T HessianDet = SimplifiedHessian.Determinant();
 				if (HessianDet > UE_SMALL_NUMBER)
 				{
@@ -335,6 +347,11 @@ namespace Chaos::Softs
 		mutable int32 CurrentIt = 0;
 		TFunction<void(const Chaos::PMatrix<T, 3, 3>&, const T, const T, Chaos::PMatrix<T, 3, 3>& )> ComputeStress;
 		TFunction<void(const Chaos::PMatrix<T, 3, 3>&, const Chaos::PMatrix<T, 3, 3>&, const T, const T, const int32, const T, Chaos::PMatrix<T, 3, 3>&)> ComputeHessianHelper;
+		
+	public:
+		TFunction<void(const ParticleType&, const int32, const T, TVec3<T>&)> AddAdditionalRes;
+		TFunction<void(const ParticleType&, const int32, const T, Chaos::PMatrix<T, 3, 3>&)> AddAdditionalHessian;
+		TUniquePtr<TArray<int32>> ParticleColors;
 	};
 
 
diff --git a/Engine/Source/Runtime/Experimental/Chaos/Public/Chaos/Deformable/GaussSeidelWeakConstraints.h b/Engine/Source/Runtime/Experimental/Chaos/Public/Chaos/Deformable/GaussSeidelWeakConstraints.h
new file mode 100644
--- /dev/null
+++ b/Engine/Source/Runtime/Experimental/Chaos/Public/Chaos/Deformable/GaussSeidelWeakConstraints.h
@@ -0,0 +1,403 @@
+// Copyright Epic Games, Inc. All Rights Reserved.
+#pragma once
+
+#include "Chaos/PBDSpringConstraintsBase.h"
+#include "Chaos/XPBDCorotatedConstraints.h"
+#include "ChaosStats.h"
+#include "Chaos/ImplicitQRSVD.h"
+#include "Chaos/GraphColoring.h"
+#include "Chaos/NewtonCorotatedCache.h"
+#include "Chaos/Framework/Parallel.h"
+#include "Chaos/Utilities.h"
+#include "Chaos/DebugDrawQueue.h"
+#include "Chaos/XPBDWeakConstraints.h"
+
+
+namespace Chaos::Softs
+{
+
+	using Chaos::TVec3;
+	template <typename T, typename ParticleType>
+	struct FGaussSeidelWeakConstraints 
+	{
+		//TODO(Yizhou): Add unittest for Gauss Seidel Weak Constraints
+		FGaussSeidelWeakConstraints(
+			const TArray<TArray<int32>>& InIndices,
+			const TArray<TArray<T>>& InWeights,
+			const TArray<T>& InStiffness,
+			const TArray<TArray<int32>>& InSecondIndices,
+			const TArray<TArray<T>>& InSecondWeights,
+			const FDeformableXPBDWeakConstraintParams& InParams
+		)
+			: Indices(InIndices), Weights(InWeights), SecondIndices(InSecondIndices), SecondWeights(InSecondWeights), Stiffness(InStiffness), DebugDrawParams(InParams)
+		{
+			ensureMsgf(Indices.Num() == SecondIndices.Num(), TEXT("Input Double Bindings have wrong size"));
+
+			for (int32 i = 0; i < Indices.Num(); i++)
+			{
+				ensureMsgf(Indices[i].Num() + SecondIndices[i].Num() == 4, TEXT("Currently does not support non point tri pair weak constraints"));
+				TSet<int32> IndicesSet = TSet<int32>(Indices[i]);
+				for (int32 j = 0; j < SecondIndices[i].Num(); j++)
+				{
+					ensureMsgf(!IndicesSet.Contains(SecondIndices[i][j]), TEXT("Indices and Second Indices overlaps. Currently not supported"));
+				}
+			}
+
+			IsAnisotropic.Init(false, Indices.Num());
+		}
+
+		virtual ~FGaussSeidelWeakConstraints() {}
+
+		void ComputeInitialWCData(const ParticleType& InParticles, const TArray<TVector<int32, 4>>& MeshConstraints, const TArray <TArray<int32>>& MeshIncidentElements, const TArray <TArray<int32>>& MeshIncidentElementsLocal, TArray<TArray<int32>>& ParticlesPerColor)
+		{
+			TArray<TVector<int32, 4>> ExtraConstraints;
+			ExtraConstraints.Init(TVector<int32, 4>(-1), Indices.Num());
+			for (int32 i = 0; i < Indices.Num(); i++)
+			{
+				ensureMsgf(Indices[i].Num() + SecondIndices[i].Num() == 4, TEXT("Currently does not support non point tri pair weak constraints"));
+				for (int32 j = 0; j < Indices[i].Num(); j++)
+				{
+					ExtraConstraints[i][j] = Indices[i][j];
+				}
+				for (int32 j = 0; j < SecondIndices[i].Num(); j++)
+				{
+					ExtraConstraints[i][j+Indices[i].Num()] = SecondIndices[i][j];
+				}
+			}
+			TArray<TVector<int32, 4>> TotalConstraints = MeshConstraints;
+			TotalConstraints += ExtraConstraints;
+			TArray<TArray<int32>> TotalIncidentElements = MeshIncidentElements, TotalIncidentElementsLocal = MeshIncidentElementsLocal;
+			WCIncidentElements = Chaos::Utilities::ComputeIncidentElements(ExtraConstraints, &WCIncidentElementsLocal);
+			
+			//TODO (Yizhou): Directly computing incident elements is inefficient. Optimize the total incident elements in future.
+			TotalIncidentElements = Chaos::Utilities::ComputeIncidentElements(TotalConstraints, &TotalIncidentElementsLocal);
+			Particle2WCIndices.Init(INDEX_NONE, InParticles.Size());
+			for (int32 i = 0; i < WCIncidentElements.Num(); i++)
+			{
+				if (WCIncidentElements[i].Num() > 0)
+				{
+					int32 p = ExtraConstraints[WCIncidentElements[i][0]][WCIncidentElementsLocal[i][0]];
+					Particle2WCIndices[p] = i;
+				}
+			}
+
+			ParticlesPerColor = ComputeNodalColoring(TotalConstraints, InParticles, 0, InParticles.Size(), TotalIncidentElements, TotalIncidentElementsLocal);
+
+			NodalWeights.Init({}, InParticles.Size());
+
+			for (int32 i = 0; i < WCIncidentElements.Num(); i++) 
+			{
+				if (WCIncidentElements[i].Num() > 0)
+				{
+					int32 p = ExtraConstraints[WCIncidentElements[i][0]][WCIncidentElementsLocal[i][0]];
+					NodalWeights[p].Init(T(0), 6);
+					for (int32 j = 0; j < WCIncidentElements[i].Num(); j++)
+					{
+						int32 ConstraintIndex = WCIncidentElements[i][j];
+						int32 LocalIndex = WCIncidentElementsLocal[i][j];
+
+						T Weight = T(0);
+						if (LocalIndex >= Indices[ConstraintIndex].Num())
+						{
+							Weight = SecondWeights[ConstraintIndex][LocalIndex - Indices[ConstraintIndex].Num()];
+						}
+						else
+						{
+							Weight = Weights[ConstraintIndex][LocalIndex];
+						}
+
+						if (IsAnisotropic[ConstraintIndex])
+						{
+							for (int32 alpha = 0; alpha < 3; alpha++) 
+							{
+								NodalWeights[p][alpha] += Normals[ConstraintIndex][alpha] * Normals[ConstraintIndex][alpha] * Weight * Weight * Stiffness[ConstraintIndex];
+							}
+
+							NodalWeights[p][3] += Normals[ConstraintIndex][0] * Normals[ConstraintIndex][1] * Weight * Weight * Stiffness[ConstraintIndex];
+							NodalWeights[p][4] += Normals[ConstraintIndex][0] * Normals[ConstraintIndex][2] * Weight * Weight * Stiffness[ConstraintIndex];
+							NodalWeights[p][5] += Normals[ConstraintIndex][1] * Normals[ConstraintIndex][2] * Weight * Weight * Stiffness[ConstraintIndex];
+						}
+						else
+						{
+							for (int32 alpha = 0; alpha < 3; alpha++)
+							{
+								NodalWeights[p][alpha] += Weight * Weight * Stiffness[ConstraintIndex];
+							}
+						}
+					}
+				}
+			}
+
+			InitialWCSize = Indices.Num();
+		}
+
+		void AddWCResidual(const ParticleType& InParticles, const int32 p, const T Dt, TVec3<T>& res)
+		{
+			int32 WCIndex = Particle2WCIndices[p];
+			if (WCIndex != INDEX_NONE)
+			{
+				for (int32 j = 0; j < WCIncidentElements[WCIndex].Num(); j++)
+				{
+					int32 ConstraintIndex = WCIncidentElements[WCIndex][j];
+					int32 LocalIndex = WCIncidentElementsLocal[WCIndex][j];
+					TVec3<T> SpringEdge((T)0.);
+					for (int32 l = 0; l < Weights[ConstraintIndex].Num(); l++)
+					{
+						for (int32 beta = 0; beta < 3; beta++)
+						{
+							SpringEdge[beta] += Weights[ConstraintIndex][l] * InParticles.P(Indices[ConstraintIndex][l])[beta];
+						}
+					}
+					for (int32 l = 0; l < SecondWeights[ConstraintIndex].Num(); l++)
+					{
+						for (int32 beta = 0; beta < 3; beta++)
+						{
+							SpringEdge[beta] -= SecondWeights[ConstraintIndex][l] * InParticles.P(SecondIndices[ConstraintIndex][l])[beta];
+						}
+					}
+					T weight = T(0);
+					if (LocalIndex >= Indices[ConstraintIndex].Num())
+					{
+						weight = -SecondWeights[ConstraintIndex][LocalIndex - Indices[ConstraintIndex].Num()];
+					}
+					else
+					{
+						weight = Weights[ConstraintIndex][LocalIndex];
+					}
+					if (IsAnisotropic[ConstraintIndex])
+					{
+						T comp = TVec3<T>::DotProduct(SpringEdge, Normals[ConstraintIndex]);
+						TVec3<T> proj = Normals[ConstraintIndex] * comp;
+						for (int32 alpha = 0; alpha < 3; alpha++)
+						{
+							res[alpha] += Dt * Dt * Stiffness[ConstraintIndex] * proj[alpha] * weight;
+						}
+					}
+					else
+					{
+						for (int32 alpha = 0; alpha < 3; alpha++)
+						{
+							res[alpha] += Dt * Dt * Stiffness[ConstraintIndex] * SpringEdge[alpha] * weight;
+						}
+					}
+				}
+			}
+		}
+
+		void AddWCHessian(const int32 p, const T Dt, Chaos::PMatrix<T, 3, 3>& ParticleHessian)
+		{
+			if (NodalWeights[p].Num() > 0)
+			{
+				for (int32 alpha = 0; alpha < 3; alpha++)
+				{
+					ParticleHessian.SetAt(alpha, alpha, ParticleHessian.GetAt(alpha, alpha) + Dt * Dt * NodalWeights[p][alpha]);
+				}
+
+				ParticleHessian.SetAt(0, 1, ParticleHessian.GetAt(0, 1) + Dt * Dt * NodalWeights[p][3]);
+				ParticleHessian.SetAt(0, 2, ParticleHessian.GetAt(0, 2) + Dt * Dt * NodalWeights[p][4]);
+				ParticleHessian.SetAt(1, 2, ParticleHessian.GetAt(1, 2) + Dt * Dt * NodalWeights[p][5]);
+				ParticleHessian.SetAt(1, 0, ParticleHessian.GetAt(1, 0) + Dt * Dt * NodalWeights[p][3]);
+				ParticleHessian.SetAt(2, 0, ParticleHessian.GetAt(2, 0) + Dt * Dt * NodalWeights[p][4]);
+				ParticleHessian.SetAt(2, 1, ParticleHessian.GetAt(2, 1) + Dt * Dt * NodalWeights[p][5]);
+				//TODO(Yizhou): Clean up the following after debugging:
+				//ParticleHessian.SetAt(0, 2) += Dt * Dt * NodalWeights[p][4];
+				//ParticleHessian.SetAt(1, 2) += Dt * Dt * NodalWeights[p][5];
+				//ParticleHessian.SetAt(1, 0) += Dt * Dt * NodalWeights[p][3];
+				//ParticleHessian.SetAt(2, 0) += Dt * Dt * NodalWeights[p][4];
+				//ParticleHessian.SetAt(2, 1) += Dt * Dt * NodalWeights[p][5];
+			}
+		}
+
+		void AddExtraConstraints(const TArray<TArray<int32>>& InIndices,
+								const TArray<TArray<T>>& InWeights,
+								const TArray<T>& InStiffness,
+								const TArray<TArray<int32>>& InSecondIndices,
+								const TArray<TArray<T>>& InSecondWeights)
+		{
+			int32 Offset = Indices.Num();
+			Indices.SetNum(Offset + InIndices.Num());
+			SecondIndices.SetNum(Offset + InSecondIndices.Num());
+			Stiffness.SetNum(Offset + InStiffness.Num());
+			Weights.SetNum(Offset + InIndices.Num());
+			SecondWeights.SetNum(Offset + InIndices.Num());
+
+			for (int32 i = 0; i < InIndices.Num(); i++)
+			{
+				Indices[i + Offset] = InIndices[i];
+				SecondIndices[i + Offset] = InSecondIndices[i];
+				Stiffness[i + Offset] = InStiffness[i];
+				Weights[i + Offset] = InWeights[i];
+				SecondWeights[i + Offset] = InSecondWeights[i];
+			}
+
+			IsAnisotropic.Init(false, Indices.Num());
+		}
+
+		void VisualizeAllBindings(const FSolverParticles& InParticles, const T Dt) const
+		{
+#if WITH_EDITOR
+			auto DoubleVert = [](Chaos::TVec3<T> V) { return FVector3d(V.X, V.Y, V.Z); };
+			for (int32 i = 0; i < Indices.Num(); i++)
+			{
+				Chaos::TVec3<T> SourcePos((T)0.), TargetPos((T)0.);
+				for (int32 j = 0; j < Indices[i].Num(); j++)
+				{
+					SourcePos += Weights[i][j] * InParticles.P(Indices[i][j]);
+				}
+				for (int32 j = 0; j < SecondIndices[i].Num(); j++)
+				{
+					TargetPos += SecondWeights[i][j] * InParticles.P(SecondIndices[i][j]);
+				}
+
+				float ParticleThickness = DebugDrawParams.DebugParticleWidth;
+				float LineThickness = DebugDrawParams.DebugLineWidth;
+
+				if (Indices[i].Num() == 1)
+				{
+					Chaos::FDebugDrawQueue::GetInstance().DrawDebugPoint(DoubleVert(SourcePos), FColor::Red, false, Dt, 0, ParticleThickness);
+					for (int32 j = 0; j < SecondIndices[i].Num(); j++)
+					{
+						Chaos::FDebugDrawQueue::GetInstance().DrawDebugPoint(DoubleVert(InParticles.P(SecondIndices[i][j])), FColor::Green, false, Dt, 0, ParticleThickness);
+						Chaos::FDebugDrawQueue::GetInstance().DrawDebugLine(DoubleVert(InParticles.P(SecondIndices[i][j])), DoubleVert(InParticles.P(SecondIndices[i][(j + 1) % SecondIndices[i].Num()])), FColor::Green, false, Dt, 0, LineThickness);
+					}
+
+				}
+
+				if (SecondIndices[i].Num() == 1)
+				{
+					Chaos::FDebugDrawQueue::GetInstance().DrawDebugPoint(DoubleVert(TargetPos), FColor::Red, false, Dt, 0, ParticleThickness);
+					for (int32 j = 0; j < Indices[i].Num(); j++)
+					{
+						Chaos::FDebugDrawQueue::GetInstance().DrawDebugPoint(DoubleVert(InParticles.P(Indices[i][j])), FColor::Green, false, Dt, 0, ParticleThickness);
+						Chaos::FDebugDrawQueue::GetInstance().DrawDebugLine(DoubleVert(InParticles.P(Indices[i][j])), DoubleVert(InParticles.P(Indices[i][(j + 1) % SecondIndices[i].Num()])), FColor::Green, false, Dt, 0, LineThickness);
+					}
+				}
+
+				Chaos::FDebugDrawQueue::GetInstance().DrawDebugLine(DoubleVert(SourcePos), DoubleVert(TargetPos), FColor::Yellow, false, Dt, 0, LineThickness);
+			}
+#endif
+		}
+
+		void Init(const FSolverParticles& InParticles, const T Dt) const
+		{
+			if (DebugDrawParams.bVisualizeBindings)
+			{
+				VisualizeAllBindings(InParticles, Dt);
+			}
+
+		}
+
+		void ComputeCollisionWCData(const ParticleType& InParticles, const TArray<TVector<int32, 4>>& MeshConstraints, const TArray <TArray<int32>>& MeshIncidentElements, const TArray <TArray<int32>>& MeshIncidentElementsLocal, TArray<TArray<int32>>& ParticlesPerColor)
+		{
+			TArray<TVector<int32, 4>> ExtraConstraints;
+			ExtraConstraints.Init(TVector<int32, 4>(-1), Indices.Num());
+			for (int32 i = InitialWCSize; i < Indices.Num(); i++)
+			{
+				ensureMsgf(Indices[i].Num() + SecondIndices[i].Num() == 4, TEXT("Currently does not support non point tri pair weak constraints"));
+				for (int32 j = 0; j < Indices[i].Num(); j++)
+				{
+					ExtraConstraints[i][j] = Indices[i][j];
+				}
+				for (int32 j = 0; j < SecondIndices[i].Num(); j++)
+				{
+					ExtraConstraints[i][j + Indices[i].Num()] = SecondIndices[i][j];
+				}
+			}
+
+			TArray<TArray<int32>> ExtraWCIncidentElements, ExtraWCIncidentElementsLocal;
+			ExtraWCIncidentElements = Chaos::Utilities::ComputeIncidentElements(ExtraConstraints, &ExtraWCIncidentElementsLocal);
+
+			TArray<TVector<int32, 4>> TotalConstraints = NoCollisionConstraints;
+			TotalConstraints += ExtraConstraints;
+			//TODO (Yizhou): Make the following more efficient by computing only extra incident elements in the future. 
+			WCIncidentElements = Chaos::Utilities::ComputeIncidentElements(TotalConstraints, &WCIncidentElementsLocal);
+
+			Particle2WCIndices.Init(INDEX_NONE, InParticles.Size());
+			for (int32 i = 0; i < WCIncidentElements.Num(); i++)
+			{
+				if (WCIncidentElements[i].Num() > 0)
+				{
+					int32 p = ExtraConstraints[WCIncidentElements[i][0]][WCIncidentElementsLocal[i][0]];
+					Particle2WCIndices[p] = i;
+				}
+			}
+
+			//TODO(Yizhou): Write the collision per timestep coloring to make the following thing much faster:
+			TArray<TVector<int32, 4>> TrueAllConstraints = MeshConstraints;
+			TrueAllConstraints += TotalConstraints;
+			TArray<TArray<int32>> TrueAllIncidentElements, TrueAllIncidentElementsLocal;
+			TrueAllIncidentElements = Chaos::Utilities::ComputeIncidentElements(TrueAllConstraints, &TrueAllIncidentElementsLocal);
+			ParticlesPerColor = ComputeNodalColoring(TrueAllConstraints, InParticles, 0, InParticles.Size(), TrueAllIncidentElements, TrueAllIncidentElementsLocal);
+
+
+			NodalWeights = NoCollisionNodalWeights;
+			for (int32 i = 0; i < ExtraWCIncidentElements.Num(); i++) 
+			{
+				if (ExtraWCIncidentElements[i].Num() > 0)
+				{
+					int32 p = ExtraConstraints[ExtraWCIncidentElements[i][0]][ExtraWCIncidentElementsLocal[i][0]];
+					if (NodalWeights[p].Num() == 0) 
+					{
+						NodalWeights[p].Init(T(0), 6);
+					}
+					for (int32 j = 0; j < ExtraWCIncidentElements[i].Num(); j++) 
+					{
+						int32 LocalIndex = ExtraWCIncidentElementsLocal[i][j];
+						int32 ConstraintIndex = ExtraWCIncidentElements[i][j] + InitialWCSize;
+						T weight = T(0);
+						if (LocalIndex >= Indices[ConstraintIndex].Num()) 
+						{
+							weight = SecondWeights[ConstraintIndex][LocalIndex - Indices[ConstraintIndex].Num()];
+						}
+						else 
+						{
+							weight = Weights[ConstraintIndex][LocalIndex];
+						}
+						if (IsAnisotropic[ConstraintIndex]) 
+						{
+							for (int32 alpha = 0; alpha < 3; alpha++) 
+							{
+								NodalWeights[p][alpha] += Normals[ConstraintIndex][alpha] * Normals[ConstraintIndex][alpha] * weight * weight * Stiffness[ConstraintIndex];
+							}
+
+							NodalWeights[p][3] += Normals[ConstraintIndex][0] * Normals[ConstraintIndex][1] * weight * weight * Stiffness[ConstraintIndex];
+							NodalWeights[p][4] += Normals[ConstraintIndex][0] * Normals[ConstraintIndex][2] * weight * weight * Stiffness[ConstraintIndex];
+							NodalWeights[p][5] += Normals[ConstraintIndex][1] * Normals[ConstraintIndex][2] * weight * weight * Stiffness[ConstraintIndex];
+						}
+						else {
+							for (int32 alpha = 0; alpha < 3; alpha++) 
+							{
+								NodalWeights[p][alpha] += weight * weight * Stiffness[ConstraintIndex];
+							}
+						}
+					}
+				}
+			}
+
+		}
+
+
+
+		TArray<TArray<int32>> Indices;
+		TArray<TArray<T>> Weights;
+		TArray<TVector<T, 3>> Constraints;
+		TArray<TArray<int32>> SecondIndices;
+		TArray<TArray<T>> SecondWeights;
+		TArray<T> Stiffness;
+		TArray<bool> IsAnisotropic;
+		TArray<TVector<T, 3>> Normals;
+		TArray<TArray<T>> NodalWeights;
+		TArray<int32> Particle2WCIndices;
+
+		TArray<TArray<int32>> WCIncidentElements;
+		TArray<TArray<int32>> WCIncidentElementsLocal;
+
+		FDeformableXPBDWeakConstraintParams DebugDrawParams;
+
+		int32 InitialWCSize;
+		TArray<TArray<T>> NoCollisionNodalWeights;
+		TArray<TVector<int32, 4>> NoCollisionConstraints;
+	};
+
+
+}// End namespace Chaos::Softs
