Fix headless chaos crash when using levelset in cluster union
- added support to LevelSet for the connection graph generation

diff --git a/Engine/Source/Runtime/Experimental/Chaos/Private/Chaos/ClusterUnionManager.cpp b/Engine/Source/Runtime/Experimental/Chaos/Private/Chaos/ClusterUnionManager.cpp
--- a/Engine/Source/Runtime/Experimental/Chaos/Private/Chaos/ClusterUnionManager.cpp
+++ b/Engine/Source/Runtime/Experimental/Chaos/Private/Chaos/ClusterUnionManager.cpp
@@ -62,50 +62,79 @@ namespace Chaos
 						FAABB3 ToAABB = ToGeom->CalculateTransformedBounds(ToTransform);
 						ToAABB.Thicken(kAddThickness);
 
-						const bool bOverlap = Utilities::CastHelper(*ToGeom, ToTransform,
-							[Particle, &FromTransform, &ToAABB, kAddThickness](const auto& ToGeomDowncast, const FRigidTransform3& FinalToTransform)
+
+						bool bOverlap = false;
+						if (ToGeom->GetType() == ImplicitObjectType::LevelSet)
+						{
+							const FShapesArray& AllFromShapes = Particle->ShapesArray();
+							for (int32 FromIndex = 0; FromIndex < AllFromShapes.Num(); ++FromIndex)
 							{
-								const FShapesArray& AllFromShapes = Particle->ShapesArray();
-								for (int32 FromIndex = 0; FromIndex < AllFromShapes.Num(); ++FromIndex)
+								if (AllFromShapes[FromIndex])
 								{
-									if (!AllFromShapes[FromIndex])
-									{
-										continue;
-									}
-
 									const FPerShapeData& FromShape = *AllFromShapes[FromIndex];
-									const FImplicitObject* FromGeom = FromShape.GetGeometry().Get();
-
-									if (!FromGeom)
-									{
-										continue;
-									}
-
-									FAABB3 FromAABB = FromGeom->CalculateTransformedBounds(FromTransform);
-									FromAABB.Thicken(kAddThickness);
-
-									// First sanity check to see if the two shape AABB's intersect.
-									if (!FromAABB.Intersects(ToAABB))
+									if (const FImplicitObject* FromGeom = FromShape.GetGeometry().Get())
 									{
-										continue;
+										FAABB3 FromAABB = FromGeom->CalculateTransformedBounds(FromTransform);
+										FromAABB.Thicken(kAddThickness);
+
+										// First sanity check to see if the two shape AABB's intersect.
+										if (FromAABB.Intersects(ToAABB))
+										{
+											// level set have missing implementation for some of the shapes at the moment 
+											// for now we stop at the AABB test 
+											bOverlap = true;
+											break;
+										}
 									}
-
-									// Now do a more accurate overlap check between the two shapes.
-									// Note that passing MTD here is critical as it gets us a more accurate check for some reason...
-									FMTDInfo MTDInfo;
-									if (OverlapQuery(*FromGeom, FromTransform, ToGeomDowncast, FinalToTransform, kAddThickness, &MTDInfo))
-									{
-										return true;
-									}
-									else
+								}
+							}
+						}
+						else
+						{
+							bOverlap = Utilities::CastHelper(*ToGeom, ToTransform,
+								[Particle, &FromTransform, &ToAABB, kAddThickness](const auto& ToGeomDowncast, const FRigidTransform3& FinalToTransform)
+								{
+									const FShapesArray& AllFromShapes = Particle->ShapesArray();
+									for (int32 FromIndex = 0; FromIndex < AllFromShapes.Num(); ++FromIndex)
 									{
-										continue;
+										if (!AllFromShapes[FromIndex])
+										{
+											continue;
+										}
+
+										const FPerShapeData& FromShape = *AllFromShapes[FromIndex];
+										const FImplicitObject* FromGeom = FromShape.GetGeometry().Get();
+
+										if (!FromGeom)
+										{
+											continue;
+										}
+
+										FAABB3 FromAABB = FromGeom->CalculateTransformedBounds(FromTransform);
+										FromAABB.Thicken(kAddThickness);
+
+										// First sanity check to see if the two shape AABB's intersect.
+										if (!FromAABB.Intersects(ToAABB))
+										{
+											continue;
+										}
+
+										// Now do a more accurate overlap check between the two shapes.
+										// Note that passing MTD here is critical as it gets us a more accurate check for some reason...
+										FMTDInfo MTDInfo;
+										if (OverlapQuery(*FromGeom, FromTransform, ToGeomDowncast, FinalToTransform, kAddThickness, &MTDInfo))
+										{
+											return true;
+										}
+										else
+										{
+											continue;
+										}
 									}
-								}
-
-								return false;
-							});
 
+									return false;
+								});
+						}
 						if (bOverlap)
 						{
 							Clustering.CreateNodeConnection(Particle, ClusterUnion.ChildParticles[RootObjectIndex]);
