cluster unions should be able to lock their kinematic anchored state based on what gets set on the GT - this way SetIsAnchored can't be overridden by an UpdateKinematicProperties call later

diff --git a/Engine/Source/Runtime/Experimental/Chaos/Private/Chaos/ClusterUnionManager.cpp b/Engine/Source/Runtime/Experimental/Chaos/Private/Chaos/ClusterUnionManager.cpp
--- a/Engine/Source/Runtime/Experimental/Chaos/Private/Chaos/ClusterUnionManager.cpp
+++ b/Engine/Source/Runtime/Experimental/Chaos/Private/Chaos/ClusterUnionManager.cpp
@@ -466,7 +466,7 @@ namespace Chaos
 			}
 		}
 
-		if (bIsNewCluster && bIsAnchored)
+		if (bIsNewCluster && bIsAnchored && !Cluster->bAnchorLock)
 		{
 			// The anchored flag is taken care of in UpdateKinematicProperties so it must be set before that is called.
 			Cluster->InternalCluster->SetIsAnchored(true);
@@ -643,7 +643,23 @@ namespace Chaos
 			ClusterUnion.bNeedsXRInitialization = false;
 		}
 
-		UpdateKinematicProperties(ClusterUnion.InternalCluster, MClustering.GetChildrenMap(), MEvolution);
+		if (ClusterUnion.bAnchorLock)
+		{
+			// Skip UpdateKinematicPropertiesso as we won't need to look at our children for this.
+			if (ClusterUnion.InternalCluster->IsAnchored())
+			{
+				MEvolution.SetParticleObjectState(ClusterUnion.InternalCluster, EObjectStateType::Kinematic);
+			}
+			else
+			{
+				MEvolution.SetParticleObjectState(ClusterUnion.InternalCluster, EObjectStateType::Dynamic);
+				MEvolution.SetParticleKinematicTarget(ClusterUnion.InternalCluster, FKinematicTarget());
+			}
+		}
+		else
+		{
+			UpdateKinematicProperties(ClusterUnion.InternalCluster, MClustering.GetChildrenMap(), MEvolution);
+		}
 
 		// We must reset collisions etc on this particle since geometry will be changed
 		MEvolution.InvalidateParticle(ClusterUnion.InternalCluster);
diff --git a/Engine/Source/Runtime/Experimental/Chaos/Private/PhysicsProxy/ClusterUnionPhysicsProxy.cpp b/Engine/Source/Runtime/Experimental/Chaos/Private/PhysicsProxy/ClusterUnionPhysicsProxy.cpp
--- a/Engine/Source/Runtime/Experimental/Chaos/Private/PhysicsProxy/ClusterUnionPhysicsProxy.cpp
+++ b/Engine/Source/Runtime/Experimental/Chaos/Private/PhysicsProxy/ClusterUnionPhysicsProxy.cpp
@@ -192,19 +192,19 @@ namespace Chaos
 			{
 				if (Particle_Internal)
 				{
-					if (Particle_Internal->IsAnchored() != bIsAnchored)
+					if (FPBDRigidsEvolutionGBF* Evolution = GetEvolution(this))
 					{
-						if (FPBDRigidsEvolutionGBF* Evolution = GetEvolution(this))
+						Particle_Internal->SetIsAnchored(bIsAnchored);
+
+						// We also need to make sure the cluster union won't try to override this value that we set.
+						Chaos::FClusterUnionManager& Manager = Evolution->GetRigidClustering().GetClusterUnionManager();
+						if (Chaos::FClusterUnion* ClusterUnion = Manager.FindClusterUnion(ClusterUnionIndex))
 						{
-							Particle_Internal->SetIsAnchored(bIsAnchored);
-							if (!bIsAnchored && !Particle_Internal->IsDynamic())
-							{
-								FKinematicTarget NoKinematicTarget;
-								Evolution->SetParticleObjectState(Particle_Internal, EObjectStateType::Dynamic);
-								Evolution->SetParticleKinematicTarget(Particle_Internal, NoKinematicTarget);
-							}
-							Evolution->GetParticles().MarkTransientDirtyParticle(Particle_Internal);
+							ClusterUnion->bAnchorLock = true;
 						}
+
+						// Let the cluster union manager apply the proper properties on the particle.
+						Manager.RequestDeferredClusterPropertiesUpdate(ClusterUnionIndex, Chaos::EUpdateClusterUnionPropertiesFlags::None);
 					}
 				}
 			}
diff --git a/Engine/Source/Runtime/Experimental/Chaos/Public/Chaos/ClusterUnionManager.h b/Engine/Source/Runtime/Experimental/Chaos/Public/Chaos/ClusterUnionManager.h
--- a/Engine/Source/Runtime/Experimental/Chaos/Public/Chaos/ClusterUnionManager.h
+++ b/Engine/Source/Runtime/Experimental/Chaos/Public/Chaos/ClusterUnionManager.h
@@ -83,6 +83,9 @@ namespace Chaos
 		// Whether or not the position/rotation needs to be computed the first time a particle is added.
 		bool bNeedsXRInitialization = true;
 
+		// Whether or not the anchor got set by the GT and we shouldn't try to recompute its value.
+		bool bAnchorLock = false;
+
 		// Pending particles that need to be added into the connectivity graph.
 		TArray<TPair<FPBDRigidParticleHandle*, EClusterUnionConnectivityOperation>> PendingConnectivityOperations;
 	};
@@ -164,6 +167,9 @@ namespace Chaos
 
 		// Returns all cluster unions. Really meant only to be used for debugging.
 		const TMap<FClusterUnionIndex, FClusterUnion>& GetAllClusterUnions() const { return ClusterUnions; }
+
+		// Put in a deferred request to update the properties of the specified cluster union. Flags can limit the scope of the update.
+		void RequestDeferredClusterPropertiesUpdate(FClusterUnionIndex ClusterIndex, EUpdateClusterUnionPropertiesFlags Flags);
 	private:
 		FRigidClustering& MClustering;
 		FPBDRigidsEvolutionGBF& MEvolution;
@@ -190,7 +196,6 @@ namespace Chaos
 		TMap<FClusterUnionIndex, FClusterUnion> ClusterUnions;
 
 		// The set of cluster unions that we want to defer updating their cluster union properties for.
-		void RequestDeferredClusterPropertiesUpdate(FClusterUnionIndex ClusterIndex, EUpdateClusterUnionPropertiesFlags Flags);
 		TMap<FClusterUnionIndex, EUpdateClusterUnionPropertiesFlags> DeferredClusterUnionsForUpdateProperties;
 		
 		//
