Rate/op limiting for IAS' FIFO file cache
#rnx

diff --git a/Engine/Source/Runtime/Experimental/IoStoreOnDemand/Private/FileCache.cpp b/Engine/Source/Runtime/Experimental/IoStoreOnDemand/Private/FileCache.cpp
--- a/Engine/Source/Runtime/Experimental/IoStoreOnDemand/Private/FileCache.cpp
+++ b/Engine/Source/Runtime/Experimental/IoStoreOnDemand/Private/FileCache.cpp
@@ -6,6 +6,8 @@
 #include "GenericPlatform/GenericPlatformFile.h"
 #include "HAL/Event.h"
 #include "HAL/FileManager.h"
+#include "HAL/Platform.h"
+#include "HAL/PlatformFile.h"
 #include "HAL/Runnable.h"
 #include "HAL/RunnableThread.h"
 #include "IO/IoDispatcher.h"
@@ -33,6 +35,8 @@ DEFINE_LOG_CATEGORY(LogIoCache);
 namespace UE::IO::Private
 {
 
+using namespace UE::Tasks;
+
 ///////////////////////////////////////////////////////////////////////////////
 class FCacheFileToc
 {
@@ -171,7 +175,7 @@ public:
 	bool Contains(FIoHash Key) const;
 	bool Get(const FIoHash& Key, FCacheEntry& OutEntry) const;
 	bool InsertPending(FIoHash Key, FIoBuffer& Data, bool& bAdded);
-	bool RemovePending(FCacheEntryList& OutPending);
+	int32 RemovePending(FCacheEntryList& OutPending, uint32 MaxSize);
 	void InsertPersisted(FCacheEntryList&& InPersisted, const uint64 CursorPos);
 	void RemovePersisted(const uint64 RequiredSize);
 	uint64 GetPendingBytes() const { return TotalPendingBytes; }
@@ -227,6 +231,8 @@ bool FCacheMap::InsertPending(FIoHash Key, FIoBuffer& Data, bool& bAdded)
 {
 	check(Data.GetSize() > 0);
 
+	FScopeLock _(&Cs);
+
 	bAdded = false;
 	if (TotalPendingBytes + Data.GetSize() > MaxPendingBytes)
 	{
@@ -234,8 +240,6 @@ bool FCacheMap::InsertPending(FIoHash Key, FIoBuffer& Data, bool& bAdded)
 		return false;
 	}
 
-	FScopeLock _(&Cs);
-
 	if (Lookup.Contains(Key))
 	{
 		TRACE_COUNTER_INCREMENT(FFileIoCache_PutExistingCount);
@@ -258,21 +262,36 @@ bool FCacheMap::InsertPending(FIoHash Key, FIoBuffer& Data, bool& bAdded)
 	return bAdded = true;
 }
 
-bool FCacheMap::RemovePending(FCacheEntryList& OutPending)
+int32 FCacheMap::RemovePending(FCacheEntryList& OutPending, uint32 MaxSize)
 {
 	FScopeLock _(&Cs);
 
 	if (Pending.IsEmpty())
 	{
-		return false;
+		return 0;
 	}
 
-	OutPending.AddTail(MoveTemp(Pending));
-	Pending.Reset();
-	TotalPendingBytes = 0;
-	TRACE_COUNTER_SET(FFileIoCache_PendingBytes, 0);
+	uint32 ReturnSize = 0;
+	for (FCacheEntryList::TIterator Iter = Pending.begin(); Iter != Pending.end();)
+	{
+		FCacheEntry& Entry = *Iter;
+		++Iter;
+
+		if (MaxSize < ReturnSize + Entry.Data.GetSize())
+		{
+			break;
+		}
 
-	return true;
+		ReturnSize += Entry.Data.GetSize();
+
+		Pending.Remove(&Entry);
+		OutPending.AddTail(&Entry);
+	}
+
+	TotalPendingBytes -= ReturnSize;
+	TRACE_COUNTER_SET(FFileIoCache_PendingBytes, TotalPendingBytes);
+
+	return ReturnSize;
 }
 
 void FCacheMap::InsertPersisted(FCacheEntryList&& InPersisted, const uint64 CursorPos)
@@ -376,7 +395,119 @@ FIoStatus FCacheMap::Save(const FString& FilePath, const uint64 CursorPos)
 	return CacheFileToc.Save(FilePath, CursorPos);
 }
 
-using namespace UE::Tasks;
+
+////////////////////////////////////////////////////////////////////////////////
+class FGovernorExternal
+{
+public:
+				FGovernorExternal();
+	void		Set(...) {}
+	uint32		TickAllowance();
+	void		Return(uint32) {}
+
+private:
+	int64		CycleThreshold;
+	int64		CycleLast;
+};
+
+////////////////////////////////////////////////////////////////////////////////
+FGovernorExternal::FGovernorExternal()
+{
+	CycleThreshold = int64(1.0 / FPlatformTime::GetSecondsPerCycle());
+	CycleThreshold = (CycleThreshold * 3) / 4;
+	CycleLast = FPlatformTime::Cycles() - CycleThreshold;
+}
+
+////////////////////////////////////////////////////////////////////////////////
+uint32 FGovernorExternal::TickAllowance()
+{
+	// We'll only check the platform layer for our allowance every now and again
+	int64 Cycle = FPlatformTime::Cycles64();
+	if (Cycle - CycleLast < CycleThreshold)
+	{
+		return 0;
+	}
+
+	CycleLast = Cycle;
+
+	IPlatformFile& Ipf = IPlatformFile::GetPlatformPhysical();
+	return Ipf.GetAllowedBytesToWriteThrottledStorage();
+}
+
+
+
+////////////////////////////////////////////////////////////////////////////////
+class FGovernorInternal
+{
+public:
+	void		Set(uint32 InAllowance, uint32 InOps, uint32 Seconds);
+	uint32		TickAllowance();
+	void		Return(uint32 LeftOver);
+
+private:
+	uint32		TickAllowance(int64 Cycles);
+	int64		CycleThreshold;
+	int64		CyclePrev;
+	uint32		Allowance;
+	uint32		RunOff;
+};
+
+////////////////////////////////////////////////////////////////////////////////
+void FGovernorInternal::Set(uint32 InAllowance, uint32 Ops, uint32 Seconds)
+{
+	CycleThreshold = int64(1.0 / FPlatformTime::GetSecondsPerCycle());
+
+	CycleThreshold = (CycleThreshold * Seconds) / Ops;
+	Allowance = InAllowance / Ops;
+
+	CyclePrev = FPlatformTime::Cycles() - CycleThreshold;
+}
+
+////////////////////////////////////////////////////////////////////////////////
+uint32 FGovernorInternal::TickAllowance()
+{
+	int64 Cycle = FPlatformTime::Cycles64();
+	uint32 Ret = TickAllowance(Cycle) + RunOff;
+	RunOff = 0;
+	return Ret;
+}
+
+////////////////////////////////////////////////////////////////////////////////
+uint32 FGovernorInternal::TickAllowance(int64 Cycle)
+{
+	int64 CycleDelta = FMath::Max(0ll, Cycle - CyclePrev);
+	if (CycleDelta < CycleThreshold)
+	{
+		return 0;
+	}
+
+	// A crude guard against runaway.
+	for (; CycleDelta > CycleThreshold; CycleDelta -= CycleThreshold);
+	CyclePrev = Cycle + CycleDelta;
+
+	return Allowance;
+}
+
+////////////////////////////////////////////////////////////////////////////////
+void FGovernorInternal::Return(uint32 LeftOver)
+{
+	// Roughly keep to some arbitrary limit so as to avoid excessive growth
+	uint32 OnePointFive = Allowance + (Allowance >> 1);
+	RunOff = FMath::Min(OnePointFive, LeftOver);
+}
+
+
+
+////////////////////////////////////////////////////////////////////////////////
+#if !defined(UE_USE_PLATFORM_GOVERNOR)
+#	define UE_USE_PLATFORM_GOVERNOR 1
+#endif
+
+#if UE_USE_PLATFORM_GOVERNOR
+	using FGovernor = FGovernorExternal;
+#else
+	using FGovernor = FGovernorInternal;
+#endif
 
 ////////////////////////////////////////////////////////////////////////////////
 class FFileIoCache final
@@ -397,11 +528,6 @@ public:
 		return true;
 	}
 
-	virtual uint32 Run() override
-	{
-		return FileWriterThreadEntry();
-	}
-
 	virtual void Stop() override
 	{
 		bStopRequested = true;
@@ -409,9 +535,10 @@ public:
 	}
 
 private:
+	virtual uint32 Run() override;
 	void Initialize();
 	void Shutdown();
-	uint32 FileWriterThreadEntry();
+	void FileWriterThreadInner();
 
 	FFileIoCacheConfig CacheConfig;
 	FCacheMap CacheMap;
@@ -420,11 +547,14 @@ private:
 	FString CacheFilePath;
 	uint64 WriteCursorPos = 0;
 	std::atomic_bool bStopRequested{false};
+	FGovernor Governor;
 };
 
 FFileIoCache::FFileIoCache(const FFileIoCacheConfig& Config)
 	: CacheConfig(Config)
 {
+	Governor.Set(Config.WriteRate.Allowance, Config.WriteRate.Ops, Config.WriteRate.Seconds);
+
 	CacheMap.SetCacheLimits(Config.MemoryStorageSize, Config.DiskStorageSize);
 	Initialize();
 }
@@ -479,7 +609,7 @@ TTask<TIoStatusOr<FIoBuffer>> FFileIoCache::Get(const FIoHash& Key, const FIoRea
 				FIoBuffer Buffer = Options.GetTargetVa() ? FIoBuffer(FIoBuffer::Wrap, Options.GetTargetVa(), ReadSize) : FIoBuffer(ReadSize);
 
 				IPlatformFile& Ipf = IPlatformFile::GetPlatformPhysical();
-				if (TUniquePtr<IFileHandle> FileHandle(Ipf.OpenRead(*CacheFilePath, false)); FileHandle.IsValid())
+				if (TUniquePtr<IFileHandle> FileHandle(Ipf.OpenRead(*CacheFilePath, true)); FileHandle.IsValid())
 				{
 					UE_LOG(LogIoCache, VeryVerbose, TEXT("Read chunk, Key='%s', Hash='%s', File='%s', Offset='%llu', Size='%llu'"),
 						*LexToString(Entry.Key), *LexToString(Entry.Hash), *CacheFilePath, Entry.SerialOffset, Entry.SerialSize);
@@ -537,7 +667,6 @@ void FFileIoCache::Initialize()
 	CacheFilePath = CacheDir / TEXT("cache.ucas");
 	WriteCursorPos = 0;
 
-	IPlatformFile& Ipf = IPlatformFile::GetPlatformPhysical();
 	IFileManager& FileMgr = IFileManager::Get();
 	
 	FIoStatus CacheStatus(EIoErrorCode::Unknown);
@@ -602,80 +731,96 @@ void FFileIoCache::Shutdown()
 	const FString CacheTocPath = FPaths::ProjectPersistentDownloadDir() / TEXT("IoCache") / TEXT("cache.utoc");
 	UE_LOG(LogIoCache, Log, TEXT("Saving TOC '%s'"), *CacheTocPath);
 	CacheMap.Save(CacheTocPath, WriteCursorPos);
+
+	CacheMap.Reset();
 }
 
-uint32 FFileIoCache::FileWriterThreadEntry()
+uint32 FFileIoCache::Run()
 {
-	IPlatformFile& Ipf = IPlatformFile::GetPlatformPhysical();
-
 	while (!bStopRequested)
 	{
-		for (;;)
-		{
-			FCacheEntryList Entries;
-			if (CacheMap.RemovePending(Entries) == false)
-			{
-				break;
-			}
+		FileWriterThreadInner();
+		TickWriterEvent->Wait(10);
+	}
 
-			TUniquePtr<IFileHandle> WriteFileHandle(Ipf.OpenWrite(*CacheFilePath, true, true));
-			if (!WriteFileHandle.IsValid())
-			{
-				UE_LOG(LogIoCache, Warning, TEXT("Write chunks failed, unable to open file '%s' for writing"), *CacheFilePath);
-				break;
-			}
+	return 0;
+}
 
-			WriteFileHandle->Seek(WriteCursorPos);
+void FFileIoCache::FileWriterThreadInner()
+{
+	static uint32 _TickCount = 0;
+	++_TickCount;
 
-			TRACE_CPUPROFILER_EVENT_SCOPE(FFileIoCache::WriteCacheEntry);
+	// Update the rate limiting and see how much we are allowed to write
+	uint32 WriteAllowance = Governor.TickAllowance();
+	if (!WriteAllowance)
+	{
+		return;
+	}
 
-			uint64 TotalPendingSize = 0;
-			for (FCacheEntry& Entry : Entries)
-			{
-				Entry.State = ECacheEntryState::Writing;
-				TotalPendingSize += Entry.Data.GetSize();
-			}
-			CacheMap.RemovePersisted(TotalPendingSize);
+	// Collect pending entries up to the write allowance
+	FCacheEntryList Entries;
+	int32 PendingSize = CacheMap.RemovePending(Entries, WriteAllowance);
+	if (PendingSize <= 0)
+	{
+		return;
+	}
 
-			//TODO: Write bigger chunks
-			for (FCacheEntry& Entry : Entries)
-			{
-				const FIoBuffer& Data = Entry.Data;
-				check(Data.GetSize() > 0);
+	if (WriteAllowance -= PendingSize)
+	{
+		Governor.Return(WriteAllowance);
+	}
 
-				Entry.SerialOffset = WriteFileHandle->Tell();
-				Entry.SerialSize = Data.GetSize();
-				Entry.Hash = FIoHash::HashBuffer(Data.GetView());
+	// Wrap the write cursor if we'd end up writing over the end
+	if (WriteCursorPos + PendingSize > CacheConfig.DiskStorageSize)
+	{
+		WriteCursorPos = 0;
+	}
 
-				UE_LOG(LogIoCache, VeryVerbose, TEXT("Write chunk, Key='%s', Hash='%s', File='%s', Offset='%llu', Size='%llu'"),
-					*LexToString(Entry.Key), *LexToString(Entry.Hash), *CacheFilePath, Entry.SerialOffset, Entry.SerialSize);
+	// Copy data to be cached into a buffer
+	auto* Buffer = (uint8*)FMemory::Malloc(PendingSize);
+	auto* Cursor = (uint8*)Buffer;
+	for (FCacheEntry& Entry : Entries)
+	{
+		const FMemoryView& View = Entry.Data.GetView();
 
-				const uint64 RemainingDiskSize = CacheConfig.DiskStorageSize - WriteFileHandle->Tell();
-				const uint64 ByteCount = FMath::Min(Entry.Data.GetSize(), RemainingDiskSize);
-				WriteFileHandle->Write(Entry.Data.GetData(), ByteCount);
+		Entry.SerialOffset = WriteCursorPos + uint32(ptrdiff_t(Cursor - Buffer));
+		Entry.SerialSize = View.GetSize();
+		Entry.Hash = FIoHash::HashBuffer(View);
+		Entry.State = ECacheEntryState::Writing;
 
-				const uint64 RemainingChunkSize = Entry.Data.GetSize() - ByteCount;
-				if (RemainingChunkSize > 0)
-				{
-					WriteFileHandle->Flush();
-					WriteFileHandle->Seek(0);
-					WriteFileHandle->Write(Entry.Data.GetData() + ByteCount, RemainingChunkSize);
-				}
-			}
+		::memcpy(Cursor, View.GetData(), View.GetSize());
+		Cursor += View.GetSize();
+	}
 
-			WriteFileHandle->Flush();
-			WriteCursorPos = WriteFileHandle->Tell();
+	// Drop buffers.
+	for (FCacheEntry& Entry : Entries)
+	{
+		FIoBuffer& Data = Entry.Data;
+		Data = FIoBuffer();
+	}
 
-			CacheMap.InsertPersisted(MoveTemp(Entries), WriteFileHandle->Tell());
-		}
+	// Open cache file and write the block to it.
+	UE_LOG(LogIoCache, VeryVerbose, TEXT("Write; cursor:%llu size:%d"), WriteCursorPos, PendingSize);
+	TRACE_CPUPROFILER_EVENT_SCOPE(FFileIoCache::WriteCacheEntry);
 
-		if (bStopRequested == false)
-		{
-			TickWriterEvent->Wait();
-		}
+	IPlatformFile& Ipf = IPlatformFile::GetPlatformPhysical();
+	TUniquePtr<IFileHandle> FileHandle(Ipf.OpenWrite(*CacheFilePath, true, true));
+	if (!FileHandle.IsValid())
+	{
+		UE_LOG(LogIoCache, Warning, TEXT("Write chunks failed, unable to open file '%s' for writing"), *CacheFilePath);
+		return;
 	}
+	FileHandle->Seek(WriteCursorPos);
+	FileHandle->Write(Buffer, PendingSize);
+	FileHandle->Flush();
+	WriteCursorPos += PendingSize;
 
-	return 0;
+	CacheMap.InsertPersisted(MoveTemp(Entries), WriteCursorPos);
+
+	CacheMap.RemovePersisted(PendingSize);
+
+	Entries.Reset();
 }
 
 } // namespace UE::IO::Private
diff --git a/Engine/Source/Runtime/Experimental/IoStoreOnDemand/Private/FileCache.h b/Engine/Source/Runtime/Experimental/IoStoreOnDemand/Private/FileCache.h
--- a/Engine/Source/Runtime/Experimental/IoStoreOnDemand/Private/FileCache.h
+++ b/Engine/Source/Runtime/Experimental/IoStoreOnDemand/Private/FileCache.h
@@ -35,8 +35,16 @@ public:
 
 struct FFileIoCacheConfig
 {
+	struct FRate
+	{
+		uint32 Allowance = ~0u;
+		uint32 Ops = 1;
+		uint32 Seconds = 1;
+	};
+
 	uint64 DiskStorageSize;
 	uint64 MemoryStorageSize;
+	FRate WriteRate;
 };
 
 TUniquePtr<IIoCache> MakeFileIoCache(const FFileIoCacheConfig& Config);
