Iris
- Fixed issue with ObjectReferenceCache not supprorting loading of packages associated with client assigned references
- Disabled bIrisEnableLowNetUpdateFrequencyEnsure by default and replaced it with a warning
- Fixed log spam about ResolveObjectReference Failed to resolve clientassigned ref

diff --git a/Engine/Source/Runtime/Engine/Private/Net/Iris/ReplicationSystem/ActorReplicationBridge.cpp b/Engine/Source/Runtime/Engine/Private/Net/Iris/ReplicationSystem/ActorReplicationBridge.cpp
--- a/Engine/Source/Runtime/Engine/Private/Net/Iris/ReplicationSystem/ActorReplicationBridge.cpp
+++ b/Engine/Source/Runtime/Engine/Private/Net/Iris/ReplicationSystem/ActorReplicationBridge.cpp
@@ -59,7 +59,7 @@ namespace UE::Net::Private
 extern IRISCORE_API float PollFrequencyMultiplier;
 extern IRISCORE_API const int MaxPollFramePeriod;
 
-static bool bIrisEnableLowNetUpdateFrequencyEnsure = true;
+static bool bIrisEnableLowNetUpdateFrequencyEnsure = false;
 static FAutoConsoleVariableRef CVarEnableLowNetUpdateFrequencyEnsure(
 		TEXT("net.Iris.EnableLowNetUpdateFrequencyEnsure"),
 		bIrisEnableLowNetUpdateFrequencyEnsure,
@@ -261,7 +261,11 @@ UE::Net::FNetRefHandle UActorReplicationBridge::BeginReplication(AActor* Actor,
 	const uint32 ClampedPollPeriod = FMath::Clamp<uint32>(UnclampedPollPeriod, 1U, static_cast<uint32>(std::numeric_limits<uint8>::max()) + 1U);
 	CreateNetRefHandleParams.PollFramePeriod = (ClampedPollPeriod - 1U) & 255U;
 #if !UE_BUILD_SHIPPING
-	ensureAlwaysMsgf(!UE::Net::Private::bIrisEnableLowNetUpdateFrequencyEnsure || (ClampedPollPeriod >= UnclampedPollPeriod), TEXT("Very low NetUpdateFrequency %f for Actor %s. Suggest setting it to %f or higher."), Actor->NetUpdateFrequency, ToCStr(Actor->GetName()), GetMinSupportedNetUpdateFrequency());
+	if ((ClampedPollPeriod >= UnclampedPollPeriod))
+	{
+		ensureAlwaysMsgf(!UE::Net::Private::bIrisEnableLowNetUpdateFrequencyEnsure, TEXT("Very low NetUpdateFrequency %f for Actor %s. Suggest setting it to %f or higher."), Actor->NetUpdateFrequency, ToCStr(Actor->GetName()), GetMinSupportedNetUpdateFrequency());
+		UE_CLOG(!UE::Net::Private::bIrisEnableLowNetUpdateFrequencyEnsure, LogIrisBridge, Warning, TEXT("Very low NetUpdateFrequency %f for Actor %s. Suggest setting it to %f or higher."), Actor->NetUpdateFrequency, ToCStr(Actor->GetName()), GetMinSupportedNetUpdateFrequency());
+	}
 	ensureAlwaysMsgf(!(Actor->bAlwaysRelevant || Actor->bOnlyRelevantToOwner) || CreateNetRefHandleParams.StaticPriority >= 1.0f, TEXT("Very low NetPriority %.02f for always relevant or owner relevant Actor %s. Set it to 1.0f or higher."), Actor->NetPriority, ToCStr(Actor->GetName()));
 #endif
 
diff --git a/Engine/Source/Runtime/Experimental/Iris/Core/Private/Iris/ReplicationSystem/ObjectReferenceCache.cpp b/Engine/Source/Runtime/Experimental/Iris/Core/Private/Iris/ReplicationSystem/ObjectReferenceCache.cpp
--- a/Engine/Source/Runtime/Experimental/Iris/Core/Private/Iris/ReplicationSystem/ObjectReferenceCache.cpp
+++ b/Engine/Source/Runtime/Experimental/Iris/Core/Private/Iris/ReplicationSystem/ObjectReferenceCache.cpp
@@ -16,6 +16,7 @@
 #include "Iris/Serialization/NetSerializationContext.h"
 #include "Iris/Serialization/ObjectNetSerializer.h"
 #include "Misc/CoreMiscDefines.h"
+#include "Misc/PackageName.h"
 #include "Misc/StringBuilder.h"
 #include "Net/Core/Trace/NetDebugName.h"
 #include "Net/Core/Trace/NetTrace.h"
@@ -924,8 +925,6 @@ ENetObjectReferenceResolveResult FObjectReferenceCache::ResolveObjectReference(c
 	if (Reference.PathToken.IsValid())
 	{
 		// This path is only used by Client to Server references
-		//$TODO: $IRIS: Should we try to load or not? Currently we do not!
-
 		const TCHAR* ResolvedToken = StringTokenStore->ResolveRemoteToken(Reference.PathToken, *ResolveContext.RemoteNetTokenStoreState);
 		if (ResolvedToken)
 		{
@@ -938,13 +937,31 @@ ENetObjectReferenceResolveResult FObjectReferenceCache::ResolveObjectReference(c
 			{
 				UObject* ReplicatedOuter = GetObjectFromReferenceHandle(Reference.GetRefHandle());
 				ResolvedObject = ReplicatedOuter ? StaticFindObject(UObject::StaticClass(), ReplicatedOuter, ToCStr(ObjectPath), false) : nullptr;
+
+				if (ResolvedObject == nullptr)
+				{
+					UE_LOG_REFERENCECACHE_WARNING(TEXT("ResolveObjectReference Failed to resolve clientassigned ref: %s due to not finding object with relative path %s."), *Reference.ToString(), *ObjectPath);	
+				}
 			}
 			else
 			{
 				ResolvedObject = StaticFindObject(UObject::StaticClass(), nullptr, *ObjectPath, false);
-			}
 
-			UE_LOG_REFERENCECACHE_WARNING(TEXT("ResolveObjectReference Failed to resolve clientassigned ref: %s due to not finding object with path %s."), *Reference.ToString(), *ObjectPath);	
+				// Try to load package if it wasn't found. Note load package fails if the package is already loaded.
+				if (ResolvedObject == nullptr)
+				{
+					FPackagePath Path = FPackagePath::FromPackageNameChecked(FPackageName::ObjectPathToPackageName(ObjectPath));
+					if (UPackage* LoadedPackage = LoadPackage(nullptr, Path, LOAD_None))
+					{
+						ResolvedObject = StaticFindObject(UObject::StaticClass(), nullptr, *ObjectPath, false);	
+					}
+				}
+
+				if (ResolvedObject == nullptr)
+				{
+					UE_LOG_REFERENCECACHE_WARNING(TEXT("ResolveObjectReference Failed to resolve clientassigned ref: due to not finding object with path %s."), *ObjectPath);	
+				}
+			}
 		}
 		else
 		{
