Networked Physics: Set same history size on RewindData and MarshallingManager push data history.
RewindData rounds up to nearest power of 2.

diff --git a/Engine/Source/Runtime/Experimental/Chaos/Private/PBDRigidsSolver.cpp b/Engine/Source/Runtime/Experimental/Chaos/Private/PBDRigidsSolver.cpp
--- a/Engine/Source/Runtime/Experimental/Chaos/Private/PBDRigidsSolver.cpp
+++ b/Engine/Source/Runtime/Experimental/Chaos/Private/PBDRigidsSolver.cpp
@@ -964,7 +964,7 @@ namespace Chaos
 
 	void FPBDRigidsSolver::EnableRewindCapture(int32 NumFrames, bool InUseCollisionResimCache, TUniquePtr<IRewindCallback>&& RewindCallback)
 	{
-		//TODO: this function calls both internal and extrnal - sort of assumed during initialization. Should decide what thread it's called on and mark it as either external or internal
+		//TODO: this function calls both internal and external - sort of assumed during initialization. Should decide what thread it's called on and mark it as either external or internal
 		if (MRewindData.IsValid())
 		{
 			MRewindData->Init(((FPBDRigidsSolver*)this), NumFrames, InUseCollisionResimCache, ((FPBDRigidsSolver*)this)->GetCurrentFrame());
@@ -975,7 +975,8 @@ namespace Chaos
 		}
 		bUseCollisionResimCache = InUseCollisionResimCache;
 		MRewindCallback = MoveTemp(RewindCallback);
-		MarshallingManager.SetHistoryLength_Internal(NumFrames);
+		const int32 NumFramesSet = GetRewindData() != nullptr ? GetRewindData()->Capacity() : NumFrames;
+		MarshallingManager.SetHistoryLength_Internal(NumFramesSet);
 		MEvolution->SetRewindData(GetRewindData());
 		if(MRewindCallback)
 		{
