various changes to ensure that cluster union book-keeping is consistent

- defer update cluster union properties when needing to remove particles from parent clusters on add
- when adding particles to a cluster, always try to perform bookkeeping on it regardless of whether or not its parent is currently the same
- don't call UpdateKinematicProperties on release if that particle is getting added into some other cluster due to connectivity
- don't call RemoveChildFromParentAndChildrenArray in FRigidClustering::HandleConnectivityOnReleaseClusterParticle when adding into a cluster union since that will cause the particle to be removed from the rigid clustering bookkeeping but not the cluster union bookkeeping. defer this operation until the cluster union pending add happens instead
- remove the need to bookkeep an extra ParticleToClusterUnionIndex map
- remove the possibility of adding duplicate particles into cluster unions which could lead to partial removal later on which can result in ensures/asserts

diff --git a/Engine/Source/Runtime/Experimental/Chaos/Private/Chaos/ClusterUnionManager.cpp b/Engine/Source/Runtime/Experimental/Chaos/Private/Chaos/ClusterUnionManager.cpp
--- a/Engine/Source/Runtime/Experimental/Chaos/Private/Chaos/ClusterUnionManager.cpp
+++ b/Engine/Source/Runtime/Experimental/Chaos/Private/Chaos/ClusterUnionManager.cpp
@@ -196,7 +196,6 @@ namespace Chaos
 		MEvolution.DisableParticle(NewUnion.InternalCluster);
 
 		ClusterUnions.Add(NewIndex, NewUnion);
-		ParticleToClusterUnionIndex.Add(NewUnion.InternalCluster, NewIndex);
 		return NewIndex;
 	}
 
@@ -373,7 +372,7 @@ namespace Chaos
 
 		// If we're adding particles to a cluster we need to first make sure they're not part of any other cluster.
 		// Book-keeping might get a bit odd if we try to add a particle to a new cluster and then only later remove the particle from its old cluster.
-		HandleRemoveOperationWithClusterLookup(Particles, EClusterUnionOperationTiming::Immediate);
+		HandleRemoveOperationWithClusterLookup(Particles, EClusterUnionOperationTiming::Defer);
 
 		TGuardValue_Bitfield_Cleanup<TFunction<void()>> Cleanup(
 			[this, OldGenerateClusterBreaking=MClustering.GetDoGenerateBreakingData()]() {
@@ -394,14 +393,18 @@ namespace Chaos
 		bool bIsSleeping = true;
 		bool bIsAnchored = false;
 
+		// Use a set to guarantee that there's no duplicates. This is necessary since further removal operations assume that the particle is only in here once.
+		const TSet<FPBDRigidParticleHandle*> ParticleSet{ Particles };
+
 		TArray<FPBDRigidParticleHandle*> FinalParticlesToAdd;
-		FinalParticlesToAdd.Reserve(Particles.Num());
+		FinalParticlesToAdd.Reserve(ParticleSet.Num());
 
 		// This is only relevant when bReleaseClustersFirst=true. This is used to be able to
 		// properly notify the parent cluster about its child proxies.
 		TMap<FPBDRigidParticleHandle*, FPBDRigidParticleHandle*> ChildToParentMap;
 
-		for (FPBDRigidParticleHandle* Handle : Particles)
+		
+		for (FPBDRigidParticleHandle* Handle : ParticleSet)
 		{
 			if (!Handle)
 			{
@@ -443,7 +446,6 @@ namespace Chaos
 		Cluster->ChildParticles.Append(FinalParticlesToAdd);
 		for (FPBDRigidParticleHandle* Particle : FinalParticlesToAdd)
 		{
-			ParticleToClusterUnionIndex.Add(Particle, ClusterIndex);
 			if (!bIsNewCluster)
 			{
 				Cluster->PendingConnectivityOperations.Add({Particle, EClusterUnionConnectivityOperation ::Add});
@@ -582,8 +584,7 @@ namespace Chaos
 		for (int32 Index = ParticleIndicesToRemove.Num() - 1; Index >= 0; --Index)
 		{
 			const int32 ParticleIndex = ParticleIndicesToRemove[Index];
-			ParticleToClusterUnionIndex.Remove(Cluster->ChildParticles[ParticleIndex]);
-			Cluster->ChildParticles.RemoveAt(ParticleIndex);
+			Cluster->ChildParticles.RemoveAtSwap(ParticleIndex);
 		}
 
 		MClustering.RemoveParticlesFromCluster(Cluster->InternalCluster, Particles);
@@ -776,10 +777,27 @@ namespace Chaos
 			return INDEX_NONE;
 		}
 
-		if (FClusterUnionIndex* Index = ParticleToClusterUnionIndex.Find(ChildParticle))
+		if (FPBDRigidClusteredParticleHandle* ChildClustered = ChildParticle->CastToClustered())
 		{
-			return *Index;
+			if (FPBDRigidClusteredParticleHandle* Parent = ChildClustered->Parent())
+			{
+				// Recursion should be fine here since the hierarchy should be fairly shallow.
+				return FindClusterUnionIndexFromParticle(Parent);
+			}
+
+			// The only other check to do here is to see if this is the cluster union particle itself.
+			const int32 ClusterGroupIndex = ChildClustered->ClusterGroupIndex();
+			if (FClusterUnion* ClusterUnion = FindClusterUnion(FMath::Abs(ClusterGroupIndex)))
+			{
+				// This is a sanity check which may or may not be necessary.
+				// TODO: Can probably be safely removed once we deprecate cluster group indices on GCs.
+				if (ClusterUnion->InternalCluster == ChildClustered)
+				{
+					return ClusterUnion->InternalIndex;
+				}
+			}
 		}
+		
 		return INDEX_NONE;
 	}
 
diff --git a/Engine/Source/Runtime/Experimental/Chaos/Private/Chaos/PBDRigidClustering.cpp b/Engine/Source/Runtime/Experimental/Chaos/Private/Chaos/PBDRigidClustering.cpp
--- a/Engine/Source/Runtime/Experimental/Chaos/Private/Chaos/PBDRigidClustering.cpp
+++ b/Engine/Source/Runtime/Experimental/Chaos/Private/Chaos/PBDRigidClustering.cpp
@@ -262,15 +262,12 @@ namespace Chaos
 			{
 				if (FPBDRigidClusteredParticleHandle* ExistingParent = ClusteredChild->Parent())
 				{
-					if (ExistingParent == Cluster)
+					if (ExistingParent != Cluster)
 					{
-						// This shouldn't really ever be the case probably?
-						continue;
+						// This is needed in the case where we use intercluster edges with geometry collections that need to then get stuck into a cluster union.
+						// It's possible due to replication ordering that we create an internal cluster surrounding the piece that we want to add to the cluster union first.
+						RemoveChildFromParentAndChildrenArray(ClusteredChild, ExistingParent);
 					}
-
-					// This is needed in the case where we use intercluster edges with geometry collections that need to then get stuck into a cluster union.
-					// It's possible due to replication ordering that we create an internal cluster surrounding the piece that we want to add to the cluster union first.
-					RemoveChildFromParentAndChildrenArray(ClusteredChild, ExistingParent);
 				}
 
 				TopLevelClusterParents.Remove(ClusteredChild);
@@ -935,7 +932,6 @@ namespace Chaos
 				const FClusterUnionIndex ClusterUnionIndex = ClusterUnionManager.FindClusterUnionIndexFromParticle(Child);
 				const bool bIsInClusterUnion = ClusterUnionIndex != INDEX_NONE;
 
-				//UE_LOG(LogTemp, Warning, TEXT("Releasing child %d from parent %p due to strain %.5f Exceeding internal strain %.5f (Source: %s)"), ChildIdx, ClusteredParticle, ChildStrain, Child->Strain(), bForceRelease ? TEXT("Forced by caller") : ExternalStrainMap ? TEXT("External") : TEXT("Collision"));
 				if (bIsInClusterUnion)
 				{
 					DeferredRemoveFromClusterUnion.Add(Child);
@@ -1037,6 +1033,16 @@ namespace Chaos
 
 			for (FPBDRigidParticleHandle* Child : ActivatedChildren)
 			{
+				// If an activated child has a parent, we don't want to update their kinematic properties since they should be disabled
+				// and thus, shouldn't need to have properties updated.
+				if (FPBDRigidClusteredParticleHandle* ClusteredChild = Child->CastToClustered())
+				{
+					if (ClusteredChild->Parent())
+					{
+						check(ClusteredChild->Disabled());
+						continue;
+					}
+				}
 				UpdateKinematicProperties(Child, MChildren, MEvolution);
 			}
 
@@ -1129,10 +1135,6 @@ namespace Chaos
 
 				if (AttachedClusterUnion)
 				{
-					for (FPBDRigidParticleHandle* Child : Island)
-					{
-						RemoveChildFromParentAndChildrenArray(Child, ClusteredParticle);
-					}
 					ClusterUnionManager.AddPendingClusterIndexOperation(AttachedClusterUnion->InternalIndex, EClusterUnionOperation::Add, Island);
 					IslandIndicesToRemove.Add(IslandIndex);
 				}
diff --git a/Engine/Source/Runtime/Experimental/Chaos/Public/Chaos/ClusterUnionManager.h b/Engine/Source/Runtime/Experimental/Chaos/Public/Chaos/ClusterUnionManager.h
--- a/Engine/Source/Runtime/Experimental/Chaos/Public/Chaos/ClusterUnionManager.h
+++ b/Engine/Source/Runtime/Experimental/Chaos/Public/Chaos/ClusterUnionManager.h
@@ -214,9 +214,6 @@ namespace Chaos
 		TMap<FClusterUnionExplicitIndex, FClusterUnionIndex> ExplicitIndexMap;
 		FClusterCreationParameters DefaultClusterCreationParameters() const;
 
-		// A particle lookup table. Lets us go from a given particle pointer to the cluster it's in.
-		TMap<FPBDRigidParticleHandle*, FClusterUnionIndex> ParticleToClusterUnionIndex;
-
 		// If no cluster index is set but an explicit index is set, map the explicit index to a regular index.
 		FClusterUnionIndex GetOrCreateClusterUnionIndexFromExplicitIndex(FClusterUnionExplicitIndex InIndex);
 
