cluster unions now gate the call to UpdateKinematicProperties behind the EUpdateClusterUnionPropertiesFlags::UpdateKinematicProperties flag since this call shouldn't be necessary past creation

diff --git a/Engine/Source/Runtime/Experimental/Chaos/Private/Chaos/ClusterUnionManager.cpp b/Engine/Source/Runtime/Experimental/Chaos/Private/Chaos/ClusterUnionManager.cpp
--- a/Engine/Source/Runtime/Experimental/Chaos/Private/Chaos/ClusterUnionManager.cpp
+++ b/Engine/Source/Runtime/Experimental/Chaos/Private/Chaos/ClusterUnionManager.cpp
@@ -644,23 +644,26 @@ namespace Chaos
 			ClusterUnion.bNeedsXRInitialization = false;
 		}
 
-		if (ClusterUnion.bAnchorLock)
+		if (EnumHasAnyFlags(Flags, EUpdateClusterUnionPropertiesFlags::UpdateKinematicProperties))
 		{
-			// Skip UpdateKinematicPropertiesso as we won't need to look at our children for this.
-			if (ClusterUnion.InternalCluster->IsAnchored())
+			if (ClusterUnion.bAnchorLock)
 			{
-				MEvolution.SetParticleObjectState(ClusterUnion.InternalCluster, EObjectStateType::Kinematic);
+				// Skip UpdateKinematicPropertiesso as we won't need to look at our children for this.
+				if (ClusterUnion.InternalCluster->IsAnchored())
+				{
+					MEvolution.SetParticleObjectState(ClusterUnion.InternalCluster, EObjectStateType::Kinematic);
+				}
+				else
+				{
+					MEvolution.SetParticleObjectState(ClusterUnion.InternalCluster, EObjectStateType::Dynamic);
+					MEvolution.SetParticleKinematicTarget(ClusterUnion.InternalCluster, FKinematicTarget());
+				}
 			}
 			else
 			{
-				MEvolution.SetParticleObjectState(ClusterUnion.InternalCluster, EObjectStateType::Dynamic);
-				MEvolution.SetParticleKinematicTarget(ClusterUnion.InternalCluster, FKinematicTarget());
+				UpdateKinematicProperties(ClusterUnion.InternalCluster, MClustering.GetChildrenMap(), MEvolution);
 			}
 		}
-		else
-		{
-			UpdateKinematicProperties(ClusterUnion.InternalCluster, MClustering.GetChildrenMap(), MEvolution);
-		}
 
 		// We must reset collisions etc on this particle since geometry will be changed
 		MEvolution.InvalidateParticle(ClusterUnion.InternalCluster);
diff --git a/Engine/Source/Runtime/Experimental/Chaos/Private/PhysicsProxy/ClusterUnionPhysicsProxy.cpp b/Engine/Source/Runtime/Experimental/Chaos/Private/PhysicsProxy/ClusterUnionPhysicsProxy.cpp
--- a/Engine/Source/Runtime/Experimental/Chaos/Private/PhysicsProxy/ClusterUnionPhysicsProxy.cpp
+++ b/Engine/Source/Runtime/Experimental/Chaos/Private/PhysicsProxy/ClusterUnionPhysicsProxy.cpp
@@ -204,7 +204,7 @@ namespace Chaos
 						}
 
 						// Let the cluster union manager apply the proper properties on the particle.
-						Manager.RequestDeferredClusterPropertiesUpdate(ClusterUnionIndex, Chaos::EUpdateClusterUnionPropertiesFlags::None);
+						Manager.RequestDeferredClusterPropertiesUpdate(ClusterUnionIndex, Chaos::EUpdateClusterUnionPropertiesFlags::UpdateKinematicProperties);
 					}
 				}
 			}
diff --git a/Engine/Source/Runtime/Experimental/Chaos/Public/Chaos/ClusterUnionManager.h b/Engine/Source/Runtime/Experimental/Chaos/Public/Chaos/ClusterUnionManager.h
--- a/Engine/Source/Runtime/Experimental/Chaos/Public/Chaos/ClusterUnionManager.h
+++ b/Engine/Source/Runtime/Experimental/Chaos/Public/Chaos/ClusterUnionManager.h
@@ -39,7 +39,8 @@ namespace Chaos
 		RecomputeMassOrientation = 1 << 0,
 		ForceGenerateConnectionGraph = 1 << 1,
 		IncrementalGenerateConnectionGraph = 1 << 2,
-		All = RecomputeMassOrientation | ForceGenerateConnectionGraph
+		UpdateKinematicProperties = 1 << 3,
+		All = RecomputeMassOrientation | ForceGenerateConnectionGraph | UpdateKinematicProperties
 	};
 	ENUM_CLASS_FLAGS(EUpdateClusterUnionPropertiesFlags);
 
