Add missing #include

diff --git a/Engine/Source/Runtime/Experimental/Iris/Core/Private/Iris/ReplicationSystem/ObjectReferenceCache.h b/Engine/Source/Runtime/Experimental/Iris/Core/Private/Iris/ReplicationSystem/ObjectReferenceCache.h
--- a/Engine/Source/Runtime/Experimental/Iris/Core/Private/Iris/ReplicationSystem/ObjectReferenceCache.h
+++ b/Engine/Source/Runtime/Experimental/Iris/Core/Private/Iris/ReplicationSystem/ObjectReferenceCache.h
@@ -8,6 +8,7 @@
 #include "UObject/WeakObjectPtr.h"
 #include "Containers/Map.h"
 #include "ObjectReferenceCacheFwd.h"
+#include "UObject/ObjectPtr.h"
 
 namespace UE::Net
 {
