[ChaosVD] Implemented Particle Data Serialization and visualization

- Now we wrap the particle data into a single struct, composed of UStruct versions of existing data wrappers used by the Re-Simulation code. Once that is done, we serialize that struct as a CVD Binary Data entry instead of a Particle Trace Event.

- Updated the CVD Trace Analyzer to read the particle data from Binary Data events instead of the now removed CVD Trace Particle Event

- Create a property customization class for the new Particle Data Viewer struct. This allows us to hide properties that are part of the struct but don't have recorded data backing them up.

diff --git a/Engine/Plugins/Experimental/ChaosVD/Source/ChaosVD/Private/ChaosVDModule.cpp b/Engine/Plugins/Experimental/ChaosVD/Source/ChaosVD/Private/ChaosVDModule.cpp
--- a/Engine/Plugins/Experimental/ChaosVD/Source/ChaosVD/Private/ChaosVDModule.cpp
+++ b/Engine/Plugins/Experimental/ChaosVD/Source/ChaosVD/Private/ChaosVDModule.cpp
@@ -10,6 +10,7 @@
 #include "ChaosVDEngine.h"
 #include "WorkspaceMenuStructure.h"
 #include "WorkspaceMenuStructureModule.h"
+#include "DetailsCustomizations/ChaosVDParticleDataWrapperCustomization.h"
 #include "Trace/ChaosVDTraceManager.h"
 #include "Misc/Guid.h"
 
@@ -48,6 +49,7 @@ void FChaosVDModule::RegisterClassesCustomDetails() const
 {
 	FPropertyEditorModule& PropertyModule = FModuleManager::GetModuleChecked<FPropertyEditorModule>("PropertyEditor");
 	PropertyModule.RegisterCustomClassLayout("ChaosVDParticleActor", FOnGetDetailCustomizationInstance::CreateStatic(&FChaosVDParticleActorCustomization::MakeInstance));
+	PropertyModule.RegisterCustomPropertyTypeLayout("ChaosVDParticleDataWrapper", FOnGetPropertyTypeCustomizationInstance::CreateStatic(&FChaosVDParticleDataWrapperCustomization::MakeInstance));
 }
 
 TSharedRef<SDockTab> FChaosVDModule::SpawnMainTab(const FSpawnTabArgs& Args)
diff --git a/Engine/Plugins/Experimental/ChaosVD/Source/ChaosVD/Private/ChaosVDParticleActor.cpp b/Engine/Plugins/Experimental/ChaosVD/Source/ChaosVD/Private/ChaosVDParticleActor.cpp
--- a/Engine/Plugins/Experimental/ChaosVD/Source/ChaosVD/Private/ChaosVDParticleActor.cpp
+++ b/Engine/Plugins/Experimental/ChaosVD/Source/ChaosVD/Private/ChaosVDParticleActor.cpp
@@ -4,7 +4,7 @@
 
 #include "ChaosVDGeometryBuilder.h"
 #include "ChaosVDScene.h"
-#include "ChaosVDRecording.h"
+#include "DataWrappers/ChaosVDParticleDataWrapper.h"
 #include "Components/MeshComponent.h"
 #include "Components/SceneComponent.h"
 #include "Components/StaticMeshComponent.h"
@@ -21,29 +21,26 @@ AChaosVDParticleActor::AChaosVDParticleActor(const FObjectInitializer& ObjectIni
 	RootComponent = CreateDefaultSubobject<USceneComponent>(TEXT("SceneComponent0"));
 }
 
-void AChaosVDParticleActor::UpdateFromRecordedData(const FChaosVDParticleDebugData& InRecordedData, const Chaos::FRigidTransform3& SimulationTransform)
+void AChaosVDParticleActor::UpdateFromRecordedParticleData(const FChaosVDParticleDataWrapper& InRecordedData, const Chaos::FRigidTransform3& SimulationTransform)
 {
-	SetActorLocationAndRotation(SimulationTransform.TransformPosition(InRecordedData.Position), SimulationTransform.GetRotation() * InRecordedData.Rotation, false);
+	if (InRecordedData.ParticlePositionRotation.HasValidData())
+	{
+		// TODO: Only update if the transform (either the simulation or the particle) has changed;
+		SetActorLocationAndRotation(SimulationTransform.TransformPosition(InRecordedData.ParticlePositionRotation.MX), SimulationTransform.GetRotation() * InRecordedData.ParticlePositionRotation.MR, false);
+	}
 
-	if (RecordedDebugData.ImplicitObjectHash != InRecordedData.ImplicitObjectHash)
+	if (ParticleDataViewer.GeometryHash != InRecordedData.GeometryHash)
 	{
-		if (const TSharedPtr<FChaosVDScene>& ScenePtr = OwningScene.Pin())
-		{
-			if (const TSharedPtr<const Chaos::FImplicitObject>* Geometry = ScenePtr->GetUpdatedGeometry(InRecordedData.ImplicitObjectHash))
-			{
-				UpdateGeometryData(*Geometry, EChaosVDActorGeometryUpdateFlags::ForceUpdate);
-			}
-		}
+		UpdateGeometry(InRecordedData.GeometryHash, EChaosVDActorGeometryUpdateFlags::ForceUpdate);
 	}
 
-	// TODO: We should store a ptr to the data and in our custom details panel draw it, but the current data format is just a test
-	// So we should do this once it is final
-	RecordedDebugData = InRecordedData;
+	// TODO: We should store a ptr to the data and in our custom details panel draw it
+	ParticleDataViewer = InRecordedData;
 }
 
-void AChaosVDParticleActor::UpdateGeometryData(const TSharedPtr<const Chaos::FImplicitObject>& ImplicitObject, EChaosVDActorGeometryUpdateFlags OptionFlag)
+void AChaosVDParticleActor::UpdateGeometry(const TSharedPtr<const Chaos::FImplicitObject>& ImplicitObject, EChaosVDActorGeometryUpdateFlags OptionsFlags)
 {
-	if (EnumHasAnyFlags(OptionFlag, EChaosVDActorGeometryUpdateFlags::ForceUpdate))
+	if (EnumHasAnyFlags(OptionsFlags, EChaosVDActorGeometryUpdateFlags::ForceUpdate))
 	{
 		bIsGeometryDataGenerationStarted = false;
 
@@ -51,7 +48,7 @@ void AChaosVDParticleActor::UpdateGeometryData(const TSharedPtr<const Chaos::FIm
 		{
 			if (MeshComponent.IsValid())
 			{
-				MeshComponent->DestroyComponent();	
+				MeshComponent->DestroyComponent();
 			}
 		}
 
@@ -90,7 +87,7 @@ void AChaosVDParticleActor::UpdateGeometryData(const TSharedPtr<const Chaos::FIm
 				{
 					// This is a temp hack (and is not performant) we need until we have a proper way to filer out trigger volumes/brushes at will
 					// Without this most maps will be covered in boxes
-					if (RecordedDebugData.DebugName.Contains(TEXT("Brush")) || RecordedDebugData.DebugName.Contains(TEXT("Volume")))
+					if (ParticleDataViewer.DebugName.Contains(TEXT("Brush")) || ParticleDataViewer.DebugName.Contains(TEXT("Volume")))
 					{
 						for (TWeakObjectPtr<UMeshComponent> MeshComponent : OutGeneratedMeshComponents)
 						{
@@ -108,6 +105,17 @@ void AChaosVDParticleActor::UpdateGeometryData(const TSharedPtr<const Chaos::FIm
 	}
 }
 
+void AChaosVDParticleActor::UpdateGeometry(uint32 NewGeometryHash, EChaosVDActorGeometryUpdateFlags OptionsFlags)
+{
+	if (const TSharedPtr<FChaosVDScene>& ScenePtr = OwningScene.Pin())
+	{
+		if (const TSharedPtr<const Chaos::FImplicitObject>* Geometry = ScenePtr->GetUpdatedGeometry(NewGeometryHash))
+		{
+			UpdateGeometry(*Geometry, OptionsFlags);
+		}
+	}
+}
+
 void AChaosVDParticleActor::SetScene(const TSharedPtr<FChaosVDScene>& InScene)
 {
 	OwningScene = InScene;
@@ -116,9 +124,9 @@ void AChaosVDParticleActor::SetScene(const TSharedPtr<FChaosVDScene>& InScene)
 	{
 		GeometryUpdatedDelegate = ScenePtr->OnNewGeometryAvailable().AddWeakLambda(this, [this](const TSharedPtr<const Chaos::FImplicitObject>& ImplicitObject, const uint32 ID)
 		{
-			if (RecordedDebugData.ImplicitObjectHash == ID)
+			if (ParticleDataViewer.GeometryHash == ID)
 			{
-				UpdateGeometryData(ImplicitObject);
+				UpdateGeometry(ImplicitObject);
 			}
 		});
 	}
@@ -146,5 +154,19 @@ bool AChaosVDParticleActor::IsSelectedInEditor() const
 
 	return false;
 }
-#endif
 
+void AChaosVDParticleActor::PostEditChangeProperty(FPropertyChangedEvent& PropertyChangedEvent)
+{
+	Super::PostEditChangeProperty(PropertyChangedEvent);
+
+	if (PropertyChangedEvent.GetMemberPropertyName() == GET_MEMBER_NAME_STRING_CHECKED(AChaosVDParticleActor, ParticleDataViewer))
+	{
+		// Not particularly useful for now. This is a test code verifying we can react to changes in the data.
+		// In the future this could be part of the Re-simulation feature. When data is changed here, it can be propagated to the evolution instance that will be re-simulated
+		if (PropertyChangedEvent.GetPropertyName() == GET_MEMBER_NAME_STRING_CHECKED(FChaosVDParticleDataWrapper, GeometryHash))
+		{
+			UpdateGeometry(ParticleDataViewer.GeometryHash, EChaosVDActorGeometryUpdateFlags::ForceUpdate);
+		}
+	}
+}
+#endif //WITH_EDITOR
diff --git a/Engine/Plugins/Experimental/ChaosVD/Source/ChaosVD/Private/ChaosVDScene.cpp b/Engine/Plugins/Experimental/ChaosVD/Source/ChaosVD/Private/ChaosVDScene.cpp
--- a/Engine/Plugins/Experimental/ChaosVD/Source/ChaosVD/Private/ChaosVDScene.cpp
+++ b/Engine/Plugins/Experimental/ChaosVD/Source/ChaosVD/Private/ChaosVDScene.cpp
@@ -7,8 +7,13 @@
 #include "ChaosVDParticleActor.h"
 #include "Chaos/ImplicitObject.h"
 #include "ChaosVDRecording.h"
+#include "EditorActorFolders.h"
+#include "WorldPersistentFolders.h"
 #include "Elements/Framework/EngineElementsLibrary.h"
 #include "Elements/Framework/TypedElementSelectionSet.h"
+#include "EditorActorFolders.h"
+#include "WorldPersistentFolders.h"
+#include "DataWrappers/ChaosVDParticleDataWrapper.h"
 #include "Engine/Engine.h"
 #include "Engine/LevelStreamingDynamic.h"
 #include "Engine/World.h"
@@ -77,7 +82,7 @@ void FChaosVDScene::UpdateFromRecordedStepData(const int32 SolverID, const FStri
 	FChaosVDParticlesByIDMap& SolverParticlesByID = ParticlesBySolverID.FindChecked(SolverID);
 
 	// Go over existing Particle VD Instances and update them or create them if needed 
-	for (const FChaosVDParticleDebugData& Particle : InRecordedStepData.RecordedParticles)
+	for (const FChaosVDParticleDataWrapper& Particle : InRecordedStepData.RecordedParticlesData)
 	{
 		const int32 ParticleVDInstanceID = GetIDForRecordedParticleData(Particle);
 
@@ -91,7 +96,7 @@ void FChaosVDScene::UpdateFromRecordedStepData(const int32 SolverID, const FStri
 		{
 			if (AChaosVDParticleActor* ExistingParticleVDInstancePtr = *ExistingParticleVDInstancePtrPtr)
 			{
-				ExistingParticleVDInstancePtr->UpdateFromRecordedData(Particle, InFrameData.SimulationTransform);
+				ExistingParticleVDInstancePtr->UpdateFromRecordedParticleData(Particle, InFrameData.SimulationTransform);
 			}
 			else
 			{
@@ -104,7 +109,7 @@ void FChaosVDScene::UpdateFromRecordedStepData(const int32 SolverID, const FStri
 			if (AChaosVDParticleActor* NewParticleVDInstance = SpawnParticleFromRecordedData(Particle, InFrameData))
 			{
 				FStringFormatOrderedArguments Args {SolverName, FString::FromInt(SolverID)};
-				const FName FolderPath = *FPaths::Combine(FString::Format(TEXT("Solver {0} | ID {1}"), Args), UEnum::GetDisplayValueAsText(Particle.ParticleType).ToString());
+				const FName FolderPath = *FPaths::Combine(FString::Format(TEXT("Solver {0} | ID {1}"), Args), UEnum::GetDisplayValueAsText(NewParticleVDInstance->GetParticleData().Type).ToString());
 
 				NewParticleVDInstance->SetFolderPath(FolderPath);
 
@@ -216,7 +221,7 @@ const TSharedPtr<const Chaos::FImplicitObject>* FChaosVDScene::GetUpdatedGeometr
 	return nullptr;
 }
 
-AChaosVDParticleActor* FChaosVDScene::SpawnParticleFromRecordedData(const FChaosVDParticleDebugData& InParticleData, const FChaosVDSolverFrameData& InFrameData)
+AChaosVDParticleActor* FChaosVDScene::SpawnParticleFromRecordedData(const FChaosVDParticleDataWrapper& InParticleData, const FChaosVDSolverFrameData& InFrameData)
 {
 	using namespace Chaos;
 
@@ -226,7 +231,7 @@ AChaosVDParticleActor* FChaosVDScene::SpawnParticleFromRecordedData(const FChaos
 
 	if (AChaosVDParticleActor* NewActor = PhysicsVDWorld->SpawnActor<AChaosVDParticleActor>(Params))
 	{
-		NewActor->UpdateFromRecordedData(InParticleData, InFrameData.SimulationTransform);
+		NewActor->UpdateFromRecordedParticleData(InParticleData, InFrameData.SimulationTransform);
 
 		if (!InParticleData.DebugName.IsEmpty())
 		{
@@ -237,9 +242,9 @@ AChaosVDParticleActor* FChaosVDScene::SpawnParticleFromRecordedData(const FChaos
 		
 		if (ensure(LoadedRecording.IsValid()))
 		{
-			if (const TSharedPtr<const Chaos::FImplicitObject>* Geometry = LoadedRecording->GetGeometryDataMap().Find(InParticleData.ImplicitObjectHash))
+			if (const TSharedPtr<const Chaos::FImplicitObject>* Geometry = LoadedRecording->GetGeometryDataMap().Find(InParticleData.GeometryHash))
 			{
-				NewActor->UpdateGeometryData(*Geometry);
+				NewActor->UpdateGeometry(*Geometry);
 			}
 		}
 
@@ -249,7 +254,7 @@ AChaosVDParticleActor* FChaosVDScene::SpawnParticleFromRecordedData(const FChaos
 	return nullptr;
 }
 
-int32 FChaosVDScene::GetIDForRecordedParticleData(const FChaosVDParticleDebugData& InParticleData) const
+int32 FChaosVDScene::GetIDForRecordedParticleData(const FChaosVDParticleDataWrapper& InParticleData) const
 {
 	return InParticleData.ParticleIndex;
 }
@@ -344,6 +349,7 @@ void FChaosVDScene::HandlePostSelectionChange(const UTypedElementSelectionSet* P
 
 void FChaosVDScene::SetSelectedObject(UObject* SelectedObject)
 {
+
 	if (IsObjectSelected(SelectedObject))
 	{
 		// Already selected, nothing to do here
diff --git a/Engine/Plugins/Experimental/ChaosVD/Source/ChaosVD/Private/DetailsCustomizations/ChaosVDParticleDataWrapperCustomization.cpp b/Engine/Plugins/Experimental/ChaosVD/Source/ChaosVD/Private/DetailsCustomizations/ChaosVDParticleDataWrapperCustomization.cpp
new file mode 100644
--- /dev/null
+++ b/Engine/Plugins/Experimental/ChaosVD/Source/ChaosVD/Private/DetailsCustomizations/ChaosVDParticleDataWrapperCustomization.cpp
@@ -0,0 +1,59 @@
+﻿// Copyright Epic Games, Inc. All Rights Reserved.
+
+#include "DetailsCustomizations/ChaosVDParticleDataWrapperCustomization.h"
+
+#include "DataWrappers/ChaosVDParticleDataWrapper.h"
+#include "IDetailChildrenBuilder.h"
+#include "PropertyHandle.h"
+
+TSharedRef<IPropertyTypeCustomization> FChaosVDParticleDataWrapperCustomization::MakeInstance()
+{
+	return MakeShareable(new FChaosVDParticleDataWrapperCustomization());
+}
+
+void FChaosVDParticleDataWrapperCustomization::CustomizeChildren(TSharedRef<IPropertyHandle> StructPropertyHandle, IDetailChildrenBuilder& StructBuilder, IPropertyTypeCustomizationUtils& StructCustomizationUtils)
+{
+	uint32 NumChildren;
+	StructPropertyHandle->GetNumChildren(NumChildren);
+
+	if (NumChildren == 0)
+	{
+		return;
+	}
+
+	TSet<FName> ParticleDataViewersNames;
+	ParticleDataViewersNames.Reserve(5);
+
+	ParticleDataViewersNames.Add(GET_MEMBER_NAME_CHECKED(FChaosVDParticleDataWrapper, ParticlePositionRotation));
+	ParticleDataViewersNames.Add(GET_MEMBER_NAME_CHECKED(FChaosVDParticleDataWrapper, ParticleVelocities));
+	ParticleDataViewersNames.Add(GET_MEMBER_NAME_CHECKED(FChaosVDParticleDataWrapper, ParticleDynamics));
+	ParticleDataViewersNames.Add(GET_MEMBER_NAME_CHECKED(FChaosVDParticleDataWrapper, ParticleDynamicsMisc));
+	ParticleDataViewersNames.Add(GET_MEMBER_NAME_CHECKED(FChaosVDParticleDataWrapper, ParticleMassProps));
+
+	for (uint32 ChildIndex = 0; ChildIndex < NumChildren; ++ChildIndex)
+	{
+		TSharedRef<IPropertyHandle> ChildHandle = StructPropertyHandle->GetChildHandle(ChildIndex).ToSharedRef();
+
+		const FName PropertyName = ChildHandle->GetProperty()->GetFName();
+		if (ParticleDataViewersNames.Contains(PropertyName))
+		{
+			void* Data = nullptr;
+			ChildHandle->GetValueData(Data);
+			if (Data)
+			{
+				FChaosVDParticleDataBase* DataViewer = static_cast<FChaosVDParticleDataBase*>(Data);
+
+				// The Particle Data viewer struct has several fields that will have default values if there was no recorded data for them in the trace file
+				// As these do not represent anything real value, we only add to the details panel the ones with recorded data
+				if (DataViewer->HasValidData())
+				{
+					StructBuilder.AddProperty(ChildHandle);
+				}
+			}
+		}
+		else
+		{
+			StructBuilder.AddProperty(ChildHandle);
+		}
+	}
+}
diff --git a/Engine/Plugins/Experimental/ChaosVD/Source/ChaosVD/Private/Trace/ChaosVDTraceAnalyzer.cpp b/Engine/Plugins/Experimental/ChaosVD/Source/ChaosVD/Private/Trace/ChaosVDTraceAnalyzer.cpp
--- a/Engine/Plugins/Experimental/ChaosVD/Source/ChaosVD/Private/Trace/ChaosVDTraceAnalyzer.cpp
+++ b/Engine/Plugins/Experimental/ChaosVD/Source/ChaosVD/Private/Trace/ChaosVDTraceAnalyzer.cpp
@@ -10,7 +10,6 @@
 void FChaosVDTraceAnalyzer::OnAnalysisBegin(const FOnAnalysisContext& Context)
 {
 	FInterfaceBuilder& Builder = Context.InterfaceBuilder;
-	Builder.RouteEvent(RouteId_ChaosVDParticle, "ChaosVDLogger", "ChaosVDParticle");
 	Builder.RouteEvent(RouteId_ChaosVDParticleDestroyed, "ChaosVDLogger", "ChaosVDParticleDestroyed");
 	
 	Builder.RouteEvent(RouteId_ChaosVDSolverFrameStart, "ChaosVDLogger", "ChaosVDSolverFrameStart");
@@ -110,24 +109,6 @@ bool FChaosVDTraceAnalyzer::OnEvent(uint16 RouteId, EStyle Style, const FOnEvent
 		}
 	case RouteId_ChaosVDSolverStepEnd:
 		{
-			break;
-		}
-	case RouteId_ChaosVDParticle:
-		{
-			const int32 SolverID = EventData.GetValue<int32>("SolverID");
-
-			FChaosVDParticleDebugData ParticleData;
-			ReadParticleDataFromEvent(EventData, ParticleData);
-
-			// This can be null if the recording started Mid-Frame. In this case we just discard the data for now
-			if (FChaosVDSolverFrameData* FrameData = ChaosVDTraceProvider->GetLastSolverFrame(SolverID))
-			{
-				if (ensureMsgf(FrameData->SolverSteps.Num() > 0, TEXT("A particle was traced without a valid sstep scope")))
-				{
-					FrameData->SolverSteps.Last().RecordedParticles.Add(MoveTemp(ParticleData));
-				}
-			}
-
 			break;
 		}
 	case RouteId_ChaosVDParticleDestroyed:
@@ -207,34 +188,3 @@ bool FChaosVDTraceAnalyzer::OnEvent(uint16 RouteId, EStyle Style, const FOnEvent
 
 	return true;
 }
-
-void FChaosVDTraceAnalyzer::ReadParticleDataFromEvent(const FEventData& InEventData, FChaosVDParticleDebugData& OutParticleData)
-{
-	OutParticleData.ParticleType = static_cast<EChaosVDParticleType>(InEventData.GetValue<uint8>("ParticleType"));
-	OutParticleData.ParticleState = static_cast<EChaosVDParticleState>(InEventData.GetValue<int8>("ObjectState"));
-	OutParticleData.ParticleIndex = InEventData.GetValue<int32>("ParticleID");
-	OutParticleData.ImplicitObjectID = InEventData.GetValue<int32>("ImplicitObjectID");
-	OutParticleData.ImplicitObjectHash = InEventData.GetValue<uint32>("ImplicitObjectHash");
-
-	FWideStringView DebugNameView;
-	InEventData.GetString("DebugName", DebugNameView);
-
-	OutParticleData.DebugName = DebugNameView;
-
-	OutParticleData.Position.X = InEventData.GetValue<float>("PositionX");
-	OutParticleData.Position.Y = InEventData.GetValue<float>("PositionY");
-	OutParticleData.Position.Z = InEventData.GetValue<float>("PositionZ");
-
-	OutParticleData.Rotation.X = InEventData.GetValue<float>("RotationX");
-	OutParticleData.Rotation.Y = InEventData.GetValue<float>("RotationY");
-	OutParticleData.Rotation.Z = InEventData.GetValue<float>("RotationZ");
-	OutParticleData.Rotation.W = InEventData.GetValue<float>("RotationW");
-			
-	OutParticleData.Velocity.X = InEventData.GetValue<float>("VelocityX");
-	OutParticleData.Velocity.Y = InEventData.GetValue<float>("VelocityY");
-	OutParticleData.Velocity.Z = InEventData.GetValue<float>("VelocityZ");
-			
-	OutParticleData.AngularVelocity.X = InEventData.GetValue<float>("AngularVelocityX");
-	OutParticleData.AngularVelocity.Y = InEventData.GetValue<float>("AngularVelocityY");
-	OutParticleData.AngularVelocity.Z = InEventData.GetValue<float>("AngularVelocityZ");
-}
diff --git a/Engine/Plugins/Experimental/ChaosVD/Source/ChaosVD/Private/Trace/ChaosVDTraceProvider.cpp b/Engine/Plugins/Experimental/ChaosVD/Source/ChaosVD/Private/Trace/ChaosVDTraceProvider.cpp
--- a/Engine/Plugins/Experimental/ChaosVD/Source/ChaosVD/Private/Trace/ChaosVDTraceProvider.cpp
+++ b/Engine/Plugins/Experimental/ChaosVD/Source/ChaosVD/Private/Trace/ChaosVDTraceProvider.cpp
@@ -2,12 +2,21 @@
 
 #include "Trace/ChaosVDTraceProvider.h"
 
-#include "Chaos/ImplicitObject.h"
 #include "ChaosVDRecording.h"
-#include "ChaosVisualDebugger/ChaosVisualDebuggerTrace.h"
+
+#include "Chaos/ChaosArchive.h"
+#include "Chaos/ImplicitObject.h"
+
 #include "Compression/OodleDataCompressionUtil.h"
 #include "Serialization/MemoryReader.h"
 
+#if WITH_CHAOS_VISUAL_DEBUGGER
+#include "DataWrappers/ChaosVDImplicitObjectDataWrapper.h"
+#include "DataWrappers/ChaosVDParticleDataWrapper.h"
+
+using FChaosVDImplicitObjectWrapper = FChaosVDImplicitObjectDataWrapper<Chaos::TSerializablePtr<Chaos::FImplicitObject>, Chaos::FChaosArchive>;
+#endif
+
 FName FChaosVDTraceProvider::ProviderName("ChaosVDProvider");
 
 FChaosVDTraceProvider::FChaosVDTraceProvider(TraceServices::IAnalysisSession& InSession): Session(InSession)
@@ -133,17 +142,33 @@ void FChaosVDTraceProvider::SetBinaryDataReadyToUse(const int32 DataID)
 			// TODO: Create a system to register external "Data handlers" with the logic on how serialize each type
 			// This should not be done here but as this is the only type we have is ok for now
 			if (UnprocessedData->TypeName == TEXT("FChaosVDImplicitObjectDataWrapper"))
-			{	
-				Chaos::TSerializablePtr<Chaos::FImplicitObject> Geometry;
-
+			{
 				FMemoryReader MemeReader(*RawData);
 				Chaos::FChaosArchive Ar(MemeReader);
 
-				FChaosVDImplicitObjectDataWrapper WrappedGeometryData;
+				FChaosVDImplicitObjectWrapper WrappedGeometryData;
 				WrappedGeometryData.Serialize(Ar);
 
 				InternalRecording->AddImplicitObject(WrappedGeometryData.Hash, WrappedGeometryData.ImplicitObject.Get());
 			}
+
+			if (UnprocessedData->TypeName == TEXT("FChaosVDParticleDataWrapper"))
+			{	
+				FMemoryReader MemeReader(*RawData);
+				MemeReader.SetUseUnversionedPropertySerialization(true);
+
+				FChaosVDParticleDataWrapper ParticleData;
+				MemeReader << ParticleData;
+
+				// This can be null if the recording started Mid-Frame. In this case we just discard the data for now
+				if (FChaosVDSolverFrameData* FrameData = GetLastSolverFrame(ParticleData.SolverID))
+				{
+					if (ensureMsgf(FrameData->SolverSteps.Num() > 0, TEXT("A particle was traced without a valid step scope")))
+					{
+						FrameData->SolverSteps.Last().RecordedParticlesData.Add(MoveTemp(ParticleData));
+					}
+				}
+			}
 #endif
 		}
 	}
diff --git a/Engine/Plugins/Experimental/ChaosVD/Source/ChaosVD/Public/ChaosVDParticleActor.h b/Engine/Plugins/Experimental/ChaosVD/Source/ChaosVD/Public/ChaosVDParticleActor.h
--- a/Engine/Plugins/Experimental/ChaosVD/Source/ChaosVD/Public/ChaosVDParticleActor.h
+++ b/Engine/Plugins/Experimental/ChaosVD/Source/ChaosVD/Public/ChaosVDParticleActor.h
@@ -1,8 +1,8 @@
 // Copyright Epic Games, Inc. All Rights Reserved.
 
 #pragma once
-#include "ChaosVDRecording.h"
 #include "Chaos/Core.h"
+#include "DataWrappers/ChaosVDParticleDataWrapper.h"
 #include "GameFramework/Actor.h"
 
 #include "ChaosVDParticleActor.generated.h"
@@ -14,6 +14,11 @@ class USceneComponent;
 class UStaticMeshComponent;
 class UStaticMesh;
 
+namespace Chaos
+{
+	class FImplicitObject;
+}
+
 /** Options flags to control how geometry is updated in a ChaosVDActor */
 enum class EChaosVDActorGeometryUpdateFlags : int32
 {
@@ -30,24 +35,26 @@ class AChaosVDParticleActor : public AActor
 public:
 	AChaosVDParticleActor(const FObjectInitializer& ObjectInitializer);
 
-	void UpdateFromRecordedData(const FChaosVDParticleDebugData& InRecordedData, const Chaos::FRigidTransform3& SimulationTransform);
-	
-	const FChaosVDParticleDebugData& GetDebugData() const { return RecordedDebugData; }
+	void UpdateFromRecordedParticleData(const FChaosVDParticleDataWrapper& InRecordedData, const Chaos::FRigidTransform3& SimulationTransform);
 
-	void UpdateGeometryData(const TSharedPtr<const Chaos::FImplicitObject>& ImplicitObject, EChaosVDActorGeometryUpdateFlags OptionFlags = EChaosVDActorGeometryUpdateFlags::None);
+	void UpdateGeometry(const TSharedPtr<const Chaos::FImplicitObject>& ImplicitObject, EChaosVDActorGeometryUpdateFlags OptionsFlags = EChaosVDActorGeometryUpdateFlags::None);
+	void UpdateGeometry(uint32 NewGeometryHash, EChaosVDActorGeometryUpdateFlags OptionsFlags = EChaosVDActorGeometryUpdateFlags::None);
 
 	void SetScene(const TSharedPtr<FChaosVDScene>& InScene);
 
 	virtual void BeginDestroy() override;
+
+	const FChaosVDParticleDataWrapper& GetParticleData() { return ParticleDataViewer; }
 	
 #if WITH_EDITOR
 	virtual bool IsSelectedInEditor() const override;
+	virtual void PostEditChangeProperty(FPropertyChangedEvent& PropertyChangedEvent) override;
 #endif
 
 protected:
 
-	UPROPERTY(VisibleAnywhere, Category="Chaos Visual Debugger Data")
-	FChaosVDParticleDebugData RecordedDebugData;
+	UPROPERTY(EditAnywhere, Category="Particle Data")
+	FChaosVDParticleDataWrapper ParticleDataViewer;
 
 	bool bIsGeometryDataGenerationStarted = false;
 
diff --git a/Engine/Plugins/Experimental/ChaosVD/Source/ChaosVD/Public/ChaosVDParticleActorCustomization.h b/Engine/Plugins/Experimental/ChaosVD/Source/ChaosVD/Public/ChaosVDParticleActorCustomization.h
--- a/Engine/Plugins/Experimental/ChaosVD/Source/ChaosVD/Public/ChaosVDParticleActorCustomization.h
+++ b/Engine/Plugins/Experimental/ChaosVD/Source/ChaosVD/Public/ChaosVDParticleActorCustomization.h
@@ -10,7 +10,7 @@ class FChaosVDParticleActorCustomization : public IDetailCustomization
 {
 public:
 	
-	inline static FName ChaosVDCategoryName = FName("Chaos Visual Debugger Data");
+	inline static FName ChaosVDCategoryName = FName("Particle Data");
 
 	static TSharedRef<IDetailCustomization> MakeInstance();
 
diff --git a/Engine/Plugins/Experimental/ChaosVD/Source/ChaosVD/Public/ChaosVDScene.h b/Engine/Plugins/Experimental/ChaosVD/Source/ChaosVD/Public/ChaosVDScene.h
--- a/Engine/Plugins/Experimental/ChaosVD/Source/ChaosVD/Public/ChaosVDScene.h
+++ b/Engine/Plugins/Experimental/ChaosVD/Source/ChaosVD/Public/ChaosVDScene.h
@@ -79,10 +79,10 @@ public:
 private:
 
 	/** Creates an ChaosVDParticle actor for the Provided recorded Particle Data */
-	AChaosVDParticleActor* SpawnParticleFromRecordedData(const FChaosVDParticleDebugData& InParticleData, const FChaosVDSolverFrameData& InFrameData);
+	AChaosVDParticleActor* SpawnParticleFromRecordedData(const FChaosVDParticleDataWrapper& InParticleData, const FChaosVDSolverFrameData& InFrameData);
 
 	/** Returns the ID used to track this recorded particle data */
-	int32 GetIDForRecordedParticleData(const FChaosVDParticleDebugData& InParticleData) const;
+	int32 GetIDForRecordedParticleData(const FChaosVDParticleDataWrapper& InParticleData) const;
 
 	/** Creates the instance of the World which will be used the recorded data*/
 	UWorld* CreatePhysicsVDWorld() const;
diff --git a/Engine/Plugins/Experimental/ChaosVD/Source/ChaosVD/Public/DetailsCustomizations/ChaosVDParticleDataWrapperCustomization.h b/Engine/Plugins/Experimental/ChaosVD/Source/ChaosVD/Public/DetailsCustomizations/ChaosVDParticleDataWrapperCustomization.h
new file mode 100644
--- /dev/null
+++ b/Engine/Plugins/Experimental/ChaosVD/Source/ChaosVD/Public/DetailsCustomizations/ChaosVDParticleDataWrapperCustomization.h
@@ -0,0 +1,21 @@
+﻿// Copyright Epic Games, Inc. All Rights Reserved.
+
+#pragma once
+
+#include "IPropertyTypeCustomization.h"
+#include "PropertyEditorModule.h"
+#include "Templates/SharedPointer.h"
+
+class FDetailWidgetRow;
+class IDetailChildrenBuilder;
+
+/** Property Customization used to only show the available particle data withing the ChaosVDParticleDataViewer */
+class FChaosVDParticleDataWrapperCustomization : public IPropertyTypeCustomization
+{
+public:
+
+	static TSharedRef<IPropertyTypeCustomization> MakeInstance();
+
+	virtual void CustomizeHeader(TSharedRef<IPropertyHandle> StructPropertyHandle, FDetailWidgetRow& HeaderRow, IPropertyTypeCustomizationUtils& StructCustomizationUtils) override {}
+	virtual void CustomizeChildren(TSharedRef<IPropertyHandle> StructPropertyHandle, IDetailChildrenBuilder& StructBuilder, IPropertyTypeCustomizationUtils& StructCustomizationUtils) override;
+};
diff --git a/Engine/Plugins/Experimental/ChaosVD/Source/ChaosVD/Public/Trace/ChaosVDTraceAnalyzer.h b/Engine/Plugins/Experimental/ChaosVD/Source/ChaosVD/Public/Trace/ChaosVDTraceAnalyzer.h
--- a/Engine/Plugins/Experimental/ChaosVD/Source/ChaosVD/Public/Trace/ChaosVDTraceAnalyzer.h
+++ b/Engine/Plugins/Experimental/ChaosVD/Source/ChaosVD/Public/Trace/ChaosVDTraceAnalyzer.h
@@ -36,7 +36,6 @@ private:
 
 	enum : uint16
 	{
-		RouteId_ChaosVDParticle,
 		RouteId_ChaosVDSolverStepStart,
 		RouteId_ChaosVDSolverStepEnd,
 		RouteId_ChaosVDSolverFrameStart,
@@ -50,8 +49,6 @@ private:
 		RouteId_ChaosVDParticleDestroyed,
 	};
 
-	void ReadParticleDataFromEvent(const FEventData& InEventData, FChaosVDParticleDebugData& OutParticleData);
-
 	TraceServices::IAnalysisSession& Session;
 
 	TSharedPtr<FChaosVDTraceProvider> ChaosVDTraceProvider;
diff --git a/Engine/Source/Runtime/Experimental/Chaos/Private/ChaosVisualDebugger/ChaosVDDataWrapperUtils.cpp b/Engine/Source/Runtime/Experimental/Chaos/Private/ChaosVisualDebugger/ChaosVDDataWrapperUtils.cpp
new file mode 100644
--- /dev/null
+++ b/Engine/Source/Runtime/Experimental/Chaos/Private/ChaosVisualDebugger/ChaosVDDataWrapperUtils.cpp
@@ -0,0 +1,37 @@
+﻿// Copyright Epic Games, Inc. All Rights Reserved.
+
+#include "ChaosVisualDebugger/ChaosVDDataWrapperUtils.h"
+
+#if WITH_CHAOS_VISUAL_DEBUGGER
+#include "Chaos/ParticleHandle.h"
+#include "DataWrappers/ChaosVDParticleDataWrapper.h"
+
+FChaosVDParticleDataWrapper FChaosVDDataWrapperUtils::BuildParticleDataWrapperFromParticle(const Chaos::FGeometryParticleHandle* ParticleHandlePtr)
+{
+	check(ParticleHandlePtr);
+
+	FChaosVDParticleDataWrapper WrappedParticleData;
+
+	WrappedParticleData.ParticleIndex = ParticleHandlePtr->UniqueIdx().Idx;
+	WrappedParticleData.Type =  static_cast<EChaosVDParticleType>(ParticleHandlePtr->Type);
+
+#if CHAOS_DEBUG_NAME
+	WrappedParticleData.DebugNamePtr = ParticleHandlePtr->DebugName();
+#endif
+
+	WrappedParticleData.ParticlePositionRotation.CopyFrom(*ParticleHandlePtr);
+
+	if (const Chaos::TKinematicGeometryParticleHandleImp<Chaos::FReal, 3, true>* KinematicParticle = ParticleHandlePtr->CastToKinematicParticle())
+	{
+		WrappedParticleData.ParticleVelocities.CopyFrom(*KinematicParticle);
+	}
+
+	if (const Chaos::TPBDRigidParticleHandleImp<Chaos::FReal, 3, true>* RigidParticle = ParticleHandlePtr->CastToRigidParticle())
+	{
+		WrappedParticleData.ParticleDynamics.CopyFrom(*RigidParticle);
+		WrappedParticleData.ParticleDynamicsMisc.CopyFrom(*RigidParticle);
+		WrappedParticleData.ParticleMassProps.CopyFrom(*RigidParticle);
+	}
+	return MoveTemp(WrappedParticleData);
+}
+#endif //WITH_CHAOS_VISUAL_DEBUGGER
diff --git a/Engine/Source/Runtime/Experimental/Chaos/Private/ChaosVisualDebugger/ChaosVisualDebuggerTrace.cpp b/Engine/Source/Runtime/Experimental/Chaos/Private/ChaosVisualDebugger/ChaosVisualDebuggerTrace.cpp
--- a/Engine/Source/Runtime/Experimental/Chaos/Private/ChaosVisualDebugger/ChaosVisualDebuggerTrace.cpp
+++ b/Engine/Source/Runtime/Experimental/Chaos/Private/ChaosVisualDebugger/ChaosVisualDebuggerTrace.cpp
@@ -1,6 +1,7 @@
 // Copyright Epic Games, Inc. All Rights Reserved.
 #include "ChaosVisualDebugger/ChaosVisualDebuggerTrace.h"
 
+#include "ChaosVisualDebugger/ChaosVDDataWrapperUtils.h"
 #include "Chaos/PBDRigidsSOAs.h"
 #include "Tasks/Task.h"
 
@@ -9,6 +10,8 @@
 #include "Chaos/Framework/PhysicsSolverBase.h"
 #include "Chaos/ImplicitObject.h"
 #include "Compression/OodleDataCompressionUtil.h"
+#include "DataWrappers/ChaosVDImplicitObjectDataWrapper.h"
+#include "DataWrappers/ChaosVDParticleDataWrapper.h"
 #include "HAL/CriticalSection.h"
 #include "Misc/ScopeRWLock.h"
 #include "Serialization/MemoryWriter.h"
@@ -97,12 +100,6 @@ TSet<int32> FChaosVisualDebuggerTrace::SolverIDsForDeltaRecording = TSet<int32>(
 TSet<int32> FChaosVisualDebuggerTrace::RequestedFullCaptureSolverIDs = TSet<int32>();
 FThreadSafeBool FChaosVisualDebuggerTrace::bIsTracing = false;
 
-void FChaosVDImplicitObjectDataWrapper::Serialize(Chaos::FChaosArchive& Ar)
-{
-	Ar << Hash;
-	Ar << ImplicitObject;
-}
-
 void FChaosVisualDebuggerTrace::TraceParticle(const Chaos::FGeometryParticleHandle* ParticleHandle)
 {
 	if (!IsTracing())
@@ -132,35 +129,25 @@ void FChaosVisualDebuggerTrace::TraceParticle(Chaos::FGeometryParticleHandle* Pa
 		return;
 	}
 
-	static FString DefaultName = TEXT("NONAME");
-	FStringView ParticleNameView(DefaultName);
-
-#if CHAOS_DEBUG_NAME
-	if (const TSharedPtr<FString, ESPMode::ThreadSafe>& DebugNamePtr = ParticleHandle->DebugName())
-	{
-		ParticleNameView = FStringView(*DebugNamePtr.Get());
-	}
-#endif
-
 	const uint32 GeometryHash = GeometryTracerObject.GetGeometryHashForImplicit(ParticleHandle->Geometry().Get());
 	
 	TraceImplicitObject({ GeometryHash, ParticleHandle->Geometry() });
 
-	UE_TRACE_LOG(ChaosVDLogger, ChaosVDParticle, ChaosVDChannel)
-		<< ChaosVDParticle.SolverID(ContextData.Id)
-		<< ChaosVDParticle.Cycle(FPlatformTime::Cycles64())
+	{
+		FChaosVDParticleDataWrapper ParticleDataWrapper = FChaosVDDataWrapperUtils::BuildParticleDataWrapperFromParticle(ParticleHandle);
+		ParticleDataWrapper.GeometryHash = GeometryHash;
+		ParticleDataWrapper.SolverID = ContextData.Id;
+		
+		FChaosVDScopedTLSBufferAccessor TLSDataBuffer;
 
-		<< ChaosVDParticle.ParticleID(ParticleHandle->UniqueIdx().Idx)
-		<< ChaosVDParticle.ParticleType(static_cast<uint8>(ParticleHandle->Type))
-		<< ChaosVDParticle.DebugName(ParticleNameView.GetData(), ParticleNameView.Len())
+		FMemoryWriter MemWriterAr(TLSDataBuffer.BufferRef);
+		MemWriterAr.SetShouldSkipUpdateCustomVersion(true);
+		MemWriterAr.SetUseUnversionedPropertySerialization(true);
 
-		<< CVD_TRACE_VECTOR_ON_EVENT(ChaosVDParticle, Position, ParticleHandle->X())
-		<< CVD_TRACE_ROTATOR_ON_EVENT(ChaosVDParticle, Rotation, ParticleHandle->R())
-		<< CVD_TRACE_VECTOR_ON_EVENT(ChaosVDParticle, Velocity, Chaos::FConstGenericParticleHandle(ParticleHandle)->V())
-		<< CVD_TRACE_VECTOR_ON_EVENT(ChaosVDParticle, AngularVelocity, Chaos::FConstGenericParticleHandle(ParticleHandle)->W())
+		ParticleDataWrapper.Serialize(MemWriterAr);
 
-		<< ChaosVDParticle.ObjectState(static_cast<int8>(ParticleHandle->ObjectState()))
-		<< ChaosVDParticle.ImplicitObjectHash(GeometryHash);
+		TraceBinaryData(TLSDataBuffer.BufferRef, TEXT("FChaosVDParticleDataWrapper"));
+	}
 }
 
 void FChaosVisualDebuggerTrace::TraceParticles(const Chaos::TGeometryParticleHandles<Chaos::FReal, 3>& ParticleHandles)
@@ -426,7 +413,7 @@ void FChaosVisualDebuggerTrace::TraceBinaryData(const TArray<uint8>& InData, con
 	ensure(RemainingSize == 0);
 }
 
-void FChaosVisualDebuggerTrace::TraceImplicitObject(FChaosVDImplicitObjectDataWrapper WrappedGeometryData)
+void FChaosVisualDebuggerTrace::TraceImplicitObject(FChaosVDImplicitObjectWrapper WrappedGeometryData)
 {
 	if (!IsTracing())
 	{
@@ -545,4 +532,4 @@ void FChaosVisualDebuggerTrace::PerformFullCapture(EChaosVDFullCaptureFlags Capt
 	}
 }
 
-#endif
\ No newline at end of file
+#endif
diff --git a/Engine/Source/Runtime/Experimental/Chaos/Public/ChaosVisualDebugger/ChaosVDDataWrapperUtils.h b/Engine/Source/Runtime/Experimental/Chaos/Public/ChaosVisualDebugger/ChaosVDDataWrapperUtils.h
new file mode 100644
--- /dev/null
+++ b/Engine/Source/Runtime/Experimental/Chaos/Public/ChaosVisualDebugger/ChaosVDDataWrapperUtils.h
@@ -0,0 +1,26 @@
+﻿// Copyright Epic Games, Inc. All Rights Reserved.
+
+#pragma once
+
+#if WITH_CHAOS_VISUAL_DEBUGGER
+#include "Chaos/ParticleHandleFwd.h"
+
+class FChaosVisualDebuggerTrace;
+struct FChaosVDParticleDataWrapper;
+
+/**
+ * Helper class used to build Chaos Visual Debugger data wrappers, without directly referencing chaos' types in them.
+ *
+ * @note: This is needed for now because we want to keep the data wrapper structs/classes on a separate module where possible, but if we reference Chaos's types
+ * directly we will end up with a circular dependency issue because the ChaosVDRuntime module will need the Chaos module but the Chaos module will need the ChaosVDRuntime module to use the structs
+ * Once development is done and can we commit to backward compatibility, this helper class might go away (trough the proper deprecation process)
+ */
+class FChaosVDDataWrapperUtils
+{
+private:
+	/** Creates and populates a FChaosVDParticleDataWrapper with the data of the provided ParticleHandlePtr */
+	static FChaosVDParticleDataWrapper BuildParticleDataWrapperFromParticle(const Chaos::FGeometryParticleHandle* ParticleHandlePtr);
+
+	friend FChaosVisualDebuggerTrace;
+};
+#endif //WITH_CHAOS_VISUAL_DEBUGGER
diff --git a/Engine/Source/Runtime/Experimental/Chaos/Public/ChaosVisualDebugger/ChaosVisualDebuggerTrace.h b/Engine/Source/Runtime/Experimental/Chaos/Public/ChaosVisualDebugger/ChaosVisualDebuggerTrace.h
--- a/Engine/Source/Runtime/Experimental/Chaos/Public/ChaosVisualDebugger/ChaosVisualDebuggerTrace.h
+++ b/Engine/Source/Runtime/Experimental/Chaos/Public/ChaosVisualDebugger/ChaosVisualDebuggerTrace.h
@@ -67,6 +67,7 @@
 #else
 
 #include "ChaosVDRuntimeModule.h"
+#include "DataWrappers/ChaosVDImplicitObjectDataWrapper.h"
 #include "Trace/Trace.h"
 #include "Trace/Trace.inl"
 
@@ -114,21 +115,6 @@ UE_TRACE_EVENT_BEGIN_EXTERN(ChaosVDLogger, ChaosVDSolverFrameEnd)
 	UE_TRACE_EVENT_FIELD(uint64, Cycle)
 UE_TRACE_EVENT_END()
 
-UE_TRACE_EVENT_BEGIN_EXTERN(ChaosVDLogger, ChaosVDParticle)
-	UE_TRACE_EVENT_FIELD(int32, SolverID)
-	UE_TRACE_EVENT_FIELD(uint64, Cycle)
-	UE_TRACE_EVENT_FIELD(int32, ParticleID)
-	UE_TRACE_EVENT_FIELD(uint8, ParticleType)
-	CVD_DEFINE_TRACE_VECTOR(Chaos::FReal, Position)
-	CVD_DEFINE_TRACE_ROTATOR(Chaos::FReal, Rotation)
-	CVD_DEFINE_TRACE_VECTOR(Chaos::FReal, Velocity)
-	CVD_DEFINE_TRACE_VECTOR(Chaos::FReal, AngularVelocity)
-	UE_TRACE_EVENT_FIELD(int8, ObjectState)
-	UE_TRACE_EVENT_FIELD(UE::Trace::WideString, DebugName)
-	UE_TRACE_EVENT_FIELD(int32, ImplicitObjectID)
-	UE_TRACE_EVENT_FIELD(uint32, ImplicitObjectHash)
-UE_TRACE_EVENT_END()
-
 UE_TRACE_EVENT_BEGIN_EXTERN(ChaosVDLogger, ChaosVDParticleCreated)
 	UE_TRACE_EVENT_FIELD(int32, SolverID)
 	UE_TRACE_EVENT_FIELD(uint64, Cycle)
@@ -153,7 +139,7 @@ UE_TRACE_EVENT_BEGIN_EXTERN(ChaosVDLogger, ChaosVDSolverStepEnd)
 	UE_TRACE_EVENT_FIELD(uint32, StepNumber)
 UE_TRACE_EVENT_END()
 
-UE_TRACE_EVENT_BEGIN_EXTERN(ChaosVDLogger, ChaosVDBinaryDataStart, NoSync)
+UE_TRACE_EVENT_BEGIN_EXTERN(ChaosVDLogger, ChaosVDBinaryDataStart)
 	UE_TRACE_EVENT_FIELD(uint64, Cycle)
 	UE_TRACE_EVENT_FIELD(UE::Trace::WideString, TypeName)
 	UE_TRACE_EVENT_FIELD(int32, DataID)
@@ -162,13 +148,13 @@ UE_TRACE_EVENT_BEGIN_EXTERN(ChaosVDLogger, ChaosVDBinaryDataStart, NoSync)
 	UE_TRACE_EVENT_FIELD(bool, IsCompressed)
 UE_TRACE_EVENT_END()
 
-UE_TRACE_EVENT_BEGIN_EXTERN(ChaosVDLogger, ChaosVDBinaryDataContent, NoSync)
+UE_TRACE_EVENT_BEGIN_EXTERN(ChaosVDLogger, ChaosVDBinaryDataContent)
 	UE_TRACE_EVENT_FIELD(uint64, Cycle)
 	UE_TRACE_EVENT_FIELD(int32, DataID)
 	UE_TRACE_EVENT_FIELD(uint8[], RawData)
 UE_TRACE_EVENT_END()
 
-UE_TRACE_EVENT_BEGIN_EXTERN(ChaosVDLogger, ChaosVDBinaryDataEnd, NoSync)
+UE_TRACE_EVENT_BEGIN_EXTERN(ChaosVDLogger, ChaosVDBinaryDataEnd)
 	UE_TRACE_EVENT_FIELD(uint64, Cycle)
 	UE_TRACE_EVENT_FIELD(int32, DataID)
 UE_TRACE_EVENT_END()
@@ -260,14 +246,7 @@ namespace Chaos
 	class TGeometryParticleHandles;
 }
 
-/** Data Wrapper for used to Trace Implicit objects */
-struct CHAOS_API FChaosVDImplicitObjectDataWrapper
-{
-	uint32 Hash;
-	Chaos::TSerializablePtr<Chaos::FImplicitObject> ImplicitObject;
-
-	void Serialize(Chaos::FChaosArchive& Ar);
-};
+using FChaosVDImplicitObjectWrapper = FChaosVDImplicitObjectDataWrapper<Chaos::TSerializablePtr<Chaos::FImplicitObject>, Chaos::FChaosArchive>;
 
 /** Class containing  all the Tracing logic to record data for the Chaos Visual Debugger tool */
 class FChaosVisualDebuggerTrace
@@ -339,7 +318,7 @@ public:
 	 * The trace event is not tied to a particular Solver Frame/Step
 	 *  @param WrappedGeometryData Wrapper containing a ptr to the implicit and its ID
 	 */
-	static CHAOS_API void TraceImplicitObject(FChaosVDImplicitObjectDataWrapper WrappedGeometryData);
+	static CHAOS_API void TraceImplicitObject(FChaosVDImplicitObjectWrapper WrappedGeometryData);
 
 	/** Returns true if the provided solver ID needs a Full Capture */
 	static bool ShouldPerformFullCapture(int32 SolverID);
diff --git a/Engine/Source/Runtime/Experimental/ChaosVDData/ChaosVDData.build.cs b/Engine/Source/Runtime/Experimental/ChaosVDData/ChaosVDData.build.cs
--- a/Engine/Source/Runtime/Experimental/ChaosVDData/ChaosVDData.build.cs
+++ b/Engine/Source/Runtime/Experimental/ChaosVDData/ChaosVDData.build.cs
@@ -13,7 +13,8 @@ namespace UnrealBuildTool.Rules
 				new string[] {
 					"Core",
 					"CoreUObject",
-					"Chaos"
+					"Chaos",
+					"ChaosVDRuntime"
 				}
 			);
 		}
diff --git a/Engine/Source/Runtime/Experimental/ChaosVDData/Public/ChaosVDRecording.h b/Engine/Source/Runtime/Experimental/ChaosVDData/Public/ChaosVDRecording.h
--- a/Engine/Source/Runtime/Experimental/ChaosVDData/Public/ChaosVDRecording.h
+++ b/Engine/Source/Runtime/Experimental/ChaosVDData/Public/ChaosVDRecording.h
@@ -3,10 +3,9 @@
 #pragma once
 
 #include "Chaos/Core.h"
-#include "Chaos/ShapeInstanceFwd.h"
+#include "DataWrappers/ChaosVDParticleDataWrapper.h"
 #include "UObject/ObjectMacros.h"
 #include "Containers/UnrealString.h"
-#include "ChaosVDRecording.generated.h"
 
 namespace Chaos
 {
@@ -16,69 +15,10 @@ namespace Chaos
 DECLARE_MULTICAST_DELEGATE(FChaosVDRecordingUpdated)
 DECLARE_MULTICAST_DELEGATE_TwoParams(FChaosVDGeometryDataLoaded, const TSharedPtr<const Chaos::FImplicitObject>&, const uint32 GeometryID)
 
-UENUM()
-enum class EChaosVDParticleType : uint8
-{
-	Static,
-	Kinematic,
-	Rigid,
-	Clustered,
-	StaticMesh,
-	SkeletalMesh,
-	GeometryCollection,
-	Unknown
-};
-
-UENUM()
-enum class EChaosVDParticleState : int8
-{
-	Uninitialized = 0,
-	Sleeping = 1,
-	Kinematic = 2,
-	Static = 3,
-	Dynamic = 4,
-
-	Count
-};
-
-UENUM()
-enum class EChaosVDSolverType : uint32
-{
-	Rigid
-};
-
-USTRUCT()
-struct FChaosVDParticleDebugData
-{
-	GENERATED_BODY()
-
-	UPROPERTY(VisibleAnywhere, Category= "Chaos Visual Debugger Data")
-	EChaosVDParticleType ParticleType;
-	UPROPERTY(VisibleAnywhere, Category= "Chaos Visual Debugger Data")
-	FString DebugName;
-	UPROPERTY(VisibleAnywhere, Category= "Chaos Visual Debugger Data")
-	int32 ParticleIndex;
-	UPROPERTY(VisibleAnywhere, Category= "Chaos Visual Debugger Data")
-	FVector Position;
-	UPROPERTY(VisibleAnywhere, Category= "Chaos Visual Debugger Data")
-	FQuat Rotation;
-	UPROPERTY(VisibleAnywhere, Category= "Chaos Visual Debugger Data")
-	FVector Velocity;
-	UPROPERTY(VisibleAnywhere, Category= "Chaos Visual Debugger Data")
-	FVector AngularVelocity;
-	UPROPERTY(VisibleAnywhere, Category= "Chaos Visual Debugger Data")
-	EChaosVDParticleState ParticleState;
-
-	int32 ImplicitObjectID = INDEX_NONE;
-
-	UPROPERTY(VisibleAnywhere, Category= "Chaos Visual Debugger Data")
-	uint32 ImplicitObjectHash = 0;
-};
-
 struct FChaosVDStepData
 {
 	FString StepName;
-	TArray<FChaosVDParticleDebugData> RecordedParticles;
+	TArray<FChaosVDParticleDataWrapper> RecordedParticlesData;
 	TSet<int32> ParticlesDestroyedIDs;
 };
 
diff --git a/Engine/Source/Runtime/Experimental/ChaosVisualDebugger/Private/DataWrappers/ChaosVDParticleDataWrapper.cpp b/Engine/Source/Runtime/Experimental/ChaosVisualDebugger/Private/DataWrappers/ChaosVDParticleDataWrapper.cpp
new file mode 100644
--- /dev/null
+++ b/Engine/Source/Runtime/Experimental/ChaosVisualDebugger/Private/DataWrappers/ChaosVDParticleDataWrapper.cpp
@@ -0,0 +1,112 @@
+﻿// Copyright Epic Games, Inc. All Rights Reserved.
+
+#include "DataWrappers/ChaosVDParticleDataWrapper.h"
+
+// @note: Tracing an scene with 1000 particles moving, Manually serializing the structs is ~20% faster
+// than normal UStruct serialization in unversioned mode. As we do this at runtime when tracing in development builds, this is important.
+// One of the downside is that this will be more involved to maintain as any versioning needs to be done by hand.
+
+bool FChaosVDFRigidParticleControlFlags::Serialize(FArchive& Ar)
+{
+	Ar << bGravityEnabled;
+	Ar << bCCDEnabled;
+	Ar << bOneWayInteractionEnabled;
+	Ar << bMaxDepenetrationVelocityOverrideEnabled;
+	Ar << bInertiaConditioningEnabled;
+	Ar << GravityGroupIndex;
+
+	return true;
+}
+
+bool FChaosVDParticlePositionRotation::Serialize(FArchive& Ar)
+{
+	Ar << MX;
+	Ar << MR;
+	Ar << bHasValidData;
+
+	return true;
+}
+
+bool FChaosVDParticleVelocities::Serialize(FArchive& Ar)
+{
+	Ar << MV;
+	Ar << MW;
+	Ar << bHasValidData;
+
+	return true;
+}
+
+bool FChaosVDParticleDynamics::Serialize(FArchive& Ar)
+{
+	Ar << MAcceleration;
+	Ar << MAngularAcceleration;
+	Ar << MAngularImpulseVelocity;
+	Ar << MLinearImpulseVelocity;
+	Ar << bHasValidData;
+
+	return true;
+}
+
+bool FChaosVDParticleMassProps::Serialize(FArchive& Ar)
+{
+	Ar << MCenterOfMass;
+	Ar << MRotationOfMass;
+	Ar << MI;
+	Ar << MInvI;
+	Ar << MM;
+	Ar << MInvM;
+	Ar << bHasValidData;
+
+	return true;
+}
+
+bool FChaosVDParticleDynamicMisc::Serialize(FArchive& Ar)
+{
+	Ar << MAngularEtherDrag;
+	Ar << MMaxLinearSpeedSq;
+	Ar << MMaxAngularSpeedSq;
+	Ar << MCollisionGroup;
+	Ar << MObjectState;
+	Ar << MSleepType;
+	Ar << bDisabled;
+	Ar << bHasValidData;
+
+	MControlFlags.Serialize(Ar);
+
+	return true;
+}
+
+bool FChaosVDParticleDataWrapper::Serialize(FArchive& Ar)
+{
+	Ar << Type;
+	Ar << GeometryHash;
+
+	bHasDebugName = DebugNamePtr.IsValid();
+	Ar << bHasDebugName;
+
+	if (bHasDebugName)
+	{
+		if (Ar.IsLoading())
+		{
+			Ar << DebugName;
+		}
+		else
+		{
+			FString& DebugNameRef = *DebugNamePtr.Get();
+			Ar << DebugNameRef;
+		}
+	}
+
+	Ar << ParticleIndex;
+	Ar << SolverID;
+
+	Ar << ParticlePositionRotation;
+	Ar << ParticleVelocities;
+	Ar << ParticleDynamics;
+	Ar << ParticleDynamicsMisc;
+	Ar << ParticleMassProps;
+
+	Ar << bHasValidData;
+
+	return true;
+}
diff --git a/Engine/Source/Runtime/Experimental/ChaosVisualDebugger/Public/DataWrappers/ChaosVDImplicitObjectDataWrapper.h b/Engine/Source/Runtime/Experimental/ChaosVisualDebugger/Public/DataWrappers/ChaosVDImplicitObjectDataWrapper.h
new file mode 100644
--- /dev/null
+++ b/Engine/Source/Runtime/Experimental/ChaosVisualDebugger/Public/DataWrappers/ChaosVDImplicitObjectDataWrapper.h
@@ -0,0 +1,24 @@
+﻿// Copyright Epic Games, Inc. All Rights Reserved.
+
+#pragma once
+#include "Hal/Platform.h"
+
+/** Data Wrapper for used to Trace Implicit objects
+ * @note This should not be templated. But making it so for now so this can live in the ChaosVDRuntime module without referencing Chaos types directly.
+ * Once we are close to ship the first version of the tool, this will likely change to the specific types and to the Chaos Module
+ */
+template<class SerializableImplicitType, class ArchiveType>
+struct FChaosVDImplicitObjectDataWrapper
+{
+	uint32 Hash;
+	SerializableImplicitType ImplicitObject;
+
+	void Serialize(ArchiveType& Ar);
+};
+
+template <class SerializableImplicitType, class ArchiveType>
+void FChaosVDImplicitObjectDataWrapper<SerializableImplicitType, ArchiveType>::Serialize(ArchiveType& Ar)
+{
+	Ar << Hash;
+	Ar << ImplicitObject;
+}
diff --git a/Engine/Source/Runtime/Experimental/ChaosVisualDebugger/Public/DataWrappers/ChaosVDParticleDataWrapper.h b/Engine/Source/Runtime/Experimental/ChaosVisualDebugger/Public/DataWrappers/ChaosVDParticleDataWrapper.h
new file mode 100644
--- /dev/null
+++ b/Engine/Source/Runtime/Experimental/ChaosVisualDebugger/Public/DataWrappers/ChaosVDParticleDataWrapper.h
@@ -0,0 +1,463 @@
+﻿// Copyright Epic Games, Inc. All Rights Reserved.
+
+#pragma once
+
+#include "UObject/ObjectMacros.h"
+#include "HAL/Platform.h"
+
+#include "ChaosVDParticleDataWrapper.generated.h"
+
+UENUM()
+enum class EChaosVDParticleType : uint8
+{
+	Static,
+	Kinematic,
+	Rigid,
+	Clustered,
+	StaticMesh,
+	SkeletalMesh,
+	GeometryCollection,
+	Unknown
+};
+
+UENUM()
+enum class EChaosVDSleepType : uint8
+{
+	MaterialSleep,
+	NeverSleep
+};
+
+UENUM()
+enum class EChaosVDObjectStateType: int8
+{
+	Uninitialized = 0,
+	Sleeping = 1,
+	Kinematic = 2,
+	Static = 3,
+	Dynamic = 4,
+
+	Count
+};
+
+/** Base struct that declares the interface to be used for any ParticleData Viewer */
+USTRUCT()
+struct FChaosVDParticleDataBase
+{
+	GENERATED_BODY()
+	virtual ~FChaosVDParticleDataBase() = default;
+
+	virtual bool HasValidData() const { return bHasValidData; }
+
+protected:
+	UPROPERTY()
+	bool bHasValidData = false;
+};
+
+USTRUCT()
+struct FChaosVDFRigidParticleControlFlags : public FChaosVDParticleDataBase 
+{
+	GENERATED_BODY()
+
+	FChaosVDFRigidParticleControlFlags()
+		: bGravityEnabled(false),
+		  bCCDEnabled(false),
+		  bOneWayInteractionEnabled(false),
+		  bMaxDepenetrationVelocityOverrideEnabled(false),
+		  bInertiaConditioningEnabled(false), GravityGroupIndex(0)
+	{
+	}
+
+	CHAOSVDRUNTIME_API bool Serialize(FArchive& Ar);
+
+	template <typename TOther>
+	void CopyFrom(const TOther& Other)
+	{
+		bGravityEnabled = Other.GetGravityEnabled();
+		bCCDEnabled = Other.GetCCDEnabled();
+		bOneWayInteractionEnabled = Other.GetOneWayInteractionEnabled();
+		bMaxDepenetrationVelocityOverrideEnabled = Other.GetMaxDepenetrationVelocityOverrideEnabled();
+		bInertiaConditioningEnabled = Other.GetInertiaConditioningEnabled();
+		GravityGroupIndex = Other.GetGravityGroupIndex();
+
+		bHasValidData = true;
+	}
+
+	UPROPERTY(EditAnywhere, Category= "Particle Control Flags")
+	bool bGravityEnabled;
+	UPROPERTY(EditAnywhere, Category= "Particle Control Flags")
+	bool bCCDEnabled;
+	UPROPERTY(EditAnywhere, Category= "Particle Control Flags")
+	bool bOneWayInteractionEnabled;
+	UPROPERTY(EditAnywhere, Category= "Particle Control Flags")
+	bool bMaxDepenetrationVelocityOverrideEnabled;
+	UPROPERTY(EditAnywhere, Category= "Particle Control Flags")
+	bool bInertiaConditioningEnabled;
+	UPROPERTY(EditAnywhere, Category= "Particle Control Flags")
+	int32 GravityGroupIndex;
+};
+
+inline FArchive& operator<<(FArchive& Ar, FChaosVDFRigidParticleControlFlags& Data)
+{
+	Data.Serialize(Ar);
+	return Ar;
+}
+
+template<>
+struct TStructOpsTypeTraits<FChaosVDFRigidParticleControlFlags> : public TStructOpsTypeTraitsBase2<FChaosVDFRigidParticleControlFlags>
+{
+	enum
+	{
+		WithSerializer = true,
+	};
+};
+
+/** Simplified UStruct version of FParticlePositionRotation.
+ * Used to be able to show the values in the editor and allow changes via the Property Editor.
+ */
+USTRUCT()
+struct FChaosVDParticlePositionRotation : public FChaosVDParticleDataBase
+{
+	GENERATED_BODY()
+
+	FChaosVDParticlePositionRotation()
+	{
+	}
+
+	template <typename TOther>
+	void CopyFrom(const TOther& Other)
+	{
+		MX = Other.X();
+		MR = Other.R();
+		
+		bHasValidData = true;
+	}
+
+	CHAOSVDRUNTIME_API bool Serialize(FArchive& Ar);
+
+	UPROPERTY(EditAnywhere, Category= "Particle Velocities")
+	FVector MX;
+
+	UPROPERTY(EditAnywhere, Category= "Particle Velocities")
+	FQuat MR;
+};
+
+inline FArchive& operator<<(FArchive& Ar, FChaosVDParticlePositionRotation& Data)
+{
+	Data.Serialize(Ar);
+	return Ar;
+}
+
+template<>
+struct TStructOpsTypeTraits<FChaosVDParticlePositionRotation> : public TStructOpsTypeTraitsBase2<FChaosVDParticlePositionRotation>
+{
+	enum
+	{
+		WithSerializer = true,
+	};
+};
+
+/** Simplified UStruct version of FParticleVelocities.
+ * Used to be able to show the values in the editor and allow changes via the Property Editor.
+ */
+USTRUCT()
+struct FChaosVDParticleVelocities : public FChaosVDParticleDataBase
+{
+	GENERATED_BODY()
+
+	FChaosVDParticleVelocities()
+	{
+	}
+
+	CHAOSVDRUNTIME_API bool Serialize(FArchive& Ar);
+
+	template <typename TOther>
+	void CopyFrom(const TOther& Other)
+	{
+		MV = Other.V();
+		MW = Other.W();
+		bHasValidData = true;
+	}
+	
+	UPROPERTY(EditAnywhere, Category= "Particle Velocities")
+	FVector MV;
+
+	UPROPERTY(EditAnywhere, Category= "Particle Velocities")
+	FVector MW;
+};
+
+inline FArchive& operator<<(FArchive& Ar, FChaosVDParticleVelocities& Data)
+{
+	Data.Serialize(Ar);
+	return Ar;
+}
+
+template<>
+struct TStructOpsTypeTraits<FChaosVDParticleVelocities> : public TStructOpsTypeTraitsBase2<FChaosVDParticleVelocities>
+{
+	enum
+	{
+		WithSerializer = true,
+	};
+};
+
+/** Simplified UStruct version of FParticleDynamics.
+ * Used to be able to show the values in the editor and allow changes via the Property Editor.
+ */
+USTRUCT()
+struct FChaosVDParticleDynamics : public FChaosVDParticleDataBase
+{
+	GENERATED_BODY()
+	
+	FChaosVDParticleDynamics()
+	{
+	}
+
+	CHAOSVDRUNTIME_API bool Serialize(FArchive& Ar);
+
+	template <typename TOther>
+	void CopyFrom(const TOther& Other)
+	{
+		MAcceleration = Other.Acceleration();
+		MAngularAcceleration = Other.AngularAcceleration();
+		MLinearImpulseVelocity = Other.LinearImpulseVelocity();
+		MAngularImpulseVelocity = Other.AngularImpulseVelocity();
+
+		bHasValidData = true;
+	}
+
+	UPROPERTY(EditAnywhere, Category= "Particle Dynamics")
+	FVector MAcceleration;
+
+	UPROPERTY(EditAnywhere, Category= "Particle Dynamics")
+	FVector MAngularAcceleration;
+
+	UPROPERTY(EditAnywhere, Category= "Particle Dynamics")
+	FVector MLinearImpulseVelocity;
+
+	UPROPERTY(EditAnywhere, Category= "Particle Dynamics")
+	FVector MAngularImpulseVelocity;
+};
+
+inline FArchive& operator<<(FArchive& Ar, FChaosVDParticleDynamics& Data)
+{
+	Data.Serialize(Ar);
+	return Ar;
+}
+
+template<>
+struct TStructOpsTypeTraits<FChaosVDParticleDynamics> : public TStructOpsTypeTraitsBase2<FChaosVDParticleDynamics>
+{
+	enum
+	{
+		WithSerializer = true,
+	};
+};
+
+
+/** Simplified UStruct version of FParticleMassProps.
+ * Used to be able to show the values in the editor and allow changes via the Property Editor.
+ */
+USTRUCT()
+struct FChaosVDParticleMassProps : public FChaosVDParticleDataBase
+{
+	GENERATED_BODY()
+
+	FChaosVDParticleMassProps(): MM(0), MInvM(0)
+	{
+	}
+
+	CHAOSVDRUNTIME_API bool Serialize(FArchive& Ar);
+	
+	template <typename TOther>
+	void CopyFrom(const TOther& Other)
+	{
+		MCenterOfMass = Other.CenterOfMass();
+		MRotationOfMass = Other.RotationOfMass();
+		MI = FVector(Other.I());
+		MInvI = FVector(Other.InvI());
+		MM = Other.M();
+		MInvM = Other.InvM();
+
+		bHasValidData = true;
+	}
+	
+	UPROPERTY(EditAnywhere, Category= "Particle Mass Props")
+	FVector MCenterOfMass;
+
+	UPROPERTY(EditAnywhere, Category= "Particle Mass Props")
+	FQuat MRotationOfMass;
+
+	UPROPERTY(EditAnywhere, Category= "Particle Mass Props")
+	FVector MI;
+
+	UPROPERTY(EditAnywhere, Category= "Particle Mass Props")
+	FVector MInvI;
+
+	UPROPERTY(EditAnywhere, Category= "Particle Mass Props")
+	double MM;
+
+	UPROPERTY(EditAnywhere, Category= "Particle Mass Props")
+	double MInvM;
+};
+
+inline FArchive& operator<<(FArchive& Ar, FChaosVDParticleMassProps& Data)
+{
+	Data.Serialize(Ar);
+	return Ar;
+}
+
+template<>
+struct TStructOpsTypeTraits<FChaosVDParticleMassProps> : public TStructOpsTypeTraitsBase2<FChaosVDParticleMassProps>
+{
+	enum
+	{
+		WithSerializer = true,
+	};
+};
+
+/** Simplified UStruct version of FParticleDynamicMisc.
+ * Used to be able to show the values in the editor and allow changes via the Property Editor.
+ */
+USTRUCT()
+struct FChaosVDParticleDynamicMisc : public FChaosVDParticleDataBase
+{
+	GENERATED_BODY()
+
+	FChaosVDParticleDynamicMisc(): MAngularEtherDrag(0), MMaxLinearSpeedSq(0), MMaxAngularSpeedSq(0),
+	                               MCollisionGroup(0), MObjectState(), MSleepType(), bDisabled(false)
+	{
+	}
+
+	CHAOSVDRUNTIME_API bool Serialize(FArchive& Ar);
+
+	template <typename TOther>
+	void CopyFrom(const TOther& Other)
+	{
+		MAngularEtherDrag = Other.LinearEtherDrag();
+		MAngularEtherDrag = Other.AngularEtherDrag();
+		MMaxLinearSpeedSq = Other.MaxLinearSpeedSq();
+		MMaxAngularSpeedSq = Other.MaxAngularSpeedSq();
+		MObjectState = static_cast<EChaosVDObjectStateType>(Other.ObjectState());
+		MCollisionGroup = Other.CollisionGroup();
+		MSleepType =  static_cast<EChaosVDSleepType>(Other.SleepType());
+		MCollisionConstraintFlag = Other.CollisionConstraintFlags();
+	
+		MControlFlags.CopyFrom(Other.ControlFlags());
+		
+		bDisabled = Other.Disabled();
+
+		bHasValidData = true;
+	}
+
+	UPROPERTY(EditAnywhere, Category= "Particle Dynamic Misc")
+	double MAngularEtherDrag;
+
+	UPROPERTY(EditAnywhere, Category= "Particle Dynamic Misc")
+	double MMaxLinearSpeedSq;
+	
+	UPROPERTY(EditAnywhere, Category= "Particle Dynamic Misc")
+	double MMaxAngularSpeedSq;
+
+	UPROPERTY(EditAnywhere, Category= "Particle Dynamic Misc")
+	int32 MCollisionGroup;
+
+	UPROPERTY(EditAnywhere, Category= "Particle Dynamic Misc")
+	EChaosVDObjectStateType MObjectState;
+
+	UPROPERTY(EditAnywhere, Category= "Particle Dynamic Misc")
+	EChaosVDSleepType MSleepType;
+
+	UPROPERTY(EditAnywhere, Category= "Particle Dynamic Misc")
+	uint32 MCollisionConstraintFlag = 0;
+
+	UPROPERTY(EditAnywhere, Category= "Particle Dynamic Misc")
+	FChaosVDFRigidParticleControlFlags MControlFlags;
+
+	UPROPERTY(EditAnywhere, Category= "Particle Dynamic Misc")
+	bool bDisabled;
+};
+
+inline FArchive& operator<<(FArchive& Ar, FChaosVDParticleDynamicMisc& Data)
+{
+	Data.Serialize(Ar);
+	return Ar;
+}
+
+template<>
+struct TStructOpsTypeTraits<FChaosVDParticleDynamicMisc> : public TStructOpsTypeTraitsBase2<FChaosVDParticleDynamicMisc>
+{
+	enum
+	{
+		WithSerializer = true,
+	};
+};
+
+/** Simplified UStruct version of FChaosVDParticleDataWrapper.
+ * Used to be able to show the values in the editor and allow changes via the Property Editor.
+ */
+USTRUCT()
+struct FChaosVDParticleDataWrapper : public FChaosVDParticleDataBase
+{
+	virtual ~FChaosVDParticleDataWrapper() override = default;
+
+	GENERATED_BODY()
+
+	FChaosVDParticleDataWrapper()
+	{
+	}
+
+	UPROPERTY(EditAnywhere, Category= "Particle Non Frequent Data")
+	uint32 GeometryHash = 0;
+
+	UPROPERTY(VisibleAnywhere, Category= "Particle Non Frequent Data")
+	FString DebugName;
+
+	UPROPERTY(VisibleAnywhere, Category= "Particle Non Frequent Data")
+	int32 ParticleIndex = INDEX_NONE;
+
+	UPROPERTY(VisibleAnywhere, Category= "Particle Non Frequent Data")
+	int32 SolverID = INDEX_NONE;
+
+	UPROPERTY(EditAnywhere, Category= "Particle Non Frequent Data")
+	EChaosVDParticleType Type = EChaosVDParticleType::Unknown;
+
+	UPROPERTY(EditAnywhere, Category= "Particle Position Rotation")
+	FChaosVDParticlePositionRotation ParticlePositionRotation;
+
+	UPROPERTY(EditAnywhere, Category= "Particle Particle Velocities")
+	FChaosVDParticleVelocities ParticleVelocities;
+
+	UPROPERTY(EditAnywhere, Category= "Particle Particle Dynamics")
+	FChaosVDParticleDynamics ParticleDynamics;
+
+	UPROPERTY(EditAnywhere, Category= "Particle Particle Dynamics Misc")
+	FChaosVDParticleDynamicMisc ParticleDynamicsMisc;
+
+	UPROPERTY(EditAnywhere, Category= "Particle Particle Mass Props")
+	FChaosVDParticleMassProps ParticleMassProps;
+
+	/** Only used during recording */
+	TSharedPtr<FString> DebugNamePtr;
+	bool bHasDebugName = false;
+
+	virtual bool HasValidData() const override { return bHasValidData; }
+
+	void SetHasValidData(bool bNewIsValid){ bHasValidData = bNewIsValid; }
+
+	CHAOSVDRUNTIME_API bool Serialize(FArchive& Ar);
+};
+
+inline FArchive& operator<<(FArchive& Ar,FChaosVDParticleDataWrapper& Data)
+{
+	Data.Serialize(Ar);
+	return Ar;
+}
+
+template<>
+struct TStructOpsTypeTraits<FChaosVDParticleDataWrapper> : public TStructOpsTypeTraitsBase2<FChaosVDParticleDataWrapper>
+{
+	enum
+	{
+		WithSerializer = true,
+	};
+};
