Fix mass attribute being accessed as double instead of float  causing unchecked out of bounds access

diff --git a/Engine/Source/Runtime/Experimental/Chaos/Private/PhysicsProxy/GeometryCollectionPhysicsProxy.cpp b/Engine/Source/Runtime/Experimental/Chaos/Private/PhysicsProxy/GeometryCollectionPhysicsProxy.cpp
--- a/Engine/Source/Runtime/Experimental/Chaos/Private/PhysicsProxy/GeometryCollectionPhysicsProxy.cpp
+++ b/Engine/Source/Runtime/Experimental/Chaos/Private/PhysicsProxy/GeometryCollectionPhysicsProxy.cpp
@@ -541,7 +541,7 @@ void FGeometryCollectionPhysicsProxy::Initialize(Chaos::FPBDRigidsEvolutionBase
 
 				GTParticles[Index]->SetUniqueIdx(Evolution->GenerateUniqueIdx());
 
-				const Chaos::FReal ScaledMass = AdjustMassForScale(Mass[Index]);
+				const float ScaledMass = AdjustMassForScale(Mass[Index]);
 
 				// Note that this transform must match the physics thread transform computation for initialization.
 				// Take for example, a geometry collection in the editor that is linked with a joint constraint. The joint
@@ -989,7 +989,7 @@ void FGeometryCollectionPhysicsProxy::InitializeBodiesPT(Chaos::FPBDRigidsSolver
 
 		for (int32 Index = 0; Index < NumTransforms; Index++)
 		{
-			const float ScaledMass = (float)AdjustMassForScale(Mass[Index]);
+			const float ScaledMass = AdjustMassForScale(Mass[Index]);
 			const Chaos::FVec3f ScaledInertia = AdjustInertiaForScale((Chaos::FVec3f)InertiaTensor[Index]);
 			Masses.Add(ScaledMass);
 			Inertias.Add(ScaledInertia);
@@ -1505,10 +1505,10 @@ float FGeometryCollectionPhysicsProxy::ComputeUserDefinedDamageThreshold_Interna
 	return DamageThreshold;
 }
 
-Chaos::FReal FGeometryCollectionPhysicsProxy::AdjustMassForScale(Chaos::FReal Mass) const
+float FGeometryCollectionPhysicsProxy::AdjustMassForScale(float Mass) const
 {
 	const FVector WorldScale = Parameters.WorldTransform.GetScale3D();
-	Chaos::FReal MassScale = WorldScale.X * WorldScale.Y * WorldScale.Z;
+	float MassScale = (float)(WorldScale.X * WorldScale.Y * WorldScale.Z);
 	MassScale *= Parameters.MaterialOverrideMassScaleMultiplier;
 	return (Mass * MassScale);
 }
@@ -1720,13 +1720,10 @@ FGeometryCollectionPhysicsProxy::BuildClusters_Internal(
 	}
 
 	check(Parameters.RestCollection);
-	const TManagedArray<Chaos::FReal>& Mass = Parameters.RestCollection->GetAttribute<Chaos::FReal>("Mass", FTransformCollection::TransformGroup);
+	const TManagedArray<float>& Mass = Parameters.RestCollection->GetAttribute<float>("Mass", FTransformCollection::TransformGroup);
 	const TManagedArray<FVector3f>& InertiaTensor = Parameters.RestCollection->GetAttribute<FVector3f>("InertiaTensor", FTransformCollection::TransformGroup);
 
-	Mass.RangeCheck(CollectionClusterIndex);
-	InertiaTensor.RangeCheck(CollectionClusterIndex);
-
-	const Chaos::FReal ScaledMass = AdjustMassForScale(Mass[CollectionClusterIndex]);
+	const float ScaledMass = AdjustMassForScale(Mass[CollectionClusterIndex]);
 	const Chaos::FVec3f ScaledInertia = AdjustInertiaForScale((Chaos::FVec3f)InertiaTensor[CollectionClusterIndex]);
 	
 	// NOTE: The particle creation call above (CreateClusterParticle) activates the particle, which does various things including adding
@@ -1738,7 +1735,7 @@ FGeometryCollectionPhysicsProxy::BuildClusters_Internal(
 		nullptr, // Parent->Geometry() ? Parent->Geometry() : Implicits[CollectionClusterIndex], 
 		SimFilter,
 		QueryFilter,
-		Parent->M() > 0.0 ? Parent->M() : ScaledMass,
+		Parent->M() > 0.0 ? Parent->M() : (Chaos::FReal)ScaledMass,
 		Parent->I() != Chaos::FVec3f(0.0) ? Parent->I() : ScaledInertia,
 		ParticleTM, 
 		(uint8)DynamicState[CollectionClusterIndex], 
diff --git a/Engine/Source/Runtime/Experimental/Chaos/Public/GeometryCollection/ManagedArrayCollection.h b/Engine/Source/Runtime/Experimental/Chaos/Public/GeometryCollection/ManagedArrayCollection.h
--- a/Engine/Source/Runtime/Experimental/Chaos/Public/GeometryCollection/ManagedArrayCollection.h
+++ b/Engine/Source/Runtime/Experimental/Chaos/Public/GeometryCollection/ManagedArrayCollection.h
@@ -367,6 +367,7 @@ public:
 	{
 		check(HasAttribute(Name, Group))
 		const FKeyType Key = FManagedArrayCollection::MakeMapKey(Name, Group);
+		checkSlow(Map[Key].ArrayType == ManagedArrayType<T>());
 		return *(static_cast<TManagedArray<T>*>(Map[Key].Value));
 	};
 
diff --git a/Engine/Source/Runtime/Experimental/Chaos/Public/PhysicsProxy/GeometryCollectionPhysicsProxy.h b/Engine/Source/Runtime/Experimental/Chaos/Public/PhysicsProxy/GeometryCollectionPhysicsProxy.h
--- a/Engine/Source/Runtime/Experimental/Chaos/Public/PhysicsProxy/GeometryCollectionPhysicsProxy.h
+++ b/Engine/Source/Runtime/Experimental/Chaos/Public/PhysicsProxy/GeometryCollectionPhysicsProxy.h
@@ -475,7 +475,7 @@ protected:
 	float ComputeUserDefinedDamageThreshold_Internal(int32 TransformIndex) const;
 
 	/** adjust scalar mass to account for per component scale properties ( from material override and world transform scale ) */ 
-	Chaos::FReal AdjustMassForScale(Chaos::FReal Mass) const;
+	float AdjustMassForScale(float Mass) const;
 
 	/** adjust inertia to account for per component scale properties ( from material override and world transform scale ) */
 	Chaos::FVec3f AdjustInertiaForScale(const Chaos::FVec3f& Inertia) const;
