Iris - Added NetObjectConnectionFilter which can be used to dynamically filter objects to only replicate to certain connections. It's main use case is when objects added to this filter may also be dependent objects as dependent objects will override dynamic filtering, but not any type of static filtering. Added GetNetObjectFilteringInfo to UNetObjectFilter, allowing access to that info outside of filtering functions.

diff --git a/Engine/Source/Runtime/Experimental/Iris/Core/Private/Iris/ReplicationSystem/Filtering/NetObjectFilter.cpp b/Engine/Source/Runtime/Experimental/Iris/Core/Private/Iris/ReplicationSystem/Filtering/NetObjectFilter.cpp
--- a/Engine/Source/Runtime/Experimental/Iris/Core/Private/Iris/ReplicationSystem/Filtering/NetObjectFilter.cpp
+++ b/Engine/Source/Runtime/Experimental/Iris/Core/Private/Iris/ReplicationSystem/Filtering/NetObjectFilter.cpp
@@ -1,6 +1,7 @@
 // Copyright Epic Games, Inc. All Rights Reserved.
 
 #include "Iris/ReplicationSystem/Filtering/NetObjectFilter.h"
+#include "Iris/ReplicationSystem/Filtering/ReplicationFiltering.h"
 
 FNetObjectFilteringParams::FNetObjectFilteringParams(const UE::Net::FNetBitArrayView InFilteredObjects)
 : FilteredObjects(InFilteredObjects)
@@ -21,6 +22,12 @@ void UNetObjectFilter::Init(FNetObjectFilterInitParams& Params)
 		FilterType = Params.Config->FilterType;
 	}
 
+	{
+		UE::Net::Private::FNetObjectFilteringInfoAccessor FilteringInfoAccessor;
+		TArrayView<FNetObjectFilteringInfo> NetObjectFilteringInfos = FilteringInfoAccessor.GetNetObjectFilteringInfos(Params.ReplicationSystem);
+		FilterInfo = FFilterInfo(Params.FilteredObjects, NetObjectFilteringInfos);
+	}
+
 	OnInit(Params);
 }
 
@@ -43,3 +50,30 @@ void UNetObjectFilter::Filter(FNetObjectFilteringParams&)
 void UNetObjectFilter::PostFilter(FNetObjectPostFilteringParams&)
 {
 }
+
+// FFilterInfo implementation
+UNetObjectFilter::FFilterInfo::FFilterInfo(const UE::Net::FNetBitArrayView InFilteredObjects, const TArrayView<FNetObjectFilteringInfo> InFilteringInfos)
+: FilteredObjects(InFilteredObjects)
+, FilteringInfos(InFilteringInfos)
+{
+}
+
+UNetObjectFilter::FFilterInfo& UNetObjectFilter::FFilterInfo::operator=(const UNetObjectFilter::FFilterInfo& Other)
+{
+	// Re-initialize this instance.
+	this->~FFilterInfo();
+	new (this) FFilterInfo(Other.FilteredObjects, Other.FilteringInfos);
+	return *this;
+}
+
+FNetObjectFilteringInfo* UNetObjectFilter::FFilterInfo::GetFilteringInfo(uint32 ObjectIndex)
+{
+	// Only allow retriving infos for objects handled by this instance.
+	if (!IsAddedToFilter(ObjectIndex))
+	{
+		return nullptr;
+	}
+
+	return &FilteringInfos[ObjectIndex];
+}
+
diff --git a/Engine/Source/Runtime/Experimental/Iris/Core/Private/Iris/ReplicationSystem/Filtering/ReplicationFiltering.cpp b/Engine/Source/Runtime/Experimental/Iris/Core/Private/Iris/ReplicationSystem/Filtering/ReplicationFiltering.cpp
--- a/Engine/Source/Runtime/Experimental/Iris/Core/Private/Iris/ReplicationSystem/Filtering/ReplicationFiltering.cpp
+++ b/Engine/Source/Runtime/Experimental/Iris/Core/Private/Iris/ReplicationSystem/Filtering/ReplicationFiltering.cpp
@@ -1119,6 +1119,8 @@ void FReplicationFiltering::PreUpdateDynamicFiltering(ENetFilterType FilterType)
 	// Give filters a chance to prepare for filtering. It's only called if any object has the filter set.
 	{
 		FNetObjectPreFilteringParams PreFilteringParams;
+		PreFilteringParams.ValidConnections = MakeNetBitArrayView(ValidConnections);
+
 		for (FFilterInfo& Info : DynamicFilterInfos)
 		{
 			if (Info.ObjectCount == 0U)
@@ -2034,23 +2036,23 @@ void FReplicationFiltering::InitFilters()
 			}
 		}
 
-		TStrongObjectPtr<UNetObjectFilter> Filter(NewObject<UNetObjectFilter>((UObject*)GetTransientPackage(), NetObjectFilterClass, MakeUniqueObjectName(nullptr, NetObjectFilterClass, FilterDefinition.FilterName)));
-		check(Filter.IsValid());
+		
+		FFilterInfo& Info = DynamicFilterInfos.Emplace_GetRef();
+		Info.Filter = TStrongObjectPtr<UNetObjectFilter>(NewObject<UNetObjectFilter>((UObject*)GetTransientPackage(), NetObjectFilterClass, MakeUniqueObjectName(nullptr, NetObjectFilterClass, FilterDefinition.FilterName)));
+		check(Info.Filter.IsValid());
+		Info.Name = FilterDefinition.FilterName;
+		Info.ObjectCount = 0;
+		Info.FilteredObjects.Init(MaxObjectCount);
 
-		FNetObjectFilterInitParams InitParams;
+		FNetObjectFilterInitParams InitParams(MakeNetBitArrayView(Info.FilteredObjects));
 		InitParams.ReplicationSystem = ReplicationSystem;
 		InitParams.Config = (NetObjectFilterConfigClass ? NewObject<UNetObjectFilterConfig>((UObject*)GetTransientPackage(), NetObjectFilterConfigClass) : nullptr);
 		InitParams.MaxObjectCount = MaxObjectCount;
 		InitParams.MaxConnectionCount = Connections->GetMaxConnectionCount();
 
-		Filter->Init(InitParams);
-
-		FFilterInfo& Info = DynamicFilterInfos.Emplace_GetRef();
-		Info.Type = Filter->GetFilterType();
-		Info.Filter = MoveTemp(Filter);
-		Info.Name = FilterDefinition.FilterName;
-		Info.ObjectCount = 0;
-		Info.FilteredObjects.Init(MaxObjectCount);
+		Info.Filter->Init(InitParams);
+		// Filter type can be changed in Init.
+		Info.Type = Info.Filter->GetFilterType();
 
 		bHasDynamicFilters = true;
 		bHasDynamicRawFilters |= (Info.Type == ENetFilterType::PrePoll_Raw);
@@ -2108,6 +2110,24 @@ void FReplicationFiltering::InvalidateBaselinesForObject(uint32 ObjectIndex, uin
 	}
 }
 
+TArrayView<FNetObjectFilteringInfo> FReplicationFiltering::GetNetObjectFilteringInfos()
+{
+	return MakeArrayView(NetObjectFilteringInfos);
+}
+
+//*************************************************************************************************
+// FNetObjectFilteringInfoAccessor
+//*************************************************************************************************
+TArrayView<FNetObjectFilteringInfo> FNetObjectFilteringInfoAccessor::GetNetObjectFilteringInfos(UReplicationSystem* ReplicationSystem) const
+{
+	if (ReplicationSystem)
+	{
+		return ReplicationSystem->GetReplicationSystemInternal()->GetFiltering().GetNetObjectFilteringInfos();
+	}
+
+	return TArrayView<FNetObjectFilteringInfo>();
+}
+
 //*************************************************************************************************
 // FNetObjectFilterHandleUtil
 //*************************************************************************************************
diff --git a/Engine/Source/Runtime/Experimental/Iris/Core/Private/Iris/ReplicationSystem/Filtering/ReplicationFiltering.h b/Engine/Source/Runtime/Experimental/Iris/Core/Private/Iris/ReplicationSystem/Filtering/ReplicationFiltering.h
--- a/Engine/Source/Runtime/Experimental/Iris/Core/Private/Iris/ReplicationSystem/Filtering/ReplicationFiltering.h
+++ b/Engine/Source/Runtime/Experimental/Iris/Core/Private/Iris/ReplicationSystem/Filtering/ReplicationFiltering.h
@@ -26,6 +26,18 @@ namespace UE::Net
 namespace UE::Net::Private
 {
 
+class FNetObjectFilteringInfoAccessor
+{
+private:
+	/** Returns all the NetObjectFilteringInfos for the filtering system. */
+	TArrayView<FNetObjectFilteringInfo> GetNetObjectFilteringInfos(UReplicationSystem* ReplicationSystem) const;
+
+private:
+	// Friends
+	friend UNetObjectFilter;
+};
+
+
 struct FReplicationFilteringInitParams
 {
 	TObjectPtr<UReplicationSystem> ReplicationSystem;
@@ -163,6 +175,7 @@ private:
 
 private:
 	class FUpdateDirtyObjectsBatchHelper;
+	friend FNetObjectFilteringInfoAccessor;
 	
 	static void StaticChecks();
 
@@ -217,6 +230,9 @@ private:
 
 	void InvalidateBaselinesForObject(uint32 ObjectIndex, uint32 NewOwningConnectionId, uint32 PrevOwningConnectionId);
 
+	/** Returns all the filtering infos. */
+	TArrayView<FNetObjectFilteringInfo> GetNetObjectFilteringInfos();
+
 private:
 	// Used for ObjectIndexToDynamicFilterIndex lookup
 	static constexpr uint8 InvalidDynamicFilterIndex = 255U;
diff --git a/Engine/Source/Runtime/Experimental/Iris/Core/Private/Iris/ReplicationSystem/Prioritization/NetObjectConnectionFilter.cpp b/Engine/Source/Runtime/Experimental/Iris/Core/Private/Iris/ReplicationSystem/Prioritization/NetObjectConnectionFilter.cpp
new file mode 100644
--- /dev/null
+++ b/Engine/Source/Runtime/Experimental/Iris/Core/Private/Iris/ReplicationSystem/Prioritization/NetObjectConnectionFilter.cpp
@@ -0,0 +1,125 @@
+// Copyright Epic Games, Inc. All Rights Reserved.
+
+#include "Iris/ReplicationSystem/Filtering/NetObjectConnectionFilter.h"
+#include "Iris/ReplicationSystem/NetRefHandleManager.h"
+#include "Iris/ReplicationSystem/ReplicationSystem.h"
+#include "Iris/ReplicationSystem/ReplicationSystemInternal.h"
+#include "Iris/Core/IrisProfiler.h"
+
+void UNetObjectConnectionFilter::SetReplicateToConnection(UE::Net::FNetRefHandle RefHandle, uint32 ConnectionId, UE::Net::ENetFilterStatus FilterStatus)
+{
+	if (!ensure(ConnectionId < static_cast<uint32>(PerConnectionInfos.Num())))
+	{
+		return;
+	}
+
+	const UE::Net::Private::FInternalNetRefIndex ObjectIndex = NetRefHandleManager->GetInternalIndex(RefHandle);
+	if (ObjectIndex == UE::Net::Private::FNetRefHandleManager::InvalidInternalIndex)
+	{
+		return;
+	}
+
+	if (const FFilteringInfo* FilteringInfo = static_cast<const FFilteringInfo*>(GetFilteringInfo(ObjectIndex)))
+	{
+		const uint32 LocalIndex = FilteringInfo->GetLocalObjectIndex();
+		FPerConnectionInfo& PerConnectionInfo = PerConnectionInfos[ConnectionId];
+		PerConnectionInfo.ReplicationEnabledObjects.SetBit(LocalIndex);
+	}
+}
+
+void UNetObjectConnectionFilter::OnInit(FNetObjectFilterInitParams& Params)
+{
+	NetRefHandleManager = &Params.ReplicationSystem->GetReplicationSystemInternal()->GetNetRefHandleManager();
+
+	Config = TStrongObjectPtr<UNetObjectConnectionFilterConfig>(CastChecked<UNetObjectConnectionFilterConfig>(Params.Config));
+
+	const uint16 MaxLocalObjectCount = static_cast<uint16>(FPlatformMath::Min<uint32>(Params.MaxObjectCount, Config->MaxObjectCount));
+	UsedLocalInfoIndices.Init(MaxLocalObjectCount);
+	LocalToNetRefIndex.SetNumZeroed(MaxLocalObjectCount);
+
+	PerConnectionInfos.SetNum(Params.MaxConnectionCount + 1);
+}
+
+void UNetObjectConnectionFilter::AddConnection(uint32 ConnectionId)
+{
+	FPerConnectionInfo& ConnInfo = PerConnectionInfos[ConnectionId];
+	ConnInfo.ReplicationEnabledObjects.Init(UsedLocalInfoIndices.GetNumBits());
+}
+
+void UNetObjectConnectionFilter::RemoveConnection(uint32 ConnectionId)
+{
+	// Free everything that was allocated for the connection.
+	FPerConnectionInfo& ConnInfo = PerConnectionInfos[ConnectionId];
+	ConnInfo = FPerConnectionInfo();
+}
+
+bool UNetObjectConnectionFilter::AddObject(uint32 ObjectIndex, FNetObjectFilterAddObjectParams& Params)
+{
+	const uint32 LocalIndex = UsedLocalInfoIndices.FindFirstZero();
+	if (!ensureMsgf(LocalIndex != UE::Net::FNetBitArrayBase::InvalidIndex, TEXT("%hs. MaxObjectCount: %u. Config type %s."), "Too many objects added to NetObjectConnectionFilter. Object will not be handled by filter!", (Config.IsValid() ? Config->MaxObjectCount : uint16(0)), ToCStr(GetNameSafe(Config.Get()))))
+	{
+		return false;
+	}
+
+	UsedLocalInfoIndices.SetBit(LocalIndex);
+	LocalToNetRefIndex[LocalIndex] = ObjectIndex;
+
+	FFilteringInfo& Info = static_cast<FFilteringInfo&>(Params.OutInfo);
+	Info.SetLocalObjectIndex(static_cast<uint16>(LocalIndex));
+
+	return true;
+}
+
+void UNetObjectConnectionFilter::RemoveObject(uint32 ObjectIndex, const FNetObjectFilteringInfo& InInfo)
+{
+	bObjectRemoved = true;
+
+	const FFilteringInfo& Info = static_cast<const FFilteringInfo&>(InInfo);
+
+	const uint16 LocalIndex = Info.GetLocalObjectIndex();
+	UsedLocalInfoIndices.ClearBit(LocalIndex);
+
+	// For good measure. It's not strictly needed.
+	LocalToNetRefIndex[LocalIndex] = UE::Net::Private::FNetRefHandleManager::InvalidInternalIndex;
+}
+
+void UNetObjectConnectionFilter::UpdateObjects(FNetObjectFilterUpdateParams&)
+{
+	// Intentionally left empty.
+}
+
+void UNetObjectConnectionFilter::PreFilter(FNetObjectPreFilteringParams& Params)
+{
+	if (!bObjectRemoved)
+	{
+		return;
+	}
+
+	bObjectRemoved = false;
+
+	// Mask out no longer filtered objects to minimize looping during filtering.
+	Params.ValidConnections.ForAllSetBits(
+		[this](uint32 ConnectionId)
+		{
+			FPerConnectionInfo& Info = this->PerConnectionInfos[ConnectionId];
+			Info.ReplicationEnabledObjects.Combine(this->UsedLocalInfoIndices, UE::Net::FNetBitArrayBase::AndOp);
+		}
+	);
+}
+
+void UNetObjectConnectionFilter::Filter(FNetObjectFilteringParams& Params)
+{
+	IRIS_PROFILER_SCOPE(UNetObjectConnectionFilter_Filter);
+
+	UE::Net::FNetBitArrayView& AllowedObjects = Params.OutAllowedObjects;
+	AllowedObjects.Reset();
+
+	FPerConnectionInfo& ConnectionInfo = PerConnectionInfos[Params.ConnectionId];
+	ConnectionInfo.ReplicationEnabledObjects.ForAllSetBits(
+		[&AllowedObjects, &ToNetRefIndex = LocalToNetRefIndex](uint32 LocalObjectIndex)
+		{
+			const uint32 ObjectIndex = ToNetRefIndex[LocalObjectIndex];
+			AllowedObjects.SetBit(ObjectIndex);
+		}
+	);
+}
diff --git a/Engine/Source/Runtime/Experimental/Iris/Core/Public/Iris/ReplicationSystem/Filtering/NetObjectConnectionFilter.h b/Engine/Source/Runtime/Experimental/Iris/Core/Public/Iris/ReplicationSystem/Filtering/NetObjectConnectionFilter.h
new file mode 100644
--- /dev/null
+++ b/Engine/Source/Runtime/Experimental/Iris/Core/Public/Iris/ReplicationSystem/Filtering/NetObjectConnectionFilter.h
@@ -0,0 +1,68 @@
+// Copyright Epic Games, Inc. All Rights Reserved.
+
+#pragma once
+
+#include "Containers/ContainersFwd.h"
+#include "Iris/ReplicationSystem/Filtering/NetObjectFilter.h"
+#include "Iris/ReplicationSystem/NetRefHandle.h"
+#include "Net/Core/NetBitArray.h"
+#include "UObject/StrongObjectPtr.h"
+#include "NetObjectConnectionFilter.generated.h"
+
+namespace UE::Net::Private
+{
+	class FNetRefHandleManager;
+}
+
+UCLASS(transient, MinimalAPI)
+class UNetObjectConnectionFilterConfig : public UNetObjectFilterConfig
+{
+	GENERATED_BODY()
+
+public:
+	/** The maximum amount of objects that may be added to the filter. It's not designed to handle massive amounts- static connection filtering via the ReplicationSystem API is preferred. */
+	UPROPERTY(Config)
+	uint16 MaxObjectCount = 4096;
+};
+
+/**
+ * The NetObjectConnectionFilter is a dynamic pre-poll filter that supports per connection filtering. It should be seen as an alternative to the ReplicationSystem SetConnectionFilter API for use cases where
+ * for example the object in question can be a dependent object. Dependent objects will override dynamic filtering and only dynamic filtering.
+ */
+UCLASS(transient, MinimalAPI) 
+class UNetObjectConnectionFilter : public UNetObjectFilter
+{
+	GENERATED_BODY()
+
+public:
+	IRISCORE_API void SetReplicateToConnection(UE::Net::FNetRefHandle RefHandle, uint32 ConnectionId, UE::Net::ENetFilterStatus FilterStatus);
+
+protected:
+	struct FFilteringInfo : public FNetObjectFilteringInfo
+	{
+		void SetLocalObjectIndex(uint16 Index) { Data[0] = Index; }
+		uint16 GetLocalObjectIndex() const { return Data[0]; }
+	};
+
+	struct FPerConnectionInfo
+	{
+		UE::Net::FNetBitArray ReplicationEnabledObjects;
+	};
+
+	// UNetObjectFilter interface
+	IRISCORE_API virtual void OnInit(FNetObjectFilterInitParams&) override;
+	IRISCORE_API virtual void AddConnection(uint32 ConnectionId) override;
+	IRISCORE_API virtual void RemoveConnection(uint32 ConnectionId) override;
+	IRISCORE_API virtual bool AddObject(uint32 ObjectIndex, FNetObjectFilterAddObjectParams&) override;
+	IRISCORE_API virtual void RemoveObject(uint32 ObjectIndex, const FNetObjectFilteringInfo&) override;
+	IRISCORE_API virtual void UpdateObjects(FNetObjectFilterUpdateParams&) override;
+	IRISCORE_API virtual void PreFilter(FNetObjectPreFilteringParams&) override;
+	IRISCORE_API virtual void Filter(FNetObjectFilteringParams&) override;
+
+	const UE::Net::Private::FNetRefHandleManager* NetRefHandleManager = nullptr;
+	TStrongObjectPtr<UNetObjectConnectionFilterConfig> Config;
+	TArray<uint32> LocalToNetRefIndex;
+	TArray<FPerConnectionInfo> PerConnectionInfos;
+	UE::Net::FNetBitArray UsedLocalInfoIndices;
+	bool bObjectRemoved = false;
+};
diff --git a/Engine/Source/Runtime/Experimental/Iris/Core/Public/Iris/ReplicationSystem/Filtering/NetObjectFilter.h b/Engine/Source/Runtime/Experimental/Iris/Core/Public/Iris/ReplicationSystem/Filtering/NetObjectFilter.h
--- a/Engine/Source/Runtime/Experimental/Iris/Core/Public/Iris/ReplicationSystem/Filtering/NetObjectFilter.h
+++ b/Engine/Source/Runtime/Experimental/Iris/Core/Public/Iris/ReplicationSystem/Filtering/NetObjectFilter.h
@@ -90,6 +90,8 @@ struct FNetObjectFilteringParams
  */
 struct FNetObjectPreFilteringParams
 {
+	// The IDs of all valid connections.
+	UE::Net::FNetBitArrayView ValidConnections;
 };
 
 /**
@@ -146,13 +148,21 @@ public:
 /** Parameters passed to the filter's Init() call. */
 struct FNetObjectFilterInitParams
 {
-	TObjectPtr<const UReplicationSystem> ReplicationSystem;
+public:
+	FNetObjectFilterInitParams(const UE::Net::FNetBitArrayView InFilteredObjects)
+	: FilteredObjects(InFilteredObjects)
+	{
+	}
+
+	TObjectPtr<UReplicationSystem> ReplicationSystem;
 	/** Optional config as set in the FNetObjectFilterDefinition. */
 	UNetObjectFilterConfig* Config = nullptr;
 	/** The maximum number of objects in the system. */
 	uint32 MaxObjectCount = 0;
 	/** The maximum number of connections in the system. */
 	uint32 MaxConnectionCount = 0;
+	/** The objects that are handled by the filter. */
+	const UE::Net::FNetBitArrayView FilteredObjects;
 };
 
 struct FNetObjectFilterAddObjectParams
@@ -244,10 +254,49 @@ protected:
 	/** Called right after constructor for enabled filters. Must be overriden. */
 	IRISCORE_API virtual void OnInit(FNetObjectFilterInitParams&) PURE_VIRTUAL(OnInit, );
 
+	/* Returns true if the object is added to this filter, false otherwise. */
+	bool IsAddedToFilter(uint32 ObjectIndex) const;
+
+	/* Returns the filtering info for this object if it's handled by this filter, nullptr otherwise. */
+	FNetObjectFilteringInfo* GetFilteringInfo(uint32 ObjectIndex);
+
 	/** Directly set when you want your dynamic filter to be executed. */
 	void SetupFilterType(ENetFilterType NewFilterType) { FilterType = NewFilterType; }
 
 private:
+	class FFilterInfo
+	{
+	public:
+		FFilterInfo() = default;
+		FFilterInfo(const UE::Net::FNetBitArrayView FilteredObjects, const TArrayView<FNetObjectFilteringInfo> FileringInfos);
+
+		FFilterInfo& operator=(const FFilterInfo&);
+
+		bool IsAddedToFilter(uint32 ObjectIndex) const;
+
+		FNetObjectFilteringInfo* GetFilteringInfo(uint32 ObjectIndex);
+
+	private:
+		const UE::Net::FNetBitArrayView FilteredObjects;
+		TArrayView<FNetObjectFilteringInfo> FilteringInfos;
+	};
 
 	ENetFilterType FilterType = ENetFilterType::PrePoll_Raw;
+	FFilterInfo FilterInfo;
 };
+
+inline bool UNetObjectFilter::IsAddedToFilter(uint32 ObjectIndex) const
+{
+	return FilterInfo.IsAddedToFilter(ObjectIndex);
+}
+
+inline FNetObjectFilteringInfo* UNetObjectFilter::GetFilteringInfo(uint32 ObjectIndex)
+{
+	return FilterInfo.GetFilteringInfo(ObjectIndex);
+}
+
+inline bool UNetObjectFilter::FFilterInfo::IsAddedToFilter(uint32 ObjectIndex) const
+{
+	return ObjectIndex < FilteredObjects.GetNumBits() && FilteredObjects.IsBitSet(ObjectIndex);
+}
+
