Iris
* Formalize the concept of ForceNetUpdate by splitting MarkDirty and ForceNetUpdate into seperate lists.
* Now dirty objects will not be forced to be polled immediately anymore. It's only if you call ForceNetUpdate that you can skip your polling frequency.
* Another exeception is when a subobject gets added to a root object. There we are forcing a call to ForceNetUpdate.
* Added code to ensure that only the currently updated object can be marked dirty during the poll phase.

#rnx

diff --git a/Engine/Plugins/Runtime/ReplicationSystemTestPlugin/Source/Private/Tests/ReplicationSystem/Filtering/TestGridFilter.cpp b/Engine/Plugins/Runtime/ReplicationSystemTestPlugin/Source/Private/Tests/ReplicationSystem/Filtering/TestGridFilter.cpp
--- a/Engine/Plugins/Runtime/ReplicationSystemTestPlugin/Source/Private/Tests/ReplicationSystem/Filtering/TestGridFilter.cpp
+++ b/Engine/Plugins/Runtime/ReplicationSystemTestPlugin/Source/Private/Tests/ReplicationSystem/Filtering/TestGridFilter.cpp
@@ -100,12 +100,16 @@ private:
 // Test world loc filter
 UE_NET_TEST_FIXTURE(FTestGridFilterFixture, TestWorldLocGridFilter)
 {
-	// Spawn object on server
-	UReplicatedTestObject* ServerObjectZero = Server->CreateObjectWithWorldLocation();
-	UReplicatedTestObject* ServerObjectLimit = Server->CreateObjectWithWorldLocation();
-	UReplicatedTestObject* ServerObjectCulled = Server->CreateObjectWithWorldLocation();
-	UReplicatedTestObject* ServerObjectVeryFar = Server->CreateObjectWithWorldLocation();
-
+	// Spawn object with WorldLocation's on server
+	UObjectReplicationBridge::FCreateNetRefHandleParams Params;
+	Params.bCanReceive = true;
+	Params.bNeedsWorldLocationUpdate = true;
+	Params.bAllowDynamicFilter = true;
+
+	UReplicatedTestObject* ServerObjectZero = Server->CreateObject(Params);
+	UReplicatedTestObject* ServerObjectLimit = Server->CreateObject(Params);
+	UReplicatedTestObject* ServerObjectCulled = Server->CreateObject(Params);
+	UReplicatedTestObject* ServerObjectVeryFar = Server->CreateObject(Params);
 
 	const UNetObjectGridFilterConfig* DefaultGridConfig = Cast<UNetObjectGridFilterConfig>(UNetObjectGridFilterConfig::StaticClass()->GetDefaultObject());
 
diff --git a/Engine/Plugins/Runtime/ReplicationSystemTestPlugin/Source/Private/Tests/ReplicationSystem/ReplicatedTestObject.cpp b/Engine/Plugins/Runtime/ReplicationSystemTestPlugin/Source/Private/Tests/ReplicationSystem/ReplicatedTestObject.cpp
--- a/Engine/Plugins/Runtime/ReplicationSystemTestPlugin/Source/Private/Tests/ReplicationSystem/ReplicatedTestObject.cpp
+++ b/Engine/Plugins/Runtime/ReplicationSystemTestPlugin/Source/Private/Tests/ReplicationSystem/ReplicatedTestObject.cpp
@@ -240,6 +240,11 @@ void UReplicatedTestObjectBridge::SetExternalWorldLocationUpdateFunctor(TFunctio
 	SetInstanceGetWorldObjectInfoFunction(LocUpdateFunctor);
 }
 
+void UReplicatedTestObjectBridge::SetExternalPreUpdateFunctor(TFunction<void(FNetRefHandle, UObject*, const UReplicationBridge*)> PreUpdateFunctor)
+{
+	SetInstancePreUpdateFunction(PreUpdateFunctor);
+}
+
 //////////////////////////////////////////////////////////////////////////
 // Implementation for UTestReplicatedIrisPropertyComponent
 //////////////////////////////////////////////////////////////////////////
diff --git a/Engine/Plugins/Runtime/ReplicationSystemTestPlugin/Source/Private/Tests/ReplicationSystem/ReplicatedTestObject.h b/Engine/Plugins/Runtime/ReplicationSystemTestPlugin/Source/Private/Tests/ReplicationSystem/ReplicatedTestObject.h
--- a/Engine/Plugins/Runtime/ReplicationSystemTestPlugin/Source/Private/Tests/ReplicationSystem/ReplicatedTestObject.h
+++ b/Engine/Plugins/Runtime/ReplicationSystemTestPlugin/Source/Private/Tests/ReplicationSystem/ReplicatedTestObject.h
@@ -450,6 +450,8 @@ public:
 
 	void SetExternalWorldLocationUpdateFunctor(TFunction<void(FNetRefHandle NetHandle, const UObject* ReplicatedObject, FVector& OutLocation, float& OutCullDistance)> LocUpdateFunctor);
 
+	void SetExternalPreUpdateFunctor(TFunction<void(FNetRefHandle, UObject*, const UReplicationBridge*)> PreUpdateFunctor);
+
 protected:
 	// Type specifics for serializing creation data this will most likely be made into a separate interface to support different types of header data for different types
 	// But if we can avoid having custom data per type for instantiating remote objects as we would like to be able to fully express the state of a replicated object using the define protocol alone. 
diff --git a/Engine/Plugins/Runtime/ReplicationSystemTestPlugin/Source/Private/Tests/ReplicationSystem/ReplicationSystemServerClientTestFixture.cpp b/Engine/Plugins/Runtime/ReplicationSystemTestPlugin/Source/Private/Tests/ReplicationSystem/ReplicationSystemServerClientTestFixture.cpp
--- a/Engine/Plugins/Runtime/ReplicationSystemTestPlugin/Source/Private/Tests/ReplicationSystem/ReplicationSystemServerClientTestFixture.cpp
+++ b/Engine/Plugins/Runtime/ReplicationSystemTestPlugin/Source/Private/Tests/ReplicationSystem/ReplicationSystemServerClientTestFixture.cpp
@@ -100,6 +100,17 @@ uint32 FReplicationSystemTestNode::GetNetTraceId() const
 	return ReplicationSystem ? ReplicationSystem->GetId() : ~0U;
 }
 
+UTestReplicatedIrisObject* FReplicationSystemTestNode::CreateObject(const UObjectReplicationBridge::FCreateNetRefHandleParams& Params)
+{
+	UTestReplicatedIrisObject* CreatedObject = NewObject<UTestReplicatedIrisObject>();
+	CreatedObjects.Add(TStrongObjectPtr<UObject>(CreatedObject));
+
+	// Add it to the bridge for replication
+	ReplicationBridge->BeginReplication(CreatedObject, Params);
+
+	return CreatedObject;
+}
+
 UTestReplicatedIrisObject* FReplicationSystemTestNode::CreateObject(uint32 NumComponents, uint32 NumIrisComponents)
 {
 	UTestReplicatedIrisObject* CreatedObject = NewObject<UTestReplicatedIrisObject>();
@@ -160,22 +171,6 @@ UTestReplicatedIrisObject* FReplicationSystemTestNode::CreateObjectWithDynamicSt
 	return CreatedObject;
 }
 
-UTestReplicatedIrisObject* FReplicationSystemTestNode::CreateObjectWithWorldLocation()
-{
-	UTestReplicatedIrisObject* CreatedObject = NewObject<UTestReplicatedIrisObject>();
-	CreatedObjects.Add(TStrongObjectPtr<UObject>(CreatedObject));
-
-	UObjectReplicationBridge::FCreateNetRefHandleParams Params;
-	Params.bCanReceive = true;
-	Params.bNeedsWorldLocationUpdate = true;
-	Params.bAllowDynamicFilter = true;
-
-	// Add it to the bridge for replication
-	ReplicationBridge->BeginReplication(CreatedObject, Params);
-
-	return CreatedObject;
-}
-
 void FReplicationSystemTestNode::DestroyObject(UReplicatedTestObject* Object, EEndReplicationFlags EndReplicationFlags)
 {
 	// Destroy handle
diff --git a/Engine/Plugins/Runtime/ReplicationSystemTestPlugin/Source/Private/Tests/ReplicationSystem/ReplicationSystemServerClientTestFixture.h b/Engine/Plugins/Runtime/ReplicationSystemTestPlugin/Source/Private/Tests/ReplicationSystem/ReplicationSystemServerClientTestFixture.h
--- a/Engine/Plugins/Runtime/ReplicationSystemTestPlugin/Source/Private/Tests/ReplicationSystem/ReplicationSystemServerClientTestFixture.h
+++ b/Engine/Plugins/Runtime/ReplicationSystemTestPlugin/Source/Private/Tests/ReplicationSystem/ReplicationSystemServerClientTestFixture.h
@@ -114,12 +114,13 @@ public:
 		return static_cast<T*>(ReplicationBridge->GetReplicatedObject(Handle));
 	}
 
+	UTestReplicatedIrisObject* CreateObject(const UObjectReplicationBridge::FCreateNetRefHandleParams& Params);
 	UTestReplicatedIrisObject* CreateObject(uint32 NumComponents, uint32 NumIrisComponents);
 	UTestReplicatedIrisObject* CreateSubObject(FNetRefHandle Owner, uint32 NumComponents, uint32 NumIrisComponents);
 	UTestReplicatedIrisObject* CreateObject(const UTestReplicatedIrisObject::FComponents& Components);
 	UTestReplicatedIrisObject* CreateSubObject(FNetRefHandle Owner, const UTestReplicatedIrisObject::FComponents& Components);
 	UTestReplicatedIrisObject* CreateObjectWithDynamicState(uint32 NumComponents, uint32 NumIrisComponents, uint32 NumDynamicStateComponents);
-	UTestReplicatedIrisObject* CreateObjectWithWorldLocation();
+
 	void DestroyObject(UReplicatedTestObject*, EEndReplicationFlags EndReplicationFlags = EEndReplicationFlags::Destroy);
 
 	// Connection
diff --git a/Engine/Plugins/Runtime/ReplicationSystemTestPlugin/Source/Private/Tests/ReplicationSystem/TestDependentObjects.cpp b/Engine/Plugins/Runtime/ReplicationSystemTestPlugin/Source/Private/Tests/ReplicationSystem/TestDependentObjects.cpp
--- a/Engine/Plugins/Runtime/ReplicationSystemTestPlugin/Source/Private/Tests/ReplicationSystem/TestDependentObjects.cpp
+++ b/Engine/Plugins/Runtime/ReplicationSystemTestPlugin/Source/Private/Tests/ReplicationSystem/TestDependentObjects.cpp
@@ -369,7 +369,7 @@ UE_NET_TEST_FIXTURE(FReplicationSystemServerClientTestFixture, TestDependentObje
 	UE_NET_ASSERT_LT(It, MaxIterationCount);
 
 	// Mark dependent dirty
-	ReplicationSystem->MarkDirty(ServerDependentObject->NetRefHandle);
+	ReplicationSystem->ForceNetUpdate(ServerDependentObject->NetRefHandle);
 
 	// Send and deliver packet, we expect dependent object to have replicated 
 	Server->PreSendUpdate();
@@ -385,7 +385,7 @@ UE_NET_TEST_FIXTURE(FReplicationSystemServerClientTestFixture, TestDependentObje
 	ServerDependentObject->IntA += 2;
 
 	// Mark parent dirty
-	ReplicationSystem->MarkDirty(ServerObject->NetRefHandle);
+	ReplicationSystem->ForceNetUpdate(ServerObject->NetRefHandle);
 
 	// Send and deliver packet
 	Server->PreSendUpdate();
diff --git a/Engine/Plugins/Runtime/ReplicationSystemTestPlugin/Source/Private/Tests/ReplicationSystem/TestDirtyNetObjectTracker.cpp b/Engine/Plugins/Runtime/ReplicationSystemTestPlugin/Source/Private/Tests/ReplicationSystem/TestDirtyNetObjectTracker.cpp
--- a/Engine/Plugins/Runtime/ReplicationSystemTestPlugin/Source/Private/Tests/ReplicationSystem/TestDirtyNetObjectTracker.cpp
+++ b/Engine/Plugins/Runtime/ReplicationSystemTestPlugin/Source/Private/Tests/ReplicationSystem/TestDirtyNetObjectTracker.cpp
@@ -5,6 +5,7 @@
 #include "Iris/ReplicationSystem/DirtyNetObjectTracker.h"
 #include "Iris/ReplicationSystem/ReplicationSystemInternal.h"
 #include "Tests/ReplicationSystem/ReplicationSystemTestFixture.h"
+#include "Tests/ReplicationSystem/ReplicationSystemServerClientTestFixture.h"
 
 namespace UE::Net::Private
 {
@@ -135,4 +136,124 @@ UE_NET_TEST(DirtyNetObjectTracker, MarkingObjectAsDirtyInNonExistingSystemDoesNo
 	MarkNetObjectStateDirty(NonExistingReplicationSystemId, ArbitraryNetObjectIndex);
 }
 
+UE_NET_TEST_FIXTURE(FReplicationSystemServerClientTestFixture, GlobalDirtyTrackerTest)
+{
+	// Add client
+	FReplicationSystemTestClient* Client = CreateClient();
+
+	// Spawn object on server that is polled only every 3 frames
+	UObjectReplicationBridge::FCreateNetRefHandleParams Params;
+	Params.PollFramePeriod = 2;
+	Params.bCanReceive = true;
+	Params.bAllowDynamicFilter = true;
+	Params.bNeedsPreUpdate = true;
+
+	UTestReplicatedIrisObject* ServerObject = Server->CreateObject(Params);
+
+	// Send and deliver packet
+	Server->UpdateAndSend({Client});
+
+	// Object should have been created on the client
+	UTestReplicatedIrisObject* ClientObject = Cast<UTestReplicatedIrisObject>(Client->GetReplicationBridge()->GetReplicatedObject(ServerObject->NetRefHandle));
+	UE_NET_ASSERT_NE(ClientObject, nullptr);
+
+	// Set a replicated variable, but don't mark it dirty
+	ServerObject->IntA = 0xFF;
+
+	// Send and deliver packet
+	Server->UpdateAndSend({ Client });
+
+	// Client replicated property should not have changed
+	UE_NET_ASSERT_NE(ClientObject->IntA, ServerObject->IntA);
+
+	// Now mark the object dirty
+	Server->ReplicationSystem->ForceNetUpdate(ServerObject->NetRefHandle);
+
+	// Send and deliver packet
+	Server->UpdateAndSend({ Client });
+
+	// Client replicated propertyshould have changed now
+	UE_NET_ASSERT_EQ(ClientObject->IntA, ServerObject->IntA);
+
+	Server->DestroyObject(ServerObject);
+}
+
+/* This test exposes a problem when an object's PreUpdate calls dirty on a different object.
+UE_NET_TEST_FIXTURE(FReplicationSystemServerClientTestFixture, DirtyInsidePreUpdateTest)
+{
+	// Add client
+	FReplicationSystemTestClient* Client = CreateClient();
+
+	// Spawn object on server that is polled late in order to test ForceNetUpdate
+	UObjectReplicationBridge::FCreateNetRefHandleParams Params;
+	Params.PollFramePeriod = 10;
+	Params.bCanReceive = true;
+	Params.bAllowDynamicFilter = true;	
+	Params.bNeedsPreUpdate = true;
+	UTestReplicatedIrisObject* ServerObjectA = Server->CreateObject(Params);
+	UTestReplicatedIrisObject* ServerObjectB = Server->CreateObject(Params);
+
+	// Send and deliver packet
+	Server->UpdateAndSend({ Client });
+
+	// Object should have been created on the client
+	UTestReplicatedIrisObject* ClientObjectA = Cast<UTestReplicatedIrisObject>(Client->GetReplicationBridge()->GetReplicatedObject(ServerObjectA->NetRefHandle));
+	UE_NET_ASSERT_NE(ClientObjectA, nullptr);
+	UTestReplicatedIrisObject* ClientObjectB = Cast<UTestReplicatedIrisObject>(Client->GetReplicationBridge()->GetReplicatedObject(ServerObjectB->NetRefHandle));
+	UE_NET_ASSERT_NE(ClientObjectB, nullptr);
+	
+
+	// Send and deliver packet twice
+	Server->UpdateAndSend({ Client });
+	Server->UpdateAndSend({ Client });
+
+	// Set a replicated variable, but don't mark it dirty
+	ServerObjectA->IntA = 0xFF;
+
+	// Send and deliver packet
+	Server->UpdateAndSend({ Client });
+
+	// Client replicated property should not have changed
+	UE_NET_ASSERT_NE(ClientObjectA->IntA, ServerObjectA->IntA);
+
+	// Set a replicated variable, but don't mark it dirty
+	ServerObjectB->IntA = 0xFF;
+
+	// Send and deliver packet
+	Server->UpdateAndSend({ Client });
+
+	// Client replicated property should not have changed
+	UE_NET_ASSERT_NE(ClientObjectA->IntA, ServerObjectA->IntA);
+	UE_NET_ASSERT_NE(ClientObjectB->IntA, ServerObjectB->IntA);
+	
+	// Make the first object dirty
+	Server->ReplicationSystem->ForceNetUpdate(ServerObjectA->NetRefHandle);
+
+	auto PreUpdate = [&](FNetRefHandle NetHandle, UObject* ReplicatedObject, const UReplicationBridge* ReplicationBridge)
+	{
+		// When ObjectA is updated, make ObjectB dirty
+		if (ServerObjectA == ReplicatedObject)
+		{
+			// This will cause an ensure now
+			Server->ReplicationSystem->MarkDirty(ServerObjectB->NetRefHandle);
+		}
+	};
+
+	// Now add a dependency where the poll of the first object makes the second one dirty
+	Server->GetReplicationBridge()->SetExternalPreUpdateFunctor(PreUpdate);
+
+	// Send and deliver packet
+	Server->UpdateAndSend({ Client });
+
+	// Client property should have changed now
+	UE_NET_ASSERT_EQ(ClientObjectA->IntA, ServerObjectA->IntA);
+
+	// This fails because MarkDirty of a different object in PreUpdate is ignored and flushed.
+	UE_NET_ASSERT_EQ(ClientObjectB->IntA, ServerObjectB->IntA);
+
+	Server->DestroyObject(ServerObjectA);
+	Server->DestroyObject(ServerObjectB);
+}*/
+
+
 }
diff --git a/Engine/Plugins/Runtime/ReplicationSystemTestPlugin/Source/Private/Tests/ReplicationSystem/TestObjectPollFrequencyLimiter.cpp b/Engine/Plugins/Runtime/ReplicationSystemTestPlugin/Source/Private/Tests/ReplicationSystem/TestObjectPollFrequencyLimiter.cpp
--- a/Engine/Plugins/Runtime/ReplicationSystemTestPlugin/Source/Private/Tests/ReplicationSystem/TestObjectPollFrequencyLimiter.cpp
+++ b/Engine/Plugins/Runtime/ReplicationSystemTestPlugin/Source/Private/Tests/ReplicationSystem/TestObjectPollFrequencyLimiter.cpp
@@ -1,9 +1,13 @@
 // Copyright Epic Games, Inc. All Rights Reserved.
 
-#include "NetworkAutomationTest.h"
-#include "NetworkAutomationTestMacros.h"
+#include "ReplicationSystemServerClientTestFixture.h"
+
 #include "Net/Core/NetBitArray.h"
+
+#include "Iris/ReplicationSystem/ReplicationSystem.h"
+#include "Iris/ReplicationSystem/ReplicationSystemInternal.h"
 #include "Iris/ReplicationSystem/ObjectPollFrequencyLimiter.h"
+
 #include "Math/UnrealMathUtility.h"
 
 namespace UE::Net::Private
@@ -105,4 +109,121 @@ UE_NET_TEST(ObjectPollFrequencyLimiter, ObjectsToPollAreSpreadOutRegardlessOfUpd
 	}
 }
 
+
+UE_NET_TEST_FIXTURE(FReplicationSystemServerClientTestFixture, TestSubObjectsPollingAndForceNetUpdate)
+{
+	UReplicationSystem* ReplicationSystem = Server->ReplicationSystem;
+	UReplicatedTestObjectBridge* Bridge = Server->GetReplicationBridge();
+
+	// Add a client
+	FReplicationSystemTestClient* Client = CreateClient();
+
+	// Spawn object on server polled every 3 frames
+	const uint32 PollPeriod = 2;
+	UObjectReplicationBridge::FCreateNetRefHandleParams Params;
+	Params.PollFramePeriod = PollPeriod;
+	Params.bCanReceive = true;
+	Params.bAllowDynamicFilter = true;
+	Params.bNeedsPreUpdate = true;
+	UTestReplicatedIrisObject* ServerRootObject = Server->CreateObject(Params);
+
+	// Spawn subobject
+	UTestReplicatedIrisObject* ServerSubObject = Server->CreateSubObject(ServerRootObject->NetRefHandle, 0, 0);
+
+	// Send and deliver packet
+	Server->UpdateAndSend({ Client });
+
+	// Root Object should be created on the client
+	UTestReplicatedIrisObject* ClientRootObject = Cast<UTestReplicatedIrisObject>(Client->GetReplicationBridge()->GetReplicatedObject(ServerRootObject->NetRefHandle));
+	UE_NET_ASSERT_NE(ClientRootObject, nullptr);
+
+	// Subobject should be created on the client
+	UTestReplicatedIrisObject* ClientSubObject = Cast<UTestReplicatedIrisObject>(Client->GetReplicationBridge()->GetReplicatedObject(ServerSubObject->NetRefHandle));
+	UE_NET_ASSERT_NE(ClientSubObject, nullptr);
+
+	// Dirty a root object property
+	ServerRootObject->IntA = 0x0A;
+
+	// Update for maximum poll period or when property is received
+	auto UpdateUntilPolled = [&](UTestReplicatedIrisObject* InServerObject, UTestReplicatedIrisObject* InClientObject) -> uint32
+	{
+		uint32 FramesPolled = 0;
+		for (; FramesPolled <= PollPeriod; ++FramesPolled)
+		{
+			// Send and deliver packet
+			Server->UpdateAndSend({ Client });
+
+			if (InServerObject->IntA == InClientObject->IntA)
+			{
+				return FramesPolled+1;
+			}
+		}
+		return FramesPolled;
+	};
+
+	// Test polling on root object
+	{
+		// Update until property is received
+		uint32 FramesPolled = UpdateUntilPolled(ServerRootObject, ClientRootObject);
+
+		// Root object should have received property
+		UE_NET_ASSERT_EQ(ServerRootObject->IntA, ClientRootObject->IntA);
+
+		// We should have needed the max poll period
+		UE_NET_ASSERT_EQ(FramesPolled, PollPeriod + 1);
+	}
+
+	// Test polling on sub object
+	{
+		// Dirty a sub object property
+		ServerSubObject->IntA = 0x0B;
+
+		// Update until property is received
+		uint32 FramesPolled = UpdateUntilPolled(ServerSubObject, ClientSubObject);
+
+		// Subobject should have received property change
+		UE_NET_ASSERT_EQ(ServerSubObject->IntA, ClientSubObject->IntA);
+
+		// We should have needed the max poll period
+		UE_NET_ASSERT_EQ(FramesPolled, PollPeriod + 1);
+	}	
+
+	// Test force update on sub object
+	{
+		// Dirty a sub object property
+		ServerSubObject->IntA = 0x0C;
+
+		// Force update it
+		ReplicationSystem->ForceNetUpdate(ServerSubObject->NetRefHandle);
+
+		// Update until property is received
+		uint32 FramesPolled = UpdateUntilPolled(ServerSubObject, ClientSubObject);
+
+		// Subobject should have received property change
+		UE_NET_ASSERT_EQ(ServerSubObject->IntA, ClientSubObject->IntA);
+
+		// We should have needed only 1 update 
+		UE_NET_ASSERT_EQ(FramesPolled, 1U);
+	}
+
+	// Test force update on root object
+	{
+		// Dirty a sub object property
+		ServerSubObject->IntA = 0x0D;
+
+		// Force update it
+		ReplicationSystem->ForceNetUpdate(ServerRootObject->NetRefHandle);
+
+		// Update until property is received
+		uint32 FramesPolled = UpdateUntilPolled(ServerSubObject, ClientSubObject);
+
+		// Subobject should have received property change
+		UE_NET_ASSERT_EQ(ServerSubObject->IntA, ClientSubObject->IntA);
+
+		// We should have needed only 1 update 
+		UE_NET_ASSERT_EQ(FramesPolled, 1U);
+	}
+}
+
+
 }
diff --git a/Engine/Source/Runtime/Engine/Private/NetDriver.cpp b/Engine/Source/Runtime/Engine/Private/NetDriver.cpp
--- a/Engine/Source/Runtime/Engine/Private/NetDriver.cpp
+++ b/Engine/Source/Runtime/Engine/Private/NetDriver.cpp
@@ -3653,7 +3653,7 @@ void UNetDriver::ForceNetUpdate(AActor* Actor)
 		if (UObjectReplicationBridge* Bridge = ReplicationSystem->GetReplicationBridgeAs<UObjectReplicationBridge>())
 		{
 			const UE::Net::FNetRefHandle ActorRefHandle = Bridge->GetReplicatedRefHandle(Actor);
-			ReplicationSystem->MarkDirty(ActorRefHandle);
+			ReplicationSystem->ForceNetUpdate(ActorRefHandle);
 		}
 		return;
 	}
diff --git a/Engine/Source/Runtime/Experimental/Iris/Core/Private/Iris/ReplicationSystem/DirtyNetObjectTracker.cpp b/Engine/Source/Runtime/Experimental/Iris/Core/Private/Iris/ReplicationSystem/DirtyNetObjectTracker.cpp
--- a/Engine/Source/Runtime/Experimental/Iris/Core/Private/Iris/ReplicationSystem/DirtyNetObjectTracker.cpp
+++ b/Engine/Source/Runtime/Experimental/Iris/Core/Private/Iris/ReplicationSystem/DirtyNetObjectTracker.cpp
@@ -30,6 +30,9 @@ namespace UE::Net::Private
 FDirtyNetObjectTracker::FDirtyNetObjectTracker()
 : ReplicationSystemId(InvalidReplicationSystemId)
 {
+#if UE_NET_IRIS_VALIDATE_POLLED_OBJECT
+	CurrentPolledObject = FNetRefHandleManager::InvalidInternalIndex;
+#endif
 }
 
 FDirtyNetObjectTracker::~FDirtyNetObjectTracker()
@@ -60,6 +63,7 @@ void FDirtyNetObjectTracker::Init(const FDirtyNetObjectTrackerInitParams& Params
 	FMemory::Memzero(DirtyNetObjectContainer, DirtyNetObjectWordCount * sizeof(StorageType));
 
 	AccumulatedDirtyNetObjects.Init(NetObjectIdCount);
+	ForceNetUpdateObjects.Init(NetObjectIdCount);
 
 	AllowExternalAccess();
 
@@ -126,7 +130,7 @@ void FDirtyNetObjectTracker::UpdateAccumulatedDirtyList()
 	MakeNetBitArrayView(AccumulatedDirtyNetObjects).Combine(DirtyObjectsThisFrame, FNetBitArrayView::OrOp);
 }
 
-void FDirtyNetObjectTracker::MarkNetObjectDirty(uint32 NetObjectIndex)
+void FDirtyNetObjectTracker::MarkNetObjectDirty(FInternalNetRefIndex NetObjectIndex)
 {
 #if UE_NET_THREAD_SAFETY_CHECK
 	checkf(bIsExternalAccessAllowed, TEXT("Cannot mark objects dirty while the bitarray is locked for modifications."));
@@ -134,6 +138,12 @@ void FDirtyNetObjectTracker::MarkNetObjectDirty(uint32 NetObjectIndex)
 
 	if ((NetObjectIndex >= NetObjectIdRangeStart) & (NetObjectIndex <= NetObjectIdRangeEnd))
 	{
+#if UE_NET_IRIS_VALIDATE_POLLED_OBJECT
+		ensureMsgf(FNetRefHandleManager::InvalidInternalIndex == CurrentPolledObject || CurrentPolledObject == NetObjectIndex, 
+			TEXT("While calling NetUpdate on %s, a different replicated object %s was dirtied. Only the updated object can be dirtied for now."),
+			*GetNameSafe(NetRefHandleManager->GetReplicatedObjectInstance(CurrentPolledObject)), *GetNameSafe(NetRefHandleManager->GetReplicatedObjectInstance(NetObjectIndex)));
+#endif
+
 		const uint32 BitOffset = NetObjectIndex;
 		const StorageType BitMask = StorageType(1) << (BitOffset & (StorageTypeBitCount - 1));
 
@@ -144,6 +154,16 @@ void FDirtyNetObjectTracker::MarkNetObjectDirty(uint32 NetObjectIndex)
 	}
 }
 
+void FDirtyNetObjectTracker::ForceNetUpdate(FInternalNetRefIndex NetObjectIndex)
+{
+	ForceNetUpdateObjects.SetBit(NetObjectIndex);
+
+	// Flag the object dirty so we update his filters too
+	MarkNetObjectDirty(NetObjectIndex);
+
+	UE_LOG_DIRTYOBJECTTRACKER(TEXT("FDirtyNetObjectTracker::ForceNetUpdateObjects %u ( InternalIndex: %u )"), ReplicationSystemId, NetObjectIndex);
+}
+
 void FDirtyNetObjectTracker::LockExternalAccess()
 {
 #if UE_NET_THREAD_SAFETY_CHECK
@@ -182,12 +202,18 @@ void FDirtyNetObjectTracker::ClearDirtyNetObjects(const FNetBitArrayView& CleanN
 	// Clear the current frame dirty objects
 	FMemory::Memzero(DirtyNetObjectContainer, DirtyNetObjectWordCount*sizeof(StorageType));
 
+	ForceNetUpdateObjects.Reset();
+
+#if UE_NET_IRIS_VALIDATE_POLLED_OBJECT
+	CurrentPolledObject = 0;
+#endif
+
 	AllowExternalAccess();
 
 	std::atomic_thread_fence(std::memory_order_seq_cst);
 }
 
-void MarkNetObjectStateDirty(uint32 ReplicationSystemId, uint32 NetObjectIndex)
+void MarkNetObjectStateDirty(uint32 ReplicationSystemId, FInternalNetRefIndex NetObjectIndex)
 {
 	if (UReplicationSystem* ReplicationSystem = GetReplicationSystem(ReplicationSystemId))
 	{
@@ -196,4 +222,13 @@ void MarkNetObjectStateDirty(uint32 ReplicationSystemId, uint32 NetObjectIndex)
 	}
 }
 
+void ForceNetUpdate(uint32 ReplicationSystemId, FInternalNetRefIndex NetObjectIndex)
+{
+	if (UReplicationSystem* ReplicationSystem = GetReplicationSystem(ReplicationSystemId))
+	{
+		FDirtyNetObjectTracker& DirtyNetObjectTracker = ReplicationSystem->GetReplicationSystemInternal()->GetDirtyNetObjectTracker();
+		DirtyNetObjectTracker.ForceNetUpdate(NetObjectIndex);
+	}
+}
+
 }
diff --git a/Engine/Source/Runtime/Experimental/Iris/Core/Private/Iris/ReplicationSystem/DirtyNetObjectTracker.h b/Engine/Source/Runtime/Experimental/Iris/Core/Private/Iris/ReplicationSystem/DirtyNetObjectTracker.h
--- a/Engine/Source/Runtime/Experimental/Iris/Core/Private/Iris/ReplicationSystem/DirtyNetObjectTracker.h
+++ b/Engine/Source/Runtime/Experimental/Iris/Core/Private/Iris/ReplicationSystem/DirtyNetObjectTracker.h
@@ -11,12 +11,20 @@ namespace UE::Net::Private
 {
 	class FNetRefHandleManager;
 	class FDirtyObjectsAccessor;
+	
+	typedef uint32 FInternalNetRefIndex;
 }
 
+// When enabled will detect if a polled object is dirtying another object. This is an unsupported behavior
+#ifndef UE_NET_IRIS_VALIDATE_POLLED_OBJECT
+#define UE_NET_IRIS_VALIDATE_POLLED_OBJECT !UE_BUILD_SHIPPING
+#endif
+
 namespace UE::Net::Private
 {
 
-IRISCORE_API void MarkNetObjectStateDirty(uint32 ReplicationSystemId, uint32 NetObjectIndex);
+IRISCORE_API void MarkNetObjectStateDirty(uint32 ReplicationSystemId, FInternalNetRefIndex NetObjectIndex);
+IRISCORE_API void ForceNetUpdate(uint32 ReplicationSystemId, FInternalNetRefIndex NetObjectIndex);
 
 struct FDirtyNetObjectTrackerInitParams
 {
@@ -50,11 +58,19 @@ public:
 	/** Reset the global and local dirty objects lists for those objects that are now clean */
 	void ClearDirtyNetObjects(const FNetBitArrayView& CleanNetObjects);
 
+	/** Track which object is currently being polled */
+	inline void SetCurrentPolledObject(FInternalNetRefIndex PolledObject);
+
 	/** Returns the list of objects that are dirty this frame or were dirty in previous frames but not cleaned up at that time. */
 	const FNetBitArrayView GetAccumulatedDirtyNetObjects() const { return MakeNetBitArrayView(AccumulatedDirtyNetObjects); }
 
+	/** Returns the list of objects who asked to force a replication this frame */
+	FNetBitArrayView GetForceNetUpdateObjects() { return MakeNetBitArrayView(ForceNetUpdateObjects); }
+	const FNetBitArrayView GetForceNetUpdateObjects() const { return MakeNetBitArrayView(ForceNetUpdateObjects); }
+
 private:
-	friend IRISCORE_API void MarkNetObjectStateDirty(uint32 ReplicationSystemId, uint32 NetObjectIndex);
+	friend IRISCORE_API void MarkNetObjectStateDirty(uint32 ReplicationSystemId, FInternalNetRefIndex NetObjectIndex);
+	friend IRISCORE_API void ForceNetUpdate(uint32 ReplicationSystemId, FInternalNetRefIndex NetObjectIndex);
 
 	friend FDirtyObjectsAccessor;
 
@@ -62,7 +78,8 @@ private:
 	static constexpr uint32 StorageTypeBitCount = FNetBitArrayView::WordBitCount;
 
 	void Deinit();
-	void MarkNetObjectDirty(uint32 NetObjectIndex);
+	void MarkNetObjectDirty(FInternalNetRefIndex NetObjectIndex);
+	void ForceNetUpdate(FInternalNetRefIndex NetObjectIndex);
 
 	/** Can only be accessed via FDirtyObjectsAccessor */
 	FNetBitArrayView GetDirtyNetObjectsThisFrame();
@@ -72,6 +89,9 @@ private:
 	// Dirty objects that persist across frames.
 	FNetBitArray AccumulatedDirtyNetObjects;
 
+    // Objects that want to force a replication this frame
+	FNetBitArray ForceNetUpdateObjects;
+
 	// List of objects set to be dirty this frame. Is always reset at the end of the net tick flush
 	StorageType* DirtyNetObjectContainer = nullptr;
 
@@ -79,6 +99,10 @@ private:
 	
 	FGlobalDirtyNetObjectTracker::FPollHandle GlobalDirtyTrackerPollHandle;
 
+#if UE_NET_IRIS_VALIDATE_POLLED_OBJECT
+	FInternalNetRefIndex CurrentPolledObject;
+#endif
+
 	uint32 ReplicationSystemId;
 
 	uint32 DirtyNetObjectWordCount = 0;
@@ -93,6 +117,14 @@ private:
 #endif
 };
 
+
+void FDirtyNetObjectTracker::SetCurrentPolledObject(FInternalNetRefIndex PolledObject)
+{ 
+#if UE_NET_IRIS_VALIDATE_POLLED_OBJECT
+	CurrentPolledObject = PolledObject; 
+#endif
+}
+
 /**
  * Gives access to the list of dirty objects while detecting non-thread safe access to it.
  */
diff --git a/Engine/Source/Runtime/Experimental/Iris/Core/Private/Iris/ReplicationSystem/NetRefHandleManager.h b/Engine/Source/Runtime/Experimental/Iris/Core/Private/Iris/ReplicationSystem/NetRefHandleManager.h
--- a/Engine/Source/Runtime/Experimental/Iris/Core/Private/Iris/ReplicationSystem/NetRefHandleManager.h
+++ b/Engine/Source/Runtime/Experimental/Iris/Core/Private/Iris/ReplicationSystem/NetRefHandleManager.h
@@ -202,7 +202,8 @@ public:
 	const FNetBitArray& GetAssignedInternalIndices() const { return AssignedInternalIndices; }
 
 	// SubObjects
-	const FNetBitArray& GetSubObjectInternalIndices() const { return SubObjectInternalIndices; }	
+	const FNetBitArray& GetSubObjectInternalIndices() const { return SubObjectInternalIndices; }
+	const FNetBitArrayView GetSubObjectInternalIndicesView() const { return MakeNetBitArrayView(SubObjectInternalIndices); }
 
 	bool AddSubObject(FNetRefHandle OwnerHandle, FNetRefHandle SubObjectHandle, FNetRefHandle RelativeOtherSubObjectHandle, EAddSubObjectFlags Flags = EAddSubObjectFlags::Default);
 	bool AddSubObject(FNetRefHandle OwnerHandle, FNetRefHandle SubObjectHandle, EAddSubObjectFlags Flags = EAddSubObjectFlags::Default);
diff --git a/Engine/Source/Runtime/Experimental/Iris/Core/Private/Iris/ReplicationSystem/ObjectReplicationBridge.cpp b/Engine/Source/Runtime/Experimental/Iris/Core/Private/Iris/ReplicationSystem/ObjectReplicationBridge.cpp
--- a/Engine/Source/Runtime/Experimental/Iris/Core/Private/Iris/ReplicationSystem/ObjectReplicationBridge.cpp
+++ b/Engine/Source/Runtime/Experimental/Iris/Core/Private/Iris/ReplicationSystem/ObjectReplicationBridge.cpp
@@ -8,6 +8,7 @@
 #include "Iris/Core/IrisLog.h"
 #include "Iris/Core/IrisMemoryTracker.h"
 #include "Iris/Core/IrisProfiler.h"
+#include "Iris/Core/IrisDebugging.h"
 
 #include "Net/Core/NetBitArrayPrinter.h"
 #include "Net/Core/NetHandle/NetHandleManager.h"
@@ -64,6 +65,13 @@ static FAutoConsoleVariableRef CVarEnableFilterMappings(
 		TEXT("Whether we honor filter mappings set in ObjectReplicationBridgeConfig. If filter mappings are enabled then objects may also be assigned the default spatial filter even if there aren't any specific mappings. Default is true.")
 		);
 
+static bool bEnableForceNetUpdate = false;
+static FAutoConsoleVariableRef CVarEnableForceNetUpdate(
+	TEXT("net.Iris.EnableForceNetUpdate"),
+	bEnableForceNetUpdate,
+	TEXT("When true the system only allows ForceNetUpdate to skip the poll frequency of objects. When false any MarkDirty object will be immediately polled.")
+);
+
 UObjectReplicationBridge::FCreateNetRefHandleParams UObjectReplicationBridge::DefaultCreateNetRefHandleParams =
 {
 	false, // bCanReceive
@@ -213,6 +221,8 @@ UE::Net::FNetRefHandle UObjectReplicationBridge::BeginReplication(UObject* Insta
 		return AllocatedRefHandle;
 	}
 
+	IRIS_PROFILER_SCOPE(BeginReplication);
+	
 	// Register fragments
 	EReplicationFragmentTraits Traits = EReplicationFragmentTraits::CanReplicate;
 	Traits |= Params.bCanReceive ? EReplicationFragmentTraits::CanReceive : EReplicationFragmentTraits::None;
@@ -273,6 +283,8 @@ UE::Net::FNetRefHandle UObjectReplicationBridge::BeginReplication(UObject* Insta
 		
 	if (ReplicationProtocol)
 	{
+		IRIS_PROFILER_PROTOCOL_NAME(ReplicationProtocol->DebugName->Name);		
+
 		// Create NetHandle and bind instance
 		FNetHandle NetHandle = FNetHandleManager::GetOrCreateNetHandle(Instance);
 		FNetRefHandle RefHandle = InternalCreateNetObject(AllocatedRefHandle, NetHandle, ReplicationProtocol);
@@ -627,16 +639,12 @@ void UObjectReplicationBridge::PostApplyInitialState(FNetRefHandle Handle)
 
 void UObjectReplicationBridge::PreSendUpdateSingleHandle(FNetRefHandle RefHandle)
 {
-	PreUpdateAndPollImpl(RefHandle);
+	ForcePollObject(RefHandle);
 }
 
 void UObjectReplicationBridge::PreSendUpdate()
 {
-	IRIS_PROFILER_SCOPE(UObjectReplicationBridge_OnPreSendUpdate);
-
-	// Invalid/default handle means update all objects
-	FNetRefHandle RefHandle;
-	PreUpdateAndPollImpl(RefHandle);
+	PreUpdateAndPoll();
 }
 
 
@@ -698,24 +706,30 @@ void UObjectReplicationBridge::SetInstanceGetWorldObjectInfoFunction(FInstanceGe
 	GetInstanceWorldObjectInfoFunction = InGetWorldObjectInfoFunction;
 }
 
-void UObjectReplicationBridge::PreUpdateAndPollImpl(FNetRefHandle Handle)
+void UObjectReplicationBridge::ForcePollObject(FNetRefHandle Handle)
 {
-	using namespace UE::Net;
 	using namespace UE::Net::Private;
 
-	// If we are asked to poll a single object
 	if (Handle.IsValid())
 	{
+		IRIS_PROFILER_SCOPE(UObjectReplicationBridge_ForcePollObject);
+
 		FObjectPoller::FInitParams PollerInitParams;
 		PollerInitParams.ObjectReplicationBridge = this;
 		PollerInitParams.ReplicationSystemInternal = GetReplicationSystem()->GetReplicationSystemInternal();
 		FObjectPoller Poller(PollerInitParams);
 
 		Poller.PollSingleObject(Handle);
-
-		return;
 	}
-	
+}
+
+void UObjectReplicationBridge::PreUpdateAndPoll()
+{
+	using namespace UE::Net;
+	using namespace UE::Net::Private;
+
+	IRIS_PROFILER_SCOPE(UObjectReplicationBridge_PreUpdateAndPoll);
+
 	// Update every relevant objects from here
 	FReplicationSystemInternal* ReplicationSystemInternal = GetReplicationSystem()->GetReplicationSystemInternal();
 	const FNetRefHandleManager& LocalNetRefHandleManager = ReplicationSystemInternal->GetNetRefHandleManager();
@@ -731,7 +745,15 @@ void UObjectReplicationBridge::PreUpdateAndPollImpl(FNetRefHandle Handle)
 
 	if (bUseFrequencyBasedPolling)
 	{
-		PollFrequencyLimiter->Update(RelevantObjects, AccumulatedDirtyObjects, ObjectsConsideredForPollingView);
+		if (bEnableForceNetUpdate)
+		{
+			const FNetBitArrayView ForceNetUpdateObjects = ReplicationSystemInternal->GetDirtyNetObjectTracker().GetForceNetUpdateObjects();
+			PollFrequencyLimiter->Update(RelevantObjects, ForceNetUpdateObjects, ObjectsConsideredForPollingView);
+		}
+		else
+		{
+			PollFrequencyLimiter->Update(RelevantObjects, AccumulatedDirtyObjects, ObjectsConsideredForPollingView);
+		}
 	}
 	else
 	{
@@ -741,48 +763,45 @@ void UObjectReplicationBridge::PreUpdateAndPollImpl(FNetRefHandle Handle)
 	// Mask off objects pending dormancy as we do not want to poll/pre-update them unless they are marked for flush or are dirty
 	if (bUseDormancyToFilterPolling)
 	{
-		IRIS_PROFILER_SCOPE(PreUpdateAndPollImpl_Dormancy);
+		IRIS_PROFILER_SCOPE(PreUpdateAndPoll_Dormancy);
 
 		// Mask off objects pending dormancy that are not dirty
 		ObjectsConsideredForPollingView.CombineMultiple(FNetBitArrayView::AndNotOp, WantToBeDormantObjects, FNetBitArrayView::AndNotOp, AccumulatedDirtyObjects);
 
-		// Add objects that have requested to flush dormancy
+		FNetBitArrayView ForceNetUpdateObjects = ReplicationSystemInternal->GetDirtyNetObjectTracker().GetForceNetUpdateObjects();
+
+		// Force poll objects that have requested to flush dormancy
 		for (FNetRefHandle HandlePendingFlush : MakeArrayView(DormantHandlesPendingFlush))
 		{
 			if (const uint32 InternalObjectIndex = LocalNetRefHandleManager.GetInternalIndex(HandlePendingFlush))
 			{
-				if (ObjectsConsideredForPolling.GetBit(InternalObjectIndex))
+				if (!ObjectsConsideredForPolling.IsBitSet(InternalObjectIndex))
 				{
-					continue;
-				}
-
-				ObjectsConsideredForPolling.SetBit(InternalObjectIndex);
-				for (const FInternalNetRefIndex SubObjectIndex : LocalNetRefHandleManager.GetSubObjects(InternalObjectIndex))
-				{
-					ObjectsConsideredForPolling.SetBit(SubObjectIndex);
+					ObjectsConsideredForPolling.SetBit(InternalObjectIndex);
+					ForceNetUpdateObjects.SetBit(InternalObjectIndex);
 				}
 			}
 		}
 
-		//TODO: Add size of DormantHandlesPendingFlush array to NetStats ?
+		UE_NET_TRACE_FRAME_STATSCOUNTER(ReplicationSystem->GetId(), FlushDormancyObjectCount, DormantHandlesPendingFlush.Num(), ENetTraceVerbosity::Trace);
 	}
 	DormantHandlesPendingFlush.Reset();
 
 	/**
 	* Make sure to propagate polling for owners to subobjects and vice versa. If an actor is not due to update due to
-	* polling frequency it can still be dirty or a dormant object marked for flush and polled for that reason. In order to make sure all recent state updates
+	* polling frequency it can still be force net update or a dormant object marked for flush and polled for that reason. In order to make sure all recent state updates
 	* are replicated atomically this polling propagation is required.
 	*/
 	{
-		IRIS_PROFILER_SCOPE(PreUpdateAndPollImpl_PropagateDirtyness);
+		IRIS_PROFILER_SCOPE(PreUpdateAndPoll_PropagatePolling);
 
-		auto PropagateSubObjectDirtinessToOwner = [&LocalNetRefHandleManager, &ObjectsConsideredForPolling](uint32 InternalObjectIndex)
+		auto PropagateSubObjectNetForceUpdateToOwner = [&LocalNetRefHandleManager, &ObjectsConsideredForPolling](uint32 InternalObjectIndex)
 		{
 			const FNetRefHandleManager::FReplicatedObjectData& ObjectData = LocalNetRefHandleManager.GetReplicatedObjectDataNoCheck(InternalObjectIndex);
 			ObjectsConsideredForPolling.SetBit(ObjectData.SubObjectRootIndex);
 		};
 
-		auto PropagateOwnerDirtinessToSubObjectsAndDependentObjects = [&LocalNetRefHandleManager, &ObjectsConsideredForPolling](uint32 InternalObjectIndex)
+		auto PropagateOwnerNetForceUpdateToSubObject = [&LocalNetRefHandleManager, &ObjectsConsideredForPolling](uint32 InternalObjectIndex)
 		{
 			const FNetRefHandleManager::FReplicatedObjectData& ObjectData = LocalNetRefHandleManager.GetReplicatedObjectDataNoCheck(InternalObjectIndex);
 			for (const FInternalNetRefIndex SubObjectInternalIndex : LocalNetRefHandleManager.GetSubObjects(InternalObjectIndex))
@@ -791,25 +810,41 @@ void UObjectReplicationBridge::PreUpdateAndPollImpl(FNetRefHandle Handle)
 			}
 		};
 
-		// Make a list of objects that were set to be dirty this frame but are also relevant
-		FNetBitArray DirtyAndRelevantObjects(RelevantObjects.GetNumBits());
-		FNetBitArrayView DirtyAndRelevantObjectsView = MakeNetBitArrayView(DirtyAndRelevantObjects);
+		// Update subobjects' owner first and owners' subobjects second. It's the only way to properly mark all groups of objects in two passes.
+		const FNetBitArrayView SubObjects = MakeNetBitArrayView(LocalNetRefHandleManager.GetSubObjectInternalIndices());
 
+		if (bEnableForceNetUpdate)
 		{
-			FDirtyObjectsAccessor DirtyObjectsAccessor(ReplicationSystemInternal->GetDirtyNetObjectTracker());	
-			FNetBitArrayView DirtyObjectsThisFrame = DirtyObjectsAccessor.GetDirtyNetObjects();
+			// Make a list of objects which forced an update and are also relevant
+			FNetBitArray ForceNetUpdateAndRelevantObjects(RelevantObjects.GetNumBits(), FNetBitArray::NoResetNoValidate);
+			FNetBitArrayView ForceNetUpdateAndRelevantObjectsView = MakeNetBitArrayView(ForceNetUpdateAndRelevantObjects, FNetBitArrayView::NoResetNoValidate);
+			
+			const FNetBitArrayView ForceNetUpdateObjects = ReplicationSystemInternal->GetDirtyNetObjectTracker().GetForceNetUpdateObjects();
+			ForceNetUpdateAndRelevantObjectsView.Set(RelevantObjects, FNetBitArray::AndOp, ForceNetUpdateObjects);
 
-			DirtyAndRelevantObjectsView.Set(RelevantObjects, FNetBitArray::AndOp, DirtyObjectsThisFrame);
+			FNetBitArrayView::ForAllSetBits(ForceNetUpdateAndRelevantObjectsView, SubObjects, FNetBitArray::AndOp, PropagateSubObjectNetForceUpdateToOwner);
+			FNetBitArrayView::ForAllSetBits(ForceNetUpdateAndRelevantObjectsView, SubObjects, FNetBitArray::AndNotOp, PropagateOwnerNetForceUpdateToSubObject);
 		}
+		else
+		{
+			// Make the list of objects which are dirty and are also relevant
+			FNetBitArray DirtyAndRelevantObjects(RelevantObjects.GetNumBits(), FNetBitArray::NoResetNoValidate);
+			FNetBitArrayView DirtyAndRelevantObjectsView = MakeNetBitArrayView(DirtyAndRelevantObjects, FNetBitArrayView::NoResetNoValidate);
 
-		// Update subobjects' owner first and owners' subobjects second. It's the only way to properly mark all groups of objects in two passes.
-		const FNetBitArrayView SubObjects = MakeNetBitArrayView(LocalNetRefHandleManager.GetSubObjectInternalIndices());
-		FNetBitArrayView::ForAllSetBits(DirtyAndRelevantObjectsView, SubObjects, FNetBitArray::AndOp, PropagateSubObjectDirtinessToOwner);
-		FNetBitArrayView::ForAllSetBits(DirtyAndRelevantObjectsView, SubObjects, FNetBitArray::AndNotOp, PropagateOwnerDirtinessToSubObjectsAndDependentObjects);
+			{
+				FDirtyObjectsAccessor DirtyObjectsAccessor(ReplicationSystemInternal->GetDirtyNetObjectTracker());
+				FNetBitArrayView DirtyObjectsThisFrame = DirtyObjectsAccessor.GetDirtyNetObjects();
+
+				DirtyAndRelevantObjectsView.Set(RelevantObjects, FNetBitArray::AndOp, DirtyObjectsThisFrame);
+			}
+
+			FNetBitArrayView::ForAllSetBits(DirtyAndRelevantObjectsView, SubObjects, FNetBitArray::AndOp, PropagateSubObjectNetForceUpdateToOwner);
+			FNetBitArrayView::ForAllSetBits(DirtyAndRelevantObjectsView, SubObjects, FNetBitArray::AndNotOp, PropagateOwnerNetForceUpdateToSubObject);
+		}
 			
-		// If an object with dependents is about to be polled, force his dependents to poll at the same time.
+		// If an object with dependents is about to be polled, force it's dependents to poll at the same time.
 		{
-			IRIS_PROFILER_SCOPE(PreUpdateAndPollImpl_PatchDependentObjects);
+			IRIS_PROFILER_SCOPE(PreUpdateAndPoll_PatchDependentObjects);
 
 			FNetBitArray TempObjectsConsideredForPolling(ObjectsConsideredForPolling);
 			FNetBitArray::ForAllSetBits(TempObjectsConsideredForPolling, NetRefHandleManager->GetObjectsWithDependentObjectsInternalIndices(), FNetBitArray::AndOp,
@@ -823,23 +858,24 @@ void UObjectReplicationBridge::PreUpdateAndPollImpl(FNetRefHandle Handle)
 			);
 		}
 	}
-
 	
-	IRIS_PROFILER_SCOPE(PreUpdateAndPollImpl_Poll);
+	{
+		IRIS_PROFILER_SCOPE(PreUpdateAndPoll_Poll);
 
-	FObjectPoller::FInitParams PollerInitParams;
-	PollerInitParams.ObjectReplicationBridge = this;
-	PollerInitParams.ReplicationSystemInternal = GetReplicationSystem()->GetReplicationSystemInternal();
+		FObjectPoller::FInitParams PollerInitParams;
+		PollerInitParams.ObjectReplicationBridge = this;
+		PollerInitParams.ReplicationSystemInternal = GetReplicationSystem()->GetReplicationSystemInternal();
 
-	FObjectPoller Poller(PollerInitParams);
-	Poller.PollObjects(ObjectsConsideredForPollingView);
+		FObjectPoller Poller(PollerInitParams);
+		Poller.PollObjects(ObjectsConsideredForPollingView);
 	
-	FObjectPoller::FPreUpdateAndPollStats Stats = Poller.GetPollStats();
+		FObjectPoller::FPreUpdateAndPollStats Stats = Poller.GetPollStats();
 
-	// Report stats
-	UE_NET_TRACE_FRAME_STATSCOUNTER(ReplicationSystem->GetId(), ReplicationSystem.PreUpdatedObjectCount, Stats.PreUpdatedObjectCount, ENetTraceVerbosity::Trace);
-	UE_NET_TRACE_FRAME_STATSCOUNTER(ReplicationSystem->GetId(), ReplicationSystem.PolledObjectCount, Stats.PolledObjectCount, ENetTraceVerbosity::Trace);
-	UE_NET_TRACE_FRAME_STATSCOUNTER(ReplicationSystem->GetId(), ReplicationSystem.PolledReferencesObjectCount, Stats.PolledReferencesObjectCount, ENetTraceVerbosity::Trace);
+		// Report stats
+		UE_NET_TRACE_FRAME_STATSCOUNTER(ReplicationSystem->GetId(), ReplicationSystem.PreUpdatedObjectCount, Stats.PreUpdatedObjectCount, ENetTraceVerbosity::Trace);
+		UE_NET_TRACE_FRAME_STATSCOUNTER(ReplicationSystem->GetId(), ReplicationSystem.PolledObjectCount, Stats.PolledObjectCount, ENetTraceVerbosity::Trace);
+		UE_NET_TRACE_FRAME_STATSCOUNTER(ReplicationSystem->GetId(), ReplicationSystem.PolledReferencesObjectCount, Stats.PolledReferencesObjectCount, ENetTraceVerbosity::Trace);
+	}
 }
 
 void UObjectReplicationBridge::UpdateInstancesWorldLocation()
@@ -857,7 +893,6 @@ void UObjectReplicationBridge::UpdateInstancesWorldLocation()
 	FWorldLocations& WorldLocations = ReplicationSystemInternal->GetWorldLocations();
 	const TArray<UObject*>& ReplicatedInstances = LocalNetRefHandleManager.GetReplicatedInstances();
 
-
 	// Retrieve the world location for instances that supports it. Only dirty objects are considered.
 	auto UpdateInstanceWorldLocation = [this, &ReplicatedInstances, &LocalNetRefHandleManager, &WorldLocations](uint32 InternalObjectIndex)
 	{
@@ -871,6 +906,7 @@ void UObjectReplicationBridge::UpdateInstancesWorldLocation()
 		}
 	};
 
+	// $IRIS TODO: This code assumes users are calling MarkDirty whenever an actor changes location. Need to add a location tracker to ensure if that is not the case.
 	FDirtyObjectsAccessor DirtyObjectsAccessor(ReplicationSystemInternal->GetDirtyNetObjectTracker());
 	const FNetBitArrayView DirtyObjectsThisFrame = DirtyObjectsAccessor.GetDirtyNetObjects();
 	DirtyObjectsThisFrame.ForAllSetBits(UpdateInstanceWorldLocation);
diff --git a/Engine/Source/Runtime/Experimental/Iris/Core/Private/Iris/ReplicationSystem/Polling/ObjectPoller.cpp b/Engine/Source/Runtime/Experimental/Iris/Core/Private/Iris/ReplicationSystem/Polling/ObjectPoller.cpp
--- a/Engine/Source/Runtime/Experimental/Iris/Core/Private/Iris/ReplicationSystem/Polling/ObjectPoller.cpp
+++ b/Engine/Source/Runtime/Experimental/Iris/Core/Private/Iris/ReplicationSystem/Polling/ObjectPoller.cpp
@@ -55,6 +55,8 @@ void FObjectPoller::PollObjects(const FNetBitArrayView& ObjectsConsideredForPoll
 
 	// Clear ref to locked dirty bit array
 	DirtyObjectsThisFrame = FNetBitArrayView();
+
+	ReplicationSystemInternal->GetDirtyNetObjectTracker().SetCurrentPolledObject(FNetRefHandleManager::InvalidInternalIndex);
 }
 
 void FObjectPoller::PollSingleObject(FNetRefHandle Handle)
@@ -123,13 +125,14 @@ void FObjectPoller::PushModelPollObject(FInternalNetRefIndex ObjectIndex)
 
 	IRIS_PROFILER_PROTOCOL_NAME(ObjectData.Protocol->DebugName->Name);
 
+	ReplicationSystemInternal->GetDirtyNetObjectTracker().SetCurrentPolledObject(ObjectIndex);
+
 	const FReplicationInstanceProtocol* InstanceProtocol = ObjectData.InstanceProtocol;
 
 	const EReplicationInstanceProtocolTraits InstanceTraits = InstanceProtocol->InstanceTraits;
 	const bool bNeedsPoll = EnumHasAnyFlags(InstanceTraits, EReplicationInstanceProtocolTraits::NeedsPoll);
 
-	//$IRIS TODO: should use the Accumulated array in case we were dirty while culled.  Also move this test to just after PreUpdate in case it gets dirty inside the call.
-	bool bIsDirtyObject = DirtyObjectsThisFrame.GetBit(ObjectIndex);
+	bool bIsDirtyObject = AccumulatedDirtyObjects.GetBit(ObjectIndex);
 
 	// Call per-instance PreUpdate function
 	if (ObjectReplicationBridge->PreUpdateInstanceFunction && EnumHasAnyFlags(InstanceTraits, EReplicationInstanceProtocolTraits::NeedsPreSendUpdate))
@@ -139,8 +142,9 @@ void FObjectPoller::PushModelPollObject(FInternalNetRefIndex ObjectIndex)
 		ObjectReplicationBridge->PreUpdateInstanceFunction(ObjectData.RefHandle, ReplicatedInstances[ObjectIndex], ObjectReplicationBridge);
 		++PollStats.PreUpdatedObjectCount;
 
-		// Pre update may dirty push based properties. Detect it.
-		if ((bNeedsPoll & !bIsDirtyObject))
+		// Pre update may have called MarkDirty. Detect it.
+		bIsDirtyObject = bIsDirtyObject || DirtyObjectsThisFrame.GetBit(ObjectIndex);
+		if (bNeedsPoll && !bIsDirtyObject)
 		{
 			bIsDirtyObject = FGlobalDirtyNetObjectTracker::IsNetObjectStateDirty(ObjectData.NetHandle);
 			if (bIsDirtyObject)
@@ -150,7 +154,6 @@ void FObjectPoller::PushModelPollObject(FInternalNetRefIndex ObjectIndex)
 		}
 	}
 
-	const bool bIsNewInScope = !PrevFrameGlobalScopeObjects.GetBit(ObjectIndex);
 	const bool bIsGCAffectedObject = GarbageCollectionAffectedObjects.GetBit(ObjectIndex);
 	GarbageCollectionAffectedObjects.ClearBit(ObjectIndex);
 
@@ -162,11 +165,15 @@ void FObjectPoller::PushModelPollObject(FInternalNetRefIndex ObjectIndex)
 
 	IRIS_PROFILER_SCOPE_VERBOSE(PollPushBased);
 
+	const bool bIsNewInScope = !PrevFrameGlobalScopeObjects.GetBit(ObjectIndex);
+
+	
+
 	// If the object is fully push model we only need to poll it if it's dirty, unless it's a new object or was garbage collected.
 	bool bWasMarkedDirty = false;
 	if (EnumHasAnyFlags(InstanceTraits, EReplicationInstanceProtocolTraits::HasFullPushBasedDirtiness))
 	{
-		if (bIsDirtyObject | bIsNewInScope)
+		if (bIsDirtyObject || bIsNewInScope)
 		{
 			// We need to do a poll if object is marked as dirty
 			EReplicationFragmentPollFlags PollOptions = EReplicationFragmentPollFlags::PollAllState;
@@ -210,7 +217,6 @@ void FObjectPoller::PushModelPollObject(FInternalNetRefIndex ObjectIndex)
 	{
 		DirtyObjectsThisFrame.SetBit(ObjectIndex);
 	}
-
 }
 
 } // end namespace UE::Net::Private
\ No newline at end of file
diff --git a/Engine/Source/Runtime/Experimental/Iris/Core/Private/Iris/ReplicationSystem/ReplicationBridge.cpp b/Engine/Source/Runtime/Experimental/Iris/Core/Private/Iris/ReplicationSystem/ReplicationBridge.cpp
--- a/Engine/Source/Runtime/Experimental/Iris/Core/Private/Iris/ReplicationSystem/ReplicationBridge.cpp
+++ b/Engine/Source/Runtime/Experimental/Iris/Core/Private/Iris/ReplicationSystem/ReplicationBridge.cpp
@@ -347,7 +347,7 @@ void UReplicationBridge::InternalAttachInstanceToNetRefHandle(FNetRefHandle RefH
 	if (bBindInstanceProtocol)
 	{
 		FReplicationInstanceOperationsInternal::BindInstanceProtocol(NetHandle, InstanceProtocol, NetRefHandleManager->GetReplicatedObjectDataNoCheck(InternalReplicationIndex).Protocol);
-		MarkNetObjectStateDirty(ReplicationSystemId, InternalReplicationIndex);
+		ForceNetUpdate(ReplicationSystemId, InternalReplicationIndex);
 	}
 }
 
@@ -464,8 +464,8 @@ void UReplicationBridge::InternalAddSubObject(FNetRefHandle OwnerHandle, FNetRef
 
 	if (NetRefHandleManager->AddSubObject(OwnerHandle, SubObjectHandle, InsertRelativeToSubObjectHandle, AddSubObjectFlags))
 	{
-		// If the subobject is new we need to mark it dirty to pick it up for replication with its new parent
-		MarkNetObjectStateDirty(ReplicationSystem->GetId(), NetRefHandleManager->GetInternalIndex(SubObjectHandle));
+		// If the subobject is new we need to update it immediately to pick it up for replication with its new parent
+		ForceNetUpdate(ReplicationSystem->GetId(), NetRefHandleManager->GetInternalIndex(SubObjectHandle));
 
 		// We set the priority of subobjects to be static as they will be prioritized with owner
 		ReplicationSystem->SetStaticPriority(SubObjectHandle, 1.0f);
diff --git a/Engine/Source/Runtime/Experimental/Iris/Core/Private/Iris/ReplicationSystem/ReplicationSystem.cpp b/Engine/Source/Runtime/Experimental/Iris/Core/Private/Iris/ReplicationSystem/ReplicationSystem.cpp
--- a/Engine/Source/Runtime/Experimental/Iris/Core/Private/Iris/ReplicationSystem/ReplicationSystem.cpp
+++ b/Engine/Source/Runtime/Experimental/Iris/Core/Private/Iris/ReplicationSystem/ReplicationSystem.cpp
@@ -1222,6 +1222,14 @@ void UReplicationSystem::TearOffNextUpdate(FNetRefHandle Handle)
 	Impl->ReplicationSystemInternal.GetReplicationBridge()->TearOff(Handle, DestroyFlags, false);
 }
 
+void UReplicationSystem::ForceNetUpdate(FNetRefHandle Handle)
+{
+	if (const uint32 InternalObjectIndex = Impl->ReplicationSystemInternal.GetNetRefHandleManager().GetInternalIndex(Handle))
+	{
+		UE::Net::Private::ForceNetUpdate(GetId(), InternalObjectIndex);
+	}
+}
+
 void UReplicationSystem::MarkDirty(FNetRefHandle Handle)
 {
 	if (const uint32 InternalObjectIndex = Impl->ReplicationSystemInternal.GetNetRefHandleManager().GetInternalIndex(Handle))
diff --git a/Engine/Source/Runtime/Experimental/Iris/Core/Public/Iris/ReplicationSystem/ObjectReplicationBridge.h b/Engine/Source/Runtime/Experimental/Iris/Core/Public/Iris/ReplicationSystem/ObjectReplicationBridge.h
--- a/Engine/Source/Runtime/Experimental/Iris/Core/Public/Iris/ReplicationSystem/ObjectReplicationBridge.h
+++ b/Engine/Source/Runtime/Experimental/Iris/Core/Public/Iris/ReplicationSystem/ObjectReplicationBridge.h
@@ -227,7 +227,11 @@ protected:
 
 private:
 
-	void PreUpdateAndPollImpl(FNetRefHandle RefHandle);
+	/** Forcibly poll a single replicated object */
+	void ForcePollObject(FNetRefHandle RefHandle);
+
+	/** Pre update and poll all relevant objects who hit their polling period or are force net update. */
+	void PreUpdateAndPoll();
 
 	/** Remove mapping between handle and object instance. */
 	void UnregisterInstance(FNetRefHandle RefHandle);
diff --git a/Engine/Source/Runtime/Experimental/Iris/Core/Public/Iris/ReplicationSystem/ReplicationSystem.h b/Engine/Source/Runtime/Experimental/Iris/Core/Public/Iris/ReplicationSystem/ReplicationSystem.h
--- a/Engine/Source/Runtime/Experimental/Iris/Core/Public/Iris/ReplicationSystem/ReplicationSystem.h
+++ b/Engine/Source/Runtime/Experimental/Iris/Core/Public/Iris/ReplicationSystem/ReplicationSystem.h
@@ -459,9 +459,16 @@ public:
 	IRISCORE_API void TearOffNextUpdate(FNetRefHandle Handle);
 
 	/**
-	 * Explicitly mark object as dirty. This will make it considered for replication this frame
-	 * if the object has any modified properties/data.
-	 * Normally an object is checked for dirtiness based on poll frequency.
+	 * Force the passed object to be considered for replication this frame.
+	 * This will also force it's sub objects, root object and any of it's dependents to also be considered for replication.
+	 * Normally an object is checked for replication only when it's poll frequency is hit.
+	 * @param Handle A valid handle to an object.
+	 * @see FObjectReplicationBridgePollConfig
+	 */
+	IRISCORE_API void ForceNetUpdate(FNetRefHandle Handle);
+
+	/**
+	 * Explicitly mark object as having dirty properties.
 	 * @param Handle A valid handle to an object.
 	 * @see FObjectReplicationBridgePollConfig
 	 */
