cluster unions should update kinematic properties on remove

diff --git a/Engine/Source/Runtime/Experimental/Chaos/Private/Chaos/ClusterUnionManager.cpp b/Engine/Source/Runtime/Experimental/Chaos/Private/Chaos/ClusterUnionManager.cpp
--- a/Engine/Source/Runtime/Experimental/Chaos/Private/Chaos/ClusterUnionManager.cpp
+++ b/Engine/Source/Runtime/Experimental/Chaos/Private/Chaos/ClusterUnionManager.cpp
@@ -593,7 +593,7 @@ namespace Chaos
 		// This gets changed because we go through an internal initialization route when we update the cluster union particle's properties.
 		Cluster->InternalCluster->SetPhysicsProxy(OldProxy);
 
-		constexpr EUpdateClusterUnionPropertiesFlags RemoveUpdateFlags = EUpdateClusterUnionPropertiesFlags::IncrementalGenerateConnectionGraph;
+		constexpr EUpdateClusterUnionPropertiesFlags RemoveUpdateFlags = EUpdateClusterUnionPropertiesFlags::IncrementalGenerateConnectionGraph | EUpdateClusterUnionPropertiesFlags::UpdateKinematicProperties;
 		switch (UpdateClusterPropertiesTiming)
 		{
 		case EClusterUnionOperationTiming::Immediate:
