Iris
- Added missing exports required if using IrisFastArraySerializer

diff --git a/Engine/Source/Runtime/Experimental/Iris/Core/Public/Iris/ReplicationState/Private/IrisFastArraySerializerInternal.h b/Engine/Source/Runtime/Experimental/Iris/Core/Public/Iris/ReplicationState/Private/IrisFastArraySerializerInternal.h
--- a/Engine/Source/Runtime/Experimental/Iris/Core/Public/Iris/ReplicationState/Private/IrisFastArraySerializerInternal.h
+++ b/Engine/Source/Runtime/Experimental/Iris/Core/Public/Iris/ReplicationState/Private/IrisFastArraySerializerInternal.h
@@ -20,25 +20,25 @@ namespace Private
 */
 struct FIrisFastArraySerializerPrivateAccessor
 {
-	static FNetBitArrayView GetChangeMask(FIrisFastArraySerializer& Array);
-	static FNetBitArrayView GetConditionalChangeMask(FIrisFastArraySerializer& Array);
+	IRISCORE_API static FNetBitArrayView GetChangeMask(FIrisFastArraySerializer& Array);
+	IRISCORE_API static FNetBitArrayView GetConditionalChangeMask(FIrisFastArraySerializer& Array);
 	static FReplicationStateHeader& GetReplicationStateHeader(FIrisFastArraySerializer& Array) { return Array.ReplicationStateHeader; }
 
 	/*
 	 * Mark array as dirty and notify DirtyObjectTracker that the object is dirty if it has not been done before
 	 */
-	static void MarkArrayDirty(FIrisFastArraySerializer& Array);
+	IRISCORE_API static void MarkArrayDirty(FIrisFastArraySerializer& Array);
 
 	/*
 	 * Mark array and all bits in changemask as dirty and notify DirtyObjectTracker that the object is dirty if it has not been done before
 	 */
-	static void MarkAllArrayItemsDirty(FIrisFastArraySerializer& Array, uint32 StartingIndex = 0U);
+	IRISCORE_API static void MarkAllArrayItemsDirty(FIrisFastArraySerializer& Array, uint32 StartingIndex = 0U);
 
 	/*
 	 * Mark array dirty and mark changemask bit used for array Item as dirty and notify DirtyObjectTracker that the object is dirty if it has not been done before
      * Currently the same changemask bit might be used to indicate dirtiness for multiple array items
 	 */
-	static void MarkArrayItemDirty(FIrisFastArraySerializer& Array, int32 Index);
+	IRISCORE_API static void MarkArrayItemDirty(FIrisFastArraySerializer& Array, int32 Index);
 };
 
 } // end namespace Private
