connectivity checks should still happen when releasing particles inside a cluster union - they just need to be released from the cluster union instead

diff --git a/Engine/Source/Runtime/Experimental/Chaos/Private/Chaos/ClusterUnionManager.cpp b/Engine/Source/Runtime/Experimental/Chaos/Private/Chaos/ClusterUnionManager.cpp
--- a/Engine/Source/Runtime/Experimental/Chaos/Private/Chaos/ClusterUnionManager.cpp
+++ b/Engine/Source/Runtime/Experimental/Chaos/Private/Chaos/ClusterUnionManager.cpp
@@ -584,6 +584,7 @@ namespace Chaos
 		for (int32 Index = ParticleIndicesToRemove.Num() - 1; Index >= 0; --Index)
 		{
 			const int32 ParticleIndex = ParticleIndicesToRemove[Index];
+			Cluster->ChildProperties.Remove(Cluster->ChildParticles[ParticleIndex]);
 			Cluster->ChildParticles.RemoveAtSwap(ParticleIndex);
 		}
 
diff --git a/Engine/Source/Runtime/Experimental/Chaos/Private/Chaos/PBDRigidClustering.cpp b/Engine/Source/Runtime/Experimental/Chaos/Private/Chaos/PBDRigidClustering.cpp
--- a/Engine/Source/Runtime/Experimental/Chaos/Private/Chaos/PBDRigidClustering.cpp
+++ b/Engine/Source/Runtime/Experimental/Chaos/Private/Chaos/PBDRigidClustering.cpp
@@ -604,7 +604,7 @@ namespace Chaos
 		}
 	}
 
-	TArray<FRigidClustering::FParticleIsland> FRigidClustering::FindIslandsInChildren(const FPBDRigidClusteredParticleHandle* ClusteredParticle)
+	TArray<FRigidClustering::FParticleIsland> FRigidClustering::FindIslandsInChildren(const FPBDRigidClusteredParticleHandle* ClusteredParticle, bool bTraverseInterclusterEdges)
 	{
 		const TArray<FPBDRigidParticleHandle*>& Children = MChildren[ClusteredParticle];
 		
@@ -634,7 +634,7 @@ namespace Chaos
 					Island.Add(ChildToProcess);
 					for (const TConnectivityEdge<FReal>& Edge : ChildToProcess->CastToClustered()->ConnectivityEdges())
 					{
-						if (IsInterclusterEdge(*ChildToProcess, Edge))
+						if (IsInterclusterEdge(*ChildToProcess, Edge) && !bTraverseInterclusterEdges)
 						{
 							continue;
 						}
@@ -1066,97 +1066,139 @@ namespace Chaos
 		// If we're breaking a geometry collection, we'll need to create internal clusters to parent the remaining particles.
 		// However, we do not need to do this if we're currently operating on a cluster union! Its remaining particles should stay
 		// attached to the cluster union because they can handle particles being dynamically added/removed.
-		const bool bIsClusterUnion = ClusterUnionManager.IsClusterUnionParticle(ClusteredParticle);
-		if (Children.Num() && !bIsClusterUnion)
+		const FClusterUnion* ParentClusterUnion = ClusterUnionManager.FindClusterUnionFromParticle(ClusteredParticle);
+		const bool bIsClusterUnion = ParentClusterUnion != nullptr;
+		if (Children.Num())
 		{
-			TArray<FParticleIsland> Islands = FindIslandsInChildren(ClusteredParticle);
+			TArray<FParticleIsland> Islands = FindIslandsInChildren(ClusteredParticle, bIsClusterUnion);
 
-			TArray<const FConnectivityEdge*> IslandInterclusterEdges;
-			IslandInterclusterEdges.Reserve(Islands.Num());
-
-			// By default, all these islands will all just start simulating independently. However, if bCheckForInterclusterEdgesOnRelease
-			// is true, we're going to want to check if any of these islands contain a particle with an intercluster edge that connects to
-			// something else! In that case, keep the particles in the island connected. For simplicity
-			// we'll use a kinematic target for each group of particles that needs to stay attached . This kinematic target
-			// will be driven by simulation relative to whatever it should be attached to.
-			//
-			// We need to assume that only the server will have the intercluster edges necessary to drive this kinematic target. Therefore, this
-			// information is authoritative on the PT on the server while is authoritative on the GT on the client(s) via GC replication. So we
-			// also need to make sure we store this information in a way that's replicatable to the client via GC replication.
-			EDamageEvaluationModel DamageEvaluationModel = EDamageEvaluationModel::StrainFromDamageThreshold;
-			const FGeometryCollectionPhysicsProxy* ConcreteGCProxy = GetConcreteProxy<FGeometryCollectionPhysicsProxy>(ClusteredParticle);
-			if (ConcreteGCProxy)
+			if (!bIsClusterUnion)
 			{
-				const FSimulationParameters& SimParams = ConcreteGCProxy->GetSimParameters();
-				DamageEvaluationModel = SimParams.DamageEvaluationModel;
-			}
+				TArray<const FConnectivityEdge*> IslandInterclusterEdges;
+				IslandInterclusterEdges.Reserve(Islands.Num());
+
+				// By default, all these islands will all just start simulating independently. However, if bCheckForInterclusterEdgesOnRelease
+				// is true, we're going to want to check if any of these islands contain a particle with an intercluster edge that connects to
+				// something else! In that case, keep the particles in the island connected. For simplicity
+				// we'll use a kinematic target for each group of particles that needs to stay attached . This kinematic target
+				// will be driven by simulation relative to whatever it should be attached to.
+				//
+				// We need to assume that only the server will have the intercluster edges necessary to drive this kinematic target. Therefore, this
+				// information is authoritative on the PT on the server while is authoritative on the GT on the client(s) via GC replication. So we
+				// also need to make sure we store this information in a way that's replicatable to the client via GC replication.
+				EDamageEvaluationModel DamageEvaluationModel = EDamageEvaluationModel::StrainFromDamageThreshold;
+				const FGeometryCollectionPhysicsProxy* ConcreteGCProxy = GetConcreteProxy<FGeometryCollectionPhysicsProxy>(ClusteredParticle);
+				if (ConcreteGCProxy)
+				{
+					const FSimulationParameters& SimParams = ConcreteGCProxy->GetSimParameters();
+					DamageEvaluationModel = SimParams.DamageEvaluationModel;
+				}
 
-			int32 IslandIndex = 0;
-			TArray<int32> IslandIndicesToRemove;
+				int32 IslandIndex = 0;
+				TArray<int32> IslandIndicesToRemove;
 
-			for (const FParticleIsland& Island : Islands)
-			{
-				const bool bNeedRecomputeConnectivityStrain = ConcreteGCProxy && DamageEvaluationModel == EDamageEvaluationModel::StrainFromMaterialStrengthAndConnectivity;
+				for (const FParticleIsland& Island : Islands)
+				{
+					const bool bNeedRecomputeConnectivityStrain = ConcreteGCProxy && DamageEvaluationModel == EDamageEvaluationModel::StrainFromMaterialStrengthAndConnectivity;
 
-				const FClusterUnion* AttachedClusterUnion = nullptr;
+					FClusterUnion* AttachedClusterUnion = nullptr;
 
-				// Protect a potentially work-intensive loop behind the conditions that actually require us to step through every child in every island.
-				if (bNeedRecomputeConnectivityStrain || bCheckForInterclusterEdgesOnRelease)
-				{
-					for (FPBDRigidParticleHandle* ChildParticle : Island)
+					// Protect a potentially work-intensive loop behind the conditions that actually require us to step through every child in every island.
+					if (bNeedRecomputeConnectivityStrain || bCheckForInterclusterEdgesOnRelease)
 					{
-						if (FPBDRigidClusteredParticleHandle* ClusteredChild = ChildParticle->CastToClustered())
+						for (FPBDRigidParticleHandle* ChildParticle : Island)
 						{
-							// recompute the strain as its connectivity has changed
-							if (bNeedRecomputeConnectivityStrain)
+							if (FPBDRigidClusteredParticleHandle* ClusteredChild = ChildParticle->CastToClustered())
 							{
-								ConcreteGCProxy->ComputeMaterialBasedDamageThreshold_Internal(*ClusteredChild);
-							}
+								// recompute the strain as its connectivity has changed
+								if (bNeedRecomputeConnectivityStrain)
+								{
+									ConcreteGCProxy->ComputeMaterialBasedDamageThreshold_Internal(*ClusteredChild);
+								}
 
-							if (bCheckForInterclusterEdgesOnRelease && !AttachedClusterUnion)
-							{
-								// Check if any edge is connected to "something else". This entire island
-								// will be connected to that "something else". For now, we're enforcing that it's a cluster union.
-								for (const FConnectivityEdge& Edge : ClusteredChild->ConnectivityEdges())
+								if (bCheckForInterclusterEdgesOnRelease && !AttachedClusterUnion)
 								{
-									if (IsInterclusterEdge(*ClusteredChild, Edge))
+									// Check if any edge is connected to "something else". This entire island
+									// will be connected to that "something else". For now, we're enforcing that it's a cluster union.
+									for (const FConnectivityEdge& Edge : ClusteredChild->ConnectivityEdges())
 									{
-										if (FClusterUnion* ClusterUnion = ClusterUnionManager.FindClusterUnionFromParticle(GetActiveParentParticle(Edge.Sibling)))
+										if (IsInterclusterEdge(*ClusteredChild, Edge))
 										{
-											AttachedClusterUnion = ClusterUnion;
+											if (FClusterUnion* ClusterUnion = ClusterUnionManager.FindClusterUnionFromParticle(GetActiveParentParticle(Edge.Sibling)))
+											{
+												AttachedClusterUnion = ClusterUnion;
+											}
+											break;
 										}
-										break;
 									}
 								}
 							}
 						}
 					}
+
+					if (AttachedClusterUnion)
+					{
+						// Need to manually add cluster union properties since the API for adding a pending operation doesn't have an option for that.
+						for (FPBDRigidParticleHandle* ChildParticle : Island)
+						{
+							FClusterUnionParticleProperties Properties;
+							Properties.bIsAuxiliaryParticle = true;
+							AttachedClusterUnion->ChildProperties.Add(ChildParticle, Properties);
+						}
+						ClusterUnionManager.AddPendingClusterIndexOperation(AttachedClusterUnion->InternalIndex, EClusterUnionOperation::Add, Island);
+						IslandIndicesToRemove.Add(IslandIndex);
+					}
+					else if (Island.Num() == 1 && !AttachedClusterUnion) //need to break single pieces first
+					{
+						FPBDRigidParticleHandle* Child = Island[0];
+						RemoveChildFromParent(Child, ClusteredParticle);
+						ActivatedChildren.Add(Child);
+					}
+
+					++IslandIndex;
 				}
 
-				if (AttachedClusterUnion)
+				for (int32 RemoveIndex = IslandIndicesToRemove.Num() - 1; RemoveIndex >= 0; --RemoveIndex)
 				{
-					ClusterUnionManager.AddPendingClusterIndexOperation(AttachedClusterUnion->InternalIndex, EClusterUnionOperation::Add, Island);
-					IslandIndicesToRemove.Add(IslandIndex);
+					Islands.RemoveAtSwap(IslandIndicesToRemove[RemoveIndex], 1, false);
 				}
-				else if (Island.Num() == 1 && !AttachedClusterUnion) //need to break single pieces first
+
+				if (bCreateNewClusters)
 				{
-					FPBDRigidParticleHandle* Child = Island[0];
-					RemoveChildFromParent(Child, ClusteredParticle);
-					ActivatedChildren.Add(Child);
+					TArray<FPBDRigidParticleHandle*> NewClusters = CreateClustersFromNewIslands(Islands, ClusteredParticle);
+					ActivatedChildren.Append(MoveTemp(NewClusters));
 				}
-
-				++IslandIndex;
 			}
-
-			for (int32 RemoveIndex = IslandIndicesToRemove.Num() - 1; RemoveIndex >= 0; --RemoveIndex)
+			else if (bCheckForInterclusterEdgesOnRelease)
 			{
-				Islands.RemoveAtSwap(IslandIndicesToRemove[RemoveIndex], 1, false);
-			}
+				// We know we're in an cluster union. There are pieces that we consider to be the "main body". There are certain pieces that we consider to be auxiliary as well.
+				// If an island is only made up of auxiliary pieces, then those pieces should fall off. Connectivity of main pieces should be handled by the GT.
+				for (const FParticleIsland& Island : Islands)
+				{
+					bool bHasMainParticle = false;
+					for (FPBDRigidParticleHandle* ChildParticle : Island)
+					{
+						if (const FClusterUnionParticleProperties* Props = ParentClusterUnion->ChildProperties.Find(ChildParticle))
+						{
+							if (!Props->bIsAuxiliaryParticle)
+							{
+								bHasMainParticle = true;
+								break;
+							}
+						}
+						else
+						{
+							bHasMainParticle = true;
+							break;
+						}
+					}
 
-			if (bCreateNewClusters)
-			{
-				TArray<FPBDRigidParticleHandle*> NewClusters = CreateClustersFromNewIslands(Islands, ClusteredParticle);
-				ActivatedChildren.Append(MoveTemp(NewClusters));
+					if (!bHasMainParticle)
+					{
+						ClusterUnionManager.AddPendingClusterIndexOperation(ParentClusterUnion->InternalIndex, EClusterUnionOperation::Remove, Island);
+						ActivatedChildren.Append(Island);
+					}
+				}
 			}
 		}
 
diff --git a/Engine/Source/Runtime/Experimental/Chaos/Public/Chaos/ClusterUnionManager.h b/Engine/Source/Runtime/Experimental/Chaos/Public/Chaos/ClusterUnionManager.h
--- a/Engine/Source/Runtime/Experimental/Chaos/Public/Chaos/ClusterUnionManager.h
+++ b/Engine/Source/Runtime/Experimental/Chaos/Public/Chaos/ClusterUnionManager.h
@@ -58,6 +58,12 @@ namespace Chaos
 		uint32 ComponentId = 0;
 	};
 
+	struct CHAOS_API FClusterUnionParticleProperties
+	{
+		// An auxiliary particle will be removed from the cluster union if FRigidClustering::HandleConnectivityOnReleaseClusterParticle detects an island made up of only auxiliary particles.
+		bool bIsAuxiliaryParticle = false;
+	};
+
 	struct CHAOS_API FClusterUnion
 	{
 		// The root cluster particle that we created internally to represent the cluster.
@@ -69,6 +75,9 @@ namespace Chaos
 		// All the particles that belong to this cluster.
 		TArray<FPBDRigidParticleHandle*> ChildParticles;
 
+		// Additional properties that are only relevant for particles that have been added into a cluster union.
+		TMap<FPBDRigidParticleHandle*, FClusterUnionParticleProperties> ChildProperties;
+
 		// The internal index used to reference the cluster union.
 		FClusterUnionIndex InternalIndex;
 
diff --git a/Engine/Source/Runtime/Experimental/Chaos/Public/Chaos/PBDRigidClustering.h b/Engine/Source/Runtime/Experimental/Chaos/Public/Chaos/PBDRigidClustering.h
--- a/Engine/Source/Runtime/Experimental/Chaos/Public/Chaos/PBDRigidClustering.h
+++ b/Engine/Source/Runtime/Experimental/Chaos/Public/Chaos/PBDRigidClustering.h
@@ -485,7 +485,7 @@ public:
 	TSet<FPBDRigidParticleHandle*> HandleConnectivityOnReleaseClusterParticle(FPBDRigidClusteredParticleHandle* ClusteredParticle, bool bCreateNewClusters);
 	
 	using FParticleIsland = TArray<FPBDRigidParticleHandle*>;
-	TArray<FParticleIsland> FindIslandsInChildren(const FPBDRigidClusteredParticleHandle* ClusteredParticle);
+	TArray<FParticleIsland> FindIslandsInChildren(const FPBDRigidClusteredParticleHandle* ClusteredParticle, bool bTraverseInterclusterEdges);
 	TArray<FPBDRigidParticleHandle*> CreateClustersFromNewIslands(TArray<FParticleIsland>& Islands, FPBDRigidClusteredParticleHandle* ClusteredParent);
 
 	void UpdateTopLevelParticle(FPBDRigidClusteredParticleHandle* Particle);
