apply distance threshold filter to the intersected geometry collection leaf convex hulls, making these hulls less likely to contain sliver elements

diff --git a/Engine/Source/Runtime/Experimental/Chaos/Private/GeometryCollection/GeometryCollectionConvexUtility.cpp b/Engine/Source/Runtime/Experimental/Chaos/Private/GeometryCollection/GeometryCollectionConvexUtility.cpp
--- a/Engine/Source/Runtime/Experimental/Chaos/Private/GeometryCollection/GeometryCollectionConvexUtility.cpp
+++ b/Engine/Source/Runtime/Experimental/Chaos/Private/GeometryCollection/GeometryCollectionConvexUtility.cpp
@@ -1729,7 +1729,7 @@ void FGeometryCollectionConvexUtility::GenerateLeafConvexHulls(FGeometryCollecti
 					int32 HullIdx = FinalHulls.Add(MakeUnique<Chaos::FConvex>());
 					UE::GeometryCollectionConvexUtility::IntersectConvexHulls(FinalHulls[HullIdx].Get(),
 						ComputedHulls.Hulls[GeoHullIdx].Get(), 0.0f, ExternalHulls[ExtHull].Convex.Get(),
-						nullptr, &ExternalHulls[ExtHull].Transform, &ExternalHulls[ExtHull].Transform);
+						nullptr, &ExternalHulls[ExtHull].Transform, &ExternalHulls[ExtHull].Transform, Settings.SimplificationDistanceThreshold);
 					TransformToConvexIndicesAttrib[SourceTransformIdx].Add(HullIdx);
 				}
 			}
@@ -3035,7 +3035,7 @@ namespace UE::GeometryCollectionConvexUtility
 		HullPolygons.EstimateSharpContact(HullA, HullB, OutSharpContact, OutMaxSharpContact);
 	}
 
-	void IntersectConvexHulls(Chaos::FConvex* ResultHull, const Chaos::FConvex* ClipHull, float ClipHullOffset, const Chaos::FConvex* UpdateHull, const FTransform* ClipHullTransform, const FTransform* UpdateHullTransform, const FTransform* ResultTransform)
+	void IntersectConvexHulls(Chaos::FConvex* ResultHull, const Chaos::FConvex* ClipHull, float ClipHullOffset, const Chaos::FConvex* UpdateHull, const FTransform* ClipHullTransform, const FTransform* UpdateHullTransform, const FTransform* ResultTransform, double SimplificationDistanceThreshold)
 	{
 		FHullPolygons HullPolygons(*ClipHull);
 		bool bNeedRecomputeBounds = false;
@@ -3071,6 +3071,7 @@ namespace UE::GeometryCollectionConvexUtility
 				V = (Chaos::FVec3f)ResultTransform->TransformPosition((FVector)V);
 			}
 		}
+		FilterHullPoints(HullPolygons.Vertices, SimplificationDistanceThreshold);
 		*ResultHull = Chaos::FConvex(HullPolygons.Vertices, UpdateHull->GetMargin());
 	}
 }
diff --git a/Engine/Source/Runtime/Experimental/Chaos/Public/GeometryCollection/GeometryCollectionConvexUtility.h b/Engine/Source/Runtime/Experimental/Chaos/Public/GeometryCollection/GeometryCollectionConvexUtility.h
--- a/Engine/Source/Runtime/Experimental/Chaos/Public/GeometryCollection/GeometryCollectionConvexUtility.h
+++ b/Engine/Source/Runtime/Experimental/Chaos/Public/GeometryCollection/GeometryCollectionConvexUtility.h
@@ -76,7 +76,7 @@ namespace UE::GeometryCollectionConvexUtility
 	// Compute the intersection of ClipHull and UpdateHull, with ClipHull optionally offset, and place the result in ResultHull. Note ResultHull can point to the same FConvex as UpdateHull.
 	// If optional transforms are provided, ClipHull will be transformed into the local space of UpdateHull, and finally transformed by UpdateToResultTransform. Otherwise, all FConvex are assumed to be in the same coordinate space.
 	void CHAOS_API IntersectConvexHulls(::Chaos::FConvex* ResultHull, const ::Chaos::FConvex* ClipHull, float ClipHullOffset, const ::Chaos::FConvex* UpdateHull,
-		const FTransform* ClipHullTransform = nullptr, const FTransform* UpdateHullTransform = nullptr, const FTransform* UpdateToResultTransform = nullptr);
+		const FTransform* ClipHullTransform = nullptr, const FTransform* UpdateHullTransform = nullptr, const FTransform* UpdateToResultTransform = nullptr, double SimplificationDistanceThreshold = 0.0);
 
 }
 
