-> Add overload of UpdateClusterMassProperties which only changes mass properties of a particle, does not set it's transforms
-> Use the new overload to set mass properties in GCs when building clusters so that we get a correct MassToLocal matrix

diff --git a/Engine/Source/Runtime/Experimental/Chaos/Private/Chaos/PBDRigidClusteringAlgo.cpp b/Engine/Source/Runtime/Experimental/Chaos/Private/Chaos/PBDRigidClusteringAlgo.cpp
--- a/Engine/Source/Runtime/Experimental/Chaos/Private/Chaos/PBDRigidClusteringAlgo.cpp
+++ b/Engine/Source/Runtime/Experimental/Chaos/Private/Chaos/PBDRigidClusteringAlgo.cpp
@@ -142,6 +142,60 @@ namespace Chaos
 	}
 
 
+	DECLARE_CYCLE_STAT(TEXT("TPBDRigidClustering<>::UpdateClusterMassProperties()"), STAT_UpdateClusterMassProperties_NoForceMassOrientation, STATGROUP_Chaos);
+	void UpdateClusterMassProperties(
+		FPBDRigidClusteredParticleHandle* Parent,
+		const TSet<FPBDRigidParticleHandle*>& Children)
+	{
+		SCOPE_CYCLE_COUNTER(STAT_UpdateClusterMassProperties);
+
+		// NOTE: Default initializer will have zero mass
+		FMassProperties ParentMass;
+
+		// Combine mass properties of all children in parent's space
+		if (Children.Num() > 0)
+		{
+			const FRigidTransform3 ParentTM = Parent->GetTransformXR();
+			TArray<FMassProperties> ChildMasses;
+			ChildMasses.Reserve(Children.Num());
+			for (const FPBDRigidParticleHandle* Child : Children)
+			{
+				// Get the child's transform relative to the parent
+				const FRigidTransform3 ChildTM = Child->GetTransformXR();
+				const FRigidTransform3 ChildToParentTM = ChildTM.GetRelativeTransform(ParentTM);
+				const FRigidTransform3 ChildCoM(Child->CenterOfMass(), Child->RotationOfMass());
+				const FRigidTransform3 ChildCoMInParentSpace = ChildCoM * ChildToParentTM;
+
+				// Get the child's mass properties
+				FMassProperties ChildMass;
+				ChildMass.Mass = Child->M();
+				ChildMass.InertiaTensor = FMatrix33(Child->I());
+				ChildMass.CenterOfMass = ChildCoMInParentSpace.GetTranslation();
+				ChildMass.RotationOfMass = ChildCoMInParentSpace.GetRotation();
+				ChildMasses.Add(ChildMass);
+			}
+
+			ParentMass = Chaos::Combine(ChildMasses);
+		}
+
+		// Set properties of particle based on combined mass properties
+		// NOTE: The combine method will have diagonalized the inertia.
+		const FVec3 Inertia = ParentMass.InertiaTensor.GetDiagonal();
+		Parent->SetCenterOfMass(ParentMass.CenterOfMass);
+		Parent->SetRotationOfMass(ParentMass.RotationOfMass);
+		Parent->M() = ParentMass.Mass;
+		Parent->InvM()
+			= FMath::IsNearlyZero(ParentMass.Mass)
+			? 0.f
+			: 1.f / ParentMass.Mass;
+		Parent->I() = Inertia;
+		Parent->InvI()
+			= (FMath::IsNearlyZero(Inertia[0]) || FMath::IsNearlyZero(Inertia[1]) || FMath::IsNearlyZero(Inertia[2]))
+			? FVec3::ZeroVector
+			: FReal(1) / Inertia;
+	}
+
+
 
 	DECLARE_CYCLE_STAT(TEXT("TPBDRigidClustering<>::UpdateKinematicProperties()"), STAT_UpdateKinematicProperties, STATGROUP_Chaos);
 	void 
diff --git a/Engine/Source/Runtime/Experimental/Chaos/Private/PhysicsProxy/GeometryCollectionPhysicsProxy.cpp b/Engine/Source/Runtime/Experimental/Chaos/Private/PhysicsProxy/GeometryCollectionPhysicsProxy.cpp
--- a/Engine/Source/Runtime/Experimental/Chaos/Private/PhysicsProxy/GeometryCollectionPhysicsProxy.cpp
+++ b/Engine/Source/Runtime/Experimental/Chaos/Private/PhysicsProxy/GeometryCollectionPhysicsProxy.cpp
@@ -1735,8 +1735,8 @@ FGeometryCollectionPhysicsProxy::BuildClusters_Internal(
 		nullptr, // Parent->Geometry() ? Parent->Geometry() : Implicits[CollectionClusterIndex], 
 		SimFilter,
 		QueryFilter,
-		Parent->M() > 0.0 ? Parent->M() : (Chaos::FReal)ScaledMass,
-		Parent->I() != Chaos::FVec3f(0.0) ? Parent->I() : ScaledInertia,
+		ScaledMass,
+		ScaledInertia,
 		ParticleTM, 
 		(uint8)DynamicState[CollectionClusterIndex], 
 		0, // static_cast<int16>(CollisionGroup[TransformGroupIndex])
@@ -4297,7 +4297,7 @@ void FGeometryCollectionPhysicsProxy::InitializeSharedCollisionStructures(
 
  			if (CollectionSimulatableParticles[TransformGroupIndex])
  			{
-				FTransform GeometryWorldTransform = CollectionMassToLocal[TransformGroupIndex] * CollectionSpaceTransforms[TransformGroupIndex];
+				FTransform CollectionSpaceParticleTransform = CollectionMassToLocal[TransformGroupIndex] * CollectionSpaceTransforms[TransformGroupIndex];
 
  				PopulateSimulatedParticle(
  					Handles[TransformGroupIndex].Get(),
@@ -4308,7 +4308,7 @@ void FGeometryCollectionPhysicsProxy::InitializeSharedCollisionStructures(
  					FCollisionFilterData(),		// QueryFilter
  					CollectionMass[TransformGroupIndex],
  					CollectionInertiaTensor[TransformGroupIndex], 
-					GeometryWorldTransform,
+					CollectionSpaceParticleTransform,
  					(uint8)EObjectStateTypeEnum::Chaos_Object_Dynamic, 
  					INDEX_NONE,  // CollisionGroup
 					1.0f // todo(chaos) CollisionParticlesPerObjectFraction is not accessible right there for now but we can pass 1.0 for the time being
@@ -4347,17 +4347,35 @@ void FGeometryCollectionPhysicsProxy::InitializeSharedCollisionStructures(
 				//CollectionSimulatableParticles[TransformGroupIndex] = true;
 				IsClusterSimulated[TransformGroupIndex] = true;
 
+				// NOTE: Particles are in collection space for this function
+
+				// Initialize XR to the collection space transform for this parent particle.
+				// We have the individual mass properties for the children, but not for the parent
+				// yet, so MassToLocal is not available.
+				FPBDRigidClusteredParticleHandle* ClusterHandle = Handles[ClusterTransformIdx].Get();
+				const FRigidTransform3 OriginalXR = CollectionSpaceTransforms[ClusterTransformIdx];
+				ClusterHandle->SetX(OriginalXR.GetTranslation());
+				ClusterHandle->SetR(OriginalXR.GetRotation());
+
+				// This method will compute a mass and inertia, and center and rotation of mass and
+				// store them on ClusterHandle based on the combined mass properties of its children.
+				UpdateClusterMassProperties(ClusterHandle, ChildrenIndices);
+
+				// Store the relative transform from the particle's configuration space to its
+				// MassToLocal because GeometryCollections assume that the particle's XR is at its
+				// CoM and RoM - MassToLocal accounts for this offset.
+				const FRigidTransform3 MassToLocal = FRigidTransform3(ClusterHandle->CenterOfMass(), ClusterHandle->RotationOfMass());
+				CollectionMassToLocal[ClusterTransformIdx] = MassToLocal;
+				ClusterHandle->SetCenterOfMass(FVec3::ZeroVector);
+				ClusterHandle->SetRotationOfMass(FQuat::Identity);
+				
+				// Since we've zerod out CoM and RoM, we need to update XR accordingly so that when
+				// UpdateClusterMassProperties is called in the next level up with ClusterHandle as
+				// one of the the children it's XR will be correct.
+				const FRigidTransform3 NewXR = MassToLocal * OriginalXR;
+				ClusterHandle->SetX(NewXR.GetTranslation());
+				ClusterHandle->SetR(NewXR.GetRotation());
 
-				// TODO: This needs to be rotated to diagonal, used to update I()/InvI() from diagonal, and update transform with rotation.
-				FMatrix33 ClusterInertia(0);
-				UpdateClusterMassProperties(Handles[ClusterTransformIdx].Get(), ChildrenIndices, ClusterInertia);	//compute mass properties
-				const FTransform ClusterMassToCollection = 
-					FTransform(CollectionSpaceParticles->R(ClusterTransformIdx), 
-							   CollectionSpaceParticles->X(ClusterTransformIdx));
-
-				CollectionMassToLocal[ClusterTransformIdx] = 
-					ClusterMassToCollection.GetRelativeTransform(
-						CollectionSpaceTransforms[ClusterTransformIdx]);
 
 				//update geometry
 				//merge children meshes and move them into cluster's mass space
diff --git a/Engine/Source/Runtime/Experimental/Chaos/Public/Chaos/PBDRigidClusteringAlgo.h b/Engine/Source/Runtime/Experimental/Chaos/Public/Chaos/PBDRigidClusteringAlgo.h
--- a/Engine/Source/Runtime/Experimental/Chaos/Public/Chaos/PBDRigidClusteringAlgo.h
+++ b/Engine/Source/Runtime/Experimental/Chaos/Public/Chaos/PBDRigidClusteringAlgo.h
@@ -5,11 +5,15 @@
 
 namespace Chaos
 {
-	void UpdateClusterMassProperties(
+	void CHAOS_API UpdateClusterMassProperties(
 		FPBDRigidClusteredParticleHandle* Parent,
 		TSet<FPBDRigidParticleHandle*>& Children,
 		FMatrix33& ParentInertia,
-		const FRigidTransform3* ForceMassOrientation = nullptr);
+		const FRigidTransform3* ForceMassOrientation);
+
+	void CHAOS_API UpdateClusterMassProperties(
+		FPBDRigidClusteredParticleHandle* Parent,
+		const TSet<FPBDRigidParticleHandle*>& Children);
 
 	void CHAOS_API UpdateKinematicProperties(
 		FPBDRigidParticleHandle* Parent,
