-> Use new UpdateClusterMassProperties in ClusterUnion
-> Add MoveClusterToMassOffset where ForceMassOrientation==nullptr used to rely on UpdateClusterMassProperties to set X/R
-> Rotate FDerivedDataGeometryCollectionCooker version string GUID so that GC inertias will be regenerated

diff --git a/Engine/Source/Runtime/Experimental/Chaos/Private/Chaos/ClusterUnionManager.cpp b/Engine/Source/Runtime/Experimental/Chaos/Private/Chaos/ClusterUnionManager.cpp
--- a/Engine/Source/Runtime/Experimental/Chaos/Private/Chaos/ClusterUnionManager.cpp
+++ b/Engine/Source/Runtime/Experimental/Chaos/Private/Chaos/ClusterUnionManager.cpp
@@ -633,15 +633,14 @@ namespace Chaos
 		FMatrix33 ClusterInertia(0);
 
 		const bool bRecomputeMassOrientation = EnumHasAnyFlags(Flags, EUpdateClusterUnionPropertiesFlags::RecomputeMassOrientation) && ClusterUnion.bNeedsXRInitialization;
-
-		// TODO: These functions are generally just re-building the cluster from scratch. Need to figure out a way
-		// to get these functions to update the already existing cluster instead.
 		TSet<FPBDRigidParticleHandle*> FullChildrenSet(ClusterUnion.ChildParticles);
 
-		const FRigidTransform3 ForceMassOrientation{ ClusterUnion.InternalCluster->X(), ClusterUnion.InternalCluster->R() };
-		UpdateClusterMassProperties(ClusterUnion.InternalCluster, FullChildrenSet, ClusterInertia, bRecomputeMassOrientation ? nullptr : &ForceMassOrientation);
+		UpdateClusterMassProperties(ClusterUnion.InternalCluster, FullChildrenSet);
+
+		// Position the internal cluster at the CoM of the children
 		if (bRecomputeMassOrientation)
 		{
+			MoveClusterToMassOffset(ClusterUnion.InternalCluster, EMassOffsetType::EPosition);
 			ClusterUnion.bNeedsXRInitialization = false;
 		}
 
diff --git a/Engine/Source/Runtime/Experimental/Chaos/Private/Chaos/PBDRigidClustering.cpp b/Engine/Source/Runtime/Experimental/Chaos/Private/Chaos/PBDRigidClustering.cpp
--- a/Engine/Source/Runtime/Experimental/Chaos/Private/Chaos/PBDRigidClustering.cpp
+++ b/Engine/Source/Runtime/Experimental/Chaos/Private/Chaos/PBDRigidClustering.cpp
@@ -210,9 +210,18 @@ namespace Chaos
 			UpdateTopLevelParticle(NewParticle);
 		}
 
-		// TODO: This needs to be rotated to diagonal, used to update I()/InvI() from diagonal, and update transform with rotation.
-		FMatrix33 ClusterInertia(0);
-		UpdateClusterMassProperties(NewParticle, ChildrenSet, ClusterInertia, ForceMassOrientation);
+		if (ForceMassOrientation)
+		{
+			NewParticle->SetX(ForceMassOrientation->GetTranslation());
+			NewParticle->SetR(ForceMassOrientation->GetRotation());
+		}
+		
+		UpdateClusterMassProperties(NewParticle, ChildrenSet);
+		
+		if (ForceMassOrientation == nullptr)
+		{
+			MoveClusterToMassOffset(NewParticle, EMassOffsetType::EPosition);
+		}
 		UpdateKinematicProperties(NewParticle, MChildren, MEvolution);
 		UpdateGeometry(NewParticle, ChildrenSet, MChildren, ProxyGeometry, Parameters);
 		GenerateConnectionGraph(NewParticle, Parameters);
@@ -431,10 +440,10 @@ namespace Chaos
 		NoCleanParams.bCopyCollisionParticles = !!UnionsHaveCollisionParticles;
 
 		TSet<FPBDRigidParticleHandle*> ChildrenSet(ChildrenArray);
-		
-		// TODO: This needs to be rotated to diagonal, used to update I()/InvI() from diagonal, and update transform with rotation.
-		FMatrix33 ClusterInertia(0);
-		UpdateClusterMassProperties(NewParticle, ChildrenSet, ClusterInertia, nullptr);
+
+		UpdateClusterMassProperties(NewParticle, ChildrenSet);
+		MoveClusterToMassOffset(NewParticle, EMassOffsetType::EPosition);
+
 		UpdateKinematicProperties(NewParticle, MChildren, MEvolution);
 
 		UpdateGeometry(NewParticle, ChildrenSet, MChildren, nullptr, NoCleanParams);
diff --git a/Engine/Source/Runtime/Experimental/Chaos/Private/Chaos/PBDRigidClusteringAlgo.cpp b/Engine/Source/Runtime/Experimental/Chaos/Private/Chaos/PBDRigidClusteringAlgo.cpp
--- a/Engine/Source/Runtime/Experimental/Chaos/Private/Chaos/PBDRigidClusteringAlgo.cpp
+++ b/Engine/Source/Runtime/Experimental/Chaos/Private/Chaos/PBDRigidClusteringAlgo.cpp
@@ -37,112 +37,6 @@ namespace Chaos
 	FAutoConsoleVariableRef CVarClusterSnapDistance2(TEXT("p.ClusterSnapDistance"), ClusterSnapDistance, TEXT(""));
 
 	DECLARE_CYCLE_STAT(TEXT("TPBDRigidClustering<>::UpdateClusterMassProperties()"), STAT_UpdateClusterMassProperties, STATGROUP_Chaos);
-	void UpdateClusterMassProperties(
-		FPBDRigidClusteredParticleHandle* Parent,
-		TSet<FPBDRigidParticleHandle*>& Children,
-		FMatrix33& ParentInertia,
-		const FRigidTransform3* ForceMassOrientation)
-	{
-		SCOPE_CYCLE_COUNTER(STAT_UpdateClusterMassProperties);
-
-		// Initialize parent
-		Parent->M() = FReal(0);
-		Parent->InvM() = FReal(0);
-		Parent->I() = FVec3(0);
-		Parent->InvI() = FVec3(0);
-		Parent->SetCenterOfMass(FVec3(0));
-		Parent->SetRotationOfMass(FQuat::Identity);
-		if (ForceMassOrientation)
-		{
-			Parent->X() = ForceMassOrientation->GetLocation();
-			Parent->R() = ForceMassOrientation->GetRotation();
-		}
-		Parent->P() = Parent->X();
-		Parent->Q() = Parent->R();
-
-		if (Children.Num() == 0)
-		{
-			return;
-		}
-
-		//
-		// Step 1: Compute the world CoM and total mass of the parent
-		//
-
-		FVec3 WorldCoM = FVec3::ZeroVector;
-		for (const FPBDRigidParticleHandle* Child : Children)
-		{
-			WorldCoM += Child->M() * Child->XCom();
-			Parent->M() += Child->M();
-		}
-		if (FMath::IsNearlyZero(Parent->M()))
-		{
-			return;
-		}
-		Parent->InvM() = FReal(1) / Parent->M();
-		WorldCoM *= Parent->InvM();
-
-
-		//
-		// Step 2: Pick the parent's orientation and location.
-		//
-		// If we have a ForceMassOrientation transform, then use that, otherwise
-		// default to X = CoM. To do this, we need to compute the world CoM of
-		// the children.
-		//
-
-		if (ForceMassOrientation == nullptr)
-		{
-			Parent->X() = WorldCoM;
-			Parent->R() = FQuat::Identity;
-		}
-		Parent->P() = Parent->X();
-		Parent->Q() = Parent->R();
-		const FRigidTransform3 ParentTM = Parent->GetTransformXR();
-		const FRigidTransform3 InvParentTM = ParentTM.Inverse();
-
-
-		//
-		// Step 3: Compute mass properties of each particle & store them in a list
-		//
-
-		TArray<FMassProperties> ChildMasses;
-		ChildMasses.Reserve(Children.Num());
-		for (const FPBDRigidParticleHandle* Child : Children)
-		{
-			// Get the child's transform relative to the parent
-			const FRigidTransform3 ChildTM = Child->GetTransformXR();
-			const FRigidTransform3 LocalTM = ChildTM * InvParentTM;
-
-			// Get the child's mass properties
-			FMassProperties ChildMass;
-			ChildMass.Mass = Child->M();
-			ChildMass.InertiaTensor = FMatrix33(Child->I());
-			ChildMass.CenterOfMass = LocalTM.TransformPosition(Child->CenterOfMass());
-			ChildMass.RotationOfMass = LocalTM.GetRotation() * Child->RotationOfMass();
-			ChildMasses.Add(ChildMass);
-		}
-
-		//
-		// Step 4: Combine mass properties of sub particles & store
-		// them in the parent particle
-		//
-
-		FMassProperties ParentMass = Chaos::Combine(ChildMasses);
-		// NOTE: The combine method will have diagonalized the inertia.
-		ParentInertia = ParentMass.InertiaTensor;
-		const FVec3 Inertia = ParentInertia.GetDiagonal();
-		Parent->SetCenterOfMass(ParentMass.CenterOfMass);
-		Parent->SetRotationOfMass(ParentMass.RotationOfMass);
-		Parent->I() = Inertia;
-		Parent->InvI()
-			= (FMath::IsNearlyZero(Inertia[0]) || FMath::IsNearlyZero(Inertia[1]) || FMath::IsNearlyZero(Inertia[2]))
-			? FVec3::ZeroVector
-			: FReal(1) / Inertia;
-	}
-
-
-	DECLARE_CYCLE_STAT(TEXT("TPBDRigidClustering<>::UpdateClusterMassProperties()"), STAT_UpdateClusterMassProperties_NoForceMassOrientation, STATGROUP_Chaos);
 	void UpdateClusterMassProperties(
 		FPBDRigidClusteredParticleHandle* Parent,
 		const TSet<FPBDRigidParticleHandle*>& Children)
@@ -195,7 +89,22 @@ namespace Chaos
 			: FReal(1) / Inertia;
 	}
 
+	void MoveClusterToMassOffset(FPBDRigidClusteredParticleHandle* Cluster, const EMassOffsetType MassOffsetTypes)
+	{
+		if (MassOffsetTypes | EMassOffsetType::EPosition)
+		{
+			Cluster->SetX(Cluster->XCom());
+			Cluster->SetP(Cluster->X());
+			Cluster->SetCenterOfMass(FVec3::ZeroVector);
+		}
 
+		if (MassOffsetTypes | EMassOffsetType::ERotation)
+		{
+			Cluster->SetR(Cluster->RCom());
+			Cluster->SetQ(Cluster->R());
+			Cluster->SetRotationOfMass(FQuat::Identity);
+		}
+	}
 
 	DECLARE_CYCLE_STAT(TEXT("TPBDRigidClustering<>::UpdateKinematicProperties()"), STAT_UpdateKinematicProperties, STATGROUP_Chaos);
 	void 
diff --git a/Engine/Source/Runtime/Experimental/Chaos/Private/PhysicsProxy/ClusterUnionPhysicsProxy.cpp b/Engine/Source/Runtime/Experimental/Chaos/Private/PhysicsProxy/ClusterUnionPhysicsProxy.cpp
--- a/Engine/Source/Runtime/Experimental/Chaos/Private/PhysicsProxy/ClusterUnionPhysicsProxy.cpp
+++ b/Engine/Source/Runtime/Experimental/Chaos/Private/PhysicsProxy/ClusterUnionPhysicsProxy.cpp
@@ -194,6 +194,7 @@ namespace Chaos
 				{
 					if (FPBDRigidsEvolutionGBF* Evolution = GetEvolution(this))
 					{
+						const bool bWasAnchored = Particle_Internal->IsAnchored();
 						Particle_Internal->SetIsAnchored(bIsAnchored);
 
 						// We also need to make sure the cluster union won't try to override this value that we set.
@@ -204,7 +205,10 @@ namespace Chaos
 						}
 
 						// Let the cluster union manager apply the proper properties on the particle.
-						Manager.RequestDeferredClusterPropertiesUpdate(ClusterUnionIndex, Chaos::EUpdateClusterUnionPropertiesFlags::UpdateKinematicProperties);
+						if (bWasAnchored != bIsAnchored)
+						{
+							Manager.RequestDeferredClusterPropertiesUpdate(ClusterUnionIndex, Chaos::EUpdateClusterUnionPropertiesFlags::UpdateKinematicProperties);
+						}
 					}
 				}
 			}
diff --git a/Engine/Source/Runtime/Experimental/Chaos/Private/PhysicsProxy/GeometryCollectionPhysicsProxy.cpp b/Engine/Source/Runtime/Experimental/Chaos/Private/PhysicsProxy/GeometryCollectionPhysicsProxy.cpp
--- a/Engine/Source/Runtime/Experimental/Chaos/Private/PhysicsProxy/GeometryCollectionPhysicsProxy.cpp
+++ b/Engine/Source/Runtime/Experimental/Chaos/Private/PhysicsProxy/GeometryCollectionPhysicsProxy.cpp
@@ -4366,16 +4366,11 @@ void FGeometryCollectionPhysicsProxy::InitializeSharedCollisionStructures(
 				// CoM and RoM - MassToLocal accounts for this offset.
 				const FRigidTransform3 MassToLocal = FRigidTransform3(ClusterHandle->CenterOfMass(), ClusterHandle->RotationOfMass());
 				CollectionMassToLocal[ClusterTransformIdx] = MassToLocal;
-				ClusterHandle->SetCenterOfMass(FVec3::ZeroVector);
-				ClusterHandle->SetRotationOfMass(FQuat::Identity);
-				
-				// Since we've zerod out CoM and RoM, we need to update XR accordingly so that when
-				// UpdateClusterMassProperties is called in the next level up with ClusterHandle as
-				// one of the the children it's XR will be correct.
-				const FRigidTransform3 NewXR = MassToLocal * OriginalXR;
-				ClusterHandle->SetX(NewXR.GetTranslation());
-				ClusterHandle->SetR(NewXR.GetRotation());
 
+				// Zero out CoM and RoM, and update XR accordingly so that when UpdateClusterMassProperties
+				// is called in the next level up with ClusterHandle as one of the the children it's XR will
+				// be correct.
+				MoveClusterToMassOffset(ClusterHandle, EMassOffsetType::EPosition | EMassOffsetType::ERotation);
 
 				//update geometry
 				//merge children meshes and move them into cluster's mass space
diff --git a/Engine/Source/Runtime/Experimental/Chaos/Public/Chaos/PBDRigidClusteringAlgo.h b/Engine/Source/Runtime/Experimental/Chaos/Public/Chaos/PBDRigidClusteringAlgo.h
--- a/Engine/Source/Runtime/Experimental/Chaos/Public/Chaos/PBDRigidClusteringAlgo.h
+++ b/Engine/Source/Runtime/Experimental/Chaos/Public/Chaos/PBDRigidClusteringAlgo.h
@@ -5,16 +5,21 @@
 
 namespace Chaos
 {
-	void CHAOS_API UpdateClusterMassProperties(
-		FPBDRigidClusteredParticleHandle* Parent,
-		TSet<FPBDRigidParticleHandle*>& Children,
-		FMatrix33& ParentInertia,
-		const FRigidTransform3* ForceMassOrientation);
+	enum CHAOS_API EMassOffsetType : uint8
+	{
+		EPosition = 1 << 0,
+		ERotation = 1 << 1
+	};
+	ENUM_CLASS_FLAGS(EMassOffsetType)
 
 	void CHAOS_API UpdateClusterMassProperties(
 		FPBDRigidClusteredParticleHandle* Parent,
 		const TSet<FPBDRigidParticleHandle*>& Children);
 
+	// If bPosition == true, set X/P of Cluster to its world CoM, and set its CoM to ZeroVector
+	// If bRotation == true, set R/Q of Cluster to its world RoM, and set its RoM to Identity
+	void CHAOS_API MoveClusterToMassOffset(FPBDRigidClusteredParticleHandle* Cluster, EMassOffsetType MassOffsetTypes);
+
 	void CHAOS_API UpdateKinematicProperties(
 		FPBDRigidParticleHandle* Parent,
 		const FRigidClustering::FClusterMap&,
diff --git a/Engine/Source/Runtime/Experimental/GeometryCollectionEngine/Private/GeometryCollection/DerivedDataGeometryCollectionCooker.cpp b/Engine/Source/Runtime/Experimental/GeometryCollectionEngine/Private/GeometryCollection/DerivedDataGeometryCollectionCooker.cpp
--- a/Engine/Source/Runtime/Experimental/GeometryCollectionEngine/Private/GeometryCollection/DerivedDataGeometryCollectionCooker.cpp
+++ b/Engine/Source/Runtime/Experimental/GeometryCollectionEngine/Private/GeometryCollection/DerivedDataGeometryCollectionCooker.cpp
@@ -60,7 +60,7 @@ bool FDerivedDataGeometryCollectionCooker::Build(TArray<uint8>& OutData)
 
 const TCHAR* FDerivedDataGeometryCollectionCooker::GetVersionString() const
 {
-	const TCHAR* VersionString = TEXT("C3F0D73CE6094A1C84E57D75C9131F55");
+	const TCHAR* VersionString = TEXT("E223057731B4465BBF7087001AF7ACBD");
 	return VersionString;
 }
 
