Moved param size handling into a file-scoped function

Removed nesting for comprehension, added elses, and tidied a little blank space

#rnx

diff --git a/Engine/Source/Runtime/Experimental/IoStoreOnDemand/Private/IoStoreOnDemand.cpp b/Engine/Source/Runtime/Experimental/IoStoreOnDemand/Private/IoStoreOnDemand.cpp
--- a/Engine/Source/Runtime/Experimental/IoStoreOnDemand/Private/IoStoreOnDemand.cpp
+++ b/Engine/Source/Runtime/Experimental/IoStoreOnDemand/Private/IoStoreOnDemand.cpp
@@ -17,6 +17,41 @@
 
 DEFINE_LOG_CATEGORY(LogIoStoreOnDemand);
 
+////////////////////////////////////////////////////////////////////////////////
+static uint64 ParseSizeParam(const TCHAR* Param)
+{
+	const TCHAR* CommandLine = FCommandLine::Get();
+
+	FString ParamValue;
+	if (!FParse::Value(CommandLine, Param, ParamValue))
+	{
+		return 0;
+	}
+
+	uint64 Size = 0;
+	if (!FParse::Value(CommandLine, Param, Size))
+	{
+		return 0;
+	}
+
+	if (ParamValue.EndsWith(TEXT("GB")))
+	{
+		Size *= 1024 * 1024 * 1024;
+	}
+
+	else if (ParamValue.EndsWith(TEXT("MB")))
+	{
+		Size *= 1024 * 1024;
+	}
+
+	else if (ParamValue.EndsWith(TEXT("KB")))
+	{
+		Size *= 1024;
+	}
+
+	return Size;
+}
+
 namespace UE
 {
 
@@ -248,40 +283,11 @@ public:
 	virtual void ShutdownModule() override;
 
 private:
-	static uint64 ParseSizeParam(const TCHAR* Param);
 	static bool ParseEncryptionKeyParam(const FString& Param, FGuid& OutGuid, FAES::FAESKey& OutKey);
 };
 
 IMPLEMENT_MODULE(FIoStoreOnDemandModule, IoStoreOnDemand);
 
-uint64 FIoStoreOnDemandModule::ParseSizeParam(const TCHAR* Param)
-{
-	FString ParamValue;
-	if (FParse::Value(FCommandLine::Get(), Param, ParamValue))
-	{
-		uint64 Size = 0;
-		if (FParse::Value(FCommandLine::Get(), Param, Size))
-		{
-			if (ParamValue.EndsWith(TEXT("GB")))
-			{
-				Size *= 1024*1024*1024;
-			}
-			if (ParamValue.EndsWith(TEXT("MB")))
-			{
-				Size *= 1024*1024;
-			}
-			if (ParamValue.EndsWith(TEXT("KB")))
-			{
-				Size *= 1024;
-			}
-
-			return Size;
-		}
-	}
-
-	return 0;
-}
-
 bool FIoStoreOnDemandModule::ParseEncryptionKeyParam(const FString& Param, FGuid& OutKeyGuid, FAES::FAESKey& OutKey)
 {
 	TArray<FString> Tokens;
