Keep content key param parsing alongside size param parsing

#rnx

diff --git a/Engine/Source/Runtime/Experimental/IoStoreOnDemand/Private/IoStoreOnDemand.cpp b/Engine/Source/Runtime/Experimental/IoStoreOnDemand/Private/IoStoreOnDemand.cpp
--- a/Engine/Source/Runtime/Experimental/IoStoreOnDemand/Private/IoStoreOnDemand.cpp
+++ b/Engine/Source/Runtime/Experimental/IoStoreOnDemand/Private/IoStoreOnDemand.cpp
@@ -17,6 +17,9 @@
 
 DEFINE_LOG_CATEGORY(LogIoStoreOnDemand);
 
+namespace UE::IO::Private
+{
+
 ////////////////////////////////////////////////////////////////////////////////
 static uint64 ParseSizeParam(const TCHAR* Param)
 {
@@ -34,24 +37,39 @@ static uint64 ParseSizeParam(const TCHAR* Param)
 		return 0;
 	}
 
-	if (ParamValue.EndsWith(TEXT("GB")))
-	{
-		Size *= 1024 * 1024 * 1024;
-	}
+	if (ParamValue.EndsWith(TEXT("GB"))) return Size << 30;
+	if (ParamValue.EndsWith(TEXT("MB"))) return Size << 20;
+	if (ParamValue.EndsWith(TEXT("KB"))) return Size << 10;
 
-	else if (ParamValue.EndsWith(TEXT("MB")))
-	{
-		Size *= 1024 * 1024;
-	}
+	return Size;
+}
+
+////////////////////////////////////////////////////////////////////////////////
+static bool ParseEncryptionKeyParam(const FString& Param, FGuid& OutKeyGuid, FAES::FAESKey& OutKey)
+{
+	TArray<FString> Tokens;
+	Param.ParseIntoArray(Tokens, TEXT(":"), true);
 
-	else if (ParamValue.EndsWith(TEXT("KB")))
+	if (Tokens.Num() == 2)
 	{
-		Size *= 1024;
+		TArray<uint8> KeyBytes;
+		if (FGuid::Parse(Tokens[0], OutKeyGuid) && FBase64::Decode(Tokens[1], KeyBytes))
+		{
+			if (OutKeyGuid != FGuid() && KeyBytes.Num() == FAES::FAESKey::KeySize)
+			{
+				FMemory::Memcpy(OutKey.Key, KeyBytes.GetData(), FAES::FAESKey::KeySize);
+				return true;
+			}
+		}
 	}
-
-	return Size;
+	
+	return false;
 }
 
+} // namespace UE::IO::Private
+
+
+
 namespace UE
 {
 
@@ -283,34 +301,14 @@ public:
 	virtual void ShutdownModule() override;
 
 private:
-	static bool ParseEncryptionKeyParam(const FString& Param, FGuid& OutGuid, FAES::FAESKey& OutKey);
 };
 
 IMPLEMENT_MODULE(FIoStoreOnDemandModule, IoStoreOnDemand);
 
-bool FIoStoreOnDemandModule::ParseEncryptionKeyParam(const FString& Param, FGuid& OutKeyGuid, FAES::FAESKey& OutKey)
-{
-	TArray<FString> Tokens;
-	Param.ParseIntoArray(Tokens, TEXT(":"), true);
-
-	if (Tokens.Num() == 2)
-	{
-		TArray<uint8> KeyBytes;
-		if (FGuid::Parse(Tokens[0], OutKeyGuid) && FBase64::Decode(Tokens[1], KeyBytes))
-		{
-			if (OutKeyGuid != FGuid() && KeyBytes.Num() == FAES::FAESKey::KeySize)
-			{
-				FMemory::Memcpy(OutKey.Key, KeyBytes.GetData(), FAES::FAESKey::KeySize);
-				return true;
-			}
-		}
-	}
-
-	return false;
-}
-
 void FIoStoreOnDemandModule::StartupModule()
 {
+	using namespace UE::IO::Private;
+
 #if !WITH_EDITOR
 	UE::FOnDemandEndpoint Endpoint;
 	
