Chaos - ImplicitObjectUnion BVH collision detection optimizations

We now have two types of particle-pair midphase. One which pre-builds a list of all potentially overlapping shape pairs (this already existed), and a new more general purpose one. The general purpose one does not cache any information or prefilter any pairs, but it will use the BVH in the collision hierarchy if there is one.

ImplicitObjectUnion now has its own implementation of a BVH (currently a loose binary tree) to avoid having to make major changes to TBoundingVolumeHierarchy which did not fit its needs.

Midphases are no longer held in an ObjectPool although this could be made to work again (except we now need a pool per midphase type)

diff --git a/Engine/Source/Runtime/Engine/Private/PhysicsEngine/ClusterUnionComponent.cpp b/Engine/Source/Runtime/Engine/Private/PhysicsEngine/ClusterUnionComponent.cpp
--- a/Engine/Source/Runtime/Engine/Private/PhysicsEngine/ClusterUnionComponent.cpp
+++ b/Engine/Source/Runtime/Engine/Private/PhysicsEngine/ClusterUnionComponent.cpp
@@ -287,8 +287,10 @@ void UClusterUnionComponent::ForceRebuildGTParticleGeometry()
 		}
 	}
 
-	TSharedPtr<Chaos::FImplicitObject, ESPMode::ThreadSafe> NewGeometry = Objects.IsEmpty() ? MakeShared<Chaos::FImplicitObjectUnionClustered>() : MakeShared<Chaos::FImplicitObjectUnion>(MoveTemp(Objects));
-	PhysicsProxy->SetSharedGeometry_External(NewGeometry, Particles);
+	Chaos::FImplicitObjectUnion* NewGeometry = Objects.IsEmpty() ? new Chaos::FImplicitObjectUnionClustered() : new Chaos::FImplicitObjectUnion(MoveTemp(Objects));
+	NewGeometry->SetAllowBVH(true);
+
+	PhysicsProxy->SetSharedGeometry_External(TSharedPtr<Chaos::FImplicitObject, ESPMode::ThreadSafe>(NewGeometry), Particles);
 }
 
 TArray<UPrimitiveComponent*> UClusterUnionComponent::GetPrimitiveComponents()
diff --git a/Engine/Source/Runtime/Engine/Private/PhysicsEngine/Experimental/PhysInterface_Chaos.cpp b/Engine/Source/Runtime/Engine/Private/PhysicsEngine/Experimental/PhysInterface_Chaos.cpp
--- a/Engine/Source/Runtime/Engine/Private/PhysicsEngine/Experimental/PhysInterface_Chaos.cpp
+++ b/Engine/Source/Runtime/Engine/Private/PhysicsEngine/Experimental/PhysInterface_Chaos.cpp
@@ -726,9 +726,10 @@ void FPhysInterface_Chaos::AddGeometry(FPhysicsActorHandle& InActor, const FGeom
 			else
 			{
 				// We always have a union so we can support any future welding operations. (Non-trivial converting the SharedPtr to UniquePtr).
-				// The root union always supports BVH and is the only Union in the hierarchy that is allowed to do so.
+				// NOTE: The root union always supports BVH (if there are enough shapes) and is the only Union in the hierarchy that is allowed 
+				// to do so, but we don't create it here because that makes welding even more expensive (bodies are welded one by one). 
+				// Search for SetAllowBVH to see where the BVH is enabled.
 				TUniquePtr<Chaos::FImplicitObjectUnion> Union = MakeUnique<Chaos::FImplicitObjectUnion>(MoveTemp(Geoms));
-				Union->SetAllowBVH(true);
 				InActor->GetGameThreadAPI().SetGeometry(MoveTemp(Union));
 			}
 		}
diff --git a/Engine/Source/Runtime/Experimental/Chaos/Private/Chaos/CCDUtilities.cpp b/Engine/Source/Runtime/Experimental/Chaos/Private/Chaos/CCDUtilities.cpp
--- a/Engine/Source/Runtime/Experimental/Chaos/Private/Chaos/CCDUtilities.cpp
+++ b/Engine/Source/Runtime/Experimental/Chaos/Private/Chaos/CCDUtilities.cpp
@@ -420,7 +420,7 @@ namespace Chaos
 					}
 					const FVec3 ContactPos = CCDConstraint->SweptConstraint->GetShapeWorldTransform1().TransformPositionNoScale(FVec3(ManifoldPoint.ContactPoint.ShapeContactPoints[1]));
 					const FVec3 ContactNormal = CCDConstraint->SweptConstraint->GetShapeWorldTransform1().TransformVectorNoScale(FVec3(ManifoldPoint.ContactPoint.ShapeContactNormal));
-					FDebugDrawQueue::GetInstance().DrawDebugLine(ContactPos, ContactPos + DebugDrawSettings.DrawScale * 50 * ContactNormal, FColor::Red, false, UE_KINDA_SMALL_NUMBER, DebugDrawSettings.DrawPriority, DebugDrawSettings.LineThickness);
+					FDebugDrawQueue::GetInstance().DrawDebugLine(ContactPos, ContactPos + DebugDrawSettings.DrawScale * 50 * ContactNormal, FColor::Red, false, UE_KINDA_SMALL_NUMBER, uint8(DebugDrawSettings.DrawPriority), DebugDrawSettings.LineThickness);
 				}
 			}
 #endif
diff --git a/Engine/Source/Runtime/Experimental/Chaos/Private/Chaos/ChaosDebugDraw.cpp b/Engine/Source/Runtime/Experimental/Chaos/Private/Chaos/ChaosDebugDraw.cpp
--- a/Engine/Source/Runtime/Experimental/Chaos/Private/Chaos/ChaosDebugDraw.cpp
+++ b/Engine/Source/Runtime/Experimental/Chaos/Private/Chaos/ChaosDebugDraw.cpp
@@ -12,6 +12,7 @@
 #include "Chaos/TriangleMesh.h"
 #include "Chaos/DebugDrawQueue.h"
 #include "Chaos/Evolution/SimulationSpace.h"
+#include "Chaos/ImplicitObjectBVH.h"
 #include "Chaos/ImplicitObjectScaled.h"
 #include "Chaos/ImplicitObjectTransformed.h"
 #include "Chaos/ImplicitObjectUnion.h"
@@ -96,6 +97,13 @@ namespace Chaos
 		bool bChaosDebugDrawConnectionGraphShowInternalStrains = false;
 		FAutoConsoleVariableRef CVarChaosDebugDrawConnectionGraphShowInternalStrains(TEXT("p.Chaos.DebugDraw.ConnectionGraph.ShowInternalStrains"), bChaosDebugDrawConnectionGraphShowInternalStrains, TEXT("When displaying connection graphs show strain values of each node"));
 
+		int32 ChaosDebugDebugDrawBVHLevel = -1;
+		bool bChaosDebugDebugDrawBVHShapes = false;
+		bool bChaosDebugDebugDrawBVHBounds = true;
+		FAutoConsoleVariableRef CVarChaosDebugDebugDrawBVHLevel(TEXT("p.Chaos.DebugDraw.BVHLevel"), ChaosDebugDebugDrawBVHLevel, TEXT("Which level of the BVH to draw (-1 = Leaf)"));
+		FAutoConsoleVariableRef CVarChaosDebugDebugDrawBVHShapes(TEXT("p.Chaos.DebugDraw.BVHShapes"), bChaosDebugDebugDrawBVHShapes, TEXT("Whether to draw the shapes in each BVH node (colorized)"));
+		FAutoConsoleVariableRef CVarChaosDebugDebugDrawBVHBounds(TEXT("p.Chaos.DebugDraw.BVHBounds"), bChaosDebugDebugDrawBVHBounds, TEXT("Whether to draw the BVH node bounds (colorized)"));
+
 		// NOTE: These settings should never really be used - they are the fallback defaults
 		// if the user does not specify settings in the debug draw call.
 		// See PBDRigidsColver.cpp and ImmediatePhysicsSimulation_Chaos.cpp for example.
@@ -192,9 +200,9 @@ namespace Chaos
 			return ChaosDefaultBoundsColorsByState;
 		}
 
-		FColor GetIslandColor(const int32 IslandIndex, const bool bIsAwake)
+		FColor GetIndexColor(const int32 Index)
 		{
-			static const FColor AwakeColors[] =
+			static const FColor Colors[] =
 			{
 				FColor::Red,
 				FColor::Orange,
@@ -207,8 +215,14 @@ namespace Chaos
 				FColor::Magenta,
 				FColor::Purple,
 			};
-			const int32 NumAwakeColors = UE_ARRAY_COUNT(AwakeColors);
+			const int32 NumColors = UE_ARRAY_COUNT(Colors);
+
+			return Colors[Index % NumColors];
+		};
+
 
+		FColor GetIslandColor(const int32 IslandIndex, const bool bIsAwake)
+		{
 			static const FColor SleepingColor = FColor::Black;
 			static const FColor NullColor = FColor::White;
 
@@ -222,7 +236,7 @@ namespace Chaos
 				return SleepingColor;
 			}
 
-			return AwakeColors[IslandIndex % NumAwakeColors];
+			return GetIndexColor(IslandIndex);
 		};
 
 
@@ -325,7 +339,7 @@ namespace Chaos
 						const FVec3 ExplodeDelta = ChaosDebugDrawConvexExplodeDistance * N0;
 
 						// Outer shape
-						FDebugDrawQueue::GetInstance().DrawDebugLine(OuterP0 + ExplodeDelta, OuterP1 + ExplodeDelta, Color, false, Duration, Settings.DrawPriority, Settings.ShapeThicknesScale * Settings.LineThickness);
+						FDebugDrawQueue::GetInstance().DrawDebugLine(OuterP0 + ExplodeDelta, OuterP1 + ExplodeDelta, Color, false, Duration, uint8(Settings.DrawPriority), Settings.ShapeThicknesScale * Settings.LineThickness);
 
 						// Core shape and lines connecting core to outer
 						if (Margin > 0.0f)
@@ -333,14 +347,14 @@ namespace Chaos
 							const FRealSingle LineThickness = 0.5f * Settings.ShapeThicknesScale * Settings.LineThickness;
 							const FVec3 InnerP0 = ShapeTransform.TransformPositionNoScale(Shape->GetMarginAdjustedVertexScaled(VertexIndex0, Margin, ShapeTransform.GetScale3D(), nullptr));
 							const FVec3 InnerP1 = ShapeTransform.TransformPositionNoScale(Shape->GetMarginAdjustedVertexScaled(VertexIndex1, Margin, ShapeTransform.GetScale3D(), nullptr));
-							FDebugDrawQueue::GetInstance().DrawDebugLine(InnerP0, InnerP1, FColor::Blue, false, Duration, Settings.DrawPriority, LineThickness);
-							FDebugDrawQueue::GetInstance().DrawDebugLine(InnerP0, OuterP0, FColor::Black, false, Duration, Settings.DrawPriority, LineThickness);
+							FDebugDrawQueue::GetInstance().DrawDebugLine(InnerP0, InnerP1, FColor::Blue, false, Duration, uint8(Settings.DrawPriority), LineThickness);
+							FDebugDrawQueue::GetInstance().DrawDebugLine(InnerP0, OuterP0, FColor::Black, false, Duration, uint8(Settings.DrawPriority), LineThickness);
 						}
 
 						// Vertex and face normal
 						if (bChaosDebugDebugDrawConvexVertices)
 						{
-							FDebugDrawQueue::GetInstance().DrawDebugLine(OuterP0 + ExplodeDelta, OuterP0 + ExplodeDelta + Settings.DrawScale * 20.0f * N0, FColor::Green, false, Duration, Settings.DrawPriority, Settings.LineThickness);
+							FDebugDrawQueue::GetInstance().DrawDebugLine(OuterP0 + ExplodeDelta, OuterP0 + ExplodeDelta + Settings.DrawScale * 20.0f * N0, FColor::Green, false, Duration, uint8(Settings.DrawPriority), Settings.LineThickness);
 						}
 					}
 				}
@@ -660,14 +674,14 @@ namespace Chaos
 			{
 				const TSphere<FReal, 3>* Sphere = Implicit->template GetObject<TSphere<FReal, 3>>();
 				const FVec3 P = ShapeTransform.TransformPosition(Sphere->GetCenter());
-				FDebugDrawQueue::GetInstance().DrawDebugSphere(P, Sphere->GetRadius(), 8, ShapeColor, false, Duration, Settings.DrawPriority, Settings.ShapeThicknesScale * Settings.LineThickness);
+				FDebugDrawQueue::GetInstance().DrawDebugSphere(P, Sphere->GetRadius(), 8, ShapeColor, false, Duration, uint8(Settings.DrawPriority), Settings.ShapeThicknesScale * Settings.LineThickness);
 				break;
 			}
 			case ImplicitObjectType::Box:
 			{
 				const TBox<FReal, 3>* Box = Implicit->template GetObject<TBox<FReal, 3>>();
 				const FVec3 P = ShapeTransform.TransformPosition(Box->GetCenter());
-				FDebugDrawQueue::GetInstance().DrawDebugBox(P, (FReal)0.5 * Box->Extents(), ShapeTransform.GetRotation(), ShapeColor, false, Duration, Settings.DrawPriority, Settings.ShapeThicknesScale * Settings.LineThickness);
+				FDebugDrawQueue::GetInstance().DrawDebugBox(P, (FReal)0.5 * Box->Extents(), ShapeTransform.GetRotation(), ShapeColor, false, Duration, uint8(Settings.DrawPriority), Settings.ShapeThicknesScale * Settings.LineThickness);
 				break;
 			}
 			case ImplicitObjectType::Plane:
@@ -677,7 +691,7 @@ namespace Chaos
 				const FCapsule* Capsule = Implicit->template GetObject<FCapsule>();
 				const FVec3 P = ShapeTransform.TransformPosition(Capsule->GetCenter());
 				const FRotation3 Q = ShapeTransform.GetRotation() * FRotationMatrix::MakeFromZ(Capsule->GetAxis());
-				FDebugDrawQueue::GetInstance().DrawDebugCapsule(P, (FReal)0.5 * Capsule->GetHeight() + Capsule->GetRadius(), Capsule->GetRadius(), Q, ShapeColor, false, Duration, Settings.DrawPriority, Settings.ShapeThicknesScale * Settings.LineThickness);
+				FDebugDrawQueue::GetInstance().DrawDebugCapsule(P, (FReal)0.5 * Capsule->GetHeight() + Capsule->GetRadius(), Capsule->GetRadius(), Q, ShapeColor, false, Duration, uint8(Settings.DrawPriority), Settings.ShapeThicknesScale * Settings.LineThickness);
 				break;
 			}
 			case ImplicitObjectType::LevelSet:
@@ -742,7 +756,7 @@ namespace Chaos
 						for (int32 ParticleIndex = 0; ParticleIndex < (int32)Particles->Size(); ++ParticleIndex)
 						{
 							FVec3 P = ShapeTransform.TransformPosition(Particles->X(ParticleIndex));
-							FDebugDrawQueue::GetInstance().DrawDebugPoint(P, ShapeColor, false, Duration, Settings.DrawPriority, Settings.PointSize);
+							FDebugDrawQueue::GetInstance().DrawDebugPoint(P, ShapeColor, false, Duration, uint8(Settings.DrawPriority), Settings.PointSize);
 						}
 					}
 				}
@@ -753,7 +767,7 @@ namespace Chaos
 				const FColor ShapeBoundsColor = FColor::Orange;
 				const FAABB3& ShapeBounds = Implicit->BoundingBox();
 				const FVec3 ShapeBoundsPos = ShapeTransform.TransformPosition(ShapeBounds.Center());
-				FDebugDrawQueue::GetInstance().DrawDebugBox(ShapeBoundsPos, 0.5f * ShapeBounds.Extents(), ShapeTransform.GetRotation(), ShapeBoundsColor, false, Duration, Settings.DrawPriority, Settings.LineThickness);
+				FDebugDrawQueue::GetInstance().DrawDebugBox(ShapeBoundsPos, 0.5f * ShapeBounds.Extents(), ShapeTransform.GetRotation(), ShapeBoundsColor, false, Duration, uint8(Settings.DrawPriority), Settings.LineThickness);
 			}
 		}
 
@@ -773,6 +787,62 @@ namespace Chaos
 			DrawShapesImpl(Particle->Handle(), FRigidTransform3(P, Q), Particle->Geometry().Get(), FShapeOrShapesArray(Particle->Handle()), 0.0f, InColor, 0.0f, Settings);
 		}
 
+		void DrawBVHImpl(const FGeometryParticleHandle* Particle, const FRigidTransform3& ShapeTransform, const Private::FImplicitBVH* BVH, const FColor& UnusedColor, const FRealSingle Duration, const FChaosDebugDrawSettings& Settings)
+		{
+			const FReal BoundsDrawShrink = 1.0f;	// 1cm
+			int32 LeafIndex = 0;
+
+			const auto& NodeVisitor = 
+				[Particle, &ShapeTransform, BVH, &LeafIndex, BoundsDrawShrink, &Duration, &Settings](const FAABB3f& NodeBounds, const int32 NodeDepth, const Private::FImplicitBVHNode& Node) -> bool
+				{
+					// Draw the node if it is at the level we are asked to draw, or it is a leaf and we want to draw leaves.
+					if ((NodeDepth == ChaosDebugDebugDrawBVHLevel) || (Node.IsLeaf() && ((ChaosDebugDebugDrawBVHLevel == INDEX_NONE) || (ChaosDebugDebugDrawBVHLevel > NodeDepth))))
+					{
+						const FColor Color = GetIndexColor(LeafIndex++);
+
+						if (bChaosDebugDebugDrawBVHBounds)
+						{
+							const FVec3 NodeBoundsPos = ShapeTransform.TransformPosition(FVec3(NodeBounds.Center()));
+							const FVec3 NodeBoundsExtents = 0.5f * FVec3(NodeBounds.Extents()) - FVec3(BoundsDrawShrink);
+							FDebugDrawQueue::GetInstance().DrawDebugBox(NodeBoundsPos, NodeBoundsExtents, ShapeTransform.GetRotation(), Color, false, Duration, uint8(Settings.DrawPriority), Settings.LineThickness);
+						}
+
+						if (bChaosDebugDebugDrawBVHShapes)
+						{
+							BVH->VisitNodeObjects(Node, 
+								[Particle, &ShapeTransform, &Color, Duration, &Settings](const FImplicitObject * Implicit, const FRigidTransform3f& RelativeTransformf, const FAABB3f& RelativeBoundsf, const int32 RootObjectIndex, const int32 LeafObjectIndex) -> void
+								{
+									const int32 ShapeIndex = (RootObjectIndex != INDEX_NONE) ? RootObjectIndex : 0;
+									const FShapeInstance* Shape = Particle->ShapeInstances()[ShapeIndex].Get();
+									DrawShapesImpl(Particle, FRigidTransform3(RelativeTransformf) * ShapeTransform, Implicit, FShapeOrShapesArray(Shape), 0, Color, Duration, Settings);
+								});
+						}
+					}
+					return true;
+				};
+
+			BVH->VisitHierarchy(NodeVisitor);
+		}
+
+		void DrawParticleBVHImpl(const FRigidTransform3& SpaceTransform, const FGeometryParticleHandle* Particle, const FColor& InColor, const FChaosDebugDrawSettings& Settings)
+		{
+			const FImplicitObjectUnion* Union = Particle->Geometry()->GetObject<FImplicitObjectUnion>();
+			if (Union == nullptr)
+			{
+				Union = Particle->Geometry()->GetObject<FImplicitObjectUnionClustered>();
+			}
+			if (Union != nullptr)
+			{
+				if (Union->GetBVH() != nullptr)
+				{
+					FVec3 P = SpaceTransform.TransformPosition(Particle->X());
+					FRotation3 Q = SpaceTransform.GetRotation() * (Particle->R());
+
+					DrawBVHImpl(Particle->Handle(), FRigidTransform3(P, Q), Union->GetBVH(), InColor, 0.0f, Settings);
+				}
+			}
+		}
+
 		static EImplicitObjectType GetFirstConcreteShapeType(const FImplicitObject* Shape)
 		{
 			EImplicitObjectType InnerType = GetInnerType(Shape->GetType());
@@ -819,7 +889,7 @@ namespace Chaos
 				}
 			}
 
-			FDebugDrawQueue::GetInstance().DrawDebugBox(P, 0.5f * Box.Extents(), Q, Color, false, UE_KINDA_SMALL_NUMBER, Settings.DrawPriority, Settings.LineThickness);
+			FDebugDrawQueue::GetInstance().DrawDebugBox(P, 0.5f * Box.Extents(), Q, Color, false, UE_KINDA_SMALL_NUMBER, uint8(Settings.DrawPriority), Settings.LineThickness);
 
 			for (const auto& Shape : InParticle->ShapesArray())
 			{
@@ -832,7 +902,7 @@ namespace Chaos
 					const FRotation3 ShapeQ = SpaceTransform.GetRotation();
 					const FColor ShapeColor = (bChaosDebugDebugDrawColorBoundsByShapeType) ? Settings.BoundsColorsPerShapeType.GetColorFromShapeType(ShapeType) : Color;
 
-					FDebugDrawQueue::GetInstance().DrawDebugBox(ShapeP, 0.5f * ShapeBox.Extents(), ShapeQ, ShapeColor, false, UE_KINDA_SMALL_NUMBER, Settings.DrawPriority, Settings.LineThickness);
+					FDebugDrawQueue::GetInstance().DrawDebugBox(ShapeP, 0.5f * ShapeBox.Extents(), ShapeQ, ShapeColor, false, UE_KINDA_SMALL_NUMBER, uint8(Settings.DrawPriority), Settings.LineThickness);
 				}
 			}
 		}
@@ -853,15 +923,15 @@ namespace Chaos
 			FVec3 PCOM = SpaceTransform.TransformPosition(FParticleUtilities::GetCoMWorldPosition(Particle));
 			FRotation3 QCOM = SpaceTransform.GetRotation() * FParticleUtilities::GetCoMWorldRotation(Particle);
 			FMatrix33 QCOMm = QCOM.ToMatrix();
-			FDebugDrawQueue::GetInstance().DrawDebugDirectionalArrow(PCOM, PCOM + Settings.DrawScale * Settings.BodyAxisLen * QCOMm.GetAxis(0), Settings.DrawScale * Settings.ArrowSize, Red, false, UE_KINDA_SMALL_NUMBER, Settings.DrawPriority, Settings.LineThickness);
-			FDebugDrawQueue::GetInstance().DrawDebugDirectionalArrow(PCOM, PCOM + Settings.DrawScale * Settings.BodyAxisLen * QCOMm.GetAxis(1), Settings.DrawScale * Settings.ArrowSize, Green, false, UE_KINDA_SMALL_NUMBER, Settings.DrawPriority, Settings.LineThickness);
-			FDebugDrawQueue::GetInstance().DrawDebugDirectionalArrow(PCOM, PCOM + Settings.DrawScale * Settings.BodyAxisLen * QCOMm.GetAxis(2), Settings.DrawScale * Settings.ArrowSize, Blue, false, UE_KINDA_SMALL_NUMBER, Settings.DrawPriority, Settings.LineThickness);
+			FDebugDrawQueue::GetInstance().DrawDebugDirectionalArrow(PCOM, PCOM + Settings.DrawScale * Settings.BodyAxisLen * QCOMm.GetAxis(0), Settings.DrawScale * Settings.ArrowSize, Red, false, UE_KINDA_SMALL_NUMBER, uint8(Settings.DrawPriority), Settings.LineThickness);
+			FDebugDrawQueue::GetInstance().DrawDebugDirectionalArrow(PCOM, PCOM + Settings.DrawScale * Settings.BodyAxisLen * QCOMm.GetAxis(1), Settings.DrawScale * Settings.ArrowSize, Green, false, UE_KINDA_SMALL_NUMBER, uint8(Settings.DrawPriority), Settings.LineThickness);
+			FDebugDrawQueue::GetInstance().DrawDebugDirectionalArrow(PCOM, PCOM + Settings.DrawScale * Settings.BodyAxisLen * QCOMm.GetAxis(2), Settings.DrawScale * Settings.ArrowSize, Blue, false, UE_KINDA_SMALL_NUMBER, uint8(Settings.DrawPriority), Settings.LineThickness);
 
 			FColor Black = FColor::Black;
 			FColor Grey = (ColorScale * FColor(64, 64, 64)).ToFColor(false);
 			FVec3 PActor = SpaceTransform.TransformPosition(FParticleUtilities::GetActorWorldTransform(Particle).GetTranslation());
-			FDebugDrawQueue::GetInstance().DrawDebugPoint(PActor, Black, false, UE_KINDA_SMALL_NUMBER, Settings.DrawPriority, Settings.DrawScale * Settings.PointSize);
-			FDebugDrawQueue::GetInstance().DrawDebugLine(PCOM, PActor, Grey, false, UE_KINDA_SMALL_NUMBER, Settings.DrawPriority, Settings.LineThickness);
+			FDebugDrawQueue::GetInstance().DrawDebugPoint(PActor, Black, false, UE_KINDA_SMALL_NUMBER, uint8(Settings.DrawPriority), Settings.DrawScale * Settings.PointSize);
+			FDebugDrawQueue::GetInstance().DrawDebugLine(PCOM, PActor, Grey, false, UE_KINDA_SMALL_NUMBER, uint8(Settings.DrawPriority), Settings.LineThickness);
 		
 			if (Index >= 0)
 			{
@@ -870,12 +940,12 @@ namespace Chaos
 
 			if ((Settings.VelScale > 0.0f) && (Particle->V().Size() > UE_KINDA_SMALL_NUMBER))
 			{
-				FDebugDrawQueue::GetInstance().DrawDebugLine(PCOM, PCOM + SpaceTransform.TransformVector(Particle->V()) * Settings.VelScale, Red, false, UE_KINDA_SMALL_NUMBER, Settings.DrawPriority, Settings.LineThickness);
-				//FDebugDrawQueue::GetInstance().DrawDebugLine(PCOM, PCOM + SpaceTransform.TransformVector(Particle->VSmooth()) * Settings.VelScale, Blue, false, KINDA_SMALL_NUMBER, Settings.DrawPriority, Settings.LineThickness);
+				FDebugDrawQueue::GetInstance().DrawDebugLine(PCOM, PCOM + SpaceTransform.TransformVector(Particle->V()) * Settings.VelScale, Red, false, UE_KINDA_SMALL_NUMBER, uint8(Settings.DrawPriority), Settings.LineThickness);
+				//FDebugDrawQueue::GetInstance().DrawDebugLine(PCOM, PCOM + SpaceTransform.TransformVector(Particle->VSmooth()) * Settings.VelScale, Blue, false, KINDA_SMALL_NUMBER, uint8(Settings.DrawPriority), Settings.LineThickness);
 			}
 			if ((Settings.AngVelScale > 0.0f) && (Particle->W().Size() > UE_KINDA_SMALL_NUMBER))
 			{
-				FDebugDrawQueue::GetInstance().DrawDebugLine(PCOM, PCOM + SpaceTransform.TransformVector(Particle->W()) * Settings.AngVelScale, Green, false, UE_KINDA_SMALL_NUMBER, Settings.DrawPriority, Settings.LineThickness);
+				FDebugDrawQueue::GetInstance().DrawDebugLine(PCOM, PCOM + SpaceTransform.TransformVector(Particle->W()) * Settings.AngVelScale, Green, false, UE_KINDA_SMALL_NUMBER, uint8(Settings.DrawPriority), Settings.LineThickness);
 			}
 
 			if (Settings.InertiaScale > 0.0f)
@@ -991,34 +1061,34 @@ namespace Chaos
 					if ((Settings.PushOutScale > 0) && ManifoldPointResult.bIsValid && !ManifoldPointResult.NetPushOut.IsNearlyZero())
 					{
 						FColor Color = (ColorScale * PushOutImpusleColor).ToFColor(false);
-						FDebugDrawQueue::GetInstance().DrawDebugLine(WorldPointPlaneLocation, WorldPointPlaneLocation + Settings.DrawScale * Settings.PushOutScale * SpaceTransform.TransformVectorNoScale(FVec3(ManifoldPointResult.NetPushOut)), Color, false, Duration, Settings.DrawPriority, Settings.LineThickness);
+						FDebugDrawQueue::GetInstance().DrawDebugLine(WorldPointPlaneLocation, WorldPointPlaneLocation + Settings.DrawScale * Settings.PushOutScale * SpaceTransform.TransformVectorNoScale(FVec3(ManifoldPointResult.NetPushOut)), Color, false, Duration, uint8(Settings.DrawPriority), Settings.LineThickness);
 					}
 					if ((Settings.ImpulseScale > 0) && ManifoldPointResult.bIsValid && !ManifoldPointResult.NetImpulse.IsNearlyZero())
 					{
 						FColor Color = (ColorScale * ImpulseColor).ToFColor(false);
-						FDebugDrawQueue::GetInstance().DrawDebugLine(WorldPointPlaneLocation, WorldPointPlaneLocation + Settings.DrawScale * Settings.ImpulseScale * SpaceTransform.TransformVectorNoScale(FVec3(ManifoldPointResult.NetImpulse)), Color, false, Duration, Settings.DrawPriority, Settings.LineThickness);
+						FDebugDrawQueue::GetInstance().DrawDebugLine(WorldPointPlaneLocation, WorldPointPlaneLocation + Settings.DrawScale * Settings.ImpulseScale * SpaceTransform.TransformVectorNoScale(FVec3(ManifoldPointResult.NetImpulse)), Color, false, Duration, uint8(Settings.DrawPriority), Settings.LineThickness);
 					}
 
 					// Manifold plane and normal
 					if (Settings.ContactWidth > 0)
 					{
 						FColor C0 = (ColorScale * DiscColor).ToFColor(false);
-						FDebugDrawQueue::GetInstance().DrawDebugCircle(WorldPlaneLocation, Settings.DrawScale * Settings.ContactWidth, 12, C0, false, Duration, Settings.DrawPriority, Settings.LineThickness, Axes.GetUnitAxis(EAxis::Y), Axes.GetUnitAxis(EAxis::Z), false);
+						FDebugDrawQueue::GetInstance().DrawDebugCircle(WorldPlaneLocation, Settings.DrawScale * Settings.ContactWidth, 12, C0, false, Duration, uint8(Settings.DrawPriority), Settings.LineThickness, Axes.GetUnitAxis(EAxis::Y), Axes.GetUnitAxis(EAxis::Z), false);
 					}
 					if (Settings.ContactLen > 0)
 					{
 						FColor NormalColor = ((ManifoldPoint.ContactPoint.ContactType != EContactPointType::EdgeEdge) ? PlaneNormalColor : EdgeNormalColor);
 						FColor C1 = (ColorScale * NormalColor).ToFColor(false);
-						FDebugDrawQueue::GetInstance().DrawDebugLine(WorldPlaneLocation, WorldPlaneLocation + Settings.DrawScale * Settings.ContactLen * WorldPlaneNormal, C1, false, Duration, Settings.DrawPriority, Settings.LineThickness);
+						FDebugDrawQueue::GetInstance().DrawDebugLine(WorldPlaneLocation, WorldPlaneLocation + Settings.DrawScale * Settings.ContactLen * WorldPlaneNormal, C1, false, Duration, uint8(Settings.DrawPriority), Settings.LineThickness);
 					}
 					if (Settings.ContactPhiWidth > 0 && (ManifoldPoint.ContactPoint.Phi < FLT_MAX))
 					{
 						FColor C2 = (ColorScale * FColor(128, 128, 0)).ToFColor(false);
-						FDebugDrawQueue::GetInstance().DrawDebugCircle(WorldPlaneLocation - ManifoldPoint.ContactPoint.Phi * WorldPlaneNormal, Settings.DrawScale * Settings.ContactPhiWidth, 12, C2, false, Duration, Settings.DrawPriority, Settings.LineThickness, Axes.GetUnitAxis(EAxis::Y), Axes.GetUnitAxis(EAxis::Z), false);
+						FDebugDrawQueue::GetInstance().DrawDebugCircle(WorldPlaneLocation - ManifoldPoint.ContactPoint.Phi * WorldPlaneNormal, Settings.DrawScale * Settings.ContactPhiWidth, 12, C2, false, Duration, uint8(Settings.DrawPriority), Settings.LineThickness, Axes.GetUnitAxis(EAxis::Y), Axes.GetUnitAxis(EAxis::Z), false);
 					}
 
 					// Manifold point
-					FDebugDrawQueue::GetInstance().DrawDebugCircle(WorldPointLocation, 0.5f * Settings.DrawScale * Settings.ContactWidth, 12, DiscColor, false, Duration, Settings.DrawPriority, Settings.LineThickness, Axes.GetUnitAxis(EAxis::Y), Axes.GetUnitAxis(EAxis::Z), false);
+					FDebugDrawQueue::GetInstance().DrawDebugCircle(WorldPointLocation, 0.5f * Settings.DrawScale * Settings.ContactWidth, 12, DiscColor, false, Duration, uint8(Settings.DrawPriority), Settings.LineThickness, Axes.GetUnitAxis(EAxis::Y), Axes.GetUnitAxis(EAxis::Z), false);
 
 					// Whether restored
 					if (Settings.ContactInfoWidth > 0)
@@ -1026,17 +1096,17 @@ namespace Chaos
 						if (Contact.WasManifoldRestored())
 						{
 							const FReal BoxScale = Settings.DrawScale * Settings.ContactInfoWidth;
-							FDebugDrawQueue::GetInstance().DrawDebugBox(WorldPlaneLocation, FVec3(BoxScale, BoxScale, FReal(0.01)), FRotation3(FRotationMatrix::MakeFromZ(WorldPlaneNormal)), FColor::Blue, false, Duration, Settings.DrawPriority, 0.5f * Settings.LineThickness);
+							FDebugDrawQueue::GetInstance().DrawDebugBox(WorldPlaneLocation, FVec3(BoxScale, BoxScale, FReal(0.01)), FRotation3(FRotationMatrix::MakeFromZ(WorldPlaneNormal)), FColor::Blue, false, Duration, uint8(Settings.DrawPriority), 0.5f * Settings.LineThickness);
 						}
 						else if (ManifoldPoint.Flags.bWasRestored)
 						{
 							const FReal BoxScale = Settings.DrawScale * Settings.ContactInfoWidth;
-							FDebugDrawQueue::GetInstance().DrawDebugBox(WorldPlaneLocation, FVec3(BoxScale, BoxScale, FReal(0.01)), FRotation3(FRotationMatrix::MakeFromZ(WorldPlaneNormal)), FColor::Purple, false, Duration, Settings.DrawPriority, 0.5f * Settings.LineThickness);
+							FDebugDrawQueue::GetInstance().DrawDebugBox(WorldPlaneLocation, FVec3(BoxScale, BoxScale, FReal(0.01)), FRotation3(FRotationMatrix::MakeFromZ(WorldPlaneNormal)), FColor::Purple, false, Duration, uint8(Settings.DrawPriority), 0.5f * Settings.LineThickness);
 						}
 						else if (ManifoldPoint.Flags.bWasReplaced)
 						{
 							const FReal BoxScale = Settings.DrawScale * Settings.ContactInfoWidth;
-							FDebugDrawQueue::GetInstance().DrawDebugBox(WorldPlaneLocation, FVec3(BoxScale, BoxScale, FReal(0.01)), FRotation3(FRotationMatrix::MakeFromZ(WorldPlaneNormal)), FColor::Orange, false, Duration, Settings.DrawPriority, 0.5f * Settings.LineThickness);
+							FDebugDrawQueue::GetInstance().DrawDebugBox(WorldPlaneLocation, FVec3(BoxScale, BoxScale, FReal(0.01)), FRotation3(FRotationMatrix::MakeFromZ(WorldPlaneNormal)), FColor::Orange, false, Duration, uint8(Settings.DrawPriority), 0.5f * Settings.LineThickness);
 						}
 					}
 
@@ -1046,7 +1116,7 @@ namespace Chaos
 						if (Contact.IsSleeping())
 						{
 							const FReal BoxScale = Settings.DrawScale * Settings.ContactInfoWidth * 1.1f;
-							FDebugDrawQueue::GetInstance().DrawDebugBox(WorldPlaneLocation, FVec3(BoxScale, BoxScale, FReal(0.01)), FRotation3(FRotationMatrix::MakeFromZ(WorldPlaneNormal)), FColor::Black, false, Duration, Settings.DrawPriority, 0.5f * Settings.LineThickness);
+							FDebugDrawQueue::GetInstance().DrawDebugBox(WorldPlaneLocation, FVec3(BoxScale, BoxScale, FReal(0.01)), FRotation3(FRotationMatrix::MakeFromZ(WorldPlaneNormal)), FColor::Black, false, Duration, uint8(Settings.DrawPriority), 0.5f * Settings.LineThickness);
 						}
 					}
 				}
@@ -1056,7 +1126,7 @@ namespace Chaos
 				{
 					FColor Color = (ColorScale * FColor::White).ToFColor(false);
 					const FVec3 ImpulsePos = SpaceTransform.TransformPosition(WorldActorTransform0.GetTranslation());
-					FDebugDrawQueue::GetInstance().DrawDebugLine(ImpulsePos, ImpulsePos + Settings.DrawScale * Settings.ImpulseScale * SpaceTransform.TransformVectorNoScale(Contact.GetAccumulatedImpulse()), Color, false, Duration, Settings.DrawPriority, Settings.LineThickness);
+					FDebugDrawQueue::GetInstance().DrawDebugLine(ImpulsePos, ImpulsePos + Settings.DrawScale * Settings.ImpulseScale * SpaceTransform.TransformVectorNoScale(Contact.GetAccumulatedImpulse()), Color, false, Duration, uint8(Settings.DrawPriority), Settings.LineThickness);
 				}
 
 			}
@@ -1069,8 +1139,8 @@ namespace Chaos
 				const FMatrix Axes = FRotationMatrix::MakeFromX(Normal);
 				const FVec3 P0 = SpaceTransform.TransformPosition(Contact.GetParticle0()->X());
 				const FVec3 P1 = SpaceTransform.TransformPosition(Contact.GetParticle1()->X());
-				FDebugDrawQueue::GetInstance().DrawDebugLine(Location, P0, C3, false, UE_KINDA_SMALL_NUMBER, Settings.DrawPriority, Settings.LineThickness * 0.5f);
-				FDebugDrawQueue::GetInstance().DrawDebugLine(Location, P1, C3, false, UE_KINDA_SMALL_NUMBER, Settings.DrawPriority, Settings.LineThickness * 0.5f);
+				FDebugDrawQueue::GetInstance().DrawDebugLine(Location, P0, C3, false, UE_KINDA_SMALL_NUMBER, uint8(Settings.DrawPriority), Settings.LineThickness * 0.5f);
+				FDebugDrawQueue::GetInstance().DrawDebugLine(Location, P1, C3, false, UE_KINDA_SMALL_NUMBER, uint8(Settings.DrawPriority), Settings.LineThickness * 0.5f);
 			}
 		}
 		
@@ -1089,7 +1159,7 @@ namespace Chaos
 					if (!Impulse.IsNearlyZero() && (Settings.ImpulseScale > 0))
 					{
 						const FVec3 Pos0 = SpaceTransform.TransformPosition(P0->PCom());
-						FDebugDrawQueue::GetInstance().DrawDebugLine(Pos0, Pos0 + Settings.ImpulseScale * Impulse, FColor::Red, false, Duration, Settings.DrawPriority, Settings.LineThickness * 0.5f);
+						FDebugDrawQueue::GetInstance().DrawDebugLine(Pos0, Pos0 + Settings.ImpulseScale * Impulse, FColor::Red, false, Duration, uint8(Settings.DrawPriority), Settings.LineThickness * 0.5f);
 					}
 				}
 			}
@@ -1226,10 +1296,10 @@ namespace Chaos
 				{
 					Sb = FMath::Lerp(Pb, Xb, FMath::Clamp<FReal>(CoMSize / Lena, 0., 1.));
 				}
-				FDebugDrawQueue::GetInstance().DrawDebugLine(Pa, Sa, FColor::White, false, UE_KINDA_SMALL_NUMBER, Settings.DrawPriority, ConnectorThickness);
-				FDebugDrawQueue::GetInstance().DrawDebugLine(Pb, Sb, FColor::White, false, UE_KINDA_SMALL_NUMBER, Settings.DrawPriority, ConnectorThickness);
-				FDebugDrawQueue::GetInstance().DrawDebugLine(Sa, Xa, R, false, UE_KINDA_SMALL_NUMBER, Settings.DrawPriority, ConnectorThickness);
-				FDebugDrawQueue::GetInstance().DrawDebugLine(Sb, Xb, C, false, UE_KINDA_SMALL_NUMBER, Settings.DrawPriority, ConnectorThickness);
+				FDebugDrawQueue::GetInstance().DrawDebugLine(Pa, Sa, FColor::White, false, UE_KINDA_SMALL_NUMBER, uint8(Settings.DrawPriority), ConnectorThickness);
+				FDebugDrawQueue::GetInstance().DrawDebugLine(Pb, Sb, FColor::White, false, UE_KINDA_SMALL_NUMBER, uint8(Settings.DrawPriority), ConnectorThickness);
+				FDebugDrawQueue::GetInstance().DrawDebugLine(Sa, Xa, R, false, UE_KINDA_SMALL_NUMBER, uint8(Settings.DrawPriority), ConnectorThickness);
+				FDebugDrawQueue::GetInstance().DrawDebugLine(Sb, Xb, C, false, UE_KINDA_SMALL_NUMBER, uint8(Settings.DrawPriority), ConnectorThickness);
 			}
 			if (FeatureMask.bCoMConnector)
 			{
@@ -1248,24 +1318,24 @@ namespace Chaos
 				{
 					Sb = FMath::Lerp(Cb, Xb, FMath::Clamp<FReal>(CoMSize / Lena, 0., 1.));
 				}
-				FDebugDrawQueue::GetInstance().DrawDebugLine(Ca, Sa, FColor::Black, false, UE_KINDA_SMALL_NUMBER, Settings.DrawPriority, ConnectorThickness);
-				FDebugDrawQueue::GetInstance().DrawDebugLine(Cb, Sb, FColor::Black, false, UE_KINDA_SMALL_NUMBER, Settings.DrawPriority, ConnectorThickness);
-				FDebugDrawQueue::GetInstance().DrawDebugLine(Sa, Xa, R, false, UE_KINDA_SMALL_NUMBER, Settings.DrawPriority, ConnectorThickness);
-				FDebugDrawQueue::GetInstance().DrawDebugLine(Sb, Xb, C, false, UE_KINDA_SMALL_NUMBER, Settings.DrawPriority, ConnectorThickness);
+				FDebugDrawQueue::GetInstance().DrawDebugLine(Ca, Sa, FColor::Black, false, UE_KINDA_SMALL_NUMBER, uint8(Settings.DrawPriority), ConnectorThickness);
+				FDebugDrawQueue::GetInstance().DrawDebugLine(Cb, Sb, FColor::Black, false, UE_KINDA_SMALL_NUMBER, uint8(Settings.DrawPriority), ConnectorThickness);
+				FDebugDrawQueue::GetInstance().DrawDebugLine(Sa, Xa, R, false, UE_KINDA_SMALL_NUMBER, uint8(Settings.DrawPriority), ConnectorThickness);
+				FDebugDrawQueue::GetInstance().DrawDebugLine(Sb, Xb, C, false, UE_KINDA_SMALL_NUMBER, uint8(Settings.DrawPriority), ConnectorThickness);
 			}
 			if (FeatureMask.bStretch)
 			{
 				const FRealSingle StretchThickness = 3.0f * Settings.LineThickness;
-				FDebugDrawQueue::GetInstance().DrawDebugLine(Xa, Xb, M, false, UE_KINDA_SMALL_NUMBER, Settings.DrawPriority, StretchThickness);
+				FDebugDrawQueue::GetInstance().DrawDebugLine(Xa, Xb, M, false, UE_KINDA_SMALL_NUMBER, uint8(Settings.DrawPriority), StretchThickness);
 			}
 			if (FeatureMask.bAxes)
 			{
-				FDebugDrawQueue::GetInstance().DrawDebugDirectionalArrow(Xa, Xa + Settings.DrawScale * Settings.ConstraintAxisLen * SpaceTransform.TransformVector(Ra.GetAxis(0)), Settings.DrawScale * Settings.ArrowSize, R, false, UE_KINDA_SMALL_NUMBER, Settings.DrawPriority, Settings.LineThickness);
-				FDebugDrawQueue::GetInstance().DrawDebugDirectionalArrow(Xa, Xa + Settings.DrawScale * Settings.ConstraintAxisLen * SpaceTransform.TransformVector(Ra.GetAxis(1)), Settings.DrawScale * Settings.ArrowSize, G, false, UE_KINDA_SMALL_NUMBER, Settings.DrawPriority, Settings.LineThickness);
-				FDebugDrawQueue::GetInstance().DrawDebugDirectionalArrow(Xa, Xa + Settings.DrawScale * Settings.ConstraintAxisLen * SpaceTransform.TransformVector(Ra.GetAxis(2)), Settings.DrawScale * Settings.ArrowSize, B, false, UE_KINDA_SMALL_NUMBER, Settings.DrawPriority, Settings.LineThickness);
-				FDebugDrawQueue::GetInstance().DrawDebugDirectionalArrow(Xb, Xb + Settings.DrawScale * Settings.ConstraintAxisLen * SpaceTransform.TransformVector(Rb.GetAxis(0)), Settings.DrawScale * Settings.ArrowSize, C, false, UE_KINDA_SMALL_NUMBER, Settings.DrawPriority, Settings.LineThickness);
-				FDebugDrawQueue::GetInstance().DrawDebugDirectionalArrow(Xb, Xb + Settings.DrawScale * Settings.ConstraintAxisLen * SpaceTransform.TransformVector(Rb.GetAxis(1)), Settings.DrawScale * Settings.ArrowSize, M, false, UE_KINDA_SMALL_NUMBER, Settings.DrawPriority, Settings.LineThickness);
-				FDebugDrawQueue::GetInstance().DrawDebugDirectionalArrow(Xb, Xb + Settings.DrawScale * Settings.ConstraintAxisLen * SpaceTransform.TransformVector(Rb.GetAxis(2)), Settings.DrawScale * Settings.ArrowSize, Y, false, UE_KINDA_SMALL_NUMBER, Settings.DrawPriority, Settings.LineThickness);
+				FDebugDrawQueue::GetInstance().DrawDebugDirectionalArrow(Xa, Xa + Settings.DrawScale * Settings.ConstraintAxisLen * SpaceTransform.TransformVector(Ra.GetAxis(0)), Settings.DrawScale * Settings.ArrowSize, R, false, UE_KINDA_SMALL_NUMBER, uint8(Settings.DrawPriority), Settings.LineThickness);
+				FDebugDrawQueue::GetInstance().DrawDebugDirectionalArrow(Xa, Xa + Settings.DrawScale * Settings.ConstraintAxisLen * SpaceTransform.TransformVector(Ra.GetAxis(1)), Settings.DrawScale * Settings.ArrowSize, G, false, UE_KINDA_SMALL_NUMBER, uint8(Settings.DrawPriority), Settings.LineThickness);
+				FDebugDrawQueue::GetInstance().DrawDebugDirectionalArrow(Xa, Xa + Settings.DrawScale * Settings.ConstraintAxisLen * SpaceTransform.TransformVector(Ra.GetAxis(2)), Settings.DrawScale * Settings.ArrowSize, B, false, UE_KINDA_SMALL_NUMBER, uint8(Settings.DrawPriority), Settings.LineThickness);
+				FDebugDrawQueue::GetInstance().DrawDebugDirectionalArrow(Xb, Xb + Settings.DrawScale * Settings.ConstraintAxisLen * SpaceTransform.TransformVector(Rb.GetAxis(0)), Settings.DrawScale * Settings.ArrowSize, C, false, UE_KINDA_SMALL_NUMBER, uint8(Settings.DrawPriority), Settings.LineThickness);
+				FDebugDrawQueue::GetInstance().DrawDebugDirectionalArrow(Xb, Xb + Settings.DrawScale * Settings.ConstraintAxisLen * SpaceTransform.TransformVector(Rb.GetAxis(1)), Settings.DrawScale * Settings.ArrowSize, M, false, UE_KINDA_SMALL_NUMBER, uint8(Settings.DrawPriority), Settings.LineThickness);
+				FDebugDrawQueue::GetInstance().DrawDebugDirectionalArrow(Xb, Xb + Settings.DrawScale * Settings.ConstraintAxisLen * SpaceTransform.TransformVector(Rb.GetAxis(2)), Settings.DrawScale * Settings.ArrowSize, Y, false, UE_KINDA_SMALL_NUMBER, uint8(Settings.DrawPriority), Settings.LineThickness);
 			}
 		}
 
@@ -1323,14 +1393,14 @@ namespace Chaos
 			const FVec3 Pos = SimSpace.Transform.GetLocation();
 			const FRotation3& Rot = SimSpace.Transform.GetRotation();
 			const FMatrix33 Rotm = Rot.ToMatrix();
-			FDebugDrawQueue::GetInstance().DrawDebugDirectionalArrow(Pos, Pos + Settings.DrawScale * Settings.BodyAxisLen * Rotm.GetAxis(0), Settings.DrawScale * Settings.ArrowSize, FColor::Red, false, UE_KINDA_SMALL_NUMBER, Settings.DrawPriority, Settings.LineThickness);
-			FDebugDrawQueue::GetInstance().DrawDebugDirectionalArrow(Pos, Pos + Settings.DrawScale * Settings.BodyAxisLen * Rotm.GetAxis(1), Settings.DrawScale * Settings.ArrowSize, FColor::Green, false, UE_KINDA_SMALL_NUMBER, Settings.DrawPriority, Settings.LineThickness);
-			FDebugDrawQueue::GetInstance().DrawDebugDirectionalArrow(Pos, Pos + Settings.DrawScale * Settings.BodyAxisLen * Rotm.GetAxis(2), Settings.DrawScale * Settings.ArrowSize, FColor::Blue, false, UE_KINDA_SMALL_NUMBER, Settings.DrawPriority, Settings.LineThickness);
-
-			FDebugDrawQueue::GetInstance().DrawDebugLine(Pos, Pos + Settings.VelScale * SimSpace.LinearVelocity, FColor::Cyan, false, UE_KINDA_SMALL_NUMBER, Settings.DrawPriority, Settings.LineThickness);
-			FDebugDrawQueue::GetInstance().DrawDebugLine(Pos, Pos + Settings.AngVelScale * SimSpace.AngularVelocity, FColor::Cyan, false, UE_KINDA_SMALL_NUMBER, Settings.DrawPriority, Settings.LineThickness);
-			FDebugDrawQueue::GetInstance().DrawDebugLine(Pos, Pos + 0.01f * Settings.VelScale * SimSpace.LinearAcceleration, FColor::Yellow, false, UE_KINDA_SMALL_NUMBER, Settings.DrawPriority, Settings.LineThickness);
-			FDebugDrawQueue::GetInstance().DrawDebugLine(Pos, Pos + 0.01f * Settings.AngVelScale * SimSpace.AngularAcceleration, FColor::Orange, false, UE_KINDA_SMALL_NUMBER, Settings.DrawPriority, Settings.LineThickness);
+			FDebugDrawQueue::GetInstance().DrawDebugDirectionalArrow(Pos, Pos + Settings.DrawScale * Settings.BodyAxisLen * Rotm.GetAxis(0), Settings.DrawScale * Settings.ArrowSize, FColor::Red, false, UE_KINDA_SMALL_NUMBER, uint8(Settings.DrawPriority), Settings.LineThickness);
+			FDebugDrawQueue::GetInstance().DrawDebugDirectionalArrow(Pos, Pos + Settings.DrawScale * Settings.BodyAxisLen * Rotm.GetAxis(1), Settings.DrawScale * Settings.ArrowSize, FColor::Green, false, UE_KINDA_SMALL_NUMBER, uint8(Settings.DrawPriority), Settings.LineThickness);
+			FDebugDrawQueue::GetInstance().DrawDebugDirectionalArrow(Pos, Pos + Settings.DrawScale * Settings.BodyAxisLen * Rotm.GetAxis(2), Settings.DrawScale * Settings.ArrowSize, FColor::Blue, false, UE_KINDA_SMALL_NUMBER, uint8(Settings.DrawPriority), Settings.LineThickness);
+
+			FDebugDrawQueue::GetInstance().DrawDebugLine(Pos, Pos + Settings.VelScale * SimSpace.LinearVelocity, FColor::Cyan, false, UE_KINDA_SMALL_NUMBER, uint8(Settings.DrawPriority), Settings.LineThickness);
+			FDebugDrawQueue::GetInstance().DrawDebugLine(Pos, Pos + Settings.AngVelScale * SimSpace.AngularVelocity, FColor::Cyan, false, UE_KINDA_SMALL_NUMBER, uint8(Settings.DrawPriority), Settings.LineThickness);
+			FDebugDrawQueue::GetInstance().DrawDebugLine(Pos, Pos + 0.01f * Settings.VelScale * SimSpace.LinearAcceleration, FColor::Yellow, false, UE_KINDA_SMALL_NUMBER, uint8(Settings.DrawPriority), Settings.LineThickness);
+			FDebugDrawQueue::GetInstance().DrawDebugLine(Pos, Pos + 0.01f * Settings.AngVelScale * SimSpace.AngularAcceleration, FColor::Orange, false, UE_KINDA_SMALL_NUMBER, uint8(Settings.DrawPriority), Settings.LineThickness);
 		}
 
 		void DrawConstraintGraphImpl(const FRigidTransform3& SpaceTransform, const Private::FPBDIslandManager& Graph, const FChaosDebugDrawSettings& Settings)
@@ -1357,18 +1427,18 @@ namespace Chaos
 				if ((bChaosDebugDebugDrawContactGraphUsed && bIsUsed) || (bChaosDebugDebugDrawContactGraphUnused && !bIsUsed))
 				{
 					FColor Color = bIsUsed ? FColor::Green : FColor::Red;
-					FDebugDrawQueue::GetInstance().DrawDebugLine(Transform0.GetLocation(), Transform1.GetLocation(), Color, false, UE_KINDA_SMALL_NUMBER, Settings.DrawPriority, Settings.LineThickness);
+					FDebugDrawQueue::GetInstance().DrawDebugLine(Transform0.GetLocation(), Transform1.GetLocation(), Color, false, UE_KINDA_SMALL_NUMBER, uint8(Settings.DrawPriority), Settings.LineThickness);
 				}
 
 				if (bChaosDebugDebugDrawContactGraph)
 				{
 					if (FConstGenericParticleHandle(Constraint->GetParticle0())->IsDynamic())
 					{
-						FDebugDrawQueue::GetInstance().DrawDebugLine(Transform0.GetLocation(), ContactPos, FColor::Red, false, UE_KINDA_SMALL_NUMBER, Settings.DrawPriority, Settings.LineThickness);
+						FDebugDrawQueue::GetInstance().DrawDebugLine(Transform0.GetLocation(), ContactPos, FColor::Red, false, UE_KINDA_SMALL_NUMBER, uint8(Settings.DrawPriority), Settings.LineThickness);
 					}
 					if (FConstGenericParticleHandle(Constraint->GetParticle1())->IsDynamic())
 					{
-						FDebugDrawQueue::GetInstance().DrawDebugLine(Transform1.GetLocation(), ContactPos, FColor::White, false, UE_KINDA_SMALL_NUMBER, Settings.DrawPriority, Settings.LineThickness);
+						FDebugDrawQueue::GetInstance().DrawDebugLine(Transform1.GetLocation(), ContactPos, FColor::White, false, UE_KINDA_SMALL_NUMBER, uint8(Settings.DrawPriority), Settings.LineThickness);
 					}
 					if (ColorIndex >= 0)
 					{
@@ -1400,7 +1470,7 @@ namespace Chaos
 
 					const FColor IslandColor = GetIslandColor(IslandIndex, !Graph.GetIsland(IslandIndex)->IsSleeping());
 					const FAABB3 Bounds = IslandAABB.TransformedAABB(SpaceTransform);
-					FDebugDrawQueue::GetInstance().DrawDebugBox(Bounds.Center(), 0.5f * Bounds.Extents(), SpaceTransform.GetRotation(), IslandColor, false, UE_KINDA_SMALL_NUMBER, Settings.DrawPriority, 3.0f * Settings.LineThickness);
+					FDebugDrawQueue::GetInstance().DrawDebugBox(Bounds.Center(), 0.5f * Bounds.Extents(), SpaceTransform.GetRotation(), IslandColor, false, UE_KINDA_SMALL_NUMBER, uint8(Settings.DrawPriority), 3.0f * Settings.LineThickness);
 				}
 			}
 
@@ -1511,13 +1581,13 @@ namespace Chaos
 
 					FColor Color = (Connection.Handle0->Parent() == Connection.Handle1->Parent()) ? FColor::Green : FColor::Blue; // FMath::Lerp(FColor::Green, FColor::Red, Connection.Strain / MaxStrain);
 
-					FDebugDrawQueue::GetInstance().DrawDebugLine(Pos0, Pos1, Color, false, UE_KINDA_SMALL_NUMBER, Settings.DrawPriority, Settings.LineThickness);
+					FDebugDrawQueue::GetInstance().DrawDebugLine(Pos0, Pos1, Color, false, UE_KINDA_SMALL_NUMBER, uint8(Settings.DrawPriority), Settings.LineThickness);
 
 					if (bChaosDebugDrawConnectionGraphShowAreas)
 					{
 						const FReal Radius = FMath::Sqrt(FMath::Abs(Connection.Strain / UE_PI)); // represent area as a circle
 						const FMatrix Axes = FRotationMatrix::MakeFromX(Normal);
-						FDebugDrawQueue::GetInstance().DrawDebugCircle((Pos0 + Pos1) * 0.5, Radius, 24, Color, false, UE_KINDA_SMALL_NUMBER, Settings.DrawPriority, Settings.LineThickness, Axes.GetUnitAxis(EAxis::Y), Axes.GetUnitAxis(EAxis::Z), false);
+						FDebugDrawQueue::GetInstance().DrawDebugCircle((Pos0 + Pos1) * 0.5, Radius, 24, Color, false, UE_KINDA_SMALL_NUMBER, uint8(Settings.DrawPriority), Settings.LineThickness, Axes.GetUnitAxis(EAxis::Y), Axes.GetUnitAxis(EAxis::Z), false);
 					}
 
 
@@ -1599,7 +1669,7 @@ namespace Chaos
 						{
 							PointColor = Settings.ShapesColorsPerState.KinematicColor;
 						}
-						FDebugDrawQueue::GetInstance().DrawDebugPoint(GetClusterTransform(Child).GetLocation(), PointColor, false, UE_KINDA_SMALL_NUMBER, Settings.DrawPriority, Settings.LineThickness * 10);
+						FDebugDrawQueue::GetInstance().DrawDebugPoint(GetClusterTransform(Child).GetLocation(), PointColor, false, UE_KINDA_SMALL_NUMBER, uint8(Settings.DrawPriority), Settings.LineThickness * 10);
 
 						if (bChaosDebugDrawConnectionGraphShowInternalStrains)
 						{
@@ -1644,7 +1714,7 @@ namespace Chaos
 			FDebugDrawQueue::GetInstance().DrawDebugLine(
 				SpaceTransform.TransformPosition(PWorld), 
 				SpaceTransform.TransformPosition(PWorld + AxisLen * AxisWorld), 
-				FColor::Green, false, UE_KINDA_SMALL_NUMBER, Settings.DrawPriority, Settings.LineThickness);
+				FColor::Green, false, UE_KINDA_SMALL_NUMBER, uint8(Settings.DrawPriority), Settings.LineThickness);
 
 			if (Settings.PushOutScale > 0)
 			{
@@ -1660,12 +1730,12 @@ namespace Chaos
 				FDebugDrawQueue::GetInstance().DrawDebugLine(
 					SpaceTransform.TransformPosition(PWorld),
 					SpaceTransform.TransformPosition(PWorld + Settings.PushOutScale * PushOutWorld),
-					SpringPushOutColor, false, UE_KINDA_SMALL_NUMBER, Settings.DrawPriority, SpringWidthScale * Settings.LineThickness);
+					SpringPushOutColor, false, UE_KINDA_SMALL_NUMBER, uint8(Settings.DrawPriority), SpringWidthScale * Settings.LineThickness);
 
 				FDebugDrawQueue::GetInstance().DrawDebugLine(
 					SpaceTransform.TransformPosition(PWorld),
 					SpaceTransform.TransformPosition(PWorld + Settings.PushOutScale * HardStopPushOutWorld),
-					HardStopPushOutColor, false, UE_KINDA_SMALL_NUMBER, Settings.DrawPriority, HardStopWidthScale * Settings.LineThickness);
+					HardStopPushOutColor, false, UE_KINDA_SMALL_NUMBER, uint8(Settings.DrawPriority), HardStopWidthScale * Settings.LineThickness);
 			}
 
 			if (Settings.ImpulseScale > 0)
@@ -1675,7 +1745,7 @@ namespace Chaos
 				FDebugDrawQueue::GetInstance().DrawDebugLine(
 					SpaceTransform.TransformPosition(PWorld),
 					SpaceTransform.TransformPosition(PWorld + Settings.PushOutScale * HardStopImpulseWorld),
-					HardStopImpulseColor, false, UE_KINDA_SMALL_NUMBER, Settings.DrawPriority, Settings.LineThickness * FRealSingle(0.5));
+					HardStopImpulseColor, false, UE_KINDA_SMALL_NUMBER, uint8(Settings.DrawPriority), Settings.LineThickness * FRealSingle(0.5));
 			}
 		}
 
@@ -1735,6 +1805,14 @@ namespace Chaos
 			}
 		}
 
+		void DrawParticleBVH(const FRigidTransform3& SpaceTransform, const FGeometryParticleHandle* Particle, const FColor& Color, const FChaosDebugDrawSettings* Settings)
+		{
+			if (FDebugDrawQueue::IsDebugDrawingEnabled())
+			{
+				DrawParticleBVHImpl(SpaceTransform, Particle, Color, GetChaosDebugDrawSettings(Settings));
+			}
+		}
+
 		void DrawParticleBounds(const FRigidTransform3& SpaceTransform, const TParticleView<FGeometryParticles>& ParticlesView, const FReal Dt, const FChaosDebugDrawSettings* Settings)
 		{
 			if (FDebugDrawQueue::IsDebugDrawingEnabled())
@@ -1929,12 +2007,12 @@ namespace Chaos
 
 			virtual void Box(const FAABB3& InBox, const TVec3<FReal>& InLinearColor, float InThickness) override
 			{
-				FDebugDrawQueue::GetInstance().DrawDebugBox(InBox.Center(), InBox.Extents() * FReal(0.5), FQuat::Identity, FLinearColor(InLinearColor).ToFColor(false), false, -1.f, Settings.DrawPriority, InThickness);
+				FDebugDrawQueue::GetInstance().DrawDebugBox(InBox.Center(), InBox.Extents() * FReal(0.5), FQuat::Identity, FLinearColor(InLinearColor).ToFColor(false), false, -1.f, uint8(Settings.DrawPriority), InThickness);
 			}
 
 			virtual void Line(const TVec3<FReal>& InBegin, const TVec3<FReal>& InEnd, const TVec3<FReal>& InLinearColor, float InThickness) override
 			{
-				FDebugDrawQueue::GetInstance().DrawDebugLine(InBegin, InEnd, FLinearColor(InLinearColor).ToFColor(false), false, -1.f, Settings.DrawPriority, InThickness);
+				FDebugDrawQueue::GetInstance().DrawDebugLine(InBegin, InEnd, FLinearColor(InLinearColor).ToFColor(false), false, -1.f, uint8(Settings.DrawPriority), InThickness);
 			}
 		private:
 			FChaosDebugDrawSettings Settings;
diff --git a/Engine/Source/Runtime/Experimental/Chaos/Private/Chaos/ClusterUnionManager.cpp b/Engine/Source/Runtime/Experimental/Chaos/Private/Chaos/ClusterUnionManager.cpp
--- a/Engine/Source/Runtime/Experimental/Chaos/Private/Chaos/ClusterUnionManager.cpp
+++ b/Engine/Source/Runtime/Experimental/Chaos/Private/Chaos/ClusterUnionManager.cpp
@@ -251,7 +251,10 @@ namespace Chaos
 			}
 		}
 
-		return MakeShared<FImplicitObjectUnion>(MoveTemp(Objects));
+		FImplicitObjectUnion* NewGeometry = new FImplicitObjectUnion(MoveTemp(Objects));
+		NewGeometry->SetAllowBVH(true);
+
+		return TSharedPtr<FImplicitObject, ESPMode::ThreadSafe>(NewGeometry);
 	}
 
 	DECLARE_CYCLE_STAT(TEXT("FClusterUnionManager::ClaimNextUnionIndex"), STAT_ClaimNextUnionIndex, STATGROUP_Chaos);
diff --git a/Engine/Source/Runtime/Experimental/Chaos/Private/Chaos/Collision/CollisionConstraintAllocator.cpp b/Engine/Source/Runtime/Experimental/Chaos/Private/Chaos/Collision/CollisionConstraintAllocator.cpp
--- a/Engine/Source/Runtime/Experimental/Chaos/Private/Chaos/Collision/CollisionConstraintAllocator.cpp
+++ b/Engine/Source/Runtime/Experimental/Chaos/Private/Chaos/Collision/CollisionConstraintAllocator.cpp
@@ -49,11 +49,11 @@ namespace Chaos
 				Particle0->ParticleCollisions().AddMidPhase(Particle0, MidPhase);
 				Particle1->ParticleCollisions().AddMidPhase(Particle1, MidPhase);
 
-	#if CHAOS_COLLISION_OBJECTPOOL_ENABLED 
+#if CHAOS_MIDPHASE_OBJECTPOOL_ENABLED 
 				Allocator->AddMidPhase(FParticlePairMidPhasePtr(MidPhase, FParticlePairMidPhaseDeleter(MidPhasePool)));
-	#else
+#else
 				Allocator->AddMidPhase(FParticlePairMidPhasePtr(MidPhase));
-	#endif
+#endif
 			}
 			NewMidPhases.Reset();
 		}
diff --git a/Engine/Source/Runtime/Experimental/Chaos/Private/Chaos/Collision/ParticlePairMidPhase.cpp b/Engine/Source/Runtime/Experimental/Chaos/Private/Chaos/Collision/ParticlePairMidPhase.cpp
--- a/Engine/Source/Runtime/Experimental/Chaos/Private/Chaos/Collision/ParticlePairMidPhase.cpp
+++ b/Engine/Source/Runtime/Experimental/Chaos/Private/Chaos/Collision/ParticlePairMidPhase.cpp
@@ -1,8 +1,10 @@
 // Copyright Epic Games, Inc. All Rights Reserved.
 #include "Chaos/Collision/ParticlePairMidPhase.h"
+#include "Chaos/ChaosDebugDraw.h"
 #include "Chaos/Collision/CollisionConstraintAllocator.h"
 #include "Chaos/Collision/CollisionContext.h"
 #include "Chaos/Collision/CollisionFilter.h"
+#include "Chaos/Collision/ContactTriangles.h"
 #include "Chaos/Collision/PBDCollisionConstraint.h"
 #include "Chaos/CollisionResolution.h"
 #include "Chaos/ImplicitObject.h"
@@ -10,11 +12,9 @@
 #include "Chaos/ParticleHandle.h"
 #include "Chaos/Particle/ParticleUtilities.h"
 #include "Chaos/PBDCollisionConstraints.h"
-
-#include "Chaos/ChaosDebugDraw.h"
 #include "ChaosStats.h"
+#include "Misc/MemStack.h"
 
-//UE_DISABLE_OPTIMIZATION
 
 extern bool Chaos_Collision_NarrowPhase_AABBBoundsCheck;
 
@@ -34,7 +34,6 @@ namespace Chaos
 
 	namespace CVars
 	{
-
 		bool bChaos_Collision_MidPhase_EnableBoundsChecks = true;
 		FAutoConsoleVariableRef CVarChaos_Collision_EnableBoundsChecks(TEXT("p.Chaos.Collision.EnableBoundsChecks"), bChaos_Collision_MidPhase_EnableBoundsChecks, TEXT(""));
 
@@ -43,25 +42,36 @@ namespace Chaos
 		FAutoConsoleVariableRef CVarChaos_Collision_CullDistanceReferenceSize(TEXT("p.Chaos.Collision.CullDistanceReferenceSize"), Chaos_Collision_CullDistanceScaleInverseSize, TEXT(""));
 		FAutoConsoleVariableRef CVarChaos_Collision_MinCullDistanecScale(TEXT("p.Chaos.Collision.MinCullDistanceScale"), Chaos_Collision_MinCullDistanceScale, TEXT(""));
 
+		// Whether we support the pre-flattened shape pair list which optimizes the common midphase case of collision between particles 
+		// with only a small number implicit objects each. If so, how many shape pairs do we allow before switching to the generic version?
+		bool bChaos_Collision_MidPhase_EnableShapePairs = true;
 		int32 Chaos_Collision_MidPhase_MaxShapePairs = 100;
+		FAutoConsoleVariableRef CVarChaos_Collision_EnableShapePairs(TEXT("p.Chaos.Collision.EnableShapePairs"), bChaos_Collision_MidPhase_EnableShapePairs, TEXT(""));
 		FAutoConsoleVariableRef CVarChaos_Collision_MaxShapePairs(TEXT("p.Chaos.Collision.MaxShapePairs"), Chaos_Collision_MidPhase_MaxShapePairs, TEXT(""));
 	}
+
 	using namespace CVars;
 
-#if CHAOS_CONSTRAINTHANDLE_DEBUG_ENABLED
-	void CheckConstraintContainsDynamic(const FPBDCollisionConstraint* Constraint)
+
+	////////////////////////////////////////////////////////////////////////////////////////
+	////////////////////////////////////////////////////////////////////////////////////////
+	////////////////////////////////////////////////////////////////////////////////////////
+	////////////////////////////////////////////////////////////////////////////////////////
+
+	inline bool ImplicitOverlapOBBToAABB(
+		const FImplicitObject* Implicit0,
+		const FImplicitObject* Implicit1,
+		const FRigidTransform3& ShapeTransform1To0,
+		const FReal CullDistance)
 	{
-		if (Constraint != nullptr)
+		if (Implicit0->HasBoundingBox() && Implicit1->HasBoundingBox())
 		{
-			const bool bIsDynamic0 = FConstGenericParticleHandle(Constraint->GetParticle0())->IsDynamic();
-			const bool bIsDynamic1 = FConstGenericParticleHandle(Constraint->GetParticle1())->IsDynamic();
-			if (!bIsDynamic0 && !bIsDynamic1)
-			{
-				ensure(false);
-			}
+			const FAABB3 Box1In0 = Implicit1->CalculateTransformedBounds(ShapeTransform1To0).Thicken(CullDistance);
+			const FAABB3 Box0 = Implicit0->BoundingBox();
+			return Box0.Intersects(Box1In0);
 		}
+		return true;
 	}
-#endif
 
 	inline bool ImplicitOverlapOBBToAABB(
 		const FImplicitObject* Implicit0,
@@ -70,14 +80,85 @@ namespace Chaos
 		const FRigidTransform3& ShapeWorldTransform1,
 		const FReal CullDistance)
 	{
-		if (Implicit0->HasBoundingBox() && Implicit1->HasBoundingBox())
+		const FRigidTransform3 ShapeTransform1To0 = ShapeWorldTransform1.GetRelativeTransform(ShapeWorldTransform0);
+		return ImplicitOverlapOBBToAABB(Implicit0, Implicit1, ShapeTransform1To0, CullDistance);
+	}
+
+	// Get the number of leaf objects in the implicit hierarchy, and set a flag if this hierarchy is a tree
+	// (i.e., contains a Union of Unions ro a BVH, rather than a single flat union at the root)
+	inline int32 GetNumLeafImplicits(const FImplicitObject* Implicit)
+	{
+		int32 NumImplicits = 0;
+
+		if (Implicit != nullptr)
 		{
-			const FRigidTransform3 Box1ToBox0TM = ShapeWorldTransform1.GetRelativeTransform(ShapeWorldTransform0);
-			const FAABB3 Box1In0 = Implicit1->CalculateTransformedBounds(Box1ToBox0TM).Thicken(CullDistance);
-			const FAABB3 Box0 = Implicit0->BoundingBox();
-			return Box0.Intersects(Box1In0);
+			// All implicit excepts unions have 1 leaf
+			NumImplicits = 1;
+
+			// @todo(chaos): we need better downcasting in FImplicitObject
+			const FImplicitObjectUnion* Union = Implicit->GetObject<FImplicitObjectUnion>();
+			if (Union == nullptr)
+			{
+				Union = Implicit->GetObject<FImplicitObjectUnionClustered>();
+			}
+
+			if (Union != nullptr)
+			{
+				NumImplicits = Union->GetNumLeafObjects();
+			}
 		}
-		return true;
+
+		return NumImplicits;
+	}
+
+	// Get the ShapeInstance data from the particle for the implicit with the specified root object index (its index
+	// in the root union implicit if there is one). Usually every implicit in the root union is represented in the
+	// ShapesArray, but not always. GCS and ClusterUnions sometimes have a Union at the root but only a single ShapeInstance
+	const FShapeInstance* GetShapeInstance(const FShapeInstanceArray& ShapeInstances, const int32 RootObjectIndex)
+	{
+		const int32 ShapeIndex = (ShapeInstances.IsValidIndex(RootObjectIndex)) ? RootObjectIndex : 0;
+		return ShapeInstances[ShapeIndex].Get();
+	}
+
+
+	// Find the LevelSet particles for an ImplicitObject. This used to just be the BVHParticles on the owning Particle,
+	// but we can now merge particles with BVHs into a ClusterUnion sp we can end up with LevelSet implicits referencing
+	// BVHParticles from other sources. Therefore We need to search the hierarchy for the implicit and if there is an 
+	// ImplicitObjectUnionClustered in the hierarchy, it will contain a mapping to the original particle so we can get 
+	// its BVHParticles.
+	// @todo(chaos): we should rework how BVHParticles are created and stored and fix this
+	const FBVHParticles* FindLevelSetParticles(const FGeometryParticleHandle* Particle, const FImplicitObject* Implicit)
+	{
+		const FBVHParticles* BVHParticles = nullptr;
+
+		if (const FPBDRigidParticleHandle* Rigid = Particle->CastToRigidParticle())
+		{
+			// Visit the hierarchy in pre-order depth-first traversal. If we hit a UnionClustered, see if it has our particles
+			Implicit->VisitObjects(
+				[Implicit, &BVHParticles](const FImplicitObject* HierarchyImplicit, const FRigidTransform3& Transform, const int32 RootObjectIndex, const int32 ObjectIndex, const int32 LeafObjectIndex) -> bool
+				{
+					if (const FImplicitObjectUnionClustered* UnionClustered = HierarchyImplicit->GetObject<FImplicitObjectUnionClustered>())
+					{
+						// If this union has our implicit in its map, extract the collision particles and we're done
+						if (const FPBDRigidParticleHandle* SourceRigid = UnionClustered->FindParticleForImplicitObject(Implicit))
+						{
+							BVHParticles = SourceRigid->CollisionParticles().Get();
+							return false;
+						}
+					}
+
+					// This is not the implicit we're looking for. Move along.
+					return true;
+				});
+
+			// If we didn't find a source particle, it must be ours (if we have any)
+			if (BVHParticles == nullptr)
+			{
+				BVHParticles = Rigid->CollisionParticles().Get();
+			}
+		}
+
+		return BVHParticles;
 	}
 
 	FPBDCollisionConstraintPtr CreateShapePairConstraint(
@@ -119,36 +200,6 @@ namespace Chaos
 		return Context.GetAllocator()->CreateConstraint(Particle0, Implicit0, Shape0, BVHParticles0, ShapeRelativeTransform0, Particle1, Implicit1, Shape1, BVHParticles1, ShapeRelativeTransform1, CullDistance, bUseManifold, ShapePairType);
 	}
 
-	struct FLeafImplicitObject
-	{
-		const FImplicitObject* Geometry;
-		const FShapeInstance* ShapeInstance;
-		FRigidTransform3 RelativeTransform;
-		FRigidTransform3 WorldTransform;
-		FAABB3 WorldBounds;
-		int32 RootObjectIndex;
-		int32 LeafObjectIndex;
-	};
-
-	void CollectLeafImplicitObjects(const FAABB3& LocalBounds, const FConstGenericParticleHandle& P, const FImplicitObjectUnion* Union, const FRigidTransform3 ParticleWorldTransform, const FReal CullDistance, TArray<FLeafImplicitObject>& Objects)
-	{
-		Union->VisitOverlappingLeafObjects(LocalBounds,
-			[&Objects, &P, &ParticleWorldTransform, CullDistance](const FImplicitObject* Implicit, const FRigidTransform3& RelativeTransform, const int32 RootObjectIndex, const int32 ObjectIndex, const int32 LeafObjectIndex)
-			{
-				const FRigidTransform3 WorldTransform = RelativeTransform * ParticleWorldTransform;
-				const FLeafImplicitObject LeafObject =
-				{
-					Implicit,
-					P->ShapeInstances()[RootObjectIndex].Get(),
-					RelativeTransform,
-					WorldTransform,
-					Implicit->CalculateTransformedBounds(WorldTransform).ThickenSymmetrically(FVec3(CullDistance)),
-					RootObjectIndex,
-					LeafObjectIndex
-				};
-				Objects.Emplace(LeafObject);
-			});
-	};
 
 	////////////////////////////////////////////////////////////////////////////////////////
 	////////////////////////////////////////////////////////////////////////////////////////
@@ -579,400 +630,296 @@ namespace Chaos
 	////////////////////////////////////////////////////////////////////////////////////////
 
 
-	FMultiShapePairCollisionDetector::FMultiShapePairCollisionDetector(
-		FGeometryParticleHandle* InParticle0,
-		const FPerShapeData* InShape0,
-		FGeometryParticleHandle* InParticle1,
-		const FPerShapeData* InShape1,
-		FParticlePairMidPhase& InMidPhase)
-		: MidPhase(InMidPhase)
-		, Constraints()
-		, Particle0(InParticle0)
-		, Particle1(InParticle1)
-		, Shape0(InShape0)
-		, Shape1(InShape1)
-	{
-	}
-
-	FMultiShapePairCollisionDetector::FMultiShapePairCollisionDetector(FMultiShapePairCollisionDetector&& R)
-		: MidPhase(R.MidPhase)
-		, Constraints(MoveTemp(R.Constraints))
-		, Particle0(R.Particle0)
-		, Particle1(R.Particle1)
-		, Shape0(R.Shape0)
-		, Shape1(R.Shape1)
+	FParticlePairMidPhase::FParticlePairMidPhase(const EParticlePairMidPhaseType InMidPhaseType)
+		: MidPhaseType(InMidPhaseType)
+		, Flags()
+		, Particle0(nullptr)
+		, Particle1(nullptr)
+		, CullDistanceScale(1)
+		, Key()
+		, LastUsedEpoch(INDEX_NONE)
+		, NumActiveConstraints(0)
+		, ParticleCollisionsIndex0(INDEX_NONE)
+		, ParticleCollisionsIndex1(INDEX_NONE)
 	{
-		check(R.NewConstraints.IsEmpty());
 	}
 
-	FMultiShapePairCollisionDetector::~FMultiShapePairCollisionDetector()
+	FParticlePairMidPhase::~FParticlePairMidPhase()
 	{
 	}
 
-	int32 FMultiShapePairCollisionDetector::GenerateCollisions(
-		const FReal CullDistance,
-		const FReal Dt,
-		const FCollisionContext& Context)
+	EParticlePairMidPhaseType FParticlePairMidPhase::CalculateMidPhaseType(FGeometryParticleHandle* InParticle0, FGeometryParticleHandle* InParticle1)
 	{
-		return GenerateCollisionsImpl(false, CullDistance, Dt, Context);
-	}
-
+		// For particles with a small number of shapes, we prebuild the set of potentially colliding
+		// shape pairs and cache various data with them. However we don't do this if there are a
+		// large number of shapes, or if the geometry hierarchy contains Unions of Unions.
+		//
+		// @todo(chaos): we should be able to support Unions of Unions in the accelerated path
+		// when there aren't many objects (with some changes to FSingleShapePairCollisionDetector)
 
-	int32 FMultiShapePairCollisionDetector::GenerateCollisionsCCD(
-		const bool bEnableCCDSweep,
-		const FReal CullDistance,
-		const FReal Dt,
-		const FCollisionContext& Context)
-	{
-		return GenerateCollisionsImpl(bEnableCCDSweep, CullDistance, Dt, Context);
-	}
+		// We can force use of the Generic midphase (mainly for testing)
+		if (!CVars::bChaos_Collision_MidPhase_EnableShapePairs)
+		{
+			return EParticlePairMidPhaseType::Generic;
+		}
 
-	int32 FMultiShapePairCollisionDetector::GenerateCollisionsImpl(
-		const bool bEnableCCDSweep,
-		const FReal CullDistance,
-		const FReal Dt,
-		const FCollisionContext& Context)
-	{
-		FConstGenericParticleHandle P0 = Particle0;
-		FConstGenericParticleHandle P1 = Particle1;
+		// How many implicits does each particle have?
+		const int32 NumImplicits0 = GetNumLeafImplicits(InParticle0->Geometry().Get());
+		const int32 NumImplicits1 = GetNumLeafImplicits(InParticle1->Geometry().Get());
 
-		const FImplicitObject* Implicit0 = Shape0->GetLeafGeometry();
-		const FBVHParticles* BVHParticles0 = P0->CollisionParticles().Get();
-		const FRigidTransform3& ShapeRelativeTransform0 = (FRigidTransform3)Shape0->GetLeafRelativeTransform();
-		const FRigidTransform3 ParticleWorldTransform0 = FParticleUtilities::GetActorWorldTransform(P0);
-		const FImplicitObject* Implicit1 = Shape1->GetLeafGeometry();
-		const FBVHParticles* BVHParticles1 = P1->CollisionParticles().Get();
-		const FRigidTransform3& ShapeRelativeTransform1 = (FRigidTransform3)Shape1->GetLeafRelativeTransform();
-		const FRigidTransform3 ParticleWorldTransform1 = FParticleUtilities::GetActorWorldTransform(P1);
-
-		FCollisionContext LocalContext = Context;
-		LocalContext.MultiShapeCollisionDetector = this;
-
-		// This recurively processes the Implicit Hierarchy and generates a collision constraint for every
-		// overlapping shape pair on the two bodies via a callback to FindOrCreateConstraint
-		Collisions::ConstructConstraints(
-			Particle0, Particle1,
-			Implicit0, Shape0, BVHParticles0, 0,
-			Implicit1, Shape1, BVHParticles1, 0,
-			ParticleWorldTransform0, ShapeRelativeTransform0,
-			ParticleWorldTransform1, ShapeRelativeTransform1,
-			CullDistance, Dt, bEnableCCDSweep, LocalContext);
+		// Do we have a ShapeInstance for every implicit object?
+		// Only the implicits in the root union are represented in the shapes array
+		const bool bIsTree0 = (NumImplicits0 > InParticle0->ShapeInstances().Num());
+		const bool bIsTree1 = (NumImplicits1 > InParticle1->ShapeInstances().Num());
 
-		// Generate manifolds for each constraint we created/recovered and (re)activate if necessary
-		int32 NumActiveConstraints = ProcessNewConstraints(CullDistance, Dt, Context);
+		// Are there too many implicits for the flattened shape pair path to be optimal?
+		const int32 MaxImplicitPairs = CVars::Chaos_Collision_MidPhase_MaxShapePairs;
+		const bool bTooManyImplicitPairs = (MaxImplicitPairs > 0) && (NumImplicits0 * NumImplicits1 > MaxImplicitPairs);
 
-		// @todo(chaos): we could clean up unused collisions between this pair, but probably not worth it
-		// and we would have to prevent cleanup for sleeping particles because the collisions will still
-		// be referenced in the IslandManager's constraint graph for the sleeping island.
-		//PruneConstraints();
+		// Is either of the implicits a Union of Unions?
+		const bool bIsTree = bIsTree0 || bIsTree1;
 
-		return NumActiveConstraints;
+		const bool bPrebuildShapePairs = !bTooManyImplicitPairs && !bIsTree;
+		return bPrebuildShapePairs ? EParticlePairMidPhaseType::ShapePair : EParticlePairMidPhaseType::Generic;
 	}
 
-	FPBDCollisionConstraint* FMultiShapePairCollisionDetector::FindOrCreateConstraint(
-		FGeometryParticleHandle* InParticle0,
-		const FImplicitObject* InImplicit0,
-		const int32 InImplicitId0,
-		const FShapeInstance* InShape0,
-		const FBVHParticles* InBVHParticles0,
-		const FRigidTransform3& InShapeRelativeTransform0,
-		FGeometryParticleHandle* InParticle1,
-		const FImplicitObject* InImplicit1,
-		const int32 InImplicitId1,
-		const FShapeInstance* InShape1,
-		const FBVHParticles* InBVHParticles1,
-		const FRigidTransform3& InShapeRelativeTransform1,
-		const FReal CullDistance,
-		const EContactShapesType ShapePairType,
-		const bool bUseManifold,
-		const bool bEnableSweep,
-		const FCollisionContext& Context)
+	FParticlePairMidPhase* FParticlePairMidPhase::Make(FGeometryParticleHandle* InParticle0, FGeometryParticleHandle* InParticle1)
 	{
-		// This is a callback from the low-level collision function. It should always be the same two particles, though the
-		// shapes may be some children in the implicit hierarchy. The particles could be in the opposite order though, and
-		// this will depend on the shape types involved. E.g., with two particles each with a sphere and a box in a union
-		// would require up to two Sphere-Box contacts, with the particles in opposite orders.
-		if (!ensure(((InParticle0 == Particle0) && (InParticle1 == Particle1)) || ((InParticle0 == Particle1) && (InParticle1 == Particle0))))
+		const EParticlePairMidPhaseType MidPhaseType = CalculateMidPhaseType(InParticle0, InParticle1);
+		if (MidPhaseType == EParticlePairMidPhaseType::ShapePair)
 		{
-			// We somehow received a callback for the wrong particle pair...this should not happen
-			return nullptr;
+			return new FShapePairParticlePairMidPhase();
+		}
+		else
+		{
+			return new FGenericParticlePairMidPhase();
 		}
+	}
 
-		const FCollisionParticlePairConstraintKey Key = FCollisionParticlePairConstraintKey(InShape0, InImplicit0, InImplicitId0, InBVHParticles0, InShape1, InImplicit1, InImplicitId1, InBVHParticles1);
-		FPBDCollisionConstraint* Constraint = FindConstraint(Key);
+	void FParticlePairMidPhase::DetachParticle(FGeometryParticleHandle* Particle)
+	{
+		Reset();
 
-		// @todo(chaos): fix key uniqueness guarantee.  We need a truly unique key gen function
-		const bool bIsKeyCollision = (Constraint != nullptr) && ((Constraint->GetImplicit0() != InImplicit0) || (Constraint->GetImplicit1() != InImplicit1) || (Constraint->GetCollisionParticles0() != InBVHParticles0) || (Constraint->GetCollisionParticles1() != InBVHParticles1));
-		if (bIsKeyCollision)
+		if (Particle == Particle0)
 		{
-			// If we get here, we have a key collision. The key uses a hash of pointers which is very likely to be unique for different implicit pairs, 
-			// especially since it only needs to be unique for this particle pair, but it is not guaranteed.
-			// Creating a new constraint with the same key could cause fatal problems (the original constraint will be deleted when we add the new one 
-			// to the map, but if it is asleep it will be referenced in the contact graph) so we just abort and accept we will miss collisions. 
-			// It is extremely unlikely to happen but we should fix it at some point.
-			ensure(false);
-			return nullptr;
+			Particle0 = nullptr;
 		}
-
-		if (Constraint == nullptr)
+		else if (Particle == Particle1)
 		{
-			// NOTE: Using InParticle0 and InParticle1 here because the order may be different to what we have stored
-			Constraint = CreateConstraint(InParticle0, InImplicit0, InShape0, InBVHParticles0, InShapeRelativeTransform0, InParticle1, InImplicit1, InShape1, InBVHParticles1, InShapeRelativeTransform1, CullDistance, ShapePairType, bUseManifold, Key, Context);
-
-			// Is this a CCD constraint?
-			Constraint->SetCCDEnabled(MidPhase.IsCCD());
+			Particle1 = nullptr;
 		}
+	}
 
-		// Do we want to sweep on this tick? 
-		// CCD may be temporarily disabled by the user or because we are moving slowly.
-		Constraint->SetCCDSweepEnabled(MidPhase.IsCCD() && bEnableSweep);
+	void FParticlePairMidPhase::Reset()
+	{
+		Flags.bIsActive = true;
+		Flags.bIsCCD = false;
+		Flags.bIsCCDActive = false;
+		Flags.bIsSleeping = false;
+		Flags.bIsModified = false;
 
-		NewConstraints.Add(Constraint);
-		return Constraint;
+		ResetImpl();
 	}
 
-	FPBDCollisionConstraint* FMultiShapePairCollisionDetector::FindConstraint(const FCollisionParticlePairConstraintKey& Key)
+	void FParticlePairMidPhase::ResetModifications()
 	{
-		FPBDCollisionConstraintPtr* PConstraint = Constraints.Find(Key.GetKey());
-		if (PConstraint != nullptr)
+		if (Flags.bIsModified)
 		{
-			return (*PConstraint).Get();
+			Flags.bIsActive = true;
+			Flags.bIsCCDActive = Flags.bIsCCD;
+			Flags.bIsModified = false;
 		}
-		return nullptr;
 	}
 
-	FPBDCollisionConstraint* FMultiShapePairCollisionDetector::CreateConstraint(
+	void FParticlePairMidPhase::Init(
 		FGeometryParticleHandle* InParticle0,
-		const FImplicitObject* Implicit0,
-		const FPerShapeData* InShape0,
-		const FBVHParticles* BVHParticles0,
-		const FRigidTransform3& ShapeRelativeTransform0,
 		FGeometryParticleHandle* InParticle1,
-		const FImplicitObject* Implicit1,
-		const FPerShapeData* InShape1,
-		const FBVHParticles* BVHParticles1,
-		const FRigidTransform3& ShapeRelativeTransform1,
-		const FReal CullDistance,
-		const EContactShapesType ShapePairType,
-		const bool bUseManifold,
-		const FCollisionParticlePairConstraintKey& Key,
+		const FCollisionParticlePairKey& InKey,
 		const FCollisionContext& Context)
 	{
-		PHYSICS_CSV_SCOPED_EXPENSIVE(PhysicsVerbose, NarrowPhase_CreateConstraint);
-		FPBDCollisionConstraintPtr Constraint = CreateImplicitPairConstraint(InParticle0, Implicit0, InShape0, BVHParticles0, ShapeRelativeTransform0, InParticle1, Implicit1, InShape1, BVHParticles1, ShapeRelativeTransform1, CullDistance, ShapePairType, bUseManifold, Context);
+		PHYSICS_CSV_SCOPED_EXPENSIVE(PhysicsVerbose, NarrowPhase_Filter);
 
-		const int32 CurrentEpoch = Context.GetAllocator()->GetCurrentEpoch();
-		Constraint->GetContainerCookie().MidPhase = &MidPhase;
-		Constraint->GetContainerCookie().bIsMultiShapePair = true;
-		Constraint->GetContainerCookie().CreationEpoch = CurrentEpoch;
+		Particle0 = InParticle0;
+		Particle1 = InParticle1;
+		Key = InKey;
+
+		Flags.bIsActive = true;
+
+		// If CCD is allowed in the current context and for at least one of
+		// the particles involved, enable it for this midphase.
+		//
+		// bIsCCDActive is reset to bIsCCD each frame, but can be overridden
+		// by modifiers.
+		const bool bIsCCD = Context.GetSettings().bAllowCCD && (
+			FConstGenericParticleHandle(Particle0)->CCDEnabled() ||
+			FConstGenericParticleHandle(Particle1)->CCDEnabled());
+		Flags.bIsCCD = bIsCCD;
+		Flags.bIsCCDActive = bIsCCD;
 
-		return Constraints.Add(Key.GetKey(), MoveTemp(Constraint)).Get();
+		BuildDetectorsImpl();
+
+		InitThresholds();
 	}
 
-	void FMultiShapePairCollisionDetector::WakeCollisions(const int32 SleepEpoch, const int32 CurrentEpoch)
+	bool FParticlePairMidPhase::ShouldEnableCCD(const FReal Dt)
 	{
-		for (auto& KVP : Constraints)
+		// bIsCCDActive is set to bIsCCD at the beginning of every frame, but may be
+		// overridden in midphase modification or potentially other systems which run
+		// in between mid and narrow phase. bIsCCDActive indicates the final
+		// overridden value so we use that here instead of bIsCCD.
+		if (Flags.bIsCCDActive)
 		{
-			FPBDCollisionConstraintPtr& Constraint = KVP.Value;
-			if (Constraint->GetContainerCookie().LastUsedEpoch >= SleepEpoch)
-			{
-				Constraint->GetContainerCookie().LastUsedEpoch = CurrentEpoch;
-				Constraint->GetContainerCookie().ConstraintIndex = INDEX_NONE;
-				Constraint->GetContainerCookie().CCDConstraintIndex = INDEX_NONE;
-			}
+			FConstGenericParticleHandle ConstParticle0 = FConstGenericParticleHandle(Particle0);
+			FConstGenericParticleHandle ConstParticle1 = FConstGenericParticleHandle(Particle1);
+
+			// See comments in FSingleShapePairCollisionDetector::GenerateCollisionCCDImpl for why we use V() here
+			FReal LengthCCD = 0;
+			FVec3 DirCCD = FVec3(0);
+			const FVec3 DeltaX0 = ConstParticle0->V() * Dt;
+			const FVec3 DeltaX1 = ConstParticle1->V() * Dt;
+			const bool bUseCCD = Collisions::ShouldUseCCD(Particle0, DeltaX0, Particle1, DeltaX1, DirCCD, LengthCCD);
+
+			return bUseCCD;
 		}
+		return false;
 	}
 
-	int32 FMultiShapePairCollisionDetector::ProcessNewConstraints(
-		const FReal CullDistance,
-		const FReal Dt,
-		const FCollisionContext& Context)
+	void FParticlePairMidPhase::InitThresholds()
 	{
-		int32 NumActiveConstraints = 0;
-		for (FPBDCollisionConstraint* Constraint : NewConstraints)
-		{
-			if (!Constraint->GetCCDSweepEnabled())
-			{
-				UpdateCollision(Constraint, CullDistance, Dt, Context);
-			}
-			else
-			{
-				UpdateCollisionCCD(Constraint, CullDistance, Dt, Context);
-			}
+		// @todo(chaos): improve this threshold calculation for thin objects? Dynamic thin objects have bigger problems so maybe we don't care
+		// @todo(chaos): Spheres and capsules need smaller position tolerance - the restore test doesn't work well with rolling
+		const bool bIsDynamic0 = FConstGenericParticleHandle(Particle0)->IsDynamic();
+		const bool bIsDynamic1 = FConstGenericParticleHandle(Particle1)->IsDynamic();
+		const bool bIsBounded0 = Particle0->HasBounds();
+		const bool bIsBounded1 = Particle1->HasBounds();
 
-			if (!Constraint->GetDisabled())
-			{
-				++NumActiveConstraints;
-			}
-		}
-		NewConstraints.Reset();
-		return NumActiveConstraints;
+		// NOTE: If CullDistance ends up smaller than the thresholds used to restore collisions, we can end up missing
+		// collisions as the objects move if we restore a "zero contact" manifold after movement greater than the cull distance. 
+		// Currently this should not happen, but it is not explicitly ensured by the way the thresholds and CullDistanceScale are calculated.
+		// @todo(chaos): Add a way to enforce a CullDistance big enough to support the reuse thresholds
+		const FReal CullDistanceReferenceSizeInv = FReal(Chaos_Collision_CullDistanceScaleInverseSize);
+		const FReal MinCullDistanceScale = FReal(Chaos_Collision_MinCullDistanceScale);
+		const FReal MaxBoundsSize0 = (bIsDynamic0 && bIsBounded0) ? Particle0->LocalBounds().Extents().GetMax() : FReal(0);
+		const FReal MaxBoundsSize1 = (bIsDynamic1 && bIsBounded1) ? Particle1->LocalBounds().Extents().GetMax() : FReal(0);
+		const FReal CullDistanceScale0 = MaxBoundsSize0 * CullDistanceReferenceSizeInv;
+		const FReal CullDistanceScale1 = MaxBoundsSize1 * CullDistanceReferenceSizeInv;
+		CullDistanceScale = (FRealSingle)FMath::Max3(CullDistanceScale0, CullDistanceScale1, MinCullDistanceScale);
 	}
 
-	void FMultiShapePairCollisionDetector::UpdateCollision(
-		FPBDCollisionConstraint* Constraint, 
-		const FReal CullDistance,
+	void FParticlePairMidPhase::GenerateCollisions(
+		const FReal InCullDistance,
 		const FReal Dt,
 		const FCollisionContext& Context)
 	{
-		// @todo(chaos): share this code with FSingleShapePairCollisionDetector
-
-		const FRigidTransform3 ShapeWorldTransform0 = Constraint->GetShapeRelativeTransform0() * FConstGenericParticleHandle(Constraint->GetParticle0())->GetTransformPQ();
-		const FRigidTransform3 ShapeWorldTransform1 = Constraint->GetShapeRelativeTransform1() * FConstGenericParticleHandle(Constraint->GetParticle1())->GetTransformPQ();
-		const int32 CurrentEpoch = Context.GetAllocator()->GetCurrentEpoch();
-		const int32 LastEpoch = CurrentEpoch - 1;
-		const bool bWasUpdatedLastTick = (Constraint->GetContainerCookie().LastUsedEpoch >= LastEpoch);
-
-		// Update the world shape transforms on the constraint (we cannot just give it the PerShapeData 
-		// pointer because of Unions - see FMultiShapePairCollisionDetector)
-		Constraint->SetShapeWorldTransforms(ShapeWorldTransform0, ShapeWorldTransform1);
-
-		Constraint->SetCullDistance(CullDistance);
+		//SCOPE_CYCLE_COUNTER(STAT_Collisions_GenerateCollisions);
+		PHYSICS_CSV_SCOPED_EXPENSIVE(PhysicsVerbose, DetectCollisions_NarrowPhase);
 
-		// If the constraint was not used last frame, it needs to be reset, otherwise we will try to reuse
-		if (!bWasUpdatedLastTick || (Constraint->GetManifoldPoints().Num() == 0))
+		if (!IsValid())
 		{
-			// Clear all manifold data including saved contact data
-			Constraint->ResetManifold();
+			return;
 		}
 
-		bool bWasManifoldRestored = false;
-		//if (Context.GetSettings().bAllowManifoldReuse && Flags.bEnableManifoldUpdate && bWasUpdatedLastTick)
-		//{
-		//	// Update the existing manifold. We can re-use as-is if none of the points have moved much and the bodies have not moved much
-		//	// NOTE: this can succeed in "restoring" even if we have no manifold points
-		//	// NOTE: this uses the transforms from SetLastShapeWorldTransforms, so we can only do this if we were updated last tick
-		//	bWasManifoldRestored = Constraint->UpdateAndTryRestoreManifold();
-		//}
-		//else
+		if (Flags.bIsActive)
 		{
-			// We are not trying to reuse manifold points, so reset them but leave stored data intact (for friction)
-			Constraint->ResetActiveManifoldContacts();
-		}
+			// CullDistance is scaled by the size of the dynamic objects.
+			FReal CullDistance = InCullDistance * CullDistanceScale;
 
-		if (!bWasManifoldRestored)
-		{
-			// We will be updating the manifold so update transforms used to check for movement in UpdateAndTryRestoreManifold on future ticks
-			Constraint->SetLastShapeWorldTransforms(ShapeWorldTransform0, ShapeWorldTransform1);
+			// If CCD is enabled, did we move far enough to require a sweep?
+			Flags.bUseSweep = Flags.bIsCCD && ShouldEnableCCD(Dt);
 
-			if (!Context.GetSettings().bDeferNarrowPhase)
+			// Extend cull distance based on velocity
+			// NOTE: We use PreV here which is the velocity after collisions from the previous tick because we want the 
+			// velocity without gravity from this tick applied. This is mainly so that we get the same CullDistance from 
+			// one tick to the next, even if one of the particles goes to sleep, and therefore its velocity is now zero 
+			// because gravity is no longer applied. Also see FPBDIslandManager::PropagateIslandSleep for other issues 
+			// related to velocity and sleeping...
+			const FReal VMax0 = FConstGenericParticleHandle(GetParticle0())->PreV().GetAbsMax();
+			const FReal VMax1 = FConstGenericParticleHandle(GetParticle1())->PreV().GetAbsMax();
+			const FReal VMaxDt = FMath::Max(VMax0, VMax1) * Dt;
+			if (!Flags.bUseSweep)
 			{
-				Collisions::UpdateConstraint(*Constraint, Constraint->GetShapeWorldTransform0(), Constraint->GetShapeWorldTransform1(), Dt);
+				// Normal (non sweep) mode: we increase CullDistance based on velocity up to a limit
+				// NOTE: This somewhat matches the bounds expansion in FPBDRigidsEvolutionGBF::Integrate
+				const FReal VelocityBoundsMultiplier = Context.GetSettings().BoundsVelocityInflation;
+				const FReal MaxVelocityBoundsExpansion = Context.GetSettings().MaxVelocityBoundsExpansion;
+				if ((VelocityBoundsMultiplier > 0) && (MaxVelocityBoundsExpansion > 0))
+				{
+					CullDistance += FMath::Min(VelocityBoundsMultiplier * VMaxDt, MaxVelocityBoundsExpansion);
+				}
 			}
-		}
-
-		// If we have a valid contact, add it to the active list
-		// We also add it to the active list if collision detection is deferred because the data will be filled in later and we
-		// don't know in advance whether we will pass the Phi check (deferred narrow phase is used with RBAN)
-		if (Constraint->GetPhi() <= CullDistance || Context.GetSettings().bDeferNarrowPhase)
-		{
-			if (Context.GetAllocator()->ActivateConstraint(Constraint))
+			else
 			{
-				return;
+				// CCD (sweept) mode: we increase CullDistance based on velocity, with no limits
+				CullDistance += VMaxDt;
 			}
-		}
-
-		// If we get here, we did not activate hte constraint and it should be disabled for this tick
-		Constraint->SetDisabled(true);
-	}
-
 
-	void FMultiShapePairCollisionDetector::UpdateCollisionCCD(
-		FPBDCollisionConstraint* Constraint,
-		const FReal CullDistance,
-		const FReal Dt,
-		const FCollisionContext& Context)
-	{
-		// @todo(chaos): share this code with FSingleShapePairCollisionDetector
-
-		if (!Constraint->GetCCDSweepEnabled())
-		{
-			return UpdateCollision(Constraint, CullDistance, Dt, Context);
+			// Run collision detection on all potentially colliding shape pairs
+			NumActiveConstraints = GenerateCollisionsImpl(CullDistance, Dt, Context);
 		}
 
-		Constraint->ResetManifold();
-		Constraint->ResetActiveManifoldContacts();
-
-		const FRigidTransform3 ShapeWorldTransform0 = Constraint->GetShapeRelativeTransform0() * FConstGenericParticleHandle(Constraint->GetParticle0())->GetTransformPQ();
-		const FRigidTransform3 ShapeWorldTransform1 = Constraint->GetShapeRelativeTransform1() * FConstGenericParticleHandle(Constraint->GetParticle1())->GetTransformPQ();
+		// Reset any modifications applied by the MidPhaseModifier.
+		// In principle this belongs at the start of the frame but it is hard to avoid the cache miss there.
+		ResetModifications();
 
-		// Update the world shape transforms on the constraint (we cannot just give it the PerShapeData 
-		// pointer because of Unions - see FMultiShapePairCollisionDetector)
-		// NOTE: these are not used by CCD which continuously moves the particles
-		Constraint->SetShapeWorldTransforms(ShapeWorldTransform0, ShapeWorldTransform1);
+		LastUsedEpoch = Context.GetAllocator()->GetCurrentEpoch();
+	}
 
-		Constraint->SetCullDistance(CullDistance);
-		Constraint->ResetManifold();
-		Constraint->ResetActiveManifoldContacts();
+	void FParticlePairMidPhase::InjectCollision(const FPBDCollisionConstraint& Constraint, const FCollisionContext& Context)
+	{
+		InjectCollisionImpl(Constraint, Context);
 
-		FConstGenericParticleHandle P0 = Constraint->GetParticle0();
-		FConstGenericParticleHandle P1 = Constraint->GetParticle1();
+		LastUsedEpoch = Context.GetAllocator()->GetCurrentEpoch();
+	}
 
-		// We need the previous transform for the swept collision detector. It assumes that the current
-		// transform has been set on the constraint. 
-		// We assume that the particle's center of mass moved in a stright line and that it's rotation has 
-		// not changed so we calculate the previous transform from the current one and the velocity.
-		// NOTE: These are actor transforms, not CoM transforms
-		// @todo(chaos): Pass both start and end transforms to the collision detector
-		const FRigidTransform3 CCDParticleWorldTransform0 = FRigidTransform3(P0->P() - P0->V() * Dt, P0->Q());
-		const FRigidTransform3 CCDParticleWorldTransform1 = FRigidTransform3(P1->P() - P1->V() * Dt, P1->Q());
-		const FRigidTransform3 CCDShapeWorldTransform0 = Constraint->GetShapeRelativeTransform0() * CCDParticleWorldTransform0;
-		const FRigidTransform3 CCDShapeWorldTransform1 = Constraint->GetShapeRelativeTransform1() * CCDParticleWorldTransform1;
-		const bool bDidSweep = Collisions::UpdateConstraintSwept(*Constraint, CCDShapeWorldTransform0, CCDShapeWorldTransform1, Dt);
+	void FParticlePairMidPhase::SetIsSleeping(const bool bInIsSleeping, const int32 CurrentEpoch)
+	{
+		// This can be called from two locations:
+		// 1)	At the start of the tick as a results of some state change from the game thread such as an explicit wake event,
+		//		applying an impulse, or moving a particle.
+		// 2)	After the constraint solver phase when we put non-moving islands to sleep.
+		// 
+		// Note that in both cases there is a collision detection phase before the next constraint solving phase.
+		//
+		// When awakening we re-activate collisions so that we have a "previous" collision to use for static friction etc.
+		// We don't need to do anything when going to sleep because sleeping particles pairs are ignored in collision detection 
+		// so the next set of active collisions generated will not contain these collisions.
 
-#if CHAOS_DEBUG_DRAW
-		if (CVars::ChaosSolverDrawCCDInteractions)
+		if (Flags.bIsSleeping != bInIsSleeping)
 		{
-			if (FConstGenericParticleHandle(Constraint->GetParticle0())->CCDEnabled())
-			{
-				DebugDraw::DrawShape(CCDShapeWorldTransform0, Constraint->GetImplicit0(), Constraint->GetShape0(), FColor::Black, DebugDraw::ChaosDebugDrawCCDDuration, &CVars::ChaosSolverDebugDebugDrawSettings);
-				DebugDraw::DrawShape(ShapeWorldTransform0, Constraint->GetImplicit0(), Constraint->GetShape0(), FColor::White, DebugDraw::ChaosDebugDrawCCDDuration, &CVars::ChaosSolverDebugDebugDrawSettings);
-			}
-			if (FConstGenericParticleHandle(Constraint->GetParticle1())->CCDEnabled())
+			// If we are waking particles, reactivate all collisions that were
+			// active when we were put to sleep, i.e., all collisions whose LastUsedEpoch
+			// is equal to our LastUsedEpoch.
+			const bool bWakingUp = !bInIsSleeping;
+			if (bWakingUp)
 			{
-				DebugDraw::DrawShape(CCDShapeWorldTransform1, Constraint->GetImplicit1(), Constraint->GetShape1(), FColor::Black, DebugDraw::ChaosDebugDrawCCDDuration, &CVars::ChaosSolverDebugDebugDrawSettings);
-				DebugDraw::DrawShape(ShapeWorldTransform1, Constraint->GetImplicit1(), Constraint->GetShape1(), FColor::White, DebugDraw::ChaosDebugDrawCCDDuration, &CVars::ChaosSolverDebugDebugDrawSettings);
-			}
-		}
-#endif
+				if (LastUsedEpoch < CurrentEpoch)
+				{
+					// Restore all constraints that were active when we were put to sleep
+					WakeCollisionsImpl(CurrentEpoch);
 
-		// If we did get a hit but it's at TOI = 1, treat this constraint as a regular non-swept constraint (skip the rewind)
-		if ((!bDidSweep) || Constraint->GetCCDTimeOfImpact() == FReal(1))
-		{
-			Collisions::UpdateConstraint(*Constraint, Constraint->GetShapeWorldTransform0(), Constraint->GetShapeWorldTransform1(), Dt);
-		}
+					LastUsedEpoch = CurrentEpoch;
+				}
+			}
+			// If we are going to sleep, there is nothing to do (see comments above)
 
-		// If the sweep did not find a hit, or we are treating it like a regular contact, we will skip the CCD rewind step for this contact
-		if (Constraint->GetCCDTimeOfImpact() >= FReal(1))
-		{
-			Constraint->SetCCDSweepEnabled(false);
+			Flags.bIsSleeping = bInIsSleeping;
 		}
-
-		Context.GetAllocator()->ActivateConstraint(Constraint);
 	}
 
-	void FMultiShapePairCollisionDetector::PruneConstraints(const int32 CurrentEpoch)
+	bool FParticlePairMidPhase::IsInConstraintGraph() const
 	{
-		// We don't prune from NewCollisions - must call ProcessNewCollisions before Prune
-		check(NewConstraints.Num() == 0);
-
-		// Find all the expired collisions
-		TArray<uint32> Pruned;
-		for (auto& KVP : Constraints)
+		// @todo(chaos): optimize
+		bool bInGraph = false;
+		VisitConstCollisions([&bInGraph](const FPBDCollisionConstraint& Constraint)
 		{
-			const uint32 Key = KVP.Key;
-			FPBDCollisionConstraintPtr& Constraint = KVP.Value;
-			if (Constraint->GetContainerCookie().LastUsedEpoch < CurrentEpoch)
+			if (Constraint.GetConstraintGraphEdge() != nullptr)
 			{
-				Pruned.Add(Key);
+				bInGraph = true;
+				return ECollisionVisitorResult::Stop;
 			}
-		}
-
-		// Destroy expired collisions
-		for (uint32 Key : Pruned)
-		{
-			Constraints.Remove(Key);
-		}
+			return ECollisionVisitorResult::Continue;
+		}, false);
+		return bInGraph;
 	}
 
 
@@ -982,130 +929,457 @@ namespace Chaos
 	////////////////////////////////////////////////////////////////////////////////////////
 
 
-	FParticlePairCollisionDetector::FParticlePairCollisionDetector(
-		FGeometryParticleHandle* InParticle0,
-		FGeometryParticleHandle* InParticle1,
-		FParticlePairMidPhase& InMidPhase)
-		: MidPhase(InMidPhase)
-		, Constraints()
-		, Particle0(InParticle0)
-		, Particle1(InParticle1)
+	FShapePairParticlePairMidPhase::FShapePairParticlePairMidPhase()
+		: FParticlePairMidPhase(EParticlePairMidPhaseType::ShapePair)
 	{
-	}
 
-	FParticlePairCollisionDetector::FParticlePairCollisionDetector(FParticlePairCollisionDetector&& R)
-		: MidPhase(R.MidPhase)
-		, Constraints(MoveTemp(R.Constraints))
-		, Particle0(R.Particle0)
-		, Particle1(R.Particle1)
-	{
-		check(R.NewConstraints.IsEmpty());
 	}
 
-	FParticlePairCollisionDetector::~FParticlePairCollisionDetector()
+	void FShapePairParticlePairMidPhase::ResetImpl()
 	{
+		ShapePairDetectors.Reset();
 	}
 
-	int32 FParticlePairCollisionDetector::GenerateCollisions(
-		const FReal CullDistance,
-		const FReal Dt,
-		const FCollisionContext& Context)
+	void FShapePairParticlePairMidPhase::BuildDetectorsImpl()
 	{
-		return GenerateCollisionsImpl(false, CullDistance, Dt, Context);
+		if (IsValid() && (Particle0 != Particle1))
+		{
+			const FShapesArray& Shapes0 = Particle0->ShapesArray();
+			const FShapesArray& Shapes1 = Particle1->ShapesArray();
+			for (int32 ShapeIndex0 = 0; ShapeIndex0 < Shapes0.Num(); ++ShapeIndex0)
+			{
+				const FPerShapeData* Shape0 = Shapes0[ShapeIndex0].Get();
+				for (int32 ShapeIndex1 = 0; ShapeIndex1 < Shapes1.Num(); ++ShapeIndex1)
+				{
+					const FPerShapeData* Shape1 = Shapes1[ShapeIndex1].Get();
+					TryAddShapePair(Shape0, Shape1);
+				}
+			}
+		}
 	}
 
-
-	int32 FParticlePairCollisionDetector::GenerateCollisionsCCD(
-		const bool bEnableCCDSweep,
-		const FReal CullDistance,
-		const FReal Dt,
-		const FCollisionContext& Context)
+	void FShapePairParticlePairMidPhase::TryAddShapePair(const FPerShapeData* Shape0, const FPerShapeData* Shape1)
 	{
-		return GenerateCollisionsImpl(bEnableCCDSweep, CullDistance, Dt, Context);
-	}
+		const FImplicitObject* Implicit0 = Shape0->GetLeafGeometry();
+		const EImplicitObjectType ImplicitType0 = Collisions::GetImplicitCollisionType(Particle0, Implicit0);
 
-	int32 FParticlePairCollisionDetector::GenerateCollisionsImpl(
-		const bool bEnableCCDSweep,
-		const FReal CullDistance,
-		const FReal Dt,
-		const FCollisionContext& Context)
-	{
-		const FConstGenericParticleHandle P0 = Particle0;
-		const FConstGenericParticleHandle P1 = Particle1;
+		const FImplicitObject* Implicit1 = Shape1->GetLeafGeometry();
+		const EImplicitObjectType ImplicitType1 = Collisions::GetImplicitCollisionType(Particle1, Implicit1);
 
-		// We should have a Union at the root of each geometry hierarchy
-		const FImplicitObjectUnion* Union0 = P0->Geometry()->template GetObject<FImplicitObjectUnion>();
-		const FImplicitObjectUnion* Union1 = P1->Geometry()->template GetObject<FImplicitObjectUnion>();
-		if ((Union0 == nullptr) || (Union1 == nullptr))
+		const bool bDoPassFilter = ShapePairNarrowPhaseFilter(ImplicitType0, Shape0, ImplicitType1, Shape1);
+		if (bDoPassFilter)
 		{
-			return 0;
-		}
-
-		// Levelset collision if present
-		const FBVHParticles* BVHParticles0 = P0->CollisionParticles().Get();
-		const FBVHParticles* BVHParticles1 = P1->CollisionParticles().Get();
+			bool bSwap = false;
+			const EContactShapesType ShapePairType = Collisions::CalculateShapePairType(Particle0, Implicit0, Particle1, Implicit1, bSwap);
 
-		// Particle transforms
-		const FRigidTransform3 ParticleWorldTransform0 = P0->GetTransformPQ();
-		const FRigidTransform3 ParticleWorldTransform1 = P1->GetTransformPQ();
+			if (ShapePairType != EContactShapesType::Unknown)
+			{
+				if (!bSwap)
+				{
+					ShapePairDetectors.Emplace(FSingleShapePairCollisionDetector(Particle0, Shape0, Particle1, Shape1, ShapePairType, *this));
+				}
+				else
+				{
+					ShapePairDetectors.Emplace(FSingleShapePairCollisionDetector(Particle1, Shape1, Particle0, Shape0, ShapePairType, *this));
+				}
+			}
+			else
+			{
+				// If we get here we have a Union of Unions, but that means we should have created 
+				// a FGenericParticlePairMidPhase. See CalculateMidPhaseType
+				ensure(false);
+			}
+		}
+	}
 
-		// Calculate the overlapping volume in each particle's space
-		FAABB3 ParticleLocalIntersect0 = FAABB3::FullAABB();
-		FAABB3 ParticleLocalIntersect1 = FAABB3::FullAABB();
-		if (P0->HasBoundingBox() && P1->HasBoundingBox())
+	int32 FShapePairParticlePairMidPhase::GenerateCollisionsImpl(
+		const FReal CullDistance,
+		const FReal Dt,
+		const FCollisionContext& Context)
+	{
+		int32 NumActive = 0;
+		if (Flags.bIsCCD)
+		{
+			for (FSingleShapePairCollisionDetector& ShapePair : ShapePairDetectors)
+			{
+				NumActive += ShapePair.GenerateCollisionCCD(Flags.bUseSweep, CullDistance, Dt, Context);
+			}
+		}
+		else
 		{
-			const FRigidTransform3 ParticleTransform1To0 = ParticleWorldTransform1.GetRelativeTransform(ParticleWorldTransform0);
-			const FRigidTransform3 ParticleTransform0To1 = ParticleWorldTransform0.GetRelativeTransform(ParticleWorldTransform1);
-			ParticleLocalIntersect0 = P1->Geometry()->CalculateTransformedBounds(ParticleTransform1To0).ThickenSymmetrically(FVec3(CullDistance));
-			ParticleLocalIntersect1 = P0->Geometry()->CalculateTransformedBounds(ParticleTransform0To1).ThickenSymmetrically(FVec3(CullDistance));
+			for (FSingleShapePairCollisionDetector& ShapePair : ShapePairDetectors)
+			{
+				NumActive += ShapePair.GenerateCollision(CullDistance, Dt, Context);
+			}
 		}
+		return NumActive;
+	}
 
-		// Get the set of implicits from each particle that are in the overlap zone
-		TArray<FLeafImplicitObject> Objects0;
-		TArray<FLeafImplicitObject> Objects1;
-		CollectLeafImplicitObjects(ParticleLocalIntersect0, P0, Union0, ParticleWorldTransform0, CullDistance, Objects0);
-		CollectLeafImplicitObjects(ParticleLocalIntersect1, P1, Union1, ParticleWorldTransform1, CullDistance, Objects1);
+	void FShapePairParticlePairMidPhase::WakeCollisionsImpl(const int32 CurrentEpoch)
+	{
+		for (FSingleShapePairCollisionDetector& ShapePair : ShapePairDetectors)
+		{
+			ShapePair.WakeCollision(LastUsedEpoch, CurrentEpoch);
+		}
+	}
 
-		// Modify the context so that Collision allocations come to us
-		// @todo(chaos): maybe pass a functor into ConstructConstraints instead
-		FCollisionContext LocalContext = Context;
-		LocalContext.ParticlePairCollisionDetector = this;
+	void FShapePairParticlePairMidPhase::InjectCollisionImpl(const FPBDCollisionConstraint& Constraint, const FCollisionContext& Context)
+	{
+		const FPerShapeData* Shape0 = Constraint.GetShape0();
+		const FPerShapeData* Shape1 = Constraint.GetShape1();
 
-		// Check all shape pairs for collisions
-		for (const FLeafImplicitObject& Object0 : Objects0)
+		// @todo(chaos): fix O(N) search for shape pair - store the index in the cookie (it will be the same
+		// as long as the ShapesArray on each particle has not changed)
+		for (FSingleShapePairCollisionDetector& ShapePair : ShapePairDetectors)
 		{
-			for (const FLeafImplicitObject& Object1 : Objects1)
+			if (((Shape0 == ShapePair.GetShape0()) && (Shape1 == ShapePair.GetShape1())) || ((Shape0 == ShapePair.GetShape1()) && (Shape1 == ShapePair.GetShape0())))
 			{
-				// NOTE: We have 2 sets of bodies from the total overlapping volume of the 2 hierarchies. 
-				// We still need to perform a shape pair bounds check
-				if (!Object0.WorldBounds.Intersects(Object1.WorldBounds))
-				{
-					continue;
-				}
+				ShapePair.SetCollision(Constraint, Context);
+				break;
+			}
+		}
+	}
+
+	////////////////////////////////////////////////////////////////////////////////////////
+	////////////////////////////////////////////////////////////////////////////////////////
+	////////////////////////////////////////////////////////////////////////////////////////
+	////////////////////////////////////////////////////////////////////////////////////////
+
+
+	FGenericParticlePairMidPhase::FGenericParticlePairMidPhase()
+		: FParticlePairMidPhase(EParticlePairMidPhaseType::Generic)
+	{
+	}
+
+	FGenericParticlePairMidPhase::~FGenericParticlePairMidPhase()
+	{
+	}
+
+	void FGenericParticlePairMidPhase::ResetImpl()
+	{
+		Constraints.Empty();
+		NewConstraints.Empty();
+	}
+
+	void FGenericParticlePairMidPhase::BuildDetectorsImpl()
+	{
+		check(Particle0->Geometry().Get() != nullptr);
+		check(Particle1->Geometry().Get() != nullptr);
+	}
+
+	int32 FGenericParticlePairMidPhase::GenerateCollisionsImpl(
+		const FReal CullDistance,
+		const FReal Dt,
+		const FCollisionContext& Context)
+	{
+		const FImplicitObject* Implicit0 = GetParticle0()->Geometry().Get();
+		const FImplicitObject* Implicit1 = GetParticle1()->Geometry().Get();
 
-				// Create the constraint for the shape pair if the bounds overlap
-				Collisions::ConstructConstraints(
-					Particle0, Particle1,
-					Object0.Geometry, Object0.ShapeInstance, BVHParticles0, Object0.LeafObjectIndex,
-					Object1.Geometry, Object1.ShapeInstance, BVHParticles1, Object1.LeafObjectIndex,
-					ParticleWorldTransform0, Object0.RelativeTransform,
-					ParticleWorldTransform1, Object1.RelativeTransform,
-					CullDistance, Dt, bEnableCCDSweep, LocalContext);
+		// See if we have a BVH for either/both of the particles
+		const Private::FImplicitBVH* BVH0 = nullptr;
+		const Private::FImplicitBVH* BVH1 = nullptr;
+		if (CVars::bChaosUnionBVHEnabled)
+		{
+			if (const FImplicitObjectUnion* Union0 = Implicit0->GetObject<FImplicitObjectUnion>())
+			{
+				BVH0 = Union0->GetBVH();
+			}
+			if (const FImplicitObjectUnion* Union1 = Implicit1->GetObject<FImplicitObjectUnion>())
+			{
+				BVH1 = Union1->GetBVH();
 			}
 		}
 
+		// Create constraints for all the shape pairs whose bounds overlap.
+		// If we have a BVH, use it. Otherwise run a recursive hierarchy sweep.
+		// @todo(chaos): if we have 2 BVHs select the deepest one as BVHA?
+		if (BVH0 != nullptr)
+		{
+			GenerateCollisionsBVH(GetParticle0(), BVH0, GetParticle1(), Implicit1, CullDistance, Dt, Context);
+		}
+		else if (BVH1 != nullptr)
+		{
+			GenerateCollisionsBVH(GetParticle1(), BVH1, GetParticle0(), Implicit0, CullDistance, Dt, Context);
+		}
+		else
+		{
+			GenerateCollisionsImplicit(GetParticle0(), Implicit0, GetParticle1(), Implicit1, CullDistance, Dt, Context);
+		}
+
 		// Generate manifolds for each constraint we created/recovered and (re)activate if necessary
-		int32 NumActiveConstraints = ProcessNewConstraints(CullDistance, Dt, Context);
+		int32 NumActive = ProcessNewConstraints(CullDistance, Dt, Context);
 
 		// @todo(chaos): we could clean up unused collisions between this pair, but probably not worth it
-		// and we would have to prevent cleanup for sleeping particles because the collisions will still
-		// be referenced in the IslandManager's constraint graph for the sleeping island.
 		//PruneConstraints();
 
-		return NumActiveConstraints;
+		return NumActive;
+	}
+
+	// Detect collisions betweena BVH and some other implicit object (which may be a hierarchy)
+	void FGenericParticlePairMidPhase::GenerateCollisionsBVH(
+		FGeometryParticleHandle* ParticleA, const Private::FImplicitBVH* BVHA,
+		FGeometryParticleHandle* ParticleB, const FImplicitObject* RootImplicitB,
+		const FReal CullDistance,
+		const FReal Dt,
+		const FCollisionContext& Context)
+	{
+		const FConstGenericParticleHandle PA = ParticleA;
+		const FConstGenericParticleHandle PB = ParticleB;
+
+		// Particle transforms
+		const FRigidTransform3 ParticleWorldTransformA = PA->GetTransformPQ();
+		const FRigidTransform3 ParticleWorldTransformB = PB->GetTransformPQ();
+		const FRigidTransform3 ParticleTransformAToB = ParticleWorldTransformA.GetRelativeTransform(ParticleWorldTransformB);
+
+		FMemMark Mark(FMemStack::Get());
+		TArray<bool, TMemStackAllocator<alignof(bool)>> bIsVisitedA;
+		bIsVisitedA.SetNumZeroed(BVHA->GetNumObjects());
+
+		// Visitor for FImplicitBVH::VisitNodeObjects
+		// Given an ImplicitObject from ParticleA (which we know overlaps the bounds of some parts of ParticleB),
+		// run collision detection on ImplicitA against the implicit object hierarchy of ParticleB.
+		// NOTE: may be called with the same ImplicitA multiple times because implicits may be in many BVH leaves
+		const auto& ObjectVisitorA = 
+			[this, ParticleA, &ParticleWorldTransformA, ParticleB, RootImplicitB, &ParticleWorldTransformB, &ParticleTransformAToB, &bIsVisitedA, CullDistance, Dt, &Context](const FImplicitObject* ImplicitA, const FRigidTransform3f& RelativeTransformfA, const FAABB3f& RelativeBoundsfA, const int32 RootObjectIndexA, const int32 LeafObjectIndexA) -> void
+			{
+				if (bIsVisitedA[LeafObjectIndexA])
+				{
+					return;
+				}
+
+				bIsVisitedA[LeafObjectIndexA] = true;
+
+				// @todo(chaos): remove this float to double conversion and use floats where possible
+				const FRigidTransform3& RelativeTransformA = FRigidTransform3(RelativeTransformfA);
+				const FAABB3& RelativeBoundsA = FAABB3(RelativeBoundsfA);
+
+				// If this is the first time we have seen ImplicitA, run the collision detection against ParticleB
+				GenerateCollisionsShapeHierarchy(
+					ParticleA, ImplicitA, ParticleWorldTransformA, RelativeTransformA, RelativeBoundsA, RootObjectIndexA, LeafObjectIndexA,
+					ParticleB, RootImplicitB, ParticleWorldTransformB,
+					ParticleTransformAToB,
+					CullDistance, Dt, Context);
+			};
+
+		// Visitor for FImplicitBVH::VisitHierarchy
+		// This will be passed the nodes of the BVH on ParticleA and its bounds and contents. If the bounds overlap something in ParticleB
+		// we will keep recursing into the BVH. When we hit a leaf, run collision detection between the leaf contents and ParticleB.
+		const auto& NodeVisitorA = 
+			[BVHA, RootImplicitB, &ParticleTransformAToB, &ObjectVisitorA, CullDistance](const FAABB3f& NodeBoundsAf, const int32 NodeDepthA, const Private::FImplicitBVHNode& NodeA) -> bool
+			{
+				const FAABB3 NodeBoundsAInB = FAABB3(NodeBoundsAf).TransformedAABB(ParticleTransformAToB).ThickenSymmetrically(FVec3(CullDistance));
+
+				// Does this Node in A overlap anything in B?
+				// NOTE: IsOverlappingBounds performs a deep bounds check, including checking for node overlaps in BVH, Heightfield and TriMesh
+				if (!RootImplicitB->IsOverlappingBounds(NodeBoundsAInB))
+				{
+					// No overlap - stop recursing down this branch
+					return false;
+				}
+
+				// If we are at a leaf of A, we need to check all objects in the node against B
+				if (NodeA.IsLeaf())
+				{
+					BVHA->VisitNodeObjects(NodeA, ObjectVisitorA);
+				}
+
+				// Keep recursing
+				return true;
+			};
+
+		// Visit all the nodes in BVHA and detect collisions with ParticleB
+		BVHA->VisitHierarchy(NodeVisitorA);
+	}
+
+	// Detect collisions between two implicits, where either or both may be a hierarchy
+	void FGenericParticlePairMidPhase::GenerateCollisionsImplicit(
+		FGeometryParticleHandle* ParticleA, const FImplicitObject* RootImplicitA,
+		FGeometryParticleHandle* ParticleB, const FImplicitObject* RootImplicitB,
+		const FReal CullDistance,
+		const FReal Dt,
+		const FCollisionContext& Context)
+	{
+		const FConstGenericParticleHandle PA = ParticleA;
+		const FConstGenericParticleHandle PB = ParticleB;
+
+		// Particle transforms
+		const FRigidTransform3 ParticleWorldTransformA = PA->GetTransformPQ();
+		const FRigidTransform3 ParticleWorldTransformB = PB->GetTransformPQ();
+
+		// Calculate the overlapping volume in each particle's space
+		const FRigidTransform3 ParticleTransformAToB = ParticleWorldTransformA.GetRelativeTransform(ParticleWorldTransformB);
+
+		// Visitor for FImplicitBVH::VisitNodeObjects
+		// Given an ImplicitObject from ParticleA (which we know overlaps the bounds of some parts of ParticleB),
+		// run collision detection on ImplicitA against the implicit object hierarchy of ParticleB.
+		// NOTE: may be called with the same ImplicitA multiple times because implicits may be in many BVH leaves
+		const auto& ObjectVisitorA =
+			[this, ParticleA, &ParticleWorldTransformA, ParticleB, RootImplicitB, &ParticleWorldTransformB, &ParticleTransformAToB, CullDistance, Dt, &Context]
+			(const FImplicitObject* ImplicitA, const FRigidTransform3& RelativeTransformA, const int32 RootObjectIndexA, const int32 ObjectIndex, const int32 LeafObjectIndexA)
+			{
+				const FAABB3 RelativeBoundsA = ImplicitA->CalculateTransformedBounds(RelativeTransformA);
+
+				// If this is the first time we have seen ImplicitA, run the collision detection against ParticleB
+				GenerateCollisionsShapeHierarchy(
+					ParticleA, ImplicitA, ParticleWorldTransformA, RelativeTransformA, RelativeBoundsA, RootObjectIndexA, LeafObjectIndexA,
+					ParticleB, RootImplicitB, ParticleWorldTransformB,
+					ParticleTransformAToB,
+					CullDistance, Dt, Context);
+			};
+
+		// Detect collisons between Implicit Hierarchy of ParticleA and Implicit Hierarchy of ParticleB
+		RootImplicitA->VisitLeafObjects(ObjectVisitorA);
+	}
+
+	// Generate collisions between lesf (non-hierarchy) implicit and some other (maybe hierarchy) implicit
+	void FGenericParticlePairMidPhase::GenerateCollisionsShapeHierarchy(
+		FGeometryParticleHandle* ParticleA, const FImplicitObject* ImplicitA, const FRigidTransform3 ParticleWorldTransformA, const FRigidTransform3& RelativeTransformA, const FAABB3& RelativeBoundsA, const int32 RootObjectIndexA, const int32 LeafObjectIndexA,
+		FGeometryParticleHandle* ParticleB, const FImplicitObject* RootImplicitB, const FRigidTransform3 ParticleWorldTransformB,
+		const FRigidTransform3 ParticleTransformAToB,
+		const FReal CullDistance,
+		const FReal Dt,
+		const FCollisionContext& Context)
+	{
+		const FShapeInstanceArray& ShapeInstancesA = ParticleA->ShapeInstances();
+		const FShapeInstanceArray& ShapeInstancesB = ParticleB->ShapeInstances();
+
+		const FShapeInstance* ShapeInstanceA = GetShapeInstance(ShapeInstancesA, RootObjectIndexA);
+
+		const FAABB3 ShapeBoundsAInB = FAABB3(RelativeBoundsA).TransformedAABB(ParticleTransformAToB).ThickenSymmetrically(FVec3(CullDistance));
+
+		const auto& OverlappingLeafVisitor = 
+			[this, ParticleA, ImplicitA, ShapeInstanceA, &ParticleWorldTransformA, &RelativeTransformA, LeafObjectIndexA, ParticleB, &ParticleWorldTransformB, &ShapeInstancesB, CullDistance, Dt, &Context]
+			(const FImplicitObject* ImplicitB, const FRigidTransform3& RelativeTransformB, const int32 RootObjectIndexB, const int32 ObjectIndexB, const int32 LeafObjectIndexB)
+			{
+				const FShapeInstance* ShapeInstanceB = GetShapeInstance(ShapeInstancesB, RootObjectIndexB);
+
+				GenerateCollisionsShapeShape(
+					ParticleA, ImplicitA, ShapeInstanceA, ParticleWorldTransformA, RelativeTransformA, LeafObjectIndexA,
+					ParticleB, ImplicitB, ShapeInstanceB, ParticleWorldTransformB, RelativeTransformB, LeafObjectIndexB,
+					CullDistance, Dt, Context);
+			};
+
+		// Detect collisons between ImplicitA and Implicit Hierarchy of ParticleB
+		RootImplicitB->VisitOverlappingLeafObjects(ShapeBoundsAInB, OverlappingLeafVisitor);
+	}
+
+	// Generate collisions between two leaf (not hierahcies) implicits
+	void FGenericParticlePairMidPhase::GenerateCollisionsShapeShape(
+		FGeometryParticleHandle* ParticleA, const FImplicitObject* ImplicitA, const FShapeInstance* ShapeInstanceA, const FRigidTransform3 ParticleWorldTransformA, const FRigidTransform3& RelativeTransformA, const int32 LeafObjectIndexA,
+		FGeometryParticleHandle* ParticleB, const FImplicitObject* ImplicitB, const FShapeInstance* ShapeInstanceB, const FRigidTransform3 ParticleWorldTransformB, const FRigidTransform3& RelativeTransformB, const int32 LeafObjectIndexB,
+		const FReal CullDistance,
+		const FReal Dt,
+		const FCollisionContext& Context)
+	{
+		// Check the sim filter to see if these shapes collide
+		const EImplicitObjectType ImplicitTypeA = Collisions::GetImplicitCollisionType(ParticleA, ImplicitA);
+		const EImplicitObjectType ImplicitTypeB = Collisions::GetImplicitCollisionType(ParticleB, ImplicitB);
+		const bool bDoPassFilter = ShapePairNarrowPhaseFilter(ImplicitTypeA, ShapeInstanceA, ImplicitTypeB, ShapeInstanceB);
+		if (!bDoPassFilter)
+		{
+			return;
+		}
+
+		// Do the objects bounds overlap?
+		const bool bDoOverlap = DoBoundsOverlap(
+			ImplicitA, ParticleWorldTransformA, RelativeTransformA,
+			ImplicitB, ParticleWorldTransformB, RelativeTransformB,
+			CullDistance);
+
+		if (!bDoOverlap)
+		{
+			return;
+		}
+
+		// Find the BVH Particles if we have LevelSet Collision
+		// @todo(chaos): This is not really right. The BVHParticles represent the whole
+		// geometry hierarchy from the orginating particle (we may be a ClusterUnion)
+		// so we shouldn't be doing a collision per LevelSet
+		const FBVHParticles* LevelSetParticlesA = nullptr;
+		const FBVHParticles* LevelSetParticlesB = nullptr;
+		if (ImplicitA->GetCollisionType() == ImplicitObjectType::LevelSet)
+		{
+			LevelSetParticlesA = FindLevelSetParticles(ParticleA, ImplicitA);
+		}
+		if (ImplicitB->GetCollisionType() == ImplicitObjectType::LevelSet)
+		{
+			LevelSetParticlesB = FindLevelSetParticles(ParticleB, ImplicitB);
+		}
+
+		// What shape pair type is this (i.e., which contact update function do we call)
+		bool bShouldSwapParticles = false;
+		const EContactShapesType ShapePairType = Collisions::CalculateShapePairType(ParticleA, ImplicitA, ParticleB, ImplicitB, bShouldSwapParticles);
+
+		// Create the constraint for this shape pair
+		// NOTE: this just creates the object. The collision detection is done in ProcessNewConstraints
+		if (!bShouldSwapParticles)
+		{
+			FindOrCreateConstraint(
+				ParticleA, ImplicitA, LeafObjectIndexA, ShapeInstanceA, LevelSetParticlesA, RelativeTransformA,
+				ParticleB, ImplicitB, LeafObjectIndexB, ShapeInstanceB, LevelSetParticlesB, RelativeTransformB,
+				CullDistance, ShapePairType, Context.GetSettings().bAllowManifolds, Flags.bUseSweep, Context);
+		}
+		else
+		{
+			FindOrCreateConstraint(
+				ParticleB, ImplicitB, LeafObjectIndexB, ShapeInstanceB, LevelSetParticlesB, RelativeTransformB,
+				ParticleA, ImplicitA, LeafObjectIndexA, ShapeInstanceA, LevelSetParticlesA, RelativeTransformA,
+				CullDistance, ShapePairType, Context.GetSettings().bAllowManifolds, Flags.bUseSweep, Context);
+		}
+	}
+
+	bool FGenericParticlePairMidPhase::DoBoundsOverlap(
+		const FImplicitObject* ImplicitA,
+		const FRigidTransform3& ParticleWorldTransformA,
+		const FRigidTransform3& ShapeRelativeTransformA,
+		const FImplicitObject* ImplicitB,
+		const FRigidTransform3& ParticleWorldTransformB,
+		const FRigidTransform3& ShapeRelativeTransformB,
+		const FReal CullDistance)
+	{
+		const bool bEnableAABBCheck = true;
+		const bool bEnableOBBCheckA = true;
+		const bool bEnableOBBCheckB = true;
+
+		const FRigidTransform3 ShapeWorldTransformA = ShapeRelativeTransformA * ParticleWorldTransformA;
+		const FRigidTransform3 ShapeWorldTransformB = ShapeRelativeTransformB * ParticleWorldTransformB;
+
+		// NOTE: only expand one bounds by cull distance
+		const FAABB3 ShapeWorldBoundsA = ImplicitA->CalculateTransformedBounds(ShapeWorldTransformA).ThickenSymmetrically(FVec3(CullDistance));
+		const FAABB3 ShapeWorldBoundsB = ImplicitB->CalculateTransformedBounds(ShapeWorldTransformB);
+
+		// World-space expanded bounds check
+		if (bEnableAABBCheck)
+		{
+			if (!ShapeWorldBoundsA.Intersects(ShapeWorldBoundsB))
+			{
+				return false;
+			}
+		}
+
+		// OBB-AABB test in both directions. This is beneficial for shapes which do not fit their AABBs very well,
+		// which includes boxes and other shapes that are not roughly spherical. It is especially beneficial when
+		// one shape is long and thin (i.e., it does not fit an AABB well when the shape is rotated).
+		if (bEnableOBBCheckA)
+		{
+			if (!ImplicitOverlapOBBToAABB(ImplicitA, ImplicitB, ShapeWorldTransformA, ShapeWorldTransformB, CullDistance))
+			{
+				return false;
+			}
+		}
+
+		if (bEnableOBBCheckB)
+		{
+			if (!ImplicitOverlapOBBToAABB(ImplicitB, ImplicitA, ShapeWorldTransformB, ShapeWorldTransformA, CullDistance))
+			{
+				return false;
+			}
+		}
+
+		return true;
 	}
 
-	FPBDCollisionConstraint* FParticlePairCollisionDetector::FindOrCreateConstraint(
+	FPBDCollisionConstraint* FGenericParticlePairMidPhase::FindOrCreateConstraint(
 		FGeometryParticleHandle* InParticle0,
 		const FImplicitObject* InImplicit0,
 		const int32 InImplicitId0,
@@ -1134,8 +1408,8 @@ namespace Chaos
 			return nullptr;
 		}
 
-		const FCollisionParticlePairConstraintKey Key = FCollisionParticlePairConstraintKey(InShape0, InImplicit0, InImplicitId0, InBVHParticles0, InShape1, InImplicit1, InImplicitId1, InBVHParticles1);
-		FPBDCollisionConstraint* Constraint = FindConstraint(Key);
+		const FCollisionParticlePairConstraintKey CollisionKey = FCollisionParticlePairConstraintKey(InShape0, InImplicit0, InImplicitId0, InBVHParticles0, InShape1, InImplicit1, InImplicitId1, InBVHParticles1);
+		FPBDCollisionConstraint* Constraint = FindConstraint(CollisionKey);
 
 		// @todo(chaos): fix key uniqueness guarantee.  We need a truly unique key gen function
 		const bool bIsKeyCollision = (Constraint != nullptr) && ((Constraint->GetImplicit0() != InImplicit0) || (Constraint->GetImplicit1() != InImplicit1) || (Constraint->GetCollisionParticles0() != InBVHParticles0) || (Constraint->GetCollisionParticles1() != InBVHParticles1));
@@ -1153,23 +1427,23 @@ namespace Chaos
 		if (Constraint == nullptr)
 		{
 			// NOTE: Using InParticle0 and InParticle1 here because the order may be different to what we have stored
-			Constraint = CreateConstraint(InParticle0, InImplicit0, InShape0, InBVHParticles0, InShapeRelativeTransform0, InParticle1, InImplicit1, InShape1, InBVHParticles1, InShapeRelativeTransform1, CullDistance, ShapePairType, bUseManifold, Key, Context);
+			Constraint = CreateConstraint(InParticle0, InImplicit0, InShape0, InBVHParticles0, InShapeRelativeTransform0, InParticle1, InImplicit1, InShape1, InBVHParticles1, InShapeRelativeTransform1, CullDistance, ShapePairType, bUseManifold, CollisionKey, Context);
 
 			// Is this a CCD constraint?
-			Constraint->SetCCDEnabled(MidPhase.IsCCD());
+			Constraint->SetCCDEnabled(IsCCD());
 		}
 
 		// Do we want to sweep on this tick? 
 		// CCD may be temporarily disabled by the user or because we are moving slowly.
-		Constraint->SetCCDSweepEnabled(MidPhase.IsCCD() && bEnableSweep);
+		Constraint->SetCCDSweepEnabled(IsCCD() && bEnableSweep);
 
 		NewConstraints.Add(Constraint);
 		return Constraint;
 	}
 
-	FPBDCollisionConstraint* FParticlePairCollisionDetector::FindConstraint(const FCollisionParticlePairConstraintKey& Key)
+	FPBDCollisionConstraint* FGenericParticlePairMidPhase::FindConstraint(const FCollisionParticlePairConstraintKey& CollisionKey)
 	{
-		FPBDCollisionConstraintPtr* PConstraint = Constraints.Find(Key.GetKey());
+		FPBDCollisionConstraintPtr* PConstraint = Constraints.Find(CollisionKey.GetKey());
 		if (PConstraint != nullptr)
 		{
 			return (*PConstraint).Get();
@@ -1177,7 +1451,7 @@ namespace Chaos
 		return nullptr;
 	}
 
-	FPBDCollisionConstraint* FParticlePairCollisionDetector::CreateConstraint(
+	FPBDCollisionConstraint* FGenericParticlePairMidPhase::CreateConstraint(
 		FGeometryParticleHandle* InParticle0,
 		const FImplicitObject* Implicit0,
 		const FPerShapeData* InShape0,
@@ -1191,21 +1465,25 @@ namespace Chaos
 		const FReal CullDistance,
 		const EContactShapesType ShapePairType,
 		const bool bUseManifold,
-		const FCollisionParticlePairConstraintKey& Key,
+		const FCollisionParticlePairConstraintKey& CollisionKey,
 		const FCollisionContext& Context)
 	{
 		PHYSICS_CSV_SCOPED_EXPENSIVE(PhysicsVerbose, NarrowPhase_CreateConstraint);
-		FPBDCollisionConstraintPtr Constraint = CreateImplicitPairConstraint(InParticle0, Implicit0, InShape0, BVHParticles0, ShapeRelativeTransform0, InParticle1, Implicit1, InShape1, BVHParticles1, ShapeRelativeTransform1, CullDistance, ShapePairType, bUseManifold, Context);
+
+		FPBDCollisionConstraintPtr Constraint = Context.GetAllocator()->CreateConstraint(
+			InParticle0, Implicit0, InShape0, BVHParticles0, ShapeRelativeTransform0,
+			InParticle1, Implicit1, InShape1, BVHParticles1, ShapeRelativeTransform1,
+			CullDistance, bUseManifold, ShapePairType);
 
 		const int32 CurrentEpoch = Context.GetAllocator()->GetCurrentEpoch();
-		Constraint->GetContainerCookie().MidPhase = &MidPhase;
+		Constraint->GetContainerCookie().MidPhase = this;
 		Constraint->GetContainerCookie().bIsMultiShapePair = true;
 		Constraint->GetContainerCookie().CreationEpoch = CurrentEpoch;
 
-		return Constraints.Add(Key.GetKey(), MoveTemp(Constraint)).Get();
+		return Constraints.Add(CollisionKey.GetKey(), MoveTemp(Constraint)).Get();
 	}
 
-	void FParticlePairCollisionDetector::WakeCollisions(const int32 SleepEpoch, const int32 CurrentEpoch)
+	void FGenericParticlePairMidPhase::WakeCollisionsImpl(const int32 SleepEpoch, const int32 CurrentEpoch)
 	{
 		for (auto& KVP : Constraints)
 		{
@@ -1219,12 +1497,13 @@ namespace Chaos
 		}
 	}
 
-	int32 FParticlePairCollisionDetector::ProcessNewConstraints(
+	int32 FGenericParticlePairMidPhase::ProcessNewConstraints(
 		const FReal CullDistance,
 		const FReal Dt,
 		const FCollisionContext& Context)
 	{
-		int32 NumActiveConstraints = 0;
+		int32 NumActive = 0;
+
 		for (FPBDCollisionConstraint* Constraint : NewConstraints)
 		{
 			if (!Constraint->GetCCDSweepEnabled())
@@ -1238,14 +1517,16 @@ namespace Chaos
 
 			if (!Constraint->GetDisabled())
 			{
-				++NumActiveConstraints;
+				++NumActive;
 			}
 		}
+
 		NewConstraints.Reset();
-		return NumActiveConstraints;
+
+		return NumActive;
 	}
 
-	void FParticlePairCollisionDetector::UpdateCollision(
+	void FGenericParticlePairMidPhase::UpdateCollision(
 		FPBDCollisionConstraint* Constraint,
 		const FReal CullDistance,
 		const FReal Dt,
@@ -1313,7 +1594,7 @@ namespace Chaos
 	}
 
 
-	void FParticlePairCollisionDetector::UpdateCollisionCCD(
+	void FGenericParticlePairMidPhase::UpdateCollisionCCD(
 		FPBDCollisionConstraint* Constraint,
 		const FReal CullDistance,
 		const FReal Dt,
@@ -1384,423 +1665,55 @@ namespace Chaos
 		Context.GetAllocator()->ActivateConstraint(Constraint);
 	}
 
-	void FParticlePairCollisionDetector::PruneConstraints(const int32 CurrentEpoch)
+	void FGenericParticlePairMidPhase::PruneConstraints(const int32 CurrentEpoch)
 	{
-		// We don't prune from NewCollisions - must call ProcessNewCollisions before Prune
+		// Must call ProcessNewCollisions before Prune
 		check(NewConstraints.Num() == 0);
 
 		// Find all the expired collisions
-		TArray<uint32> Pruned;
+		FMemMark Mark(FMemStack::Get());
+		TArray<uint32, TMemStackAllocator<alignof(uint32)>> Pruned;
+		Pruned.Reserve(Constraints.Num());
+
 		for (auto& KVP : Constraints)
 		{
-			const uint32 Key = KVP.Key;
+			const uint32 CollisionKey = KVP.Key;
 			FPBDCollisionConstraintPtr& Constraint = KVP.Value;
-			if (Constraint->GetContainerCookie().LastUsedEpoch < CurrentEpoch)
-			{
-				Pruned.Add(Key);
-			}
-		}
-
-		// Destroy expired collisions
-		for (uint32 Key : Pruned)
-		{
-			Constraints.Remove(Key);
-		}
-	}
-
-
-	////////////////////////////////////////////////////////////////////////////////////////
-	////////////////////////////////////////////////////////////////////////////////////////
-	////////////////////////////////////////////////////////////////////////////////////////
-	////////////////////////////////////////////////////////////////////////////////////////
-
-
-	FParticlePairMidPhase::FParticlePairMidPhase()
-		: Particle0(nullptr)
-		, Particle1(nullptr)
-		, CullDistanceScale(1)
-		, Flags()
-		, Key()
-		, LastUsedEpoch(INDEX_NONE)
-		, NumActiveConstraints(0)
-		, ParticleCollisionsIndex0(INDEX_NONE)
-		, ParticleCollisionsIndex1(INDEX_NONE)
-		, MultiShapePairDetectors()
-		, ShapePairDetectors()
-		, ParticlePairDetector()
-	{
-	}
-
-	FParticlePairMidPhase::~FParticlePairMidPhase()
-	{
-		Reset();
-	}
-
-	void FParticlePairMidPhase::DetachParticle(FGeometryParticleHandle* Particle)
-	{
-		Reset();
-
-		if (Particle == Particle0)
-		{
-			Particle0 = nullptr;
-		}
-		else if (Particle == Particle1)
-		{
-			Particle1 = nullptr;
-		}
-	}
-
-	void FParticlePairMidPhase::Reset()
-	{
-		ShapePairDetectors.Reset();
-		MultiShapePairDetectors.Reset();
-
-		Flags.bIsActive = true;
-		Flags.bIsCCD = false;
-		Flags.bIsCCDActive = false;
-		Flags.bIsSleeping = false;
-		Flags.bIsModified = false;
-	}
-
-	void FParticlePairMidPhase::ResetModifications()
-	{
-		if (Flags.bIsModified)
-		{
-			Flags.bIsActive = true;
-			Flags.bIsCCDActive = Flags.bIsCCD;
-			Flags.bIsModified = false;
-		}
-	}
-
-	void FParticlePairMidPhase::Init(
-		FGeometryParticleHandle* InParticle0,
-		FGeometryParticleHandle* InParticle1,
-		const FCollisionParticlePairKey& InKey,
-		const FCollisionContext& Context)
-	{
-		PHYSICS_CSV_SCOPED_EXPENSIVE(PhysicsVerbose, NarrowPhase_Filter);
-
-		Particle0 = InParticle0;
-		Particle1 = InParticle1;
-		Key = InKey;
-
-		Flags.bIsActive = true;
-
-		// If CCD is allowed in the current context and for at least one of
-		// the particles involved, enable it for this midphase.
-		//
-		// bIsCCDActive is reset to bIsCCD each frame, but can be overridden
-		// by modifiers.
-		const bool bIsCCD = Context.GetSettings().bAllowCCD && (
-			FConstGenericParticleHandle(Particle0)->CCDEnabled() ||
-			FConstGenericParticleHandle(Particle1)->CCDEnabled());
-		Flags.bIsCCD = bIsCCD;
-		Flags.bIsCCDActive = bIsCCD;
-
-		BuildDetectors();
-
-		InitThresholds();
-	}
-
-	void FParticlePairMidPhase::BuildDetectors()
-	{
-		if (IsValid() && (Particle0 != Particle1))
-		{
-			// For particles with a small number of shapes, we prebuild the set of potentially colliding
-			// shape pairs and cache various data with them. However we don't do this if there are a
-			// large number of shapes, or if the geometry hierarchy contains Unions of Unions (the
-			// FParticlePairCollisionDetector assumes that it has a leaf shape and not a Union).
-			bool bPrebuildShapePairs = true;
-			if (CVars::Chaos_Collision_MidPhase_MaxShapePairs > 0)
-			{
-				const FImplicitObjectUnion* Union0 = (Particle0->Geometry().Get() != nullptr) ? Particle0->Geometry()->template GetObject<FImplicitObjectUnion>() : nullptr;
-				const FImplicitObjectUnion* Union1 = (Particle1->Geometry().Get() != nullptr) ? Particle1->Geometry()->template GetObject<FImplicitObjectUnion>() : nullptr;
-				if ((Union0 != nullptr) && (Union1 != nullptr))
-				{
-					const int32 NumShapePairs = Union0->GetNumLeafObjects() * Union1->GetNumLeafObjects();
-					const bool bIsTree0 = (Union0->GetNumLeafObjects() > Union0->GetObjects().Num());	// Do we have Unions of Unions?
-					const bool bIsTree1 = (Union1->GetNumLeafObjects() > Union1->GetObjects().Num());	// Do we have Unions of Unions?
-					bPrebuildShapePairs = (NumShapePairs < CVars::Chaos_Collision_MidPhase_MaxShapePairs) && !bIsTree0 && !bIsTree1;
-				}
-			}
-
-			if (bPrebuildShapePairs)
-			{
-				const FShapesArray& Shapes0 = Particle0->ShapesArray();
-				const FShapesArray& Shapes1 = Particle1->ShapesArray();
-				for (int32 ShapeIndex0 = 0; ShapeIndex0 < Shapes0.Num(); ++ShapeIndex0)
-				{
-					const FPerShapeData* Shape0 = Shapes0[ShapeIndex0].Get();
-					for (int32 ShapeIndex1 = 0; ShapeIndex1 < Shapes1.Num(); ++ShapeIndex1)
-					{
-						const FPerShapeData* Shape1 = Shapes1[ShapeIndex1].Get();
-						TryAddShapePair(Shape0, Shape1);
-					}
-				}
-			}
-			else
-			{
-				ParticlePairDetector = MakeUnique<FParticlePairCollisionDetector>(Particle0, Particle1, *this);
-			}
-		}
-	}
-
-	void FParticlePairMidPhase::TryAddShapePair(const FPerShapeData* Shape0, const FPerShapeData* Shape1)
-	{
-		const FImplicitObject* Implicit0 = Shape0->GetLeafGeometry();
-		const EImplicitObjectType ImplicitType0 = Collisions::GetImplicitCollisionType(Particle0, Implicit0);
-
-		const FImplicitObject* Implicit1 = Shape1->GetLeafGeometry();
-		const EImplicitObjectType ImplicitType1 = Collisions::GetImplicitCollisionType(Particle1, Implicit1);
-
-		const bool bDoPassFilter = ShapePairNarrowPhaseFilter(ImplicitType0, Shape0, ImplicitType1, Shape1);
-		if (bDoPassFilter)
-		{
-			bool bSwap = false;
-			const EContactShapesType ShapePairType = Collisions::CalculateShapePairType(Particle0, Implicit0, Particle1, Implicit1, bSwap);
-
-			if (ShapePairType != EContactShapesType::Unknown)
-			{
-				if (!bSwap)
-				{
-					ShapePairDetectors.Emplace(FSingleShapePairCollisionDetector(Particle0, Shape0, Particle1, Shape1, ShapePairType, *this));
-				}
-				else
-				{
-					ShapePairDetectors.Emplace(FSingleShapePairCollisionDetector(Particle1, Shape1, Particle0, Shape0, ShapePairType, *this));
-				}
-			}
-			else
-			{
-				if (ensure(!bSwap))
-				{
-					MultiShapePairDetectors.Emplace(FMultiShapePairCollisionDetector(Particle0, Shape0, Particle1, Shape1, *this));
-				}
-			}
-		}
-	}
-
-	bool FParticlePairMidPhase::ShouldEnableCCD(const FReal Dt)
-	{
-		// bIsCCDActive is set to bIsCCD at the beginning of every frame, but may be
-		// overridden in midphase modification or potentially other systems which run
-		// in between mid and narrow phase. bIsCCDActive indicates the final
-		// overridden value so we use that here instead of bIsCCD.
-		if (Flags.bIsCCDActive)
-		{
-			FConstGenericParticleHandle ConstParticle0 = FConstGenericParticleHandle(Particle0);
-			FConstGenericParticleHandle ConstParticle1 = FConstGenericParticleHandle(Particle1);
-
-			// See comments in FSingleShapePairCollisionDetector::GenerateCollisionCCDImpl for why we use V() here
-			FReal LengthCCD = 0;
-			FVec3 DirCCD = FVec3(0);
-			const FVec3 DeltaX0 = ConstParticle0->V() * Dt;
-			const FVec3 DeltaX1 = ConstParticle1->V() * Dt;
-			const bool bUseCCD = Collisions::ShouldUseCCD(Particle0, DeltaX0, Particle1, DeltaX1, DirCCD, LengthCCD);
-
-			return bUseCCD;
-		}
-		return false;
-	}
-
-	void FParticlePairMidPhase::InitThresholds()
-	{
-		// @todo(chaos): improve this threshold calculation for thin objects? Dynamic thin objects have bigger problems so maybe we don't care
-		// @todo(chaos): Spheres and capsules need smaller position tolerance - the restore test doesn't work well with rolling
-		const bool bIsDynamic0 = FConstGenericParticleHandle(Particle0)->IsDynamic();
-		const bool bIsDynamic1 = FConstGenericParticleHandle(Particle1)->IsDynamic();
-		const bool bIsBounded0 = Particle0->HasBounds();
-		const bool bIsBounded1 = Particle1->HasBounds();
-
-		// NOTE: If CullDistance ends up smaller than the thresholds used to restore collisions, we can end up missing
-		// collisions as the objects move if we restore a "zero contact" manifold after movement greater than the cull distance. 
-		// Currently this should not happen, but it is not explicitly ensured by the way the thresholds and CullDistanceScale are calculated.
-		// @todo(chaos): Add a way to enforce a CullDistance big enough to support the reuse thresholds
-		const FReal CullDistanceReferenceSizeInv = FReal(Chaos_Collision_CullDistanceScaleInverseSize);
-		const FReal MinCullDistanceScale = FReal(Chaos_Collision_MinCullDistanceScale);
-		const FReal MaxBoundsSize0 = (bIsDynamic0 && bIsBounded0) ? Particle0->LocalBounds().Extents().GetMax() : FReal(0);
-		const FReal MaxBoundsSize1 = (bIsDynamic1 && bIsBounded1) ? Particle1->LocalBounds().Extents().GetMax() : FReal(0);
-		const FReal CullDistanceScale0 = MaxBoundsSize0 * CullDistanceReferenceSizeInv;
-		const FReal CullDistanceScale1 = MaxBoundsSize1 * CullDistanceReferenceSizeInv;
-		CullDistanceScale = (FRealSingle)FMath::Max3(CullDistanceScale0, CullDistanceScale1, MinCullDistanceScale);
-	}
-
-	void FParticlePairMidPhase::GenerateCollisions(
-		const FReal InCullDistance,
-		const FReal Dt,
-		const FCollisionContext& Context)
-	{
-		//SCOPE_CYCLE_COUNTER(STAT_Collisions_GenerateCollisions);
-		PHYSICS_CSV_SCOPED_EXPENSIVE(PhysicsVerbose, DetectCollisions_NarrowPhase);
-
-		if (!IsValid())
-		{
-			return;
-		}
-
-		if (Flags.bIsActive)
-		{
-			// CullDistance is scaled by the size of the dynamic objects.
-			FReal CullDistance = InCullDistance * CullDistanceScale;
-
-			// If CCD is enabled, did we move far enough to require a sweep?
-			const bool bUseSweep = Flags.bIsCCD && ShouldEnableCCD(Dt);
-
-			// Extend cull distance based on velocity
-			// NOTE: We use PreV here which is the velocity after collisions from the previous tick because we want the 
-			// velocity without gravity from this tick applied. This is mainly so that we get the same CullDistance from 
-			// one tick to the next, even if one of the particles goes to sleep, and therefore its velocity is now zero 
-			// because gravity is no longer applied. Also see FPBDIslandManager::PropagateIslandSleep for other issues 
-			// related to velocity and sleeping...
-			const FReal VMax0 = FConstGenericParticleHandle(GetParticle0())->PreV().GetAbsMax();
-			const FReal VMax1 = FConstGenericParticleHandle(GetParticle1())->PreV().GetAbsMax();
-			const FReal VMaxDt = FMath::Max(VMax0, VMax1) * Dt;
-			if (!bUseSweep)
-			{
-				// Normal (non sweep) mode: we increase CullDistance based on velocity up to a limit
-				// NOTE: This somewhat matches the bounds expansion in FPBDRigidsEvolutionGBF::Integrate
-				const FReal VelocityBoundsMultiplier = Context.GetSettings().BoundsVelocityInflation;
-				const FReal MaxVelocityBoundsExpansion = Context.GetSettings().MaxVelocityBoundsExpansion;
-				if ((VelocityBoundsMultiplier > 0) && (MaxVelocityBoundsExpansion > 0))
-				{
-					CullDistance += FMath::Min(VelocityBoundsMultiplier * VMaxDt, MaxVelocityBoundsExpansion);
-				}
-			}
-			else
-			{
-				// CCD (sweept) mode: we increase CullDistance based on velocity, with no limits
-				CullDistance += VMaxDt;
-			}
 
-			// Run collision detection on all potentially colliding shape pairs
-			NumActiveConstraints = 0;
-			if (Flags.bIsCCD)
-			{
-				for (FSingleShapePairCollisionDetector& ShapePair : ShapePairDetectors)
-				{
-					NumActiveConstraints += ShapePair.GenerateCollisionCCD(bUseSweep, CullDistance, Dt, Context);
-				}
-				for (FMultiShapePairCollisionDetector& MultiShapePair : MultiShapePairDetectors)
-				{
-					NumActiveConstraints += MultiShapePair.GenerateCollisionsCCD(bUseSweep, CullDistance, Dt, Context);
-				}
-				if (ParticlePairDetector.IsValid())
-				{
-					ParticlePairDetector->GenerateCollisionsCCD(bUseSweep, CullDistance, Dt, Context);
-				}
-			}
-			else
+			// NOTE: Constraints in sleeping islands should be kept alive. They will still be in the graph
+			// and will be restored when the island wakes.
+			if ((Constraint->GetContainerCookie().LastUsedEpoch < CurrentEpoch) && !Constraint->IsInConstraintGraph())
 			{
-				for (FSingleShapePairCollisionDetector& ShapePair : ShapePairDetectors)
-				{
-					NumActiveConstraints += ShapePair.GenerateCollision(CullDistance, Dt, Context);
-				}
-				for (FMultiShapePairCollisionDetector& MultiShapePair : MultiShapePairDetectors)
-				{
-					NumActiveConstraints += MultiShapePair.GenerateCollisions(CullDistance, Dt, Context);
-				}
-				if (ParticlePairDetector.IsValid())
-				{
-					ParticlePairDetector->GenerateCollisions(CullDistance, Dt, Context);
-				}
+				Pruned.Add(CollisionKey);
 			}
 		}
 
-		// Reset any modifications applied by the MidPhaseModifier.
-		// In principle this belongs at the start of the frame but it is hard to avoid the cache miss there.
-		ResetModifications();
-
-		LastUsedEpoch = Context.GetAllocator()->GetCurrentEpoch();
-	}
-
-	void FParticlePairMidPhase::InjectCollision(const FPBDCollisionConstraint& Constraint, const FCollisionContext& Context)
-	{
-		if (!Constraint.GetContainerCookie().bIsMultiShapePair)
-		{
-			const FPerShapeData* Shape0 = Constraint.GetShape0();
-			const FPerShapeData* Shape1 = Constraint.GetShape1();
-
-			// @todo(chaos): fix O(N) search for shape pair - store the index in the cookie (it will be the same
-			// as long as the ShapesArray on each particle has not changed)
-			for (FSingleShapePairCollisionDetector& ShapePair : ShapePairDetectors)
-			{
-				if (((Shape0 == ShapePair.GetShape0()) && (Shape1 == ShapePair.GetShape1())) || ((Shape0 == ShapePair.GetShape1()) && (Shape1 == ShapePair.GetShape0())))
-				{
-					ShapePair.SetCollision(Constraint, Context);
-					break;
-				}
-			}
-		}
-		else
+		// Destroy expired collisions
+		for (uint32 CollisionKey : Pruned)
 		{
-			// @todo(chaos): implement cluster/large hierarchy Resim restore
-			ensure(false);
+			Constraints.Remove(CollisionKey);
 		}
-
-		LastUsedEpoch = Context.GetAllocator()->GetCurrentEpoch();
 	}
 
-	void FParticlePairMidPhase::SetIsSleeping(const bool bInIsSleeping, const int32 CurrentEpoch)
+	void FGenericParticlePairMidPhase::WakeCollisionsImpl(const int32 CurrentEpoch)
 	{
-		// This can be called from two locations:
-		// 1)	At the start of the tick as a results of some state change from the game thread such as an explicit wake event,
-		//		applying an impulse, or moving a particle.
-		// 2)	After the constraint solver phase when we put non-moving islands to sleep.
-		// 
-		// Note that in both cases there is a collision detection phase before the next constraint solving phase.
-		//
-		// When awakening we re-activate collisions so that we have a "previous" collision to use for static friction etc.
-		// We don't need to do anything when going to sleep because sleeping particles pairs are ignored in collision detection 
-		// so the next set of active collisions generated will not contain these collisions.
-
-		if (Flags.bIsSleeping != bInIsSleeping)
+		for (auto& KVP : Constraints)
 		{
-			// If we are waking particles, reactivate all collisions that were
-			// active when we were put to sleep, i.e., all collisions whose LastUsedEpoch
-			// is equal to our LastUsedEpoch.
-			const bool bWakingUp = !bInIsSleeping;
-			if (bWakingUp)
+			FPBDCollisionConstraintPtr& Constraint = KVP.Value;
+			if (Constraint->GetContainerCookie().LastUsedEpoch >= LastUsedEpoch)
 			{
-				if (LastUsedEpoch < CurrentEpoch)
-				{
-					// Restore all constraints that were active when we were put to sleep
-					for (FSingleShapePairCollisionDetector& ShapePair : ShapePairDetectors)
-					{
-						ShapePair.WakeCollision(LastUsedEpoch, CurrentEpoch);
-					}
-					for (FMultiShapePairCollisionDetector& MultiShapePair : MultiShapePairDetectors)
-					{
-						MultiShapePair.WakeCollisions(LastUsedEpoch, CurrentEpoch);
-					}
-					if (ParticlePairDetector.IsValid())
-					{
-						ParticlePairDetector->WakeCollisions(LastUsedEpoch, CurrentEpoch);
-					}
-					LastUsedEpoch = CurrentEpoch;
-				}
+				Constraint->GetContainerCookie().LastUsedEpoch = CurrentEpoch;
+				Constraint->GetContainerCookie().ConstraintIndex = INDEX_NONE;
+				Constraint->GetContainerCookie().CCDConstraintIndex = INDEX_NONE;
 			}
-			// If we are going to sleep, there is nothing to do (see comments above)
-
-			Flags.bIsSleeping = bInIsSleeping;
 		}
 	}
 
-	bool FParticlePairMidPhase::IsInConstraintGraph() const
+	void FGenericParticlePairMidPhase::InjectCollisionImpl(const FPBDCollisionConstraint& Constraint, const FCollisionContext& Context)
 	{
-		// @todo(chaos): optimize
-		bool bInGraph = false;
-		VisitConstCollisions([&bInGraph](const FPBDCollisionConstraint& Constraint)
-		{
-			if (Constraint.GetConstraintGraphEdge() != nullptr)
-			{
-				bInGraph = true;
-				return ECollisionVisitorResult::Stop;
-			}
-			return ECollisionVisitorResult::Continue;
-		}, false);
-		return bInGraph;
+		// @todo(chaos): support rewind/resim here
+		ensure(false);
 	}
-
 }
 
 
diff --git a/Engine/Source/Runtime/Experimental/Chaos/Private/Chaos/CollisionResolution.cpp b/Engine/Source/Runtime/Experimental/Chaos/Private/Chaos/CollisionResolution.cpp
--- a/Engine/Source/Runtime/Experimental/Chaos/Private/Chaos/CollisionResolution.cpp
+++ b/Engine/Source/Runtime/Experimental/Chaos/Private/Chaos/CollisionResolution.cpp
@@ -1172,8 +1172,8 @@ namespace Chaos
 		}
 
 		bool DoParticleImplicitsOverlap(
-			TGeometryParticleHandle<FReal, 3>* Particle0,
-			TGeometryParticleHandle<FReal, 3>* Particle1,
+			const TGeometryParticleHandle<FReal, 3>* Particle0,
+			const TGeometryParticleHandle<FReal, 3>* Particle1,
 			const FImplicitObject* Implicit0,
 			const FImplicitObject* Implicit1,
 			const FRigidTransform3& ParticleWorldTransform0,
@@ -1274,16 +1274,9 @@ namespace Chaos
 
 			if (bDoOverlap)
 			{
-				if (Context.MultiShapeCollisionDetector != nullptr)
-				{
-					Context.MultiShapeCollisionDetector->FindOrCreateConstraint(
-						Particle0, Implicit0, ImplicitID0, Shape0, Simplicial0, LocalTransform0,
-						Particle1, Implicit1, ImplicitID1, Shape1, Simplicial1, LocalTransform1,
-						CullDistance, ShapePairType, Context.GetSettings().bAllowManifolds, bEnableSweep, Context);
-				}
-				else if (Context.ParticlePairCollisionDetector != nullptr)
+				if (Context.MidPhase != nullptr)
 				{
-					Context.ParticlePairCollisionDetector->FindOrCreateConstraint(
+					Context.MidPhase->FindOrCreateConstraint(
 						Particle0, Implicit0, ImplicitID0, Shape0, Simplicial0, LocalTransform0,
 						Particle1, Implicit1, ImplicitID1, Shape1, Simplicial1, LocalTransform1,
 						CullDistance, ShapePairType, Context.GetSettings().bAllowManifolds, bEnableSweep, Context);
@@ -1672,32 +1665,28 @@ namespace Chaos
 			return EContactShapesType::Unknown;
 		}
 
-		void ConstructConstraints(
-			FGeometryParticleHandle* Particle0, 
-			FGeometryParticleHandle* Particle1,
-			const FImplicitObject* Implicit0, 
-			const FPerShapeData* Shape0, 
-			const FBVHParticles* Simplicial0, 
-			const int32 ImplicitID0,
-			const FImplicitObject* Implicit1, 
-			const FPerShapeData* Shape1, 
-			const FBVHParticles* Simplicial1, 
-			const int32 ImplicitID1,
-			const FRigidTransform3& ParticleWorldTransform0,
-			const FRigidTransform3& LocalTransform0, 
-			const FRigidTransform3& ParticleWorldTransform1, 
-			const FRigidTransform3& LocalTransform1, 
-			const FReal CullDistance, 
-			const FReal dT, 
-			const bool bEnableSweep, 
-			const FCollisionContext& Context)
-		{
-			INC_DWORD_STAT(STAT_ChaosCollisionCounter_NumParticlePairs);
 
+		//
+		//
+		// DEPRECATED STUFF
+		//
+		//
+
+		void ConstructConstraints(FGeometryParticleHandle* Particle0, FGeometryParticleHandle* Particle1, const FImplicitObject* Implicit0,  const FPerShapeData* Shape0,  const FBVHParticles* Simplicial0,  const int32 ImplicitID0, const FImplicitObject* Implicit1, const FPerShapeData* Shape1,  const FBVHParticles* Simplicial1,  const int32 ImplicitID1, const FRigidTransform3& ParticleWorldTransform0, const FRigidTransform3& LocalTransform0,  const FRigidTransform3& ParticleWorldTransform1,  const FRigidTransform3& LocalTransform1, const FReal CullDistance,  const FReal dT,  const bool bEnableSweep,  const FCollisionContext& Context)
+		{
 			ConstructConstraintsRecursive(Particle0, Particle1, Implicit0, Shape0->AsShapeInstance(), Simplicial0, ImplicitID0, Implicit1, Shape1->AsShapeInstance(), Simplicial1, ImplicitID1, ParticleWorldTransform0, LocalTransform0, ParticleWorldTransform1, LocalTransform1, CullDistance, dT, bEnableSweep, Context);
 		}
 
-		// This is deprecated
+		void ConstructConstraints(TGeometryParticleHandle<FReal, 3>* Particle0, TGeometryParticleHandle<FReal, 3>* Particle1, const FImplicitObject* Implicit0, const FPerShapeData* Shape0, const FBVHParticles* Simplicial0, const FImplicitObject* Implicit1, const FPerShapeData* Shape1, const FBVHParticles* Simplicial1, const FRigidTransform3& ParticleWorldTransform0, const FRigidTransform3& Transform0, const FRigidTransform3& ParticleWorldTransform1, const FRigidTransform3& Transform1, const FReal CullDistance, const FReal Dt, const bool bEnableSweep, const FCollisionContext& Context)
+		{
+			ConstructConstraintsRecursive(Particle0, Particle1, Implicit0, Shape0->AsShapeInstance(), Simplicial0, 0, Implicit1, Shape1->AsShapeInstance(), Simplicial1, 0, ParticleWorldTransform0, Transform0, ParticleWorldTransform1, Transform1, CullDistance, Dt, bEnableSweep, Context);
+		}
+
+		void ConstructConstraints(TGeometryParticleHandle<FReal, 3>* Particle0, TGeometryParticleHandle<FReal, 3>* Particle1, const FImplicitObject* Implicit0, const FPerShapeData* Shape0, const FBVHParticles* Simplicial0, const FImplicitObject* Implicit1, const FPerShapeData* Shape1, const FBVHParticles* Simplicial1, const FRigidTransform3& ParticleWorldTransform0, const FRigidTransform3& Transform0, const FRigidTransform3& ParticleWorldTransform1, const FRigidTransform3& Transform1, const FReal CullDistance, const FReal Dt, const FCollisionContext& Context)
+		{
+			ConstructConstraintsRecursive(Particle0, Particle1, Implicit0, Shape0->AsShapeInstance(), Simplicial0, 0, Implicit1, Shape1->AsShapeInstance(), Simplicial1, 0, ParticleWorldTransform0, Transform0, ParticleWorldTransform1, Transform1, CullDistance, Dt, false, Context);
+		}
+
 		template<ECollisionUpdateType UpdateType>
 		void UpdateConstraintFromGeometry(FPBDCollisionConstraint& Constraint, const FRigidTransform3& ParticleWorldTransform0, const FRigidTransform3& ParticleWorldTransform1, const FReal Dt)
 		{
diff --git a/Engine/Source/Runtime/Experimental/Chaos/Private/Chaos/ImplicitObjectBVH.cpp b/Engine/Source/Runtime/Experimental/Chaos/Private/Chaos/ImplicitObjectBVH.cpp
--- a/Engine/Source/Runtime/Experimental/Chaos/Private/Chaos/ImplicitObjectBVH.cpp
+++ b/Engine/Source/Runtime/Experimental/Chaos/Private/Chaos/ImplicitObjectBVH.cpp
@@ -1,19 +1,17 @@
 // Copyright Epic Games, Inc. All Rights Reserved.
 #include "Chaos/ImplicitObjectBVH.h"
 #include "Chaos/ChaosArchive.h"
-
-//UE_DISABLE_OPTIMIZATION
+#include "Misc/MemStack.h"
 
 namespace Chaos
 {
 	namespace CVars
 	{
-		extern int32 ChaosUnionBVHMaxDepth;
+		extern FRealSingle ChaosUnionBVHSplitBias;
 	}
 
 	namespace Private
 	{
-
 		FImplicitBVHObject::FImplicitBVHObject()
 		{
 		}
@@ -34,6 +32,8 @@ namespace Chaos
 		{
 		}
 
+		///////////////////////////////////////////////////////////////////////////////////////////
+		///////////////////////////////////////////////////////////////////////////////////////////
 		///////////////////////////////////////////////////////////////////////////////////////////
 		///////////////////////////////////////////////////////////////////////////////////////////
 
@@ -65,12 +65,13 @@ namespace Chaos
 				const TUniquePtr<FImplicitObject>& RootObject = InRootObjects[RootObjectIndex];
 
 				RootObject->VisitLeafObjects(
-					[RootObjectIndex, &Objects](const FImplicitObject* Object, const FRigidTransform3& ParentTransform, const int32 UnusedRootObjectIndex, const int32 ObjectIndex, const int32 UnusedLeafObjectIndex)
+					[RootObjectIndex, &Objects](const FImplicitObject* Object, const FRigidTransform3& ParentTransform, const int32 UnusedRootObjectIndex, const int32 UnusedObjectIndex, const int32 UnusedLeafObjectIndex)
 					{
 						// @todo(chaos): clean this up (SetFromRawLowLevel). We know all the objects we visit are children of a UniquePtr because we own it
 						TSerializablePtr<FImplicitObject> SerializableObject;
 						SerializableObject.SetFromRawLowLevel(Object);
 
+						const int32 ObjectIndex = Objects.Num();
 						Objects.Emplace(
 							SerializableObject,
 							ParentTransform.GetTranslation(),
@@ -103,22 +104,185 @@ namespace Chaos
 		///////////////////////////////////////////////////////////////////////////////////////////
 
 		FImplicitBVH::FImplicitBVH()
+			: NodeObjectIndices()
+			, Nodes()
+			, TreeDepth(0)
 		{
 		}
 
-		FImplicitBVH::FImplicitBVH(FObjects&& InObjects, const int32 InMaxBVHDepth)
+		FImplicitBVH::FImplicitBVH(FObjects&& InObjects, const int32 MaxDepth)
+			: NodeObjectIndices()
+			, Nodes()
+			, TreeDepth(0)
 		{
-			Init(MoveTemp(InObjects), InMaxBVHDepth);
+			Init(MoveTemp(InObjects), MaxDepth, 1);
 		}
 
 		FImplicitBVH::~FImplicitBVH()
 		{
 		}
 
-		void FImplicitBVH::Init(FObjects&& InObjects, const int32 InMaxBVHDepth)
+		void FImplicitBVH::Init(FObjects&& InObjects, const int32 MaxDepth, const int32 MaxLeafObjects)
 		{
 			Objects = MoveTemp(InObjects);
-			BVH = FBVH(Objects, FMath::Max(1, InMaxBVHDepth));
+			BuildTree(MaxDepth, MaxLeafObjects);
+		}
+
+		void FImplicitBVH::BuildTree(const int32 MaxDepth, const int32 MaxLeafObjects)
+		{
+			QUICK_SCOPE_CYCLE_COUNTER(Chaos_FImplicitBVH_BuildTree);
+
+			Nodes.Reset();
+
+			NodeObjectIndices.SetNumUninitialized(Objects.Num());
+
+			FAABB3f NodeBounds = FAABB3f::EmptyAABB();
+
+			FWorkData WorkData;
+			WorkData.ObjectCenters.SetNumUninitialized(GetNumObjects());
+			WorkData.NodeObjectIndices0.Reserve(Objects.Num());
+			WorkData.NodeObjectIndices1.Reserve(Objects.Num());
+
+			for (int32 ObjectIndex = 0; ObjectIndex < Objects.Num(); ++ObjectIndex)
+			{
+				const FImplicitBVHObject& Object = Objects[ObjectIndex];
+
+				NodeObjectIndices[ObjectIndex] = ObjectIndex;
+
+				NodeBounds.GrowToInclude(Object.GetBounds());
+
+				WorkData.ObjectCenters[ObjectIndex] = Object.GetBoundsCenter();
+			}
+
+			AddNodeRecursive(NodeBounds, 0, GetNumObjects(), 0, MaxDepth, MaxLeafObjects, WorkData);
+		}
+
+		int32 FImplicitBVH::AddNodeRecursive(const FAABB3f& NodeBounds, const int32 ObjectBeginIndex, const int32 ObjectEndIndex, const int32 NodeDepth, const int32 MaxDepth, const int32 MaxLeafObjects, FWorkData& WorkData)
+		{
+			// Create a new node. By default this is a leaf node, but we may add children below
+			const int32 NodeIndex = Nodes.Emplace(NodeBounds, ObjectBeginIndex, ObjectEndIndex);
+
+			// Keep track of tree depth (used for stack sizing)
+			TreeDepth = FMath::Max(NodeDepth + 1, TreeDepth);
+
+			const int32 NumNodeObjects = ObjectEndIndex - ObjectBeginIndex;
+			if ((NumNodeObjects > MaxLeafObjects) && (NodeDepth < MaxDepth - 1))
+			{
+
+				// NOTE: if we fail to partition, the node we just created will be a leaf node
+				FAABB3f ChildNodeBounds0, ChildNodeBounds1;
+				int32 ChildBeginIndex1;
+				const bool bPartitioned = Partition(NodeBounds, ObjectBeginIndex, ObjectEndIndex, ChildNodeBounds0, ChildNodeBounds1, ChildBeginIndex1, WorkData);
+
+				if (bPartitioned)
+				{
+					// @todo(chaos): could be run in parallel, but only worth it when each partition has a large number of children (and we couldn't share the WorkData arrays)
+					Nodes[NodeIndex].ChildNodeIndices[0] = AddNodeRecursive(ChildNodeBounds0, ObjectBeginIndex, ChildBeginIndex1, NodeDepth + 1, MaxDepth, MaxLeafObjects, WorkData);
+					Nodes[NodeIndex].ChildNodeIndices[1] = AddNodeRecursive(ChildNodeBounds1, ChildBeginIndex1, ObjectEndIndex, NodeDepth + 1, MaxDepth, MaxLeafObjects, WorkData);
+				}
+			}
+
+			return NodeIndex;
+		}
+
+		bool FImplicitBVH::Partition(const FAABB3f& NodeBounds, const int32 ObjectBeginIndex, const int32 ObjectEndIndex, FAABB3f& OutChildNodeBounds0, FAABB3f& OutChildNodeBounds1, int32& OutPartitionIndex, FWorkData& WorkData)
+		{
+			OutChildNodeBounds0 = FAABB3f::EmptyAABB();
+			OutChildNodeBounds1 = FAABB3f::EmptyAABB();
+			OutPartitionIndex = INDEX_NONE;
+
+			const int32 NumNodeObjects = ObjectEndIndex - ObjectBeginIndex;
+			if (NumNodeObjects <= 0)
+			{
+				return false;
+			}
+
+			// Determine the split point with a nudge
+			// @todo(chaos): currently using the mean center point - experiment with other options
+			FVec3f MeanCenter = FVec3f(0);
+			for (int32 NodeObjectIndex = ObjectBeginIndex; NodeObjectIndex < ObjectEndIndex; ++NodeObjectIndex)
+			{
+				MeanCenter += WorkData.ObjectCenters[NodeObjectIndices[NodeObjectIndex]];
+			}
+			MeanCenter /= FRealSingle(NumNodeObjects);
+			const FVec3f SplitPoint = MeanCenter + FVec3(CVars::ChaosUnionBVHSplitBias);
+
+			// Determine the split axis based on how many objects are exclusively one one side of the split
+			// @todo(chaos): look at how much volume we add to each side of the split bounds instead?
+			TVec3<int32> RejectLeft = TVec3<int32>(0);
+			TVec3<int32> RejectRight = TVec3<int32>(0);
+			for (int32 NodeObjectIndex = ObjectBeginIndex; NodeObjectIndex < ObjectEndIndex; ++NodeObjectIndex)
+			{
+				const FImplicitBVHObject& Object = Objects[NodeObjectIndices[NodeObjectIndex]];
+				for (int32 Axis = 0; Axis < 3; ++Axis)
+				{
+					if (Object.GetBounds().Max()[Axis] < SplitPoint[Axis])
+					{
+						RejectLeft[Axis] += 1;
+					}
+					if (Object.GetBounds().Min()[Axis] > SplitPoint[Axis])
+					{
+						RejectRight[Axis] += 1;
+					}
+				}
+			}
+
+			// Pick the axis which has the most objects on both sides of the split
+			int32 BestAxis = INDEX_NONE;
+			int32 MaxRejects = 0;
+			for (int32 Axis = 0; Axis < 3; ++Axis)
+			{
+				const int32 NumRejects = FMath::Min(RejectLeft[Axis], RejectRight[Axis]);
+				if (NumRejects > MaxRejects)
+				{
+					MaxRejects = NumRejects;
+					BestAxis = Axis;
+				}
+			}
+
+			// If all objects span the split point on all axes, there's no point dividing further
+			// @todo(chaos): we could try to choose a different split point
+			if (BestAxis == INDEX_NONE)
+			{
+				return false;
+			}
+
+			// Collect all the object indices and net bounds for each side of the split
+			WorkData.NodeObjectIndices0.Reset(NumNodeObjects);
+			WorkData.NodeObjectIndices1.Reset(NumNodeObjects);
+			for (int32 NodeObjectIndex = ObjectBeginIndex; NodeObjectIndex < ObjectEndIndex; ++NodeObjectIndex)
+			{
+				const int32 ObjectIndex = NodeObjectIndices[NodeObjectIndex];
+				const FImplicitBVHObject& Object = Objects[ObjectIndex];
+
+				if (WorkData.ObjectCenters[ObjectIndex][BestAxis] < SplitPoint[BestAxis])
+				{
+					WorkData.NodeObjectIndices0.Add(ObjectIndex);
+					OutChildNodeBounds0.GrowToInclude(Object.GetBounds());
+				}
+				else
+				{
+					WorkData.NodeObjectIndices1.Add(ObjectIndex);
+					OutChildNodeBounds1.GrowToInclude(Object.GetBounds());
+				}
+			}
+
+			// Update the node indices to reflect the partitioning
+			int32 NextNodeObjectIndex = ObjectBeginIndex;
+			for (int32 NodeObjectIndex0 : WorkData.NodeObjectIndices0)
+			{
+				NodeObjectIndices[NextNodeObjectIndex++] = NodeObjectIndex0;
+			}
+			for (int32 NodeObjectIndex1 : WorkData.NodeObjectIndices1)
+			{
+				NodeObjectIndices[NextNodeObjectIndex++] = NodeObjectIndex1;
+			}
+			check(NextNodeObjectIndex == ObjectEndIndex);
+
+			// The point in NodeObjectIndices where we split the objects in two
+			OutPartitionIndex = ObjectBeginIndex + WorkData.NodeObjectIndices0.Num();
+
+			return (OutPartitionIndex != INDEX_NONE);
 		}
 
 		///////////////////////////////////////////////////////////////////////////////////////////
@@ -142,7 +306,12 @@ namespace Chaos
 		FChaosArchive& FImplicitBVH::Serialize(FChaosArchive& Ar)
 		{
 			Ar << Objects;
+
+			// @todo(chaos): add new object version and serialize new BVH data
+			// Load BVH intop temp and throw away, or save empty BVH
+			TBoundingVolumeHierarchy<FObjects, TArray<int32>> BVH;
 			Ar << BVH;
+
 			return Ar;
 		}
 
@@ -155,5 +324,5 @@ namespace Chaos
 			return Ar;
 		}
 
-	}
-}
\ No newline at end of file
+	} // namespace Private
+} // namespace Chaos
\ No newline at end of file
diff --git a/Engine/Source/Runtime/Experimental/Chaos/Private/Chaos/ImplicitObjectUnion.cpp b/Engine/Source/Runtime/Experimental/Chaos/Private/Chaos/ImplicitObjectUnion.cpp
--- a/Engine/Source/Runtime/Experimental/Chaos/Private/Chaos/ImplicitObjectUnion.cpp
+++ b/Engine/Source/Runtime/Experimental/Chaos/Private/Chaos/ImplicitObjectUnion.cpp
@@ -13,9 +13,11 @@ namespace Chaos
 	{
 		int32 ChaosUnionBVHMinShapes = 32;
 		int32 ChaosUnionBVHMaxDepth = 8;
+		FRealSingle ChaosUnionBVHSplitBias = 0.1f;
 		bool bChaosUnionBVHEnabled = true;
 		FAutoConsoleVariableRef CVarChaosUnionBVHMinShapes(TEXT("p.Chaos.Collision.UnionBVH.NumShapes"), ChaosUnionBVHMinShapes, TEXT("If a geometry hierarchy has this many shapes, wrap it in a BVH for collision detection (negative to disable BVH)"));
 		FAutoConsoleVariableRef CVarChaosUnionBVHMaxDepth(TEXT("p.Chaos.Collision.UnionBVH.MaxDepth"), ChaosUnionBVHMaxDepth, TEXT("The allowed depth of the BVH when used to wrap a shape hiererchy"));
+		FAutoConsoleVariableRef CVarChaosUnionBVHSplitBias(TEXT("p.Chaos.Collision.UnionBVH.SplitBias"), ChaosUnionBVHSplitBias, TEXT(""));
 		FAutoConsoleVariableRef CVarChaosUnionBVHEnabled(TEXT("p.Chaos.Collision.UnionBVH.Enabled"), bChaosUnionBVHEnabled, TEXT("Set to false to disable use of BVH during collision detection (without affecting creations and serialization)"));
 	}
 
@@ -119,7 +121,7 @@ void FImplicitObjectUnion::SetNumLeafObjects(int32 InNumLeafObjects)
 
 void FImplicitObjectUnion::CreateBVH()
 {
-	if (Flags.bAllowBVH)
+	if (Flags.bAllowBVH && CVars::bChaosUnionBVHEnabled)
 	{
 		const int32 MinBVHShapes = CVars::ChaosUnionBVHMinShapes;
 		const int32 MaxBVHDepth = CVars::ChaosUnionBVHMaxDepth;
@@ -147,13 +149,18 @@ void FImplicitObjectUnion::FindAllIntersectingObjects(TArray<Pair<const FImplici
 {
 	if (BVH.IsValid() && CVars::bChaosUnionBVHEnabled)
 	{
-		TArray<int32> Overlaps = BVH->GetBVH().FindAllIntersections(LocalBounds);
-		Out.Reserve(Out.Num() + Overlaps.Num());
-		for (int32 Idx : Overlaps)
-		{
-			const FImplicitObject* Obj = BVH->GetGeometry(Idx);
-			Out.Add(MakePair(Obj, BVH->GetTransform(Idx)));
-		}
+		BVH->VisitOverlappingNodes(LocalBounds,
+			[this, &Out](const Private::FImplicitBVHNode& Node)
+			{
+				if (Node.IsLeaf())
+				{
+					BVH->VisitNodeObjects(Node,
+						[&Out](const FImplicitObject* Implicit, const FRigidTransform3f& RelativeTransformf, const FAABB3f& RelativeBoundsf, const int32 RootObjectIndex, const int32 LeafObjectIndex) -> void
+						{
+							Out.Add(MakePair(Implicit, FRigidTransform3(RelativeTransformf)));
+						});
+				}
+			});
 	}
 	else
 	{
@@ -168,30 +175,31 @@ void FImplicitObjectUnion::VisitOverlappingLeafObjectsImpl(
 	const FAABB3& LocalBounds,
 	const FRigidTransform3& ObjectTransform,
 	const int32 InRootObjectIndex,
-	int32& ObjectIndex,
-	int32& LeafObjectIndex,
+	int32& InOutObjectIndex,
+	int32& InOutLeafObjectIndex,
 	const FImplicitHierarchyVisitor& VisitorFunc) const
 {
+	// Skip self
+	InOutObjectIndex++;
+
 	if (BVH.IsValid() && CVars::bChaosUnionBVHEnabled)
 	{
 		// Visit children
+		// NOTE: ObjectIndex passed to the visitor isn't really correct here. Maybe it should be removed...
 		BVH->VisitAllIntersections(LocalBounds,
-			[this, &ObjectIndex, &LeafObjectIndex, &ObjectTransform, &VisitorFunc](const int32 BVHObjectIndex)
+			[this, &ObjectTransform, &InOutObjectIndex, &VisitorFunc](const FImplicitObject* Implicit, const FRigidTransform3f& RelativeTransformf, const FAABB3f& RelativeBoundsf, const int32 RootObjectIndex, const int32 LeafObjectIndex)
 			{
-				VisitorFunc(BVH->GetGeometry(BVHObjectIndex), BVH->GetTransform(BVHObjectIndex) * ObjectTransform, BVH->GetRootObjectIndex(BVHObjectIndex), BVH->GetObjectIndex(BVHObjectIndex), BVHObjectIndex);
+				VisitorFunc(Implicit, FRigidTransform3(RelativeTransformf) * ObjectTransform, RootObjectIndex, InOutObjectIndex++, LeafObjectIndex);
 			});
 	}
 	else
 	{
-		// Skip self
-		++ObjectIndex;
-
 		for (int32 BVHObjectIndex = 0; BVHObjectIndex < MObjects.Num(); ++BVHObjectIndex)
 		{
 			// If we are the root our object index is the root index, otherwise just pass on the value we were given (from the actual root)
 			const int32 RootObjectIndex = (InRootObjectIndex != INDEX_NONE) ? InRootObjectIndex : BVHObjectIndex;
 
-			MObjects[BVHObjectIndex]->VisitOverlappingLeafObjectsImpl(LocalBounds, ObjectTransform, RootObjectIndex, ObjectIndex, LeafObjectIndex, VisitorFunc);
+			MObjects[BVHObjectIndex]->VisitOverlappingLeafObjectsImpl(LocalBounds, ObjectTransform, RootObjectIndex, InOutObjectIndex, InOutLeafObjectIndex, VisitorFunc);
 		}
 	}
 }
@@ -215,27 +223,53 @@ void FImplicitObjectUnion::VisitLeafObjectsImpl(
 	}
 }
 
-void FImplicitObjectUnion::VisitObjectsImpl(
+bool FImplicitObjectUnion::VisitObjectsImpl(
 	const FRigidTransform3& ObjectTransform,
 	const int32 InRootObjectIndex,
 	int32& ObjectIndex,
 	int32& LeafObjectIndex,
-	const FImplicitHierarchyVisitor& VisitorFunc) const
+	const FImplicitHierarchyVisitorBool& VisitorFunc) const
 {
 	// Visit self
-	VisitorFunc(this, ObjectTransform, InRootObjectIndex, ObjectIndex, INDEX_NONE);
+	bool bContinue = VisitorFunc(this, ObjectTransform, InRootObjectIndex, ObjectIndex, INDEX_NONE);
 	++ObjectIndex;
 
 	// Visit Children
-	for (int32 BVHObjectIndex = 0; BVHObjectIndex < MObjects.Num(); ++BVHObjectIndex)
+	for (int32 BVHObjectIndex = 0; (BVHObjectIndex < MObjects.Num()) && bContinue; ++BVHObjectIndex)
 	{
 		// If we are the root our object index is the root index, otherwise just pass on the value we were given (from the actual root)
 		const int32 RootObjectIndex = (InRootObjectIndex != INDEX_NONE) ? InRootObjectIndex : BVHObjectIndex;
 
-		MObjects[BVHObjectIndex]->VisitObjectsImpl(ObjectTransform, RootObjectIndex, ObjectIndex, LeafObjectIndex, VisitorFunc);
+		bContinue = MObjects[BVHObjectIndex]->VisitObjectsImpl(ObjectTransform, RootObjectIndex, ObjectIndex, LeafObjectIndex, VisitorFunc);
 	}
+
+	return bContinue;
 }
 
+bool FImplicitObjectUnion::IsOverlappingBoundsImpl(const FAABB3& LocalBounds) const
+{
+	if (BVH.IsValid() && CVars::bChaosUnionBVHEnabled)
+	{
+		return BVH->IsOverlappingBounds(LocalBounds);
+	}
+	else
+	{
+		if (LocalBounds.Intersects(BoundingBox()))
+		{
+			for (int32 BVHObjectIndex = 0; BVHObjectIndex < MObjects.Num(); ++BVHObjectIndex)
+			{
+				if (MObjects[BVHObjectIndex]->IsOverlappingBoundsImpl(LocalBounds))
+				{
+					return true;
+				}
+			}
+		}
+	}
+
+	return false;
+}
+
+
 TUniquePtr<FImplicitObject> FImplicitObjectUnion::Copy() const
 {
 	TArray<TUniquePtr<FImplicitObject>> CopyOfObjects;
@@ -286,7 +320,7 @@ void FImplicitObjectUnion::ForEachObject(TFunctionRef<bool(const FImplicitObject
 	// we only visit our immediate children, and not their children. It should probably just ignore the BVH?
 	if (BVH.IsValid())
 	{
-		for (int32 Index = 0; Index < BVH->NumObjects(); ++Index)
+		for (int32 Index = 0; Index < BVH->GetNumObjects(); ++Index)
 		{
 			if (const FImplicitObject* SubObject = BVH->GetGeometry(Index))
 			{
@@ -346,7 +380,13 @@ void FImplicitObjectUnion::Serialize(FChaosArchive& Ar)
 			{
 				BVH = Private::FImplicitBVH::MakeEmpty();
 			}
+
 			Ar << *BVH;
+
+			if (Ar.IsLoading())
+			{
+				RebuildBVH();
+			}
 		}
 	}
 }
@@ -420,61 +460,31 @@ FImplicitObjectUnionClustered::FImplicitObjectUnionClustered(FImplicitObjectUnio
 	Type = ImplicitObjectType::UnionClustered;
 }
 
-void FImplicitObjectUnionClustered::FindAllIntersectingClusteredObjects(TArray<FLargeUnionClusteredImplicitInfo>& Out, const FAABB3& LocalBounds) const
-{
-	if (BVH.IsValid())
-	{
-		TArray<int32> Overlaps = BVH->GetBVH().FindAllIntersections(LocalBounds);
-		Out.Reserve(Out.Num() + Overlaps.Num());
-		for (int32 Idx : Overlaps)
-		{
-			const FImplicitObject* Obj = BVH->GetGeometry(Idx);
-			const FBVHParticles* Simplicial = MOriginalParticleLookupHack.IsValidIndex(Idx) ? MOriginalParticleLookupHack[Idx]->CollisionParticles().Get() : nullptr;
-			Out.Add(FLargeUnionClusteredImplicitInfo(Obj, BVH->GetTransform(Idx), Simplicial));
-		}
-	}
-	else
-	{
-		TArray<Pair<const FImplicitObject*, FRigidTransform3>> LocalOut;
-		TArray<int32> Idxs;
-		for (int32 Idx = 0; Idx < MObjects.Num(); ++Idx)
-		{
-			int32 NumOut = LocalOut.Num();
-			const TUniquePtr<FImplicitObject>& Object = MObjects[Idx];
-			Object->FindAllIntersectingObjects(LocalOut, LocalBounds);
-			for (int32 i = NumOut; i < LocalOut.Num(); ++i)
-			{
-				Idxs.Add(Idx);
-			}
-		}
-		for (int32 Idx = 0; Idx < LocalOut.Num(); ++Idx)
-		{
-			auto& OutElem = LocalOut[Idx];
-			const FBVHParticles* Simplicial = MOriginalParticleLookupHack.IsValidIndex(Idxs[Idx]) ? MOriginalParticleLookupHack[Idxs[Idx]]->CollisionParticles().Get() : nullptr;
-
-			Out.Add(FLargeUnionClusteredImplicitInfo(OutElem.First, OutElem.Second, Simplicial));
-		}
-	}
-}
-
 TArray<FPBDRigidParticleHandle*>
 FImplicitObjectUnionClustered::FindAllIntersectingChildren(const FAABB3& LocalBounds) const
 {
 	TArray<FPBDRigidParticleHandle*> IntersectingChildren;
-	if (BVH.IsValid()) //todo: make this work when hierarchy is not built
+	if (BVH.IsValid())
 	{
-		TArray<int32> IntersectingIndices = BVH->GetBVH().FindAllIntersections(LocalBounds);
-		IntersectingChildren.Reserve(IntersectingIndices.Num());
-		for (const int32 Idx : IntersectingIndices)
-		{
-			if (MOriginalParticleLookupHack.IsValidIndex(Idx))
+		BVH->VisitOverlappingNodes(LocalBounds,
+			[this, &IntersectingChildren](const Private::FImplicitBVHNode& Node)
 			{
-				IntersectingChildren.Add(MOriginalParticleLookupHack[Idx]);
-			}
-		}
+				if (Node.IsLeaf())
+				{
+					BVH->VisitNodeObjects(Node,
+						[this, &IntersectingChildren](const FImplicitObject* Implicit, const FRigidTransform3f& RelativeTransformf, const FAABB3f& RelativeBoundsf, const int32 RootObjectIndex, const int32 LeafObjectIndex) -> void
+						{
+							if (ensure(MOriginalParticleLookupHack.IsValidIndex(LeafObjectIndex)))
+							{
+								IntersectingChildren.Add(MOriginalParticleLookupHack[LeafObjectIndex]);
+							}
+						});
+				}
+			});
 	}
 	else
 	{
+		// @todo(chaos): make this work when there's no BVH
 		IntersectingChildren = MOriginalParticleLookupHack;
 	}
 	return IntersectingChildren;
diff --git a/Engine/Source/Runtime/Experimental/Chaos/Private/Chaos/PBDRigidClusteringAlgo.cpp b/Engine/Source/Runtime/Experimental/Chaos/Private/Chaos/PBDRigidClusteringAlgo.cpp
--- a/Engine/Source/Runtime/Experimental/Chaos/Private/Chaos/PBDRigidClusteringAlgo.cpp
+++ b/Engine/Source/Runtime/Experimental/Chaos/Private/Chaos/PBDRigidClusteringAlgo.cpp
@@ -423,9 +423,17 @@ namespace Chaos
 			Parent->SetLocalBounds(Implicit->BoundingBox());
 			const Chaos::FRigidTransform3 Xf(Parent->X(), Parent->R());
 			Parent->UpdateWorldSpaceState(Xf, FVec3(0));
-		}
 
-		
+			if (const FImplicitObjectUnion* ImplicitUnion = Implicit->GetObject<FImplicitObjectUnion>())
+			{
+				const_cast<FImplicitObjectUnion*>(ImplicitUnion)->SetAllowBVH(true);
+			}
+			else if (const FImplicitObjectUnion* ImplicitUnionClustered = Implicit->GetObject<FImplicitObjectUnionClustered>())
+			{
+				const_cast<FImplicitObjectUnion*>(ImplicitUnionClustered)->SetAllowBVH(true);
+			}
+		}
+	
 		// Update filter data on new shapes
 		const FRigidClustering::FRigidHandleArray& ChildrenArray = ChildrenMap[Parent];
 		UpdateClusterFilterDataFromChildren(Parent, ChildrenArray);
diff --git a/Engine/Source/Runtime/Experimental/Chaos/Private/Chaos/ParticleHandle.cpp b/Engine/Source/Runtime/Experimental/Chaos/Private/Chaos/ParticleHandle.cpp
--- a/Engine/Source/Runtime/Experimental/Chaos/Private/Chaos/ParticleHandle.cpp
+++ b/Engine/Source/Runtime/Experimental/Chaos/Private/Chaos/ParticleHandle.cpp
@@ -4,6 +4,7 @@
 
 #include "Chaos/ImplicitObjectTransformed.h"
 #include "Chaos/ImplicitObjectScaled.h"
+#include "Chaos/ImplicitObjectUnion.h"
 #include "Chaos/TriangleMeshImplicitObject.h"
 #include "Chaos/Framework/PhysicsProxy.h"
 #include "Chaos/Framework/PhysicsSolverBase.h"
@@ -113,6 +114,18 @@ namespace Chaos
 
 	}
 
+	template <typename T, int d>
+	void Chaos::TGeometryParticle<T, d>::PrepareBVHImpl()
+	{
+		if (MNonFrequentData.IsDirty(MDirtyFlags))
+		{
+			if (const FImplicitObjectUnion* Union = MNonFrequentData.Read().Geometry()->GetObject<FImplicitObjectUnion>())
+			{
+				// This will rebuild the BVH if the geometry is new, otherwise do nothing
+				const_cast<FImplicitObjectUnion*>(Union)->SetAllowBVH(true);
+			}
+		}
+	}
 
 	template <typename T, int d>
 	void Chaos::TGeometryParticle<T, d>::SetIgnoreAnalyticCollisionsImp(FImplicitObject* Implicit, bool bIgnoreAnalyticCollisions)
diff --git a/Engine/Source/Runtime/Experimental/Chaos/Private/Chaos/ShapeInstance.cpp b/Engine/Source/Runtime/Experimental/Chaos/Private/Chaos/ShapeInstance.cpp
--- a/Engine/Source/Runtime/Experimental/Chaos/Private/Chaos/ShapeInstance.cpp
+++ b/Engine/Source/Runtime/Experimental/Chaos/Private/Chaos/ShapeInstance.cpp
@@ -141,6 +141,15 @@ namespace Chaos
 					return GetInnerGeometryInstanceData(Instanced->GetInnerObject().Get(), OutRelativeTransformPtr);
 				}
 			}
+			else if ((ImplicitOuterType == FImplicitObjectUnion::StaticType()) || (ImplicitOuterType == FImplicitObjectUnionClustered::StaticType()))
+			{
+				// If the union only has one child, we keep recursing, otherwise we don't
+				const FImplicitObjectUnion* Union = static_cast<const FImplicitObjectUnion*>(Implicit);
+				if (Union->GetObjects().Num() == 1)
+				{
+					return GetInnerGeometryInstanceData(Union->GetObjects()[0].Get(), OutRelativeTransformPtr);
+				}
+			}
 			else if ((uint32)ImplicitOuterType & ImplicitObjectType::IsScaled)
 			{
 				// Scaled Implicit
diff --git a/Engine/Source/Runtime/Experimental/Chaos/Private/PBDRigidsSolver.cpp b/Engine/Source/Runtime/Experimental/Chaos/Private/PBDRigidsSolver.cpp
--- a/Engine/Source/Runtime/Experimental/Chaos/Private/PBDRigidsSolver.cpp
+++ b/Engine/Source/Runtime/Experimental/Chaos/Private/PBDRigidsSolver.cpp
@@ -105,6 +105,7 @@ namespace Chaos
 	namespace CVars
 	{
 		int32 ChaosSolverDebugDrawShapes = CHAOS_SOLVER_ENABLE_DEBUG_DRAW;
+		int32 ChaosSolverDebugDrawBVHs = 0;
 		int32 ChaosSolverDebugDrawCollisions = CHAOS_SOLVER_ENABLE_DEBUG_DRAW;
 		int32 ChaosSolverDebugDrawCollidingShapes = 0;
 		int32 ChaosSolverDebugDrawBounds = 0;
@@ -123,11 +124,13 @@ namespace Chaos
 		int32 ChaosSolverDebugDrawSuspensionConstraints = 0;
 		int32 ChaosSolverDrawClusterConstraints = 0;
 		int32 ChaosSolverDebugDrawMeshContacts = 0;
+		int32 ChaosSolverDebugDrawMeshBVHOverlaps = 0;
 		int32 ChaosSolverDebugDrawColorShapeByClientServer = 0;
 		int32 ChaosSolverDebugDrawShowServer = 1;
 		int32 ChaosSolverDebugDrawShowClient = 1;
 		DebugDraw::FChaosDebugDrawJointFeatures ChaosSolverDrawJointFeatures = DebugDraw::FChaosDebugDrawJointFeatures::MakeDefault();
 		FAutoConsoleVariableRef CVarChaosSolverDrawShapes(TEXT("p.Chaos.Solver.DebugDrawShapes"), ChaosSolverDebugDrawShapes, TEXT("Draw Shapes (0 = never; 1 = end of frame)."));
+		FAutoConsoleVariableRef CVarChaosSolverDrawBVHs(TEXT("p.Chaos.Solver.DebugDrawBVHs"), ChaosSolverDebugDrawBVHs, TEXT("Draw Particle BVHs where applicable (0 = never; 1 = end of frame)."));
 		FAutoConsoleVariableRef CVarChaosSolverDrawCollisions(TEXT("p.Chaos.Solver.DebugDrawCollisions"), ChaosSolverDebugDrawCollisions, TEXT("Draw Collisions (0 = never; 1 = end of frame)."));
 		FAutoConsoleVariableRef CVarChaosSolverDrawCollidingShapes(TEXT("p.Chaos.Solver.DebugDrawCollidingShapes"), ChaosSolverDebugDrawCollidingShapes, TEXT("Draw Shapes that have collisions on them (0 = never; 1 = end of frame)."));
 		FAutoConsoleVariableRef CVarChaosSolverDrawBounds(TEXT("p.Chaos.Solver.DebugDrawBounds"), ChaosSolverDebugDrawBounds, TEXT("Draw bounding volumes inside the broadphase (0 = never; 1 = end of frame)."));
@@ -154,6 +157,7 @@ namespace Chaos
 		FAutoConsoleVariableRef CVarChaosSolverDrawSuspensionConstraints(TEXT("p.Chaos.Solver.DebugDrawSuspension"), ChaosSolverDebugDrawSuspensionConstraints, TEXT("Draw Suspension (0 = never; 1 = end of frame)."));
 		FAutoConsoleVariableRef CVarChaosSolverDrawClusterConstraints(TEXT("p.Chaos.Solver.DebugDraw.Cluster.Constraints"), ChaosSolverDrawClusterConstraints, TEXT("Draw Active Cluster Constraints (0 = never; 1 = end of frame)."));
 		FAutoConsoleVariableRef CVarChaosSolverDrawMeshContacts(TEXT("p.Chaos.Solver.DebugDrawMeshContacts"), ChaosSolverDebugDrawMeshContacts, TEXT("Draw Mesh contacts"));
+		FAutoConsoleVariableRef CVarChaosSolverDrawMeshBVHOverlaps(TEXT("p.Chaos.Solver.DebugDrawMeshBVHOverlaps"), ChaosSolverDebugDrawMeshBVHOverlaps, TEXT("Draw BVH of objects overlapping meshes"));
 		FAutoConsoleVariableRef CVarChaosSolverDebugDrawColorShapeByClientServer(TEXT("p.Chaos.Solver.DebugDraw.ColorShapeByClientServer"), ChaosSolverDebugDrawColorShapeByClientServer, TEXT("Color shape according to client and server: red = server / blue = client "));
 		FAutoConsoleVariableRef CVarChaosSolverDebugDrawShowServer(TEXT("p.Chaos.Solver.DebugDraw.ShowServer"), ChaosSolverDebugDrawShowServer, TEXT("Draw server related debug data"));
 		FAutoConsoleVariableRef CVarChaosSolverDebugDrawShowClient(TEXT("p.Chaos.Solver.DebugDraw.ShowClient"), ChaosSolverDebugDrawShowClient, TEXT("Draw client related debug data"));
@@ -230,6 +234,7 @@ namespace Chaos
 		FAutoConsoleVariableRef CVarChaosSolverImpulseScale(TEXT("p.Chaos.Solver.DebugDraw.ImpulseScale"), ChaosSolverDebugDebugDrawSettings.ImpulseScale, TEXT("If >0 show impulses when drawing collisions."));
 		FAutoConsoleVariableRef CVarChaosSolverPushOutScale(TEXT("p.Chaos.Solver.DebugDraw.PushOutScale"), ChaosSolverDebugDebugDrawSettings.PushOutScale, TEXT("If >0 show pushouts when drawing collisions."));
 		FAutoConsoleVariableRef CVarChaosSolverInertiaScale(TEXT("p.Chaos.Solver.DebugDraw.InertiaScale"), ChaosSolverDebugDebugDrawSettings.InertiaScale, TEXT("When DebugDrawTransforms is enabled, show the mass-normalized inertia matrix scaled by this amount."));
+		FAutoConsoleVariableRef CVarChaosSolverDrawPriority(TEXT("p.Chaos.Solver.DebugDraw.DrawPriority"), ChaosSolverDebugDebugDrawSettings.DrawPriority, TEXT("Draw Priority for debug draw shapes (0 draw at actual Z, +ve is closer to the screen)."));
 		FAutoConsoleVariableRef CVarChaosSolverScale(TEXT("p.Chaos.Solver.DebugDraw.Scale"), ChaosSolverDebugDebugDrawSettings.DrawScale, TEXT("Scale applied to all Chaos Debug Draw line lengths etc."));
 		FAutoConsoleVariableRef CVarChaosSolverShowSimple(TEXT("p.Chaos.Solver.DebugDraw.ShowSimple"), ChaosSolverDebugDebugDrawSettings.bShowSimpleCollision, TEXT("Whether to show simple collision is shape drawing is enabled"));
 		FAutoConsoleVariableRef CVarChaosSolverShowComplex(TEXT("p.Chaos.Solver.DebugDraw.ShowComplex"), ChaosSolverDebugDebugDrawSettings.bShowComplexCollision, TEXT("Whether to show complex collision is shape drawing is enabled"));
@@ -1257,6 +1262,8 @@ namespace Chaos
 				{
 					Rigid->ApplyDynamicsWeight(DynamicsWeight);
 				}
+
+				Particle->PrepareBVH();
 				Particle->SyncRemoteData(*Manager, DataIdx, Dirty.PropertyData, Dirty.ShapeDataIndices, ShapeDirtyData);
 				Proxy->ClearAccumulatedData();
 				Proxy->ResetDirtyIdx();
@@ -2171,6 +2178,13 @@ TRACE_COUNTER_SET(ChaosTraceCounter_##Name, Value)
 				DebugDraw::DrawParticleShapes(FRigidTransform3(), Particles.GetNonDisabledDynamicView(), 1.0f, &ChaosSolverDebugDebugDrawSettings);
 			}
 		}
+		if (ChaosSolverDebugDrawBVHs == 1)
+		{
+			for (const FTransientGeometryParticleHandle& Particle : Particles.GetNonDisabledView())
+			{
+				DebugDraw::DrawParticleBVH(FRigidTransform3(), Particle.Handle(), FColor::Silver, &ChaosSolverDebugDebugDrawSettings);
+			}
+		}
 		if (ChaosSolverDebugDrawCollisions == 1) 
 		{
 			DebugDraw::DrawCollisions(FRigidTransform3(), GetEvolution()->GetCollisionConstraints().GetConstraintAllocator(), 1.f, &ChaosSolverDebugDebugDrawSettings);
diff --git a/Engine/Source/Runtime/Experimental/Chaos/Private/PhysicsProxy/GeometryCollectionPhysicsProxy.cpp b/Engine/Source/Runtime/Experimental/Chaos/Private/PhysicsProxy/GeometryCollectionPhysicsProxy.cpp
--- a/Engine/Source/Runtime/Experimental/Chaos/Private/PhysicsProxy/GeometryCollectionPhysicsProxy.cpp
+++ b/Engine/Source/Runtime/Experimental/Chaos/Private/PhysicsProxy/GeometryCollectionPhysicsProxy.cpp
@@ -19,6 +19,7 @@
 #include "Chaos/PerParticleGravity.h"
 #include "Chaos/ImplicitObject.h"
 #include "Chaos/ImplicitObjectScaled.h"
+#include "Chaos/ImplicitObjectUnion.h"
 #include "Chaos/Convex.h"
 #include "Chaos/Serializable.h"
 #include "Chaos/ErrorReporter.h"
@@ -1072,6 +1073,20 @@ void FGeometryCollectionPhysicsProxy::InitializeBodiesPT(Chaos::FPBDRigidsSolver
 			if ((ClusterHandle != nullptr) && !ClusterHandle->Disabled())
 			{
 				RigidsSolver->GetEvolution()->EnableParticle(ClusterHandle);
+
+				// Enable the BVH on the collision (if there are lots of shapes)
+				if (ClusterHandle->Geometry().Get() != nullptr)
+				{
+					if (const Chaos::FImplicitObjectUnion* Union = ClusterHandle->Geometry()->GetObject<Chaos::FImplicitObjectUnion>())
+					{
+						const_cast<Chaos::FImplicitObjectUnion*>(Union)->SetAllowBVH(true);
+					}
+					else if (const Chaos::FImplicitObjectUnion* UnionClustered = ClusterHandle->Geometry()->GetObject<Chaos::FImplicitObjectUnionClustered>())
+					{
+						const_cast<Chaos::FImplicitObjectUnion*>(UnionClustered)->SetAllowBVH(true);
+					}
+				}
+
 			}
 		}
 
diff --git a/Engine/Source/Runtime/Experimental/Chaos/Public/Chaos/ChaosDebugDraw.h b/Engine/Source/Runtime/Experimental/Chaos/Public/Chaos/ChaosDebugDraw.h
--- a/Engine/Source/Runtime/Experimental/Chaos/Public/Chaos/ChaosDebugDraw.h
+++ b/Engine/Source/Runtime/Experimental/Chaos/Public/Chaos/ChaosDebugDraw.h
@@ -86,7 +86,7 @@ namespace Chaos
 				FRealSingle InImpulseScale,
 				FRealSingle InPushOutScale,
 				FRealSingle InInertiaScale,
-				uint8 InDrawPriority,
+				int32 InDrawPriority,
 				bool bInShowSimpleCollision,
 				bool bInShowComplexCollision,
 				bool bInShowLevelSetCollision,
@@ -147,7 +147,7 @@ namespace Chaos
 			FRealSingle PushOutScale;
 			FRealSingle InertiaScale;
 			FRealSingle DrawDuration;
-			uint8 DrawPriority;
+			int32 DrawPriority;
 			bool bShowSimpleCollision;
 			bool bShowComplexCollision;
 			bool bShowLevelSetCollision;
@@ -206,6 +206,7 @@ namespace Chaos
 		CHAOS_API void DrawParticleShapes(const FRigidTransform3& SpaceTransform, const TParticleView<FPBDRigidParticles>& ParticlesView,  FReal ColorScale, const FChaosDebugDrawSettings* Settings = nullptr);
 		CHAOS_API void DrawParticleShapes(const FRigidTransform3& SpaceTransform, const FGeometryParticleHandle* Particle, const FColor& Color, const FChaosDebugDrawSettings* Settings = nullptr);
 		CHAOS_API void DrawParticleShapes(const FRigidTransform3& SpaceTransform, const FGeometryParticle* Particle, const FColor& Color, const FChaosDebugDrawSettings* Settings = nullptr);
+		CHAOS_API void DrawParticleBVH(const FRigidTransform3& SpaceTransform, const FGeometryParticleHandle* Particle, const FColor& Color, const FChaosDebugDrawSettings* Settings = nullptr);
 		CHAOS_API void DrawParticleBounds(const FRigidTransform3& SpaceTransform, const TParticleView<FGeometryParticles>& ParticlesView, const FReal Dt, const FChaosDebugDrawSettings* Settings = nullptr);
 		CHAOS_API void DrawParticleBounds(const FRigidTransform3& SpaceTransform, const TParticleView<FKinematicGeometryParticles>& ParticlesView, const FReal Dt, const FChaosDebugDrawSettings* Settings = nullptr);
 		CHAOS_API void DrawParticleBounds(const FRigidTransform3& SpaceTransform, const TParticleView<FPBDRigidParticles>& ParticlesView, const FReal Dt, const FChaosDebugDrawSettings* Settings = nullptr);
diff --git a/Engine/Source/Runtime/Experimental/Chaos/Public/Chaos/Collision/CollisionConstraintAllocator.h b/Engine/Source/Runtime/Experimental/Chaos/Public/Chaos/Collision/CollisionConstraintAllocator.h
--- a/Engine/Source/Runtime/Experimental/Chaos/Public/Chaos/Collision/CollisionConstraintAllocator.h
+++ b/Engine/Source/Runtime/Experimental/Chaos/Public/Chaos/Collision/CollisionConstraintAllocator.h
@@ -29,10 +29,12 @@ namespace Chaos
 			FCollisionContextAllocator(FPBDCollisionConstraints* InCollisionContainer, const int32 InNumCollisionsPerBlock, const int32 InCurrentEpoch)
 				: CollisionContainer(InCollisionContainer)
 				, CurrentEpoch(InCurrentEpoch)
-	#if CHAOS_COLLISION_OBJECTPOOL_ENABLED
+#if CHAOS_COLLISION_OBJECTPOOL_ENABLED
 				, ConstraintPool(InNumCollisionsPerBlock, 0)
+#endif
+#if CHAOS_MIDPHASE_OBJECTPOOL_ENABLED
 				, MidPhasePool(InNumCollisionsPerBlock, 0)
-	#endif
+#endif
 			{
 			}
 
@@ -67,12 +69,12 @@ namespace Chaos
 			 */
 			FPBDCollisionConstraintPtr CreateConstraint()
 			{
-	#if CHAOS_COLLISION_OBJECTPOOL_ENABLED 
+#if CHAOS_COLLISION_OBJECTPOOL_ENABLED 
 				FPBDCollisionConstraint* Constraint = ConstraintPool.Alloc();
 				return FPBDCollisionConstraintPtr(Constraint, FPBDCollisionConstraintDeleter(ConstraintPool));
-	#else
+#else
 				return MakeUnique<FPBDCollisionConstraint>();
-	#endif
+#endif
 			}
 
 			/**
@@ -131,10 +133,12 @@ namespace Chaos
 
 			void Reset()
 			{
-	#if CHAOS_COLLISION_OBJECTPOOL_ENABLED
+#if CHAOS_COLLISION_OBJECTPOOL_ENABLED
 				ConstraintPool.Reset();
+#endif
+#if CHAOS_MIDPHASE_OBJECTPOOL_ENABLED
 				MidPhasePool.Reset();
-	#endif
+#endif
 			}
 
 			void BeginDetectCollisions(const int32 InEpoch)
@@ -167,11 +171,11 @@ namespace Chaos
 				const FCollisionParticlePairKey Key = FCollisionParticlePairKey(Particle0, Particle1);
 
 				// We temporarily hold new midphases as raw pointers and wrap them in a unique ptr later in ProcessNewMidPhases
-	#if CHAOS_COLLISION_OBJECTPOOL_ENABLED 
+#if CHAOS_MIDPHASE_OBJECTPOOL_ENABLED 
 				FParticlePairMidPhase* MidPhase = MidPhasePool.Alloc();
-	#else
-				FParticlePairMidPhase* MidPhase = new FParticlePairMidPhase();
-	#endif
+#else
+				FParticlePairMidPhase* MidPhase = FParticlePairMidPhase::Make(Particle0, Particle1);
+#endif
 
 				NewMidPhases.Add(MidPhase);
 
@@ -185,10 +189,13 @@ namespace Chaos
 			TArray<FPBDCollisionConstraint*> NewActiveConstraints;
 			TArray<FParticlePairMidPhase*> NewMidPhases;
 
-	#if CHAOS_COLLISION_OBJECTPOOL_ENABLED
+#if CHAOS_COLLISION_OBJECTPOOL_ENABLED
 			FPBDCollisionConstraintPool ConstraintPool;
+#endif
+#if CHAOS_MIDPHASE_OBJECTPOOL_ENABLED
 			FParticlePairMidPhasePool MidPhasePool;
-	#endif
+			FParticlePairMidPhasePool MidPhasePool;
+#endif
 		};
 
 
@@ -619,108 +626,41 @@ namespace Chaos
 	///////////////////////////////////////////////////////////////////////////////////////////////
 	///////////////////////////////////////////////////////////////////////////////////////////////
 
-	template<typename TLambda>
-	inline ECollisionVisitorResult FMultiShapePairCollisionDetector::VisitCollisions(const TLambda& Visitor, const bool bOnlyActive)
-	{
-		for (auto& KVP : Constraints)
-		{
-			const FPBDCollisionConstraintPtr& Constraint = KVP.Value;
-
-			if (Constraint.IsValid() && (!bOnlyActive || !Constraint->GetDisabled()))
-			{
-				if (Visitor(*Constraint) == ECollisionVisitorResult::Stop)
-				{
-					return ECollisionVisitorResult::Stop;
-				}
-			}
-		}
-		return ECollisionVisitorResult::Continue;
-	}
-
-	template<typename TLambda>
-	inline ECollisionVisitorResult FMultiShapePairCollisionDetector::VisitConstCollisions(const TLambda& Visitor, const bool bOnlyActive) const
-	{
-		for (auto& KVP : Constraints)
-		{
-			const FPBDCollisionConstraintPtr& Constraint = KVP.Value;
-
-			if (Constraint.IsValid() && (!bOnlyActive || !Constraint->GetDisabled()))
-			{
-				if (!bOnlyActive || (Visitor(*Constraint) == ECollisionVisitorResult::Stop))
-				{
-					return ECollisionVisitorResult::Stop;
-				}
-			}
-		}
-		return ECollisionVisitorResult::Continue;
-	}
-
-	template<typename TLambda>
-	inline ECollisionVisitorResult FParticlePairCollisionDetector::VisitCollisions(const TLambda& Visitor, const bool bOnlyActive)
-	{
-		for (auto& KVP : Constraints)
-		{
-			const FPBDCollisionConstraintPtr& Constraint = KVP.Value;
-
-			if (Constraint.IsValid() && (!bOnlyActive || !Constraint->GetDisabled()))
-			{
-				if (Visitor(*Constraint) == ECollisionVisitorResult::Stop)
-				{
-					return ECollisionVisitorResult::Stop;
-				}
-			}
-		}
-		return ECollisionVisitorResult::Continue;
-	}
 
 	template<typename TLambda>
-	inline ECollisionVisitorResult FParticlePairCollisionDetector::VisitConstCollisions(const TLambda& Visitor, const bool bOnlyActive) const
+	inline ECollisionVisitorResult FParticlePairMidPhase::VisitCollisions(const TLambda& Visitor, const bool bOnlyActive)
 	{
-		for (auto& KVP : Constraints)
+		if (GetMidPhaseType() == EParticlePairMidPhaseType::ShapePair)
 		{
-			const FPBDCollisionConstraintPtr& Constraint = KVP.Value;
-
-			if (Constraint.IsValid() && (!bOnlyActive || !Constraint->GetDisabled()))
+			FShapePairParticlePairMidPhase* This = static_cast<FShapePairParticlePairMidPhase*>(this);
+			for (FSingleShapePairCollisionDetector& ShapePair : This->ShapePairDetectors)
 			{
-				if (!bOnlyActive || (Visitor(*Constraint) == ECollisionVisitorResult::Stop))
+				if ((ShapePair.GetConstraint() != nullptr) && (!bOnlyActive || !ShapePair.GetConstraint()->GetDisabled()))
 				{
-					return ECollisionVisitorResult::Stop;
+					if (Visitor(*ShapePair.GetConstraint()) == ECollisionVisitorResult::Stop)
+					{
+						return ECollisionVisitorResult::Stop;
+					}
 				}
 			}
 		}
-		return ECollisionVisitorResult::Continue;
-	}
 
-	template<typename TLambda>
-	inline ECollisionVisitorResult FParticlePairMidPhase::VisitCollisions(const TLambda& Visitor, const bool bOnlyActive)
-	{
-		for (FSingleShapePairCollisionDetector& ShapePair : ShapePairDetectors)
+		if (GetMidPhaseType() == EParticlePairMidPhaseType::Generic)
 		{
-			if ((ShapePair.GetConstraint() != nullptr) && (!bOnlyActive || !ShapePair.GetConstraint()->GetDisabled()))
+			FGenericParticlePairMidPhase* This = static_cast<FGenericParticlePairMidPhase*>(this);
+			for (auto& KVP : This->Constraints)
 			{
-				if (Visitor(*ShapePair.GetConstraint()) == ECollisionVisitorResult::Stop)
+				FPBDCollisionConstraintPtr& Constraint = KVP.Value;
+				if (!bOnlyActive || Constraint->IsEnabled())
 				{
-					return ECollisionVisitorResult::Stop;
+					if (Visitor(*Constraint) == ECollisionVisitorResult::Stop)
+					{
+						return ECollisionVisitorResult::Stop;
+					}
 				}
 			}
 		}
 
-		for (FMultiShapePairCollisionDetector& MultiShapePair : MultiShapePairDetectors)
-		{
-			if (MultiShapePair.VisitCollisions(Visitor, bOnlyActive) == ECollisionVisitorResult::Stop)
-			{
-				return ECollisionVisitorResult::Stop;
-			}
-		}
-
-		if (ParticlePairDetector.IsValid())
-		{
-			if (ParticlePairDetector->VisitCollisions(Visitor, bOnlyActive) == ECollisionVisitorResult::Stop)
-			{
-				return ECollisionVisitorResult::Stop;
-			}
-		}
-
 		return ECollisionVisitorResult::Continue;
 	}
 
@@ -728,30 +668,34 @@ namespace Chaos
 	template<typename TLambda>
 	inline ECollisionVisitorResult FParticlePairMidPhase::VisitConstCollisions(const TLambda& Visitor, const bool bOnlyActive) const
 	{
-		for (const FSingleShapePairCollisionDetector& ShapePair : ShapePairDetectors)
+		if (GetMidPhaseType() == EParticlePairMidPhaseType::ShapePair)
 		{
-			if ((ShapePair.GetConstraint() != nullptr) && (!bOnlyActive || !ShapePair.GetConstraint()->GetDisabled()))
+			const FShapePairParticlePairMidPhase* This = static_cast<const FShapePairParticlePairMidPhase*>(this);
+			for (const FSingleShapePairCollisionDetector& ShapePair : This->ShapePairDetectors)
 			{
-				if (Visitor(*ShapePair.GetConstraint()) == ECollisionVisitorResult::Stop)
+				if ((ShapePair.GetConstraint() != nullptr) && (!bOnlyActive || !ShapePair.GetConstraint()->GetDisabled()))
 				{
-					return ECollisionVisitorResult::Stop;
+					if (Visitor(*ShapePair.GetConstraint()) == ECollisionVisitorResult::Stop)
+					{
+						return ECollisionVisitorResult::Stop;
+					}
 				}
 			}
 		}
 
-		for (const FMultiShapePairCollisionDetector& MultiShapePair : MultiShapePairDetectors)
+		if (GetMidPhaseType() == EParticlePairMidPhaseType::Generic)
 		{
-			if (MultiShapePair.VisitConstCollisions(Visitor, bOnlyActive) == ECollisionVisitorResult::Stop)
+			const FGenericParticlePairMidPhase* This = static_cast<const FGenericParticlePairMidPhase*>(this);
+			for (const auto& KVP : This->Constraints)
 			{
-				return ECollisionVisitorResult::Stop;
-			}
-		}
-
-		if (ParticlePairDetector.IsValid())
-		{
-			if (ParticlePairDetector->VisitConstCollisions(Visitor, bOnlyActive) == ECollisionVisitorResult::Stop)
-			{
-				return ECollisionVisitorResult::Stop;
+				const FPBDCollisionConstraintPtr& Constraint = KVP.Value;
+				if (!bOnlyActive || Constraint->IsEnabled())
+				{
+					if (Visitor(*Constraint) == ECollisionVisitorResult::Stop)
+					{
+						return ECollisionVisitorResult::Stop;
+					}
+				}
 			}
 		}
 
diff --git a/Engine/Source/Runtime/Experimental/Chaos/Public/Chaos/Collision/CollisionContext.h b/Engine/Source/Runtime/Experimental/Chaos/Public/Chaos/Collision/CollisionContext.h
--- a/Engine/Source/Runtime/Experimental/Chaos/Public/Chaos/Collision/CollisionContext.h
+++ b/Engine/Source/Runtime/Experimental/Chaos/Public/Chaos/Collision/CollisionContext.h
@@ -9,8 +9,8 @@ namespace Chaos
 	{
 		class FCollisionContextAllocator;
 	}
-	class FMultiShapePairCollisionDetector;
-	class FParticlePairCollisionDetector;
+
+	class FParticlePairMidPhase;
 
 	class FCollisionDetectorSettings
 	{
@@ -64,8 +64,7 @@ namespace Chaos
 		FCollisionContext()
 			: Settings(nullptr)
 			, Allocator(nullptr)
-			, MultiShapeCollisionDetector(nullptr)
-			, ParticlePairCollisionDetector(nullptr)
+			, MidPhase(nullptr)
 		{
 		}
 
@@ -73,8 +72,7 @@ namespace Chaos
 		{
 			Settings = nullptr;
 			Allocator = nullptr;
-			MultiShapeCollisionDetector = nullptr;
-			ParticlePairCollisionDetector = nullptr;
+			MidPhase = nullptr;
 		}
 
 		const FCollisionDetectorSettings& GetSettings() const { return *Settings; }
@@ -90,7 +88,6 @@ namespace Chaos
 		// This is used in the older collision detection path which is still used for particles that do not flatten their implicit hierrarchies
 		// into the Particle's ShapesArray. Currently this is only Clusters.
 		// @todo(chaos): remove thsi from here and make it a parameter on ConstructCollisions and all inner functions.
-		FMultiShapePairCollisionDetector* MultiShapeCollisionDetector;
-		FParticlePairCollisionDetector* ParticlePairCollisionDetector;
+		FParticlePairMidPhase* MidPhase;
 	};
 }
\ No newline at end of file
diff --git a/Engine/Source/Runtime/Experimental/Chaos/Public/Chaos/Collision/PBDCollisionConstraint.h b/Engine/Source/Runtime/Experimental/Chaos/Public/Chaos/Collision/PBDCollisionConstraint.h
--- a/Engine/Source/Runtime/Experimental/Chaos/Public/Chaos/Collision/PBDCollisionConstraint.h
+++ b/Engine/Source/Runtime/Experimental/Chaos/Public/Chaos/Collision/PBDCollisionConstraint.h
@@ -167,10 +167,10 @@ namespace Chaos
 	{
 		friend class Private::FCollisionConstraintAllocator;
 		friend class Private::FCollisionContextAllocator;
-		friend class FMultiShapePairCollisionDetector;
-		friend class FParticlePairCollisionDetector;
+		friend class FGenericParticlePairMidPhase;
 		friend class FParticlePairMidPhase;
 		friend class FPBDCollisionConstraints;
+		friend class FShapePairParticlePairMidPhase;
 		friend class FSingleShapePairCollisionDetector;
 
 		friend CHAOS_API bool ContactConstraintSortPredicate(const FPBDCollisionConstraint& L, const FPBDCollisionConstraint& R);
diff --git a/Engine/Source/Runtime/Experimental/Chaos/Public/Chaos/Collision/ParticlePairMidPhase.h b/Engine/Source/Runtime/Experimental/Chaos/Public/Chaos/Collision/ParticlePairMidPhase.h
--- a/Engine/Source/Runtime/Experimental/Chaos/Public/Chaos/Collision/ParticlePairMidPhase.h
+++ b/Engine/Source/Runtime/Experimental/Chaos/Public/Chaos/Collision/ParticlePairMidPhase.h
@@ -1,19 +1,15 @@
 // Copyright Epic Games, Inc. All Rights Reserved.
 #pragma once
 
-#include "CoreMinimal.h"
 #include "Chaos/Core.h"
 
 #include "Chaos/CollisionResolutionTypes.h"
 #include "Chaos/Collision/CollisionKeys.h"
 #include "Chaos/Collision/CollisionVisitor.h"
 #include "Chaos/Collision/PBDCollisionConstraint.h"
-#include "Chaos/ObjectPool.h"
 #include "Chaos/ParticleHandleFwd.h"
 #include "ProfilingDebugging/CsvProfiler.h"
 
-	// Whether to use a pool for MidPhases
-#define CHAOS_COLLISION_OBJECTPOOL_ENABLED 1
 
 namespace Chaos
 {
@@ -21,31 +17,27 @@
 	{
 		class FCollisionConstraintAllocator;
 	}
-	class FCollisionContext;
-	class FParticlePairMidPhase;
-	class FPBDCollisionConstraints;
-	class FPerShapeData;
-	class FSingleShapePairCollisionDetector;
-
-#if CHAOS_COLLISION_OBJECTPOOL_ENABLED 
-	// Collision and MidPhases are stored in an ObjectPool (if CHAOS_COLLISION_OBJECTPOOL_ENABLED is set). @see FCollisionConstraintAllocator
-	using FPBDCollisionConstraintPool = TObjectPool<FPBDCollisionConstraint>;
-	using FParticlePairMidPhasePool = TObjectPool<FParticlePairMidPhase>;
-	using FPBDCollisionConstraintDeleter = TObjectPoolDeleter<FPBDCollisionConstraintPool>;
-	using FPBDCollisionConstraintPtr = TUniquePtr<FPBDCollisionConstraint, FPBDCollisionConstraintDeleter>;
-	using FParticlePairMidPhaseDeleter = TObjectPoolDeleter<FParticlePairMidPhasePool>;
-	using FParticlePairMidPhasePtr = TUniquePtr<FParticlePairMidPhase, FParticlePairMidPhaseDeleter>;
-#else 
-	using FPBDCollisionConstraintPtr = TUniquePtr<FPBDCollisionConstraint>;
-	using FParticlePairMidPhasePtr = TUniquePtr<FParticlePairMidPhase>;
-#endif
 
 	/**
+	 * The type of the particle pair midphase.
+	 */
+	enum EParticlePairMidPhaseType : int8
+	{
+		// A general purpose midphase that handle BVHs, Meshes, 
+		// Unions of Unions, etc in the geometry hierarchy.
+		Generic,
+
+		// A midphase optimized for particle pairs with a small
+		// number of shapes. Pre-expands the set of potentially
+		// colliding shape pairs.
+		ShapePair,
+	};
+
+	/**
 	 * @brief Handles collision detection for a pair of simple shapes (i.e., not compound shapes)
 	 * 
-	 * @note this is not used for collisions involving Unions that require a recursive collision test.
-	 * @see FMultiShapePairCollisionDetector
-	*/
+	 * This is used by FShapePairParticlePairMidPhase
+	 */
 	class CHAOS_API FSingleShapePairCollisionDetector
 	{
 	public:
@@ -166,336 +158,33 @@
 
 
 	/**
-	 * @brief A collision detector for shape pairs which are containers of other shapes
-	 * This is primarily used by clustered particles that leave their shapes in a Union
-	 * rather than flattening into the particle's ShapesArray.
-	*/
-	class CHAOS_API FMultiShapePairCollisionDetector
-	{
-	public:
-		FMultiShapePairCollisionDetector(
-			FGeometryParticleHandle* InParticle0,
-			const FPerShapeData* InShape0,
-			FGeometryParticleHandle* InParticle1,
-			const FPerShapeData* InShape1,
-			FParticlePairMidPhase& MidPhase);
-		FMultiShapePairCollisionDetector(FMultiShapePairCollisionDetector&& R);
-		FMultiShapePairCollisionDetector(const FMultiShapePairCollisionDetector& R) = delete;
-		FMultiShapePairCollisionDetector& operator=(const FMultiShapePairCollisionDetector& R) = delete;
-		~FMultiShapePairCollisionDetector();
-
-		/**
-		 * @brief Perform a bounds check and run the narrow phase if necessary
-		 * @return The number of collisions constraints that were activated
-		*/
-		int32 GenerateCollisions(
-			const FReal CullDistance,
-			const FReal Dt,
-			const FCollisionContext& Context);
-
-		int32 GenerateCollisionsCCD(
-			const bool bEnableCCDSweep,
-			const FReal CullDistance,
-			const FReal Dt,
-			const FCollisionContext& Context);
-
-		/**
-		 * @brief Callback from the narrow phase to create a collision constraint for this particle pair.
-		 * We should never be asked for a collision for a different particle pair, but the 
-		 * implicit objects may be children of the root shape.
-		*/
-		FPBDCollisionConstraint* FindOrCreateConstraint(
-			FGeometryParticleHandle* InParticle0,
-			const FImplicitObject* InImplicit0,
-			const int32 InImplicitId0,
-			const FShapeInstance* InShape0,
-			const FBVHParticles* InBVHParticles0,
-			const FRigidTransform3& InShapeRelativeTransform0,
-			FGeometryParticleHandle* InParticle1,
-			const FImplicitObject* Implicit1,
-			const int32 InImplicitId1,
-			const FShapeInstance* InShape1,
-			const FBVHParticles* InBVHParticles1,
-			const FRigidTransform3& InhapeRelativeTransform1,
-			const FReal CullDistance,
-			const EContactShapesType ShapePairType,
-			const bool bUseManifold,
-			const bool bEnableSweep,
-			const FCollisionContext& Context);
-
-		UE_DEPRECATED(5.3, "Replaced with version taking ImplicitIDs")
-		FPBDCollisionConstraint* FindOrCreateConstraint(
-			FGeometryParticleHandle* InParticle0,
-			const FImplicitObject* InImplicit0,
-			const FPerShapeData* InShape0,
-			const FBVHParticles* InBVHParticles0,
-			const FRigidTransform3& InShapeRelativeTransform0,
-			FGeometryParticleHandle* InParticle1,
-			const FImplicitObject* InImplicit1,
-			const FPerShapeData* InShape1,
-			const FBVHParticles* InBVHParticles1,
-			const FRigidTransform3& InShapeRelativeTransform1,
-			const FReal CullDistance,
-			const EContactShapesType ShapePairType,
-			const bool bUseManifold,
-			const FCollisionContext& Context)
-		{
-			return FindOrCreateConstraint(
-				InParticle0, InImplicit0, 0, InShape0->AsShapeInstance(), InBVHParticles0, InShapeRelativeTransform0,
-				InParticle1, InImplicit1, 0, InShape1->AsShapeInstance(), InBVHParticles1, InShapeRelativeTransform1,
-				CullDistance, ShapePairType, bUseManifold, false, Context);
-		}
-
-		UE_DEPRECATED(5.3, "Replaced with version taking ImplicitIDs")
-		FPBDCollisionConstraint* FindOrCreateSweptConstraint(
-			FGeometryParticleHandle* InParticle0,
-			const FImplicitObject* InImplicit0,
-			const FPerShapeData* InShape0,
-			const FBVHParticles* InBVHParticles0,
-			const FRigidTransform3& InShapeRelativeTransform0,
-			FGeometryParticleHandle* InParticle1,
-			const FImplicitObject* InImplicit1,
-			const FPerShapeData* InShape1,
-			const FBVHParticles* InBVHParticles1,
-			const FRigidTransform3& InShapeRelativeTransform1,
-			const FReal CullDistance,
-			const EContactShapesType ShapePairType,
-			const FCollisionContext& Context)
-		{
-			return FindOrCreateConstraint(
-				InParticle0, InImplicit0, 0, InShape0->AsShapeInstance(), InBVHParticles0, InShapeRelativeTransform0,
-				InParticle1, InImplicit1, 0, InShape1->AsShapeInstance(), InBVHParticles1, InShapeRelativeTransform1,
-				CullDistance, ShapePairType, true, true, Context);
-		}
-
-		/**
-		 * @brief Reactivate the constraint
-		 * @parame SleepEpoch The tick on which the particle went to sleep.
-		 * Only constraints that were active when the particle went to sleep should be reactivated.
-		*/
-		void WakeCollisions(const int32 SleepEpoch, const int32 CurrentEpoch);
-
-		/**
-		 * @brief Visit all the collisions
-		 * @param Visitor functor with signature ECollisionVisitorResult(FPBDCollisionConstaint& Constraint)
-		*/
-		template<typename TLambda>
-		ECollisionVisitorResult VisitCollisions(const TLambda& Visitor, const bool bOnlyActive = true);
-
-		/**
-		 * @brief Visit all the collisions
-		 * @param Visitor functor with signature ECollisionVisitorResult(const FPBDCollisionConstaint& Constraint)
-		*/
-		template<typename TLambda>
-		ECollisionVisitorResult VisitConstCollisions(const TLambda& Visitor, const bool bOnlyActive = true) const;
-
-	private:
-		int32 GenerateCollisionsImpl(
-			const bool bEnableCCDSweep,
-			const FReal CullDistance,
-			const FReal Dt,
-			const FCollisionContext& Context);
-
-		FPBDCollisionConstraint* FindConstraint(const FCollisionParticlePairConstraintKey& Key);
-
-		FPBDCollisionConstraint* CreateConstraint(
-			FGeometryParticleHandle* Particle0,
-			const FImplicitObject* Implicit0,
-			const FPerShapeData* Shape0,
-			const FBVHParticles* BVHParticles0,
-			const FRigidTransform3& ShapeRelativeTransform0,
-			FGeometryParticleHandle* Particle1,
-			const FImplicitObject* Implicit1,
-			const FPerShapeData* Shape1,
-			const FBVHParticles* BVHParticles1,
-			const FRigidTransform3& ShapeRelativeTransform1,
-			const FReal CullDistance,
-			const EContactShapesType ShapePairType,
-			const bool bInUseManifold,
-			const FCollisionParticlePairConstraintKey& Key,
-			const FCollisionContext& Context);
-
-		int32 ProcessNewConstraints(
-			const FReal CullDistance,
-			const FReal Dt,
-			const FCollisionContext& Context);
-
-		void PruneConstraints(const int32 CurrentEpoch);
-
-		void UpdateCollision(
-			FPBDCollisionConstraint* Constraint,
-			const FReal CullDistance,
-			const FReal Dt,
-			const FCollisionContext& Context);
-
-		void UpdateCollisionCCD(
-			FPBDCollisionConstraint* Constraint,
-			const FReal CullDistance,
-			const FReal Dt,
-			const FCollisionContext& Context);
-
-		FParticlePairMidPhase& MidPhase;
-		TMap<uint32, FPBDCollisionConstraintPtr> Constraints;
-		TArray<FPBDCollisionConstraint*> NewConstraints;
-		FGeometryParticleHandle* Particle0;
-		FGeometryParticleHandle* Particle1;
-		const FPerShapeData* Shape0;
-		const FPerShapeData* Shape1;
-	};
-
-	/**
-	 * 
-	 */
-	class FParticlePairCollisionDetector
-	{
-	public:
-		FParticlePairCollisionDetector(
-			FGeometryParticleHandle* InParticle0,
-			FGeometryParticleHandle* InParticle1,
-			FParticlePairMidPhase& MidPhase);
-		FParticlePairCollisionDetector(FParticlePairCollisionDetector&& R);
-		FParticlePairCollisionDetector(const FParticlePairCollisionDetector& R) = delete;
-		FParticlePairCollisionDetector& operator=(const FParticlePairCollisionDetector& R) = delete;
-		~FParticlePairCollisionDetector();
-
-		/**
-		 * @brief Perform a bounds check and run the narrow phase if necessary
-		 * @return The number of collisions constraints that were activated
-		*/
-		int32 GenerateCollisions(
-			const FReal CullDistance,
-			const FReal Dt,
-			const FCollisionContext& Context);
-
-		int32 GenerateCollisionsCCD(
-			const bool bEnableCCDSweep,
-			const FReal CullDistance,
-			const FReal Dt,
-			const FCollisionContext& Context);
-
-		/**
-		 * @brief Callback from the narrow phase to create a collision constraint for this particle pair.
-		 * We should never be asked for a collision for a different particle pair, but the
-		 * implicit objects may be children of the root shape.
-		*/
-		FPBDCollisionConstraint* FindOrCreateConstraint(
-			FGeometryParticleHandle* InParticle0,
-			const FImplicitObject* InImplicit0,
-			const int32 InImplicitId0,
-			const FShapeInstance* InShape0,
-			const FBVHParticles* InBVHParticles0,
-			const FRigidTransform3& InShapeRelativeTransform0,
-			FGeometryParticleHandle* InParticle1,
-			const FImplicitObject* Implicit1,
-			const int32 InImplicitId1,
-			const FShapeInstance* InShape1,
-			const FBVHParticles* InBVHParticles1,
-			const FRigidTransform3& InhapeRelativeTransform1,
-			const FReal CullDistance,
-			const EContactShapesType ShapePairType,
-			const bool bUseManifold,
-			const bool bEnableSweep,
-			const FCollisionContext& Context);
-
-		/**
-		 * @brief Reactivate the constraint
-		 * @parame SleepEpoch The tick on which the particle went to sleep.
-		 * Only constraints that were active when the particle went to sleep should be reactivated.
-		*/
-		void WakeCollisions(const int32 SleepEpoch, const int32 CurrentEpoch);
-
-		/**
-		 * @brief Visit all the collisions
-		 * @param Visitor functor with signature ECollisionVisitorResult(FPBDCollisionConstaint& Constraint)
-		*/
-		template<typename TLambda>
-		ECollisionVisitorResult VisitCollisions(const TLambda& Visitor, const bool bOnlyActive = true);
-
-		/**
-		 * @brief Visit all the collisions
-		 * @param Visitor functor with signature ECollisionVisitorResult(const FPBDCollisionConstaint& Constraint)
-		*/
-		template<typename TLambda>
-		ECollisionVisitorResult VisitConstCollisions(const TLambda& Visitor, const bool bOnlyActive = true) const;
-
-	private:
-		int32 GenerateCollisionsImpl(
-			const bool bEnableCCDSweep,
-			const FReal CullDistance,
-			const FReal Dt,
-			const FCollisionContext& Context);
-
-		FPBDCollisionConstraint* FindConstraint(const FCollisionParticlePairConstraintKey& Key);
-
-		FPBDCollisionConstraint* CreateConstraint(
-			FGeometryParticleHandle* Particle0,
-			const FImplicitObject* Implicit0,
-			const FPerShapeData* Shape0,
-			const FBVHParticles* BVHParticles0,
-			const FRigidTransform3& ShapeRelativeTransform0,
-			FGeometryParticleHandle* Particle1,
-			const FImplicitObject* Implicit1,
-			const FPerShapeData* Shape1,
-			const FBVHParticles* BVHParticles1,
-			const FRigidTransform3& ShapeRelativeTransform1,
-			const FReal CullDistance,
-			const EContactShapesType ShapePairType,
-			const bool bInUseManifold,
-			const FCollisionParticlePairConstraintKey& Key,
-			const FCollisionContext& Context);
-
-		int32 ProcessNewConstraints(
-			const FReal CullDistance,
-			const FReal Dt,
-			const FCollisionContext& Context);
-
-		void PruneConstraints(const int32 CurrentEpoch);
-
-		void UpdateCollision(
-			FPBDCollisionConstraint* Constraint,
-			const FReal CullDistance,
-			const FReal Dt,
-			const FCollisionContext& Context);
-
-		void UpdateCollisionCCD(
-			FPBDCollisionConstraint* Constraint,
-			const FReal CullDistance,
-			const FReal Dt,
-			const FCollisionContext& Context);
-
-		FParticlePairMidPhase& MidPhase;
-		TMap<uint32, FPBDCollisionConstraintPtr> Constraints;
-		TArray<FPBDCollisionConstraint*> NewConstraints;
-		FGeometryParticleHandle* Particle0;
-		FGeometryParticleHandle* Particle1;
-	};
-
-	/**
 	 * @brief Produce collisions for a particle pair
 	 * A FParticlePairMidPhase object is created for every particle pair whose bounds overlap. It is 
 	 * responsible for building a set of potentially colliding shape pairs and running collision
 	 * detection on those pairs each tick.
 	 * 
-	 * Most particles have a array of shapes, but not all shapes participate in collision detection
-	 * (some are query-only). The cached shape pair list prevents us from repeatesdly testing the
-	 * filters of shape pairs that can never collide.
-	 * 
-	 * @note Geometry collections and clusters do not have arrays of simple shapes. Clustered particles
-	 * typically have a Union as one of the root shapes. In this case we do not attempt to cache the
-	 * potentially colliding shape pair set, and must process the unions every tick.
-	 * 
-	 * @note The lifetime of these objects is handled entirely by the CollisionConstraintAllocator. 
+	 * @note The lifetime of midphase objects is handled entirely by the CollisionConstraintAllocator. 
 	 * Nothing outside of the CollisionConstraintAllocator should hold a pointer to the detector 
 	 * or any constraints it creates for more than the duration of the tick.
-	*/
+	 * 
+	 * @see FShapePairParticlePairMidPhase, FGenericParticlePairMidPhase
+	 * 
+	 */
 	class CHAOS_API FParticlePairMidPhase
 	{
 	public:
-		FParticlePairMidPhase();
+		static EParticlePairMidPhaseType CalculateMidPhaseType(FGeometryParticleHandle* InParticle0, FGeometryParticleHandle* InParticle1);
+		static FParticlePairMidPhase* Make(FGeometryParticleHandle* InParticle0, FGeometryParticleHandle* InParticle1);
 
 		UE_NONCOPYABLE(FParticlePairMidPhase);
 
-		~FParticlePairMidPhase();
+		FParticlePairMidPhase(const EParticlePairMidPhaseType InMidPhaseType);
+		virtual ~FParticlePairMidPhase();
+
+		EParticlePairMidPhaseType GetMidPhaseType() const
+		{
+			return MidPhaseType;
+		}
 
 		/**
 		 * @brief Set up the midphase based on the SHapesArrays of the two particles
@@ -661,59 +350,276 @@
 			return Flags.bIsCCDActive;
 		}
 
-	private:
+		virtual FPBDCollisionConstraint* FindOrCreateConstraint(
+			FGeometryParticleHandle* InParticle0,
+			const FImplicitObject* InImplicit0,
+			const int32 InImplicitId0,
+			const FShapeInstance* InShape0,
+			const FBVHParticles* InBVHParticles0,
+			const FRigidTransform3& InShapeRelativeTransform0,
+			FGeometryParticleHandle* InParticle1,
+			const FImplicitObject* Implicit1,
+			const int32 InImplicitId1,
+			const FShapeInstance* InShape1,
+			const FBVHParticles* InBVHParticles1,
+			const FRigidTransform3& InhapeRelativeTransform1,
+			const FReal CullDistance,
+			const EContactShapesType ShapePairType,
+			const bool bUseManifold,
+			const bool bEnableSweep,
+			const FCollisionContext& Context)
+		{
+			// Currently should only be called on FGenericParticlePairMidPhase
+			check(false);
+			return nullptr;
+		}
+
+	protected:
+
+		virtual void ResetImpl() = 0;
+
 		/**
 		 * @brief Build the list of potentially colliding shape pairs.
 		 * This is all the shape pairs in the partilces' shapes arrays that pass the collision filter.
 		*/
-		void BuildDetectors();
+		virtual void BuildDetectorsImpl() = 0;
+
+		virtual int32 GenerateCollisionsImpl(
+			const FReal CullDistance,
+			const FReal Dt,
+			const FCollisionContext& Context) = 0;
+
+		virtual void WakeCollisionsImpl(const int32 CurrentEpoch) = 0;
+
+		virtual void InjectCollisionImpl(const FPBDCollisionConstraint& Constraint, const FCollisionContext& Context) = 0;
 
 		/**
-		 * @brief Add the shape pair to the list of potentially colliding pairs
-		*/
-		void TryAddShapePair(
-			const FPerShapeData* Shape0, 
-			const FPerShapeData* Shape1);
-
-		/**
 		 * @brief Decide whether we should have CCD enabled on this constraint
 		 * @return true if CCD is enabled this tick, false otherwise
 		 * This may return false, even for collisions on CCD-enabled bodies when the bodies are moving slowly
-		*/
+		 */
 		bool ShouldEnableCCD(const FReal Dt);
 
 		void InitThresholds();
 
-		FGeometryParticleHandle* Particle0; // 8 bytes
-		FGeometryParticleHandle* Particle1; // 8 bytes
-		// A number based on the size of the dynamic objects used to scale cull distance
-		FRealSingle CullDistanceScale; // 4 bytes
 
 		union FFlags
 		{
 			FFlags() : Bits(0) {}
 			struct
 			{
-				uint32 bIsActive : 1;    // True if this midphase should generate a narrow phase at all
-				uint32 bIsCCD : 1;       // True if CCD is supported by either particle
-				uint32 bIsCCDActive : 1; // True if CCD is active for this midphase on this frame. This can be changed by modifiers and resets to bIsCCD each frame.
-				uint32 bIsSleeping : 1;
-				uint32 bIsModified : 1;  // True if a modifier applied any changes to this midphase
+				uint16 bIsActive : 1;    // True if this midphase should generate a narrow phase at all
+				uint16 bIsCCD : 1;       // True if CCD is supported by either particle
+				uint16 bIsCCDActive : 1; // True if CCD is active for this midphase on this frame. This can be changed by modifiers and resets to bIsCCD each frame.
+				uint16 bUseSweep : 1;    // True if CCD is active (this frame) and we are moving fast enough to require a sweep
+				uint16 bIsSleeping : 1;
+				uint16 bIsModified : 1;  // True if a modifier applied any changes to this midphase
 			};
-			uint32 Bits;
-		} Flags; // 4 bytes
+			uint16 Bits;
+		};
 
-		FCollisionParticlePairKey Key; //8 bytes
+		// VTable Ptr									// 8 bytes
+		EParticlePairMidPhaseType MidPhaseType;			// 1 byte
+		FFlags Flags;									// 2 bytes
 
-		int32 LastUsedEpoch;  //4 bytes
-		int32 NumActiveConstraints; // 4 bytes
+		FGeometryParticleHandle* Particle0;				// 8 bytes
+		FGeometryParticleHandle* Particle1;				// 8 bytes
+		// A number based on the size of the dynamic objects used to scale cull distance
+		FRealSingle CullDistanceScale;					// 4 bytes
 
+		FCollisionParticlePairKey Key;					// 8 bytes
+
+		int32 LastUsedEpoch;							// 4 bytes
+		int32 NumActiveConstraints;						// 4 bytes
+
 		// Indices into the arrays of collisions on the particles. This is a cookie for use by FParticleCollisions
-		int32 ParticleCollisionsIndex0; // 4 bytes
-		int32 ParticleCollisionsIndex1; // 4 bytes
+		int32 ParticleCollisionsIndex0;					// 4 bytes
+		int32 ParticleCollisionsIndex1;					// 4 bytes
+	};
 
-		TArray<FMultiShapePairCollisionDetector> MultiShapePairDetectors; // 16 bytes
-		TArray<FSingleShapePairCollisionDetector, TInlineAllocator<1>> ShapePairDetectors; //88 bytes
-		TUniquePtr<FParticlePairCollisionDetector> ParticlePairDetector;
+
+	/**
+	 * A midphase for a particle pair that pre-builds a set of all potentially colliding shape
+	 * pairs. This is the fast path used when each particle has a small number of shapes and
+	 * does not contain a complex hierarchy. This midphase (compared to FGenericParticlePairMidPhase)
+	 * caches various data like the results of the collision filtering, the shapr pair types, etc.
+	 */
+	class CHAOS_API FShapePairParticlePairMidPhase : public FParticlePairMidPhase
+	{
+	public:
+		friend class FParticlePairMidPhase;
+
+		FShapePairParticlePairMidPhase();
+
+		virtual void ResetImpl() override final;
+		virtual void BuildDetectorsImpl() override final;
+
+	protected:
+		virtual int32 GenerateCollisionsImpl(
+			const FReal CullDistance,
+			const FReal Dt,
+			const FCollisionContext& Context) override final;
+
+		virtual void WakeCollisionsImpl(const int32 CurrentEpoch) override final;
+
+		virtual void InjectCollisionImpl(const FPBDCollisionConstraint& Constraint, const FCollisionContext& Context) override final;
+
+	private:
+		void TryAddShapePair(
+			const FPerShapeData* Shape0,
+			const FPerShapeData* Shape1);
+
+		TArray<FSingleShapePairCollisionDetector, TInlineAllocator<1>> ShapePairDetectors;	// 88 bytes
 	};
+
+	// A set of cached data about a simple implicit object on a particle
+	struct FLeafImplicitObject;
+
+	/**
+	 * A midphase for a particle pair where one or both have a large number of collisions
+	 * shapes, or a non-flat hierarchy of shapes. This is the general-purpose collision
+	 * path and does not cache as much data as the FShapePairParticlePairMidPhase and
+	 * must visit the geometry hierarchy on both shapes every time we detect collisions.
+	 * It must also rerun the collision filters on overlapping pairs, among other things.
+	 * It does, howver, take advantage of the BVH held in the root ImplicitObjectUnion
+	 * if there is one, so it can be much faster the the FShapePairParticlePairMidPhase
+	 * when the set of potentially colliding pairs is very large.
+	 */
+	class CHAOS_API FGenericParticlePairMidPhase : public FParticlePairMidPhase
+	{
+	public:
+		friend class FParticlePairMidPhase;
+
+		FGenericParticlePairMidPhase();
+		~FGenericParticlePairMidPhase();
+
+		virtual void ResetImpl() override final;
+		virtual void BuildDetectorsImpl() override final;
+
+		/**
+		 * @brief Callback from the narrow phase to create a collision constraint for this particle pair.
+		 * We should never be asked for a collision for a different particle pair, but the
+		 * implicit objects may be children of the root shape.
+		*/
+		virtual FPBDCollisionConstraint* FindOrCreateConstraint(
+			FGeometryParticleHandle* InParticle0,
+			const FImplicitObject* InImplicit0,
+			const int32 InImplicitId0,
+			const FShapeInstance* InShape0,
+			const FBVHParticles* InBVHParticles0,
+			const FRigidTransform3& InShapeRelativeTransform0,
+			FGeometryParticleHandle* InParticle1,
+			const FImplicitObject* Implicit1,
+			const int32 InImplicitId1,
+			const FShapeInstance* InShape1,
+			const FBVHParticles* InBVHParticles1,
+			const FRigidTransform3& InhapeRelativeTransform1,
+			const FReal CullDistance,
+			const EContactShapesType ShapePairType,
+			const bool bUseManifold,
+			const bool bEnableSweep,
+			const FCollisionContext& Context) override final;
+
+		/**
+		 * @brief Reactivate the constraint
+		 * @parame SleepEpoch The tick on which the particle went to sleep.
+		 * Only constraints that were active when the particle went to sleep should be reactivated.
+		*/
+		void WakeCollisionsImpl(const int32 SleepEpoch, const int32 CurrentEpoch);
+
+	protected:
+		virtual int32 GenerateCollisionsImpl(
+			const FReal CullDistance,
+			const FReal Dt,
+			const FCollisionContext& Context) override final;
+
+		virtual void WakeCollisionsImpl(const int32 CurrentEpoch) override final;
+
+		virtual void InjectCollisionImpl(const FPBDCollisionConstraint& Constraint, const FCollisionContext& Context) override final;
+
+	private:
+		// BVH on ParticleA versus implicit hierarchy of ParticleB
+		void GenerateCollisionsBVH(
+			FGeometryParticleHandle* ParticleA, const Private::FImplicitBVH* BVHA,
+			FGeometryParticleHandle* ParticleB, const FImplicitObject* RootImplicitB,
+			const FReal CullDistance,
+			const FReal Dt,
+			const FCollisionContext& Context);
+
+		// Implicit hierarchy of particle A versus implicit hierarchy of ParticleB (used when no BVHs present)
+		void GenerateCollisionsImplicit(
+			FGeometryParticleHandle* ParticleA, const FImplicitObject* RootImplicitA,
+			FGeometryParticleHandle* ParticleB, const FImplicitObject* RootImplicitB,
+			const FReal CullDistance,
+			const FReal Dt,
+			const FCollisionContext& Context);
+
+		// Leaf Implicit on ParticleA versus implicit hierarchy of ParticleB
+		void GenerateCollisionsShapeHierarchy(
+			FGeometryParticleHandle* ParticleA, const FImplicitObject* ImplicitA, const FRigidTransform3 ParticleWorldTransformA, const FRigidTransform3& RelativeTransformA, const FAABB3& RelativeBoundsA, const int32 RootObjectIndexA, const int32 LeafObjectIndexA,
+			FGeometryParticleHandle* ParticleB, const FImplicitObject* RootImplicitB, const FRigidTransform3 ParticleWorldTransformB,
+			const FRigidTransform3 ParticleTransformAToB,
+			const FReal CullDistance,
+			const FReal Dt,
+			const FCollisionContext& Context);
+
+		// Leaf Implicit on ParticleA versus Leaf Implicit of ParticleB
+		void GenerateCollisionsShapeShape(
+			FGeometryParticleHandle* ParticleA, const FImplicitObject* ImplicitA, const FShapeInstance* ShapeInstanceA, const FRigidTransform3 ParticleWorldTransformA, const FRigidTransform3& RelativeTransformA, const int32 LeafObjectIndexA,
+			FGeometryParticleHandle* ParticleB, const FImplicitObject* ImplicitB, const FShapeInstance* ShapeInstanceB, const FRigidTransform3 ParticleWorldTransformB, const FRigidTransform3& RelativeTransformB, const int32 LeafObjectIndexB,
+			const FReal CullDistance,
+			const FReal Dt,
+			const FCollisionContext& Context);
+
+		// A bounds check between two implicits
+		bool DoBoundsOverlap(
+			const FImplicitObject* ImplicitA,
+			const FRigidTransform3& ParticleWorldTransformA,
+			const FRigidTransform3& ShapeRelativeTransformA,
+			const FImplicitObject* ImplicitB,
+			const FRigidTransform3& ParticleWorldTransformB,
+			const FRigidTransform3& ShapeRelativeTransformB,
+			const FReal CullDistance);
+
+		FPBDCollisionConstraint* FindConstraint(const FCollisionParticlePairConstraintKey& Key);
+
+		FPBDCollisionConstraint* CreateConstraint(
+			FGeometryParticleHandle* Particle0,
+			const FImplicitObject* Implicit0,
+			const FPerShapeData* Shape0,
+			const FBVHParticles* BVHParticles0,
+			const FRigidTransform3& ShapeRelativeTransform0,
+			FGeometryParticleHandle* Particle1,
+			const FImplicitObject* Implicit1,
+			const FPerShapeData* Shape1,
+			const FBVHParticles* BVHParticles1,
+			const FRigidTransform3& ShapeRelativeTransform1,
+			const FReal CullDistance,
+			const EContactShapesType ShapePairType,
+			const bool bInUseManifold,
+			const FCollisionParticlePairConstraintKey& Key,
+			const FCollisionContext& Context);
+
+		int32 ProcessNewConstraints(
+			const FReal CullDistance,
+			const FReal Dt,
+			const FCollisionContext& Context);
+
+		void PruneConstraints(const int32 CurrentEpoch);
+
+		void UpdateCollision(
+			FPBDCollisionConstraint* Constraint,
+			const FReal CullDistance,
+			const FReal Dt,
+			const FCollisionContext& Context);
+
+		void UpdateCollisionCCD(
+			FPBDCollisionConstraint* Constraint,
+			const FReal CullDistance,
+			const FReal Dt,
+			const FCollisionContext& Context);
+
+		TMap<uint32, FPBDCollisionConstraintPtr> Constraints;
+		TArray<FPBDCollisionConstraint*> NewConstraints;
+	};
 }
diff --git a/Engine/Source/Runtime/Experimental/Chaos/Public/Chaos/CollisionResolution.h b/Engine/Source/Runtime/Experimental/Chaos/Public/Chaos/CollisionResolution.h
--- a/Engine/Source/Runtime/Experimental/Chaos/Public/Chaos/CollisionResolution.h
+++ b/Engine/Source/Runtime/Experimental/Chaos/Public/Chaos/CollisionResolution.h
@@ -63,9 +63,19 @@ namespace Chaos
 		*/
 		bool CHAOS_API ShouldUseCCD(const FGeometryParticleHandle* Particle0, const FVec3& DeltaX0, const FGeometryParticleHandle* Particle1, const FVec3& DeltaX1, FVec3& Dir, FReal& Length);
 
-		// Create constraints for the particle pair. This could create multiple constraints: one for each potentially colliding shape pair in multi-shape particles.
-		// @param ImplicitID0 Used to differentiate contacts when an implicit is used multiple times in a geometry hierarchy
-		// @param ImplicitID1 Used to differentiate contacts when an implicit is used multiple times in a geometry hierarchy
+		// @todo(chaos): this should not be exposed but is currently used in tests
+		void UpdateLevelsetLevelsetConstraint(const FRigidTransform3& WorldTransform0, const FRigidTransform3& WorldTransform1, const FReal Dt, FPBDCollisionConstraint& Constraint);
+
+		// Reset per-frame collision stat counters
+		void CHAOS_API ResetChaosCollisionCounters();
+
+		//
+		//
+		// DEPRECATED STUFF
+		//
+		//
+
+		UE_DEPRECATED(5.3, "No longer needed or supported")
 		void CHAOS_API ConstructConstraints(
 			TGeometryParticleHandle<FReal, 3>* Particle0, 
 			TGeometryParticleHandle<FReal, 3>* Particle1, 
@@ -86,43 +96,15 @@ namespace Chaos
 			const bool bEnableSweep, 
 			const FCollisionContext& Context);
 
-		// @todo(chaos): this should not be exposed but is currently used in tests
-		void UpdateLevelsetLevelsetConstraint(const FRigidTransform3& WorldTransform0, const FRigidTransform3& WorldTransform1, const FReal Dt, FPBDCollisionConstraint& Constraint);
+		UE_DEPRECATED(5.3, "No longer needed or supported")
+		void CHAOS_API ConstructConstraints(TGeometryParticleHandle<FReal, 3>* Particle0, TGeometryParticleHandle<FReal, 3>* Particle1, const FImplicitObject* Implicit0, const FPerShapeData* Shape0, const FBVHParticles* Simplicial0, const FImplicitObject* Implicit1, const FPerShapeData* Shape1, const FBVHParticles* Simplicial1, const FRigidTransform3& ParticleWorldTransform0, const FRigidTransform3& Transform0, const FRigidTransform3& ParticleWorldTransform1, const FRigidTransform3& Transform1, const FReal CullDistance, const FReal Dt, const bool bEnableSweep, const FCollisionContext& Context);
 
-		// Reset per-frame collision stat counters
-		void CHAOS_API ResetChaosCollisionCounters();
-
-		//
-		//
-		// DEPRECATED STUFF (see above)
-		//
-		//
-
-		UE_DEPRECATED(5.3, "Replaced with a version that takes a ImplicitIDs")
-		inline void CHAOS_API ConstructConstraints(TGeometryParticleHandle<FReal, 3>* Particle0, TGeometryParticleHandle<FReal, 3>* Particle1, const FImplicitObject* Implicit0, const FPerShapeData* Shape0, const FBVHParticles* Simplicial0, const FImplicitObject* Implicit1, const FPerShapeData* Shape1, const FBVHParticles* Simplicial1, const FRigidTransform3& ParticleWorldTransform0, const FRigidTransform3& Transform0, const FRigidTransform3& ParticleWorldTransform1, const FRigidTransform3& Transform1, const FReal CullDistance, const FReal Dt, const bool bEnableSweep, const FCollisionContext& Context)
-		{
-			ConstructConstraints(Particle0, Particle1, Implicit0, Shape0, Simplicial0, 0, Implicit1, Shape1, Simplicial1, 0, ParticleWorldTransform0, Transform0, ParticleWorldTransform1, Transform1, CullDistance, Dt, bEnableSweep, Context);
-		}
-
-		UE_DEPRECATED(5.3, "Replaced with a version that takes a bEnableSweep flag")
-		inline void CHAOS_API ConstructConstraints(TGeometryParticleHandle<FReal, 3>* Particle0, TGeometryParticleHandle<FReal, 3>* Particle1, const FImplicitObject* Implicit0, const FPerShapeData* Shape0, const FBVHParticles* Simplicial0, const FImplicitObject* Implicit1, const FPerShapeData* Shape1, const FBVHParticles* Simplicial1, const FRigidTransform3& ParticleWorldTransform0, const FRigidTransform3& Transform0, const FRigidTransform3& ParticleWorldTransform1, const FRigidTransform3& Transform1, const FReal CullDistance, const FReal Dt, const FCollisionContext& Context)
-		{
-			ConstructConstraints(Particle0, Particle1, Implicit0, Shape0, Simplicial0, 0, Implicit1, Shape1, Simplicial1, 0, ParticleWorldTransform0, Transform0, ParticleWorldTransform1, Transform1, CullDistance, Dt, false, Context);
-		}
+		// See Above
+		UE_DEPRECATED(5.3, "No longer needed or supported")
+		void CHAOS_API ConstructConstraints(TGeometryParticleHandle<FReal, 3>* Particle0, TGeometryParticleHandle<FReal, 3>* Particle1, const FImplicitObject* Implicit0, const FPerShapeData* Shape0, const FBVHParticles* Simplicial0, const FImplicitObject* Implicit1, const FPerShapeData* Shape1, const FBVHParticles* Simplicial1, const FRigidTransform3& ParticleWorldTransform0, const FRigidTransform3& Transform0, const FRigidTransform3& ParticleWorldTransform1, const FRigidTransform3& Transform1, const FReal CullDistance, const FReal Dt, const FCollisionContext& Context);
 
 		template<ECollisionUpdateType UpdateType>
 		UE_DEPRECATED(5.3, "Use UpdateConstraint, but call Constraint.SetShapeWorldTransform first (see implementation)")
 		void CHAOS_API UpdateConstraintFromGeometry(FPBDCollisionConstraint& Constraint, const FRigidTransform3& ParticleTransform0, const FRigidTransform3& ParticleTransform1, const FReal Dt);
-
-		UE_DEPRECATED(5.3, "Use the version which takes a particle and an implicit")
-		inline EContactShapesType CHAOS_API CalculateShapePairType(const FImplicitObject* Implicit0, const FBVHParticles* BVHParticles0, const FImplicitObject* Implicit1, const FBVHParticles* BVHParticles1, bool& bOutSwap)
-		{
-			return EContactShapesType::Unknown;
-		}
-		UE_DEPRECATED(5.3, "Use the version which takes a particle and an implicit")
-		inline EContactShapesType CHAOS_API CalculateShapePairType(const EImplicitObjectType Implicit0Type, const EImplicitObjectType Implicit1Type, const bool bIsConvex0, const bool bIsConvex1, const bool bIsBVH0, const bool bIsBVH1, bool& bOutSwap)
-		{
-			return EContactShapesType::Unknown;
-		}
 	}
 }
diff --git a/Engine/Source/Runtime/Experimental/Chaos/Public/Chaos/CollisionResolutionTypes.h b/Engine/Source/Runtime/Experimental/Chaos/Public/Chaos/CollisionResolutionTypes.h
--- a/Engine/Source/Runtime/Experimental/Chaos/Public/Chaos/CollisionResolutionTypes.h
+++ b/Engine/Source/Runtime/Experimental/Chaos/Public/Chaos/CollisionResolutionTypes.h
@@ -2,8 +2,13 @@
 #pragma once
 
 #include "Chaos/Core.h"
+#include "Chaos/ObjectPool.h"
 #include "Chaos/ParticleHandleFwd.h"
 
+// Whether to use a pool for MidPhases and Collision Constraints
+#define CHAOS_COLLISION_OBJECTPOOL_ENABLED 1
+#define CHAOS_MIDPHASE_OBJECTPOOL_ENABLED 0
+
 namespace Chaos
 {
 	// Float and Double versions of ContactPoint
@@ -70,7 +75,34 @@ namespace Chaos
 	//
 	//
 
-	//TODO: move into a better forward declare location
+	class FCollisionContext;
+	class FGenericParticlePairMidPhase;
+	class FParticlePairMidPhase;
+	class FPBDCollisionConstraints;
+	class FPBDCollisionConstraint;
 	class FPBDCollisionConstraintHandle;
+	class FPerShapeData;
+	class FShapePairParticlePairMidPhase;
+
+// Collision and MidPhases are stored in an ObjectPool (if CHAOS_COLLISION_OBJECTPOOL_ENABLED or CHAOS_MIDPHASE_OBJECTPOOL_ENABLED is set). 
+// @see FCollisionConstraintAllocator
+#if CHAOS_COLLISION_OBJECTPOOL_ENABLED 
+	using FPBDCollisionConstraintPool = TObjectPool<FPBDCollisionConstraint>;
+	using FPBDCollisionConstraintDeleter = TObjectPoolDeleter<FPBDCollisionConstraintPool>;
+	using FPBDCollisionConstraintPtr = TUniquePtr<FPBDCollisionConstraint, FPBDCollisionConstraintDeleter>;
+#else 
+	using FPBDCollisionConstraintPtr = TUniquePtr<FPBDCollisionConstraint>;
+#endif
 
+#if CHAOS_MIDPHASE_OBJECTPOOL_ENABLED
+	// Not yet supported due to awkwardness of defining a deleter for FParticlePairMidPhasePtr when we have pools of derived types.
+	// We probably need a custom deleter (not TObjectPoolDeleter) and to store a pointer to the CollisionConstraintAllocator.
+	static_assert(false, "CHAOS_MIDPHASE_OBJECTPOOL_ENABLED not supported yet")
+	using FGenericParticlePairMidPhasePool = TObjectPool<FGenericParticlePairMidPhase>;
+	using FShapePairParticlePairMidPhasePool = TObjectPool<FShapePairParticlePairMidPhase>;
+	using FParticlePairMidPhaseDeleter = TObjectPoolDeleter<???>;
+	using FParticlePairMidPhasePtr = TUniquePtr<FParticlePairMidPhase, FParticlePairMidPhaseDeleter>;
+#else
+	using FParticlePairMidPhasePtr = TUniquePtr<FParticlePairMidPhase>;
+#endif
 }
diff --git a/Engine/Source/Runtime/Experimental/Chaos/Public/Chaos/HeightField.h b/Engine/Source/Runtime/Experimental/Chaos/Public/Chaos/HeightField.h
--- a/Engine/Source/Runtime/Experimental/Chaos/Public/Chaos/HeightField.h
+++ b/Engine/Source/Runtime/Experimental/Chaos/Public/Chaos/HeightField.h
@@ -175,6 +175,58 @@ namespace Chaos
 			}
 		}
 
+		virtual void VisitOverlappingLeafObjectsImpl(
+			const FAABB3& QueryBounds,
+			const FRigidTransform3& ObjectTransform,
+			const int32 RootObjectIndex,
+			int32& ObjectIndex,
+			int32& LeafObjectIndex,
+			const FImplicitHierarchyVisitor& VisitorFunc) const override final
+		{
+			if (IsOverlappingAnyCells(QueryBounds))
+			{
+				VisitorFunc(this, ObjectTransform, RootObjectIndex, ObjectIndex, LeafObjectIndex);
+			}
+			++ObjectIndex;
+			++LeafObjectIndex;
+		}
+
+		virtual bool IsOverlappingBoundsImpl(const FAABB3& QueryBounds) const override final
+		{
+			return IsOverlappingAnyCells(QueryBounds);
+		}
+
+		// Does the mesh-space QueryBounds overlap the bounds of any cells in the heightfield?
+		// @todo(chaos): we can return as soon as we overlap any cell so maybe use a 
+		// custom function rather than GetBoundsScaled. Also we only care about Z overlap.
+		bool IsOverlappingAnyCells(const FAABB3& QueryBounds) const
+		{
+			// Top-level bounds check
+			if (!BoundingBox().Intersects(QueryBounds))
+			{
+				return false;
+			}
+
+			// Find all the cells that overlap in the X/Y plane
+			FVec2 Scale2D(GeomData.Scale[0], GeomData.Scale[1]);
+			const FBounds2D FlatBounds = GetFlatBounds();
+
+			FBounds2D GridQueryBounds;
+			GridQueryBounds.Min = FVec2(QueryBounds.Min()[0], QueryBounds.Min()[1]);
+			GridQueryBounds.Max = FVec2(QueryBounds.Max()[0], QueryBounds.Max()[1]);
+			GridQueryBounds = FBounds2D::FromPoints(FlatBounds.Clamp(GridQueryBounds.Min) / Scale2D, FlatBounds.Clamp(GridQueryBounds.Max) / Scale2D);
+
+			FAABBVectorized QueryBoundsSimd = FAABBVectorized(QueryBounds);
+
+			// We want to capture the first cell (delta == 0) as well
+			TVec2<int32> FirstCell = FlatGrid.Cell(GridQueryBounds.Min);
+			TVec2<int32> LastCell = FlatGrid.Cell(GridQueryBounds.Max);
+
+			FAABBVectorized RegionBounds;
+			GeomData.GetBoundsScaled(FirstCell, LastCell - FirstCell, RegionBounds);
+			return RegionBounds.Intersects(QueryBoundsSimd);
+		}
+
 		struct FClosestFaceData
 		{
 			int32 FaceIndex = INDEX_NONE;
@@ -571,7 +623,7 @@ namespace Chaos
 
 			FORCEINLINE void GetBoundsScaled(TVec2<int32> CellIdx, FAABBVectorized& OutBounds) const
 			{
-				GetBounds(CellIdx, OutBounds);
+				GetBounds(CellIdx, TVec2<int32>(1, 1), OutBounds);
 
 				VectorRegister4Float P0 = VectorMultiply(OutBounds.GetMin(), ScaleSimd);
 				VectorRegister4Float P1 = VectorMultiply(OutBounds.GetMax(), ScaleSimd);
diff --git a/Engine/Source/Runtime/Experimental/Chaos/Public/Chaos/ImplicitObject.h b/Engine/Source/Runtime/Experimental/Chaos/Public/Chaos/ImplicitObject.h
--- a/Engine/Source/Runtime/Experimental/Chaos/Public/Chaos/ImplicitObject.h
+++ b/Engine/Source/Runtime/Experimental/Chaos/Public/Chaos/ImplicitObject.h
@@ -68,6 +68,7 @@ struct TImplicitObjectPtrStorage<T, d, true>
 * @param LeafObjectIndex A counter tracking the current leaf index in the flattened hierarchy (pre-order, depth first). Used to differentiate between geometry when duplicated in the hierarchy. INDEX_NONE if not visiting a leaf.
 */
 using FImplicitHierarchyVisitor = TFunctionRef<void(const FImplicitObject* Implicit, const FRigidTransform3& Transform, const int32 RootObjectIndex, const int32 ObjectIndex, const int32 LeafObjectIndex)>;
+using FImplicitHierarchyVisitorBool = TFunctionRef<bool(const FImplicitObject* Implicit, const FRigidTransform3& Transform, const int32 RootObjectIndex, const int32 ObjectIndex, const int32 LeafObjectIndex)>;
 
 
 /*
@@ -370,7 +371,7 @@ public:
 	* Visit all the objects in the hierarchy, including inner nodes like Union and Transform
 	* @see FImplicitHierarchyVisitor for visitor notes
 	*/
-	void VisitObjects(const FImplicitHierarchyVisitor& Visitor) const
+	void VisitObjects(const FImplicitHierarchyVisitorBool& Visitor) const
 	{
 		int32 LeafObjectIndex = 0;
 		int32 ObjectIndex = 0;
@@ -378,6 +379,17 @@ public:
 		VisitObjectsImpl(FRigidTransform3::Identity, RootObjectIndex, ObjectIndex, LeafObjectIndex, Visitor);
 	}
 
+	/**
+	* Whether this implicit (possibly a hierarchy) overlaps the bounds. This is a deeper, more
+	* accurate test than simply checking BoundingBox() because it will query the BVH in Unions,
+	* Heightfields and Triangle Meshes, and only return true if we overlap the bounds of a leaf 
+	* node that contains some elements.
+	*/
+	bool IsOverlappingBounds(const FAABB3& LocalBounds) const
+	{
+		return IsOverlappingBoundsImpl(LocalBounds);
+	}
+
 //protected:
 	// This should not be public, but it needs to be callable by derived classes on another instance
 	virtual void VisitOverlappingLeafObjectsImpl(
@@ -410,18 +422,24 @@ public:
 	}
 
 	// This should not be public, but it needs to be callable by derived classes on another instance
-	virtual void VisitObjectsImpl(
+	virtual bool VisitObjectsImpl(
 		const FRigidTransform3& ObjectTransform,
 		const int32 RootObjectIndex,
 		int32& ObjectIndex,
 		int32& LeafObjectIndex,
-		const FImplicitHierarchyVisitor& VisitorFunc) const
+		const FImplicitHierarchyVisitorBool& VisitorFunc) const
 	{
-		VisitorFunc(this, ObjectTransform, RootObjectIndex, ObjectIndex, LeafObjectIndex);
+		return VisitorFunc(this, ObjectTransform, RootObjectIndex, ObjectIndex, LeafObjectIndex);
 		++ObjectIndex;
 		++LeafObjectIndex;
 	}
 
+	// This should not be public, but it needs to be callable by derived classes on another instance
+	virtual bool IsOverlappingBoundsImpl(const FAABB3& LocalBounds) const
+	{
+		return (!HasBoundingBox() || LocalBounds.Intersects(BoundingBox()));
+	}
+
 protected:
 
 	// Safely scale a normalized Vec3 - used with both scale and inverse scale
diff --git a/Engine/Source/Runtime/Experimental/Chaos/Public/Chaos/ImplicitObjectBVH.h b/Engine/Source/Runtime/Experimental/Chaos/Public/Chaos/ImplicitObjectBVH.h
--- a/Engine/Source/Runtime/Experimental/Chaos/Public/Chaos/ImplicitObjectBVH.h
+++ b/Engine/Source/Runtime/Experimental/Chaos/Public/Chaos/ImplicitObjectBVH.h
@@ -13,9 +13,10 @@ namespace Chaos
 	namespace Private
 	{
 		class FImplicitBVH;
-		class FImplicitObjectUnion;
+		class FImplicitBVTree;
+		class FImplicitBVTreeNode;
 
-		// An entry in the ImplicitObject BVH that holds the leaf geometry and transform
+		// A item in a ImplicitBVH holding the leaf geometry and transform. Each FImplicitBVHNode node holds a set of these.
 		class FImplicitBVHObject
 		{
 		public:
@@ -34,11 +35,13 @@ namespace Chaos
 
 			const FAABB3f& GetBounds() const { return Bounds; }
 
+			const FVec3f GetBoundsCenter() const { return Bounds.Center(); }
+
 			FRigidTransform3f GetTransformf() const { return FRigidTransform3f(GetX(), GetR()); }
 
 			FRigidTransform3 GetTransform() const { return FRigidTransform3(FVec3(GetX()), FRotation3(GetR())); }
 
-			// A unique index for this object in the hierarchy. E.g., if the same FImplicirtObject is referenced 
+			// A unique index for this object in the hierarchy. E.g., if the same FImplicitObject is referenced 
 			// multiple times in the hierarchy in a union of transformed objects, each will have a different ObjectIndex 
 			// (it is the index into the pre-order depth first traversal).
 			int32 GetObjectIndex() const { return ObjectIndex; }
@@ -65,17 +68,43 @@ namespace Chaos
 			int32 ObjectIndex;
 		};
 
+		// A node in an FImplicitBVH
+		class FImplicitBVHNode
+		{
+		public:
+			FImplicitBVHNode(const FAABB3f& InBounds, const int32 InObjectBeginIndex, const int32 InObjectEndIndex)
+				: Bounds(InBounds)
+				, ObjectBeginIndex(InObjectBeginIndex)
+				, ObjectEndIndex(InObjectEndIndex)
+				, ChildNodeIndices{ INDEX_NONE, INDEX_NONE }
+			{
+			}
+
+			bool IsLeaf() const
+			{
+				// NOTE: We either have zero or two children
+				return ChildNodeIndices[0] == INDEX_NONE;
+			}
+
+		private:
+			friend class FImplicitBVH;
+
+			FAABB3f Bounds;
+			int32 ObjectBeginIndex;
+			int32 ObjectEndIndex;
+			int32 ChildNodeIndices[2];
+		};
+
+
 		// A Bounding Volume Hierarchy of a set of Implicit Objects
 		class FImplicitBVH
 		{
 		public:
 			using FObjects = TArray<FImplicitBVHObject>;
-			using FLeaf = TArray<int32>;
-			using FBVH = TBoundingVolumeHierarchy<FObjects, FLeaf>;
 
 			friend FChaosArchive& operator<<(FChaosArchive& Ar, FImplicitBVH& BVH);
 
-			// Amke an empty BVH (for serialization only)
+			// Make an empty BVH (for serialization only)
 			static TUniquePtr<FImplicitBVH> MakeEmpty();
 
 			// Utility for processing the hierarchy
@@ -89,7 +118,8 @@ namespace Chaos
 
 			~FImplicitBVH();
 
-			int32 NumObjects() const { return Objects.Num(); }
+			int32 GetNumObjects() const { return Objects.Num(); }
+			int32 GetDepth() const { return TreeDepth; }
 
 			const FImplicitBVHObject& GetObject(const int32 ObjectIndex) const { return Objects[ObjectIndex]; }
 
@@ -108,30 +138,66 @@ namespace Chaos
 			int32 GetRootObjectIndex(const int32 ObjectIndex) const { return Objects[ObjectIndex].GetRootObjectIndex(); }
 			int32 GetObjectIndex(const int32 ObjectIndex) const { return Objects[ObjectIndex].GetObjectIndex(); }
 
-			const FBVH& GetBVH() const { return BVH; }
-
-			// @param Visitor void<int ObjectIndex> where ObjectIndex is an index into the array of objects in the hierarchy
-			// for use with GetObject(ObjectIndex), GetGeometry(ObjectIndex), etc.
+			// Visit all the leaf objects that overlap the specified bounds.
+			// @param ObjectVisitor [](const FImplicitObject* Implicit, const FRigidTransform3f& RelativeTransformf, const FAABB3f& RelativeBoundsf, const int32 RootObjectIndex, const int32 LeafObjectIndex) -> void {}
 			template<typename TVisitor>
-			void VisitAllIntersections(const FAABB3& LocalBounds, const TVisitor& Visitor) const
+			void VisitAllIntersections(const FAABB3& LocalBounds, const TVisitor& ObjectVisitor) const
 			{
-				// Each objects can be in multiple leafs, so we keep track of whichones we have visited
-				TArray<bool> ObjectVisited;
-				ObjectVisited.AddZeroed(NumObjects());
-			
-				const auto& LeafVisitor = [this, &ObjectVisited, &Visitor](const TArray<int32>& Indices)
+				const auto& NodeVisitor = [this, &ObjectVisitor](const FImplicitBVHNode& Node)
 				{
-					for (const int32 Index : Indices)
+					if (Node.IsLeaf())
 					{
-						if (!ObjectVisited[Index])
-						{
-							ObjectVisited[Index] = true;
-							Visitor(Index);
-						}
+						VisitNodeObjects(Node, ObjectVisitor);
 					}
 				};
 
-				BVH.VisitAllIntersections(LocalBounds, LeafVisitor);
+				VisitOverlappingNodes(LocalBounds, NodeVisitor);
+			}
+
+			// Calls the visitor for every overlapping leaf.
+			// @tparam TVisitor void(const FImplicitBVHNode& Node)
+			template<typename TVisitor>
+			void VisitOverlappingNodes(const FAABB3& LocalBounds, const TVisitor& NodeVisitor) const
+			{
+				if (!Nodes.IsEmpty())
+				{
+					VisitOverlappingNodesRecursive(Nodes[0], FAABB3f(LocalBounds), NodeVisitor);
+				}
+			}
+
+			// Recursively visit all nodes in the hierarchy. Will stop visiting children
+			// if the visitor returns false. Leaf will be null when visiting an internal node.
+			// @param NodeVisitor (const FAABB3f& NodeBounds, const int32 NodeDepth, const FImplicitBVHNode& Node) -> void
+			template<typename TVisitor>
+			void VisitHierarchy(const TVisitor& NodeVisitor) const
+			{
+				if (!Nodes.IsEmpty())
+				{
+					VisitHierarchyRecursive(Nodes[0], 0, NodeVisitor);
+				}
+			}
+
+			// Visit all the items in the specified leaf node (which is probably obtained from VisitHierarchy)
+			// @param ObjectVisitor (const FImplicitObject* Implicit, const FRigidTransform3f& RelativeTransformf, const FAABB3f& RelativeBoundsf, const int32 RootObjectIndex, const int32 LeafObjectIndex) -> void
+			template<typename TVisitor>
+			void VisitNodeObjects(const FImplicitBVHNode& Node, const TVisitor& ObjectVisitor) const
+			{
+				for (int32 NodeObjectIndex = Node.ObjectBeginIndex; NodeObjectIndex < Node.ObjectEndIndex; ++NodeObjectIndex)
+				{
+					const FImplicitBVHObject& Object = Objects[NodeObjectIndices[NodeObjectIndex]];
+
+					ObjectVisitor(Object.GetGeometry(), Object.GetTransformf(), Object.GetBounds(), Object.GetRootObjectIndex(), Object.GetObjectIndex());
+				}
+			}
+
+			// Does the bounding box overlap any leaf nodes with items in it
+			bool IsOverlappingBounds(const FAABB3& LocalBounds) const
+			{
+				if (!Nodes.IsEmpty())
+				{
+					return IsOverlappingBoundsRecursive(Nodes[0], FAABB3f(LocalBounds));
+				}
+				return false;
 			}
 
 		private:
@@ -141,45 +207,85 @@ namespace Chaos
 
 			// Initialize the BVH from the specified set of children.
 			// NOTE: InChildren should be immediate children of the BVH owner, not further-removed descendents
-			void Init(FObjects&& InObjects, const int32 InMaxBVHDepth);
+			void Init(FObjects&& InObjects, const int32 MaxDepth, const int32 MaxLeafObjects);
 
-			// A BVH leaf holds an array of indices into the Objects array
-			FObjects Objects;
-			FBVH BVH;
-		};
+			template<typename TVisitor>
+			void VisitOverlappingNodesRecursive(const FImplicitBVHNode& Node, const FAABB3f& LocalBounds, const TVisitor& NodeVisitor) const
+			{
+				if (!Node.Bounds.Intersects(LocalBounds))
+				{
+					return;
+				}
 
-		class FImplicitCache
-		{
-		public:
+				if (Node.IsLeaf())
+				{
+					NodeVisitor(Node);
+				}
+				else
+				{
+					VisitOverlappingNodesRecursive(Nodes[Node.ChildNodeIndices[0]], LocalBounds, NodeVisitor);
+					VisitOverlappingNodesRecursive(Nodes[Node.ChildNodeIndices[1]], LocalBounds, NodeVisitor);
+				}
+			}
 
-		private:
-			TUniquePtr<FImplicitBVH> BVH;
-		};
-	}
+			template<typename TVisitor>
+			void VisitHierarchyRecursive(const FImplicitBVHNode& Node, const int32 NodeDepth, const TVisitor& NodeVisitor) const
+			{
+				// Visit this (non-leaf) node
+				const bool bVisitChildren = NodeVisitor(Node.Bounds, NodeDepth, Node);
 
-	///////////////////////////////////////////////////////////////////////////////////////////////
-	///////////////////////////////////////////////////////////////////////////////////////////////
-	// 
-	// TBoundingVolumeHierarchy Template API
-	// 
-	///////////////////////////////////////////////////////////////////////////////////////////////
-	///////////////////////////////////////////////////////////////////////////////////////////////
+				// Visit children
+				if (bVisitChildren && !Node.IsLeaf())
+				{
+					VisitHierarchyRecursive(Nodes[Node.ChildNodeIndices[0]], NodeDepth + 1, NodeVisitor);
+					VisitHierarchyRecursive(Nodes[Node.ChildNodeIndices[1]], NodeDepth + 1, NodeVisitor);
+				}
+			}
 
-	template<>
-	inline bool HasBoundingBox(const TArray<Private::FImplicitBVHObject>& Objects, const int32 ObjectIndex)
-	{
-		return Objects[ObjectIndex].GetGeometry()->HasBoundingBox();
-	}
+			bool IsOverlappingBoundsRecursive(const FImplicitBVHNode& Node, const FAABB3f& LocalBounds) const
+			{
+				if (Node.Bounds.Intersects(LocalBounds))
+				{
+					// If we hit a leaf, we have an overlap
+					if (Node.IsLeaf())
+					{
+						return true;
+					}
 
-	template<class T, int d>
-	inline const TAABB<T, d> GetWorldSpaceBoundingBox(const TArray<Private::FImplicitBVHObject>& Objects, const int32 ObjectIndex, const TMap<int32, TAABB<T, d>>& WorldSpaceBoxes)
-	{
-		return TAABB<T, d>(Objects[ObjectIndex].GetBounds());
-	}
+					// Check children, stopping if we get an overlap
+					if (IsOverlappingBoundsRecursive(Nodes[Node.ChildNodeIndices[0]], LocalBounds))
+					{
+						return true;
+					}
+					if (IsOverlappingBoundsRecursive(Nodes[Node.ChildNodeIndices[1]], LocalBounds))
+					{
+						return true;
+					}
+				}
+				return false;
+			}
 
-	template<class T, int d>
-	void ComputeAllWorldSpaceBoundingBoxes(const TArray<Private::FImplicitBVHObject>& Objects, const TArray<int32>& AllObjects, const bool bUseVelocity, const T Dt, TMap<int32, TAABB<T, d>>& WorldSpaceBoxes)
-	{
+			struct FWorkData
+			{
+				TArray<FVec3f> ObjectCenters;
+				TArray<int32> NodeObjectIndices0;
+				TArray<int32> NodeObjectIndices1;
+			};
+
+			void BuildTree(const int32 MaxDepth, const int32 MaxLeafObjects);
+			int32 AddNodeRecursive(const FAABB3f& NodeBounds, const int32 ObjectBeginIndex, const int32 ObjectEndIndex, const int32 NodeDepth, const int32 MaxDepth, const int32 MaxLeafObjects, FWorkData& WorkData);
+			bool Partition(const FAABB3f& NodeBounds, const int32 ObjectBeginIndex, const int32 ObjectEndIndex, FAABB3f& OutChildNodeBounds0, FAABB3f& OutChildNodeBounds1, int32& OutPartitionIndex, FWorkData& WorkData);
+
+			// A BVH leaf holds an array of indices into the Objects array
+			FObjects Objects;
+			TArray<int32> NodeObjectIndices;
+			TArray<FImplicitBVHNode> Nodes;
+			int32 TreeDepth;
+		};
 	}
 
+	// DO NOT USE: This is only required for serializing the old BVH data
+	template<> inline bool HasBoundingBox(const TArray<Private::FImplicitBVHObject>& Objects, const int32 ObjectIndex) { return Objects[ObjectIndex].GetGeometry()->HasBoundingBox(); }
+	template<class T, int d> inline const TAABB<T, d> GetWorldSpaceBoundingBox(const TArray<Private::FImplicitBVHObject>& Objects, const int32 ObjectIndex, const TMap<int32, TAABB<T, d>>& WorldSpaceBoxes) { return TAABB<T, d>(Objects[ObjectIndex].GetBounds()); }
+	template<class T, int d> void ComputeAllWorldSpaceBoundingBoxes(const TArray<Private::FImplicitBVHObject>& Objects, const TArray<int32>& AllObjects, const bool bUseVelocity, const T Dt, TMap<int32, TAABB<T, d>>& WorldSpaceBoxes) { }
 }
\ No newline at end of file
diff --git a/Engine/Source/Runtime/Experimental/Chaos/Public/Chaos/ImplicitObjectTransformed.h b/Engine/Source/Runtime/Experimental/Chaos/Public/Chaos/ImplicitObjectTransformed.h
--- a/Engine/Source/Runtime/Experimental/Chaos/Public/Chaos/ImplicitObjectTransformed.h
+++ b/Engine/Source/Runtime/Experimental/Chaos/Public/Chaos/ImplicitObjectTransformed.h
@@ -368,7 +368,7 @@ protected:
 		const int32 RootObjectIndex,
 		int32& ObjectIndex,
 		int32& LeafObjectIndex,
-		const FImplicitHierarchyVisitor& VisitorFunc) const override
+		const FImplicitHierarchyVisitor& VisitorFunc) const override final
 	{
 		// Skip self
 		++ObjectIndex;
@@ -383,7 +383,7 @@ protected:
 		const int32 RootObjectIndex,
 		int32& ObjectIndex,
 		int32& LeafObjectIndex,
-		const FImplicitHierarchyVisitor& VisitorFunc) const override
+		const FImplicitHierarchyVisitor& VisitorFunc) const override final
 	{
 		// Skip self
 		++ObjectIndex;
@@ -392,19 +392,30 @@ protected:
 		MObject->VisitLeafObjectsImpl(MTransform * ObjectTransform, RootObjectIndex, ObjectIndex, LeafObjectIndex, VisitorFunc);
 	}
 
-	virtual void VisitObjectsImpl(
+	virtual bool VisitObjectsImpl(
 		const FRigidTransform3& ObjectTransform,
 		const int32 RootObjectIndex,
 		int32& ObjectIndex,
 		int32& LeafObjectIndex,
-		const FImplicitHierarchyVisitor& VisitorFunc) const override
+		const FImplicitHierarchyVisitorBool& VisitorFunc) const override final
 	{
 		// Visit self
-		VisitorFunc(this, ObjectTransform, RootObjectIndex, ObjectIndex, INDEX_NONE);
+		bool bContinue = VisitorFunc(this, ObjectTransform, RootObjectIndex, ObjectIndex, INDEX_NONE);
 		++ObjectIndex;
 
 		// Visit child
-		MObject->VisitObjectsImpl(MTransform * ObjectTransform, RootObjectIndex, ObjectIndex, LeafObjectIndex, VisitorFunc);
+		if (bContinue)
+		{
+			bContinue = MObject->VisitObjectsImpl(MTransform * ObjectTransform, RootObjectIndex, ObjectIndex, LeafObjectIndex, VisitorFunc);
+		}
+
+		return bContinue;
+	}
+
+	virtual bool IsOverlappingBoundsImpl(const FAABB3& InLocalBounds) const override final
+	{
+		const FAABB3 LocalBounds = InLocalBounds.InverseTransformedAABB(MTransform);
+		return MObject->IsOverlappingBoundsImpl(LocalBounds);
 	}
 
 private:
diff --git a/Engine/Source/Runtime/Experimental/Chaos/Public/Chaos/ImplicitObjectUnion.h b/Engine/Source/Runtime/Experimental/Chaos/Public/Chaos/ImplicitObjectUnion.h
--- a/Engine/Source/Runtime/Experimental/Chaos/Public/Chaos/ImplicitObjectUnion.h
+++ b/Engine/Source/Runtime/Experimental/Chaos/Public/Chaos/ImplicitObjectUnion.h
@@ -15,6 +15,7 @@ namespace Chaos
 namespace Private
 {
 	class FImplicitBVH;
+	class FImplicitBVHObject;
 }
 
 namespace CVars
@@ -41,6 +42,12 @@ class CHAOS_API FImplicitObjectUnion : public FImplicitObject
 	void Combine(TArray<TUniquePtr<FImplicitObject>>& Objects);
 	void RemoveAt(int32 RemoveIndex);
 
+	// The total number of root objects in the hierarchy (same as GetObjects().Num())
+	int32 GetNumRootObjects() const
+	{
+		return MObjects.Num();
+	}
+
 	// The total number of leaf objects in the hierarchy
 	int32 GetNumLeafObjects() const
 	{
@@ -208,6 +215,11 @@ class CHAOS_API FImplicitObjectUnion : public FImplicitObject
 		return Union;
 	}
 
+	const Private::FImplicitBVH* GetBVH() const
+	{
+		return BVH.Get();
+	}
+
 #if INTEL_ISPC
 	// See PerParticlePBDCollisionConstraint.cpp
 	// ISPC code has matching structs for interpreting FImplicitObjects.
@@ -256,17 +268,20 @@ protected:
 		int32& LeafObjectIndex,
 		const FImplicitHierarchyVisitor& VisitorFunc) const override final;
 
-	virtual void VisitObjectsImpl(
+	virtual bool VisitObjectsImpl(
 		const FRigidTransform3& ObjectTransform,
 		const int32 RootObjectIndex,
 		int32& ObjectIndex,
 		int32& LeafObjectIndex,
-		const FImplicitHierarchyVisitor& VisitorFunc) const override final;
+		const FImplicitHierarchyVisitorBool& VisitorFunc) const override final;
+
+	virtual bool IsOverlappingBoundsImpl(
+		const FAABB3& LocalBounds) const override final;
 
   protected:
-	//needed for serialization
+	// Needed for serialization
 	FImplicitObjectUnion();
-	friend FImplicitObject;	//needed for serialization
+	friend FImplicitObject;
 
 	void SetNumLeafObjects(const int32 InNumLeafObjects);
 	void CreateBVH();
@@ -326,7 +341,9 @@ public:
 		return ImplicitObjectType::UnionClustered;
 	}
 
-	void FindAllIntersectingClusteredObjects(TArray<FLargeUnionClusteredImplicitInfo>& Out, const FAABB3& LocalBounds) const;
+	UE_DEPRECATED(5.3, "Not supported")
+	void FindAllIntersectingClusteredObjects(TArray<FLargeUnionClusteredImplicitInfo>& Out, const FAABB3& LocalBounds) const {}
+
 	TArray<FPBDRigidParticleHandle*> FindAllIntersectingChildren(const FAABB3& LocalBounds) const;
 
 	// DO NOT USE!!
diff --git a/Engine/Source/Runtime/Experimental/Chaos/Public/Chaos/ParticleHandle.h b/Engine/Source/Runtime/Experimental/Chaos/Public/Chaos/ParticleHandle.h
--- a/Engine/Source/Runtime/Experimental/Chaos/Public/Chaos/ParticleHandle.h
+++ b/Engine/Source/Runtime/Experimental/Chaos/Public/Chaos/ParticleHandle.h
@@ -2675,6 +2675,16 @@ public:
 		}
 	}
 
+	// Update the BVH in the geometry hierarchy if we have new geometry
+	// This exists so that we don't repeatedly rebuild the BVH when welding, for example,
+	// although ideally we would remove this and change the way welding works
+	void PrepareBVH()
+	{
+		if (MNonFrequentData.IsDirty(MDirtyFlags))
+		{
+			PrepareBVHImpl();
+		}
+	}
 
 	class IPhysicsProxyBase* GetProxy() const
 	{
@@ -2774,6 +2784,8 @@ protected:
 
 	CHAOS_API void UpdateShapesArray();
 
+	void PrepareBVHImpl();
+
 	virtual void SyncRemoteDataImp(FDirtyPropertiesManager& Manager, int32 DataIdx, const FDirtyChaosProperties& RemoteData) const
 	{
 		MXR.SyncRemote(Manager, DataIdx, RemoteData);
