Fix Mac Linux CIS (missing "template" keyword in template call site)

diff --git a/Engine/Source/Runtime/Experimental/Chaos/Private/Chaos/Collision/ParticlePairMidPhase.cpp b/Engine/Source/Runtime/Experimental/Chaos/Private/Chaos/Collision/ParticlePairMidPhase.cpp
--- a/Engine/Source/Runtime/Experimental/Chaos/Private/Chaos/Collision/ParticlePairMidPhase.cpp
+++ b/Engine/Source/Runtime/Experimental/Chaos/Private/Chaos/Collision/ParticlePairMidPhase.cpp
@@ -96,10 +96,10 @@ namespace Chaos
 			NumImplicits = 1;
 
 			// @todo(chaos): we need better downcasting in FImplicitObject
-			const FImplicitObjectUnion* Union = Implicit->GetObject<FImplicitObjectUnion>();
+			const FImplicitObjectUnion* Union = Implicit->template GetObject<FImplicitObjectUnion>();
 			if (Union == nullptr)
 			{
-				Union = Implicit->GetObject<FImplicitObjectUnionClustered>();
+				Union = Implicit->template GetObject<FImplicitObjectUnionClustered>();
 			}
 
 			if (Union != nullptr)
@@ -137,7 +137,7 @@ namespace Chaos
 			Implicit->VisitObjects(
 				[Implicit, &BVHParticles](const FImplicitObject* HierarchyImplicit, const FRigidTransform3& Transform, const int32 RootObjectIndex, const int32 ObjectIndex, const int32 LeafObjectIndex) -> bool
 				{
-					if (const FImplicitObjectUnionClustered* UnionClustered = HierarchyImplicit->GetObject<FImplicitObjectUnionClustered>())
+					if (const FImplicitObjectUnionClustered* UnionClustered = HierarchyImplicit->template GetObject<FImplicitObjectUnionClustered>())
 					{
 						// If this union has our implicit in its map, extract the collision particles and we're done
 						if (const FPBDRigidParticleHandle* SourceRigid = UnionClustered->FindParticleForImplicitObject(Implicit))
@@ -1080,11 +1080,11 @@ namespace Chaos
 		const Private::FImplicitBVH* BVH1 = nullptr;
 		if (CVars::bChaosUnionBVHEnabled)
 		{
-			if (const FImplicitObjectUnion* Union0 = Implicit0->GetObject<FImplicitObjectUnion>())
+			if (const FImplicitObjectUnion* Union0 = Implicit0->template GetObject<FImplicitObjectUnion>())
 			{
 				BVH0 = Union0->GetBVH();
 			}
-			if (const FImplicitObjectUnion* Union1 = Implicit1->GetObject<FImplicitObjectUnion>())
+			if (const FImplicitObjectUnion* Union1 = Implicit1->template GetObject<FImplicitObjectUnion>())
 			{
 				BVH1 = Union1->GetBVH();
 			}
diff --git a/Engine/Source/Runtime/Experimental/Chaos/Private/Chaos/ParticleHandle.cpp b/Engine/Source/Runtime/Experimental/Chaos/Private/Chaos/ParticleHandle.cpp
--- a/Engine/Source/Runtime/Experimental/Chaos/Private/Chaos/ParticleHandle.cpp
+++ b/Engine/Source/Runtime/Experimental/Chaos/Private/Chaos/ParticleHandle.cpp
@@ -119,7 +119,7 @@ namespace Chaos
 	{
 		if (MNonFrequentData.IsDirty(MDirtyFlags))
 		{
-			if (const FImplicitObjectUnion* Union = MNonFrequentData.Read().Geometry()->GetObject<FImplicitObjectUnion>())
+			if (const FImplicitObjectUnion* Union = MNonFrequentData.Read().Geometry()->template GetObject<FImplicitObjectUnion>())
 			{
 				// This will rebuild the BVH if the geometry is new, otherwise do nothing
 				const_cast<FImplicitObjectUnion*>(Union)->SetAllowBVH(true);
