Chaos - fix sleeping when materials or body sleep settings are changed after the body has been added to the constraint graph

diff --git a/Engine/Source/Runtime/Experimental/Chaos/Private/Chaos/PBDRigidsEvolutionGBF.cpp b/Engine/Source/Runtime/Experimental/Chaos/Private/Chaos/PBDRigidsEvolutionGBF.cpp
--- a/Engine/Source/Runtime/Experimental/Chaos/Private/Chaos/PBDRigidsEvolutionGBF.cpp
+++ b/Engine/Source/Runtime/Experimental/Chaos/Private/Chaos/PBDRigidsEvolutionGBF.cpp
@@ -1007,6 +1007,9 @@ void FPBDRigidsEvolutionGBF::ParticleMaterialChanged(FGeometryParticleHandle* Pa
 		Collision.ClearMaterialProperties();
 		return ECollisionVisitorResult::Continue;
 	});
+
+	// The graph caches some sleep thresholds etc 
+	GetIslandManager().UpdateParticleMaterial(Particle);
 }
 
 
diff --git a/Engine/Source/Runtime/Experimental/Chaos/Private/PhysicsProxy/SingleParticlePhysicsProxy.cpp b/Engine/Source/Runtime/Experimental/Chaos/Private/PhysicsProxy/SingleParticlePhysicsProxy.cpp
--- a/Engine/Source/Runtime/Experimental/Chaos/Private/PhysicsProxy/SingleParticlePhysicsProxy.cpp
+++ b/Engine/Source/Runtime/Experimental/Chaos/Private/PhysicsProxy/SingleParticlePhysicsProxy.cpp
@@ -153,9 +153,9 @@ void PushToPhysicsStateImp(const Chaos::FDirtyPropertiesManager& Manager, Chaos:
 			}
 		}
 		
-		if (bHasMaterial)
+		// If the material, geometry, shape data, or sleep properties changed, we need to notify any systems that cache material data
+		if (bHasMaterial || bUpdateCollisionData || NewNonFrequentData || bHasDynamicData)
 		{
-			// If materials changed, collisions need to recache their material data
 			Evolution.ParticleMaterialChanged(Handle);
 		}
 
