Add tag support to ISMPool.
ISMPool uses serialized version of the default primitive custom data to allow us to copy paste ISMPools for debugging.
ISMPool uses fast version of hash combine.
Expose ISMPool actors getter from ISMPool subsystem.

diff --git a/Engine/Source/Runtime/Experimental/GeometryCollectionEngine/Private/GeometryCollection/GeometryCollectionISMPoolComponent.cpp b/Engine/Source/Runtime/Experimental/GeometryCollectionEngine/Private/GeometryCollection/GeometryCollectionISMPoolComponent.cpp
--- a/Engine/Source/Runtime/Experimental/GeometryCollectionEngine/Private/GeometryCollection/GeometryCollectionISMPoolComponent.cpp
+++ b/Engine/Source/Runtime/Experimental/GeometryCollectionEngine/Private/GeometryCollection/GeometryCollectionISMPoolComponent.cpp
@@ -90,7 +90,7 @@ FGeometryCollectionISM::FGeometryCollectionISM(AActor* InOwningActor, const FGeo
 	}
 	for (int32 DataIndex = 0; DataIndex < MeshInstance.CustomPrimitiveData.Num(); DataIndex++)
 	{
-		ISMC->SetCustomPrimitiveDataFloat(DataIndex, MeshInstance.CustomPrimitiveData[DataIndex]);
+		ISMC->SetDefaultCustomPrimitiveDataFloat(DataIndex, MeshInstance.CustomPrimitiveData[DataIndex]);
 	}
 
 	ISMC->SetRemoveSwap();
@@ -105,12 +105,13 @@ FGeometryCollectionISM::FGeometryCollectionISM(AActor* InOwningActor, const FGeo
 	ISMC->SetCollisionEnabled(ECollisionEnabled::NoCollision);
 	ISMC->bOverrideMinLOD = MeshInstance.Desc.MinLod > 0;	
 	ISMC->MinLOD = MeshInstance.Desc.MinLod;
-	
+	ISMC->ComponentTags.Append(MeshInstance.Desc.Tags);
+
 	if (HISMC)
 	{
 		HISMC->SetLODDistanceScale(MeshInstance.Desc.LodScale);
 	}
-	
+
 	InOwningActor->AddInstanceComponent(ISMC);
 	ISMC->RegisterComponent();
 	ISMComponent = ISMC;
diff --git a/Engine/Source/Runtime/Experimental/GeometryCollectionEngine/Private/GeometryCollection/GeometryCollectionISMPoolSubSystem.cpp b/Engine/Source/Runtime/Experimental/GeometryCollectionEngine/Private/GeometryCollection/GeometryCollectionISMPoolSubSystem.cpp
--- a/Engine/Source/Runtime/Experimental/GeometryCollectionEngine/Private/GeometryCollection/GeometryCollectionISMPoolSubSystem.cpp
+++ b/Engine/Source/Runtime/Experimental/GeometryCollectionEngine/Private/GeometryCollection/GeometryCollectionISMPoolSubSystem.cpp
@@ -43,3 +43,11 @@ AGeometryCollectionISMPoolActor* UGeometryCollectionISMPoolSubSystem::FindISMPoo
 	}
 	return ISMPoolActor;
 }
+
+void UGeometryCollectionISMPoolSubSystem::GetISMPoolActors(TArray<AGeometryCollectionISMPoolActor*>& OutActors) const
+{
+	if (ISMPoolActor != nullptr)
+	{
+		OutActors.Add(ISMPoolActor);
+	}
+}
diff --git a/Engine/Source/Runtime/Experimental/GeometryCollectionEngine/Public/GeometryCollection/GeometryCollectionISMPoolComponent.h b/Engine/Source/Runtime/Experimental/GeometryCollectionEngine/Public/GeometryCollection/GeometryCollectionISMPoolComponent.h
--- a/Engine/Source/Runtime/Experimental/GeometryCollectionEngine/Public/GeometryCollection/GeometryCollectionISMPoolComponent.h
+++ b/Engine/Source/Runtime/Experimental/GeometryCollectionEngine/Public/GeometryCollection/GeometryCollectionISMPoolComponent.h
@@ -110,6 +110,7 @@ struct FISMComponentDescription
 	int32 EndCullDistance = 0;
 	int32 MinLod = 0;
 	float LodScale = 1.f;
+	TArray<FName> Tags;
 
 	bool operator==(const FISMComponentDescription& Other) const
 	{
@@ -123,18 +124,26 @@ struct FISMComponentDescription
 			StartCullDistance == Other.StartCullDistance && 
 			EndCullDistance == Other.EndCullDistance &&
 			MinLod == Other.MinLod &&
-			LodScale == Other.LodScale;
+			LodScale == Other.LodScale &&
+			Tags == Other.Tags;
 	}
 };
 
 FORCEINLINE uint32 GetTypeHash(const FISMComponentDescription& Desc)
 {
-	const uint32 PackedBools = (Desc.bUseHISM ? 1 : 0) | (Desc.bReverseCulling ? 2 : 0) | (Desc.bIsStaticMobility ? 4 : 0) | (Desc.bAffectShadow ? 8 : 0) | (Desc.bAffectDistanceFieldLighting ? 16 : 0) | (Desc.bAffectDynamicIndirectLighting ? 32 : 0);
-	uint32 Hash = HashCombine(GetTypeHash(PackedBools), GetTypeHash(Desc.NumCustomDataFloats));
-	Hash = HashCombine(Hash, GetTypeHash(Desc.StartCullDistance));
-	Hash = HashCombine(Hash, GetTypeHash(Desc.EndCullDistance));
-	Hash = HashCombine(Hash, GetTypeHash(Desc.MinLod));
-	return HashCombine(Hash, GetTypeHash(Desc.LodScale));
+	const uint32 PackedBools = 
+		(Desc.bUseHISM ? 1 : 0) | 
+		(Desc.bReverseCulling ? 2 : 0) | 
+		(Desc.bIsStaticMobility ? 4 : 0) | 
+		(Desc.bAffectShadow ? 8 : 0) | 
+		(Desc.bAffectDistanceFieldLighting ? 16 : 0) |
+		(Desc.bAffectDynamicIndirectLighting ? 32 : 0);
+	uint32 Hash = HashCombineFast(GetTypeHash(PackedBools), GetTypeHash(Desc.NumCustomDataFloats));
+	Hash = HashCombineFast(Hash, GetTypeHash(Desc.StartCullDistance));
+	Hash = HashCombineFast(Hash, GetTypeHash(Desc.EndCullDistance));
+	Hash = HashCombineFast(Hash, GetTypeHash(Desc.MinLod));
+	Hash = HashCombineFast(Hash, GetTypeHash(Desc.LodScale));
+	return HashCombineFast(Hash, GetArrayHash(Desc.Tags.GetData(), Desc.Tags.Num()));
 }
 
 /**
@@ -186,16 +195,16 @@ struct FGeometryCollectionStaticMeshInstance
 FORCEINLINE uint32 GetTypeHash(const FGeometryCollectionStaticMeshInstance& MeshInstance)
 {
 	uint32 CombinedHash = GetTypeHash(MeshInstance.StaticMesh);
-	CombinedHash = HashCombine(CombinedHash, GetTypeHash(MeshInstance.MaterialsOverrides.Num()));
+	CombinedHash = HashCombineFast(CombinedHash, GetTypeHash(MeshInstance.MaterialsOverrides.Num()));
 	for (const UMaterialInterface* Material: MeshInstance.MaterialsOverrides)
 	{
-		CombinedHash = HashCombine(CombinedHash, GetTypeHash(Material));
+		CombinedHash = HashCombineFast(CombinedHash, GetTypeHash(Material));
 	}
 	for (const float CustomFloat : MeshInstance.CustomPrimitiveData)
 	{
-		CombinedHash = HashCombine(CombinedHash, GetTypeHash(CustomFloat));
+		CombinedHash = HashCombineFast(CombinedHash, GetTypeHash(CustomFloat));
 	}
-	CombinedHash = HashCombine(CombinedHash, GetTypeHash(MeshInstance.Desc));
+	CombinedHash = HashCombineFast(CombinedHash, GetTypeHash(MeshInstance.Desc));
 	return CombinedHash;
 }
 
diff --git a/Engine/Source/Runtime/Experimental/GeometryCollectionEngine/Public/GeometryCollection/GeometryCollectionISMPoolSubSystem.h b/Engine/Source/Runtime/Experimental/GeometryCollectionEngine/Public/GeometryCollection/GeometryCollectionISMPoolSubSystem.h
--- a/Engine/Source/Runtime/Experimental/GeometryCollectionEngine/Public/GeometryCollection/GeometryCollectionISMPoolSubSystem.h
+++ b/Engine/Source/Runtime/Experimental/GeometryCollectionEngine/Public/GeometryCollection/GeometryCollectionISMPoolSubSystem.h
@@ -8,15 +8,15 @@
 class AGeometryCollectionISMPoolActor;
 
 /**
-* A subsystem managing ISMPool actors ( used by geometry collection now but repurposed for more general use )
-*/
+ * A subsystem managing ISMPool actors.
+ * Used by geometry collection now but repurposed for more general use.
+ */
 UCLASS()
 class GEOMETRYCOLLECTIONENGINE_API UGeometryCollectionISMPoolSubSystem : public UWorldSubsystem
 {
 	GENERATED_BODY()
 
 public:
-
 	UGeometryCollectionISMPoolSubSystem();
 
 	// USubsystem BEGIN
@@ -24,11 +24,14 @@ public:
 	virtual void Deinitialize() override;
 	// USubsystem END
 
+	/** Finds or creates an actor. */
 	AGeometryCollectionISMPoolActor* FindISMPoolActor();
+	
+	/** Get all actors managed by the subsystem. */
+	void GetISMPoolActors(TArray<AGeometryCollectionISMPoolActor*>& OutActors) const;
 
 protected:
-
-	/** for now we only use one ISMPool actor per world, but we could extend the system to manage many more and return the right one based on  search criteria */
+	/** For now we only use one ISMPool actor per world, but we could extend the system to manage many more and return the right one based on search criteria. */
 	UPROPERTY(Transient)
 	TObjectPtr<AGeometryCollectionISMPoolActor> ISMPoolActor;
 };
