Iris - Speculative fix for a couple of rare crashes.

#rnx

diff --git a/Engine/Source/Runtime/Experimental/Iris/Core/Private/Iris/ReplicationSystem/ObjectReferenceCache.cpp b/Engine/Source/Runtime/Experimental/Iris/Core/Private/Iris/ReplicationSystem/ObjectReferenceCache.cpp
--- a/Engine/Source/Runtime/Experimental/Iris/Core/Private/Iris/ReplicationSystem/ObjectReferenceCache.cpp
+++ b/Engine/Source/Runtime/Experimental/Iris/Core/Private/Iris/ReplicationSystem/ObjectReferenceCache.cpp
@@ -232,10 +232,13 @@ bool FObjectReferenceCache::CreateObjectReferenceInternal(const UObject* Object,
 	CA_ASSUME(Object != nullptr);
 
 	// Check if we are already know about this object
-	if (FNetRefHandle* ReferenceHandle = ObjectToNetReferenceHandle.Find(Object))
+	if (FNetRefHandle* RefHandlePtr = ObjectToNetReferenceHandle.Find(Object))
 	{
+		// Cache RefHandle as the pointer can become invalidated later.
+		const FNetRefHandle RefHandle = *RefHandlePtr;
+
 		// Verify that this is the same object and clean up if it is not
-		FCachedNetObjectReference& CachedObject = ReferenceHandleToCachedReference.FindChecked(*ReferenceHandle);
+		FCachedNetObjectReference& CachedObject = ReferenceHandleToCachedReference.FindChecked(RefHandle);
 
 		if (CachedObject.Object.Get() == Object)
 		{
@@ -250,7 +253,7 @@ bool FObjectReferenceCache::CreateObjectReferenceInternal(const UObject* Object,
 			ObjectToNetReferenceHandle.Remove(CachedObject.ObjectKey);
 			ObjectToNetReferenceHandle.Remove(Object);
 
-			if ((*ReferenceHandle).IsStatic())
+			if (RefHandle.IsStatic())
 			{
 				// Note we only cleanse the object reference
 				// we still keep the cachedObject data around in order to be able to serialize static destruction infos
@@ -259,7 +262,7 @@ bool FObjectReferenceCache::CreateObjectReferenceInternal(const UObject* Object,
 			}
 			else
 			{
-				ReferenceHandleToCachedReference.Remove(*ReferenceHandle);
+				ReferenceHandleToCachedReference.Remove(RefHandle);
 			}
 		}
 	}
diff --git a/Engine/Source/Runtime/Experimental/Iris/Core/Private/Iris/ReplicationSystem/ReplicationWriter.cpp b/Engine/Source/Runtime/Experimental/Iris/Core/Private/Iris/ReplicationSystem/ReplicationWriter.cpp
--- a/Engine/Source/Runtime/Experimental/Iris/Core/Private/Iris/ReplicationSystem/ReplicationWriter.cpp
+++ b/Engine/Source/Runtime/Experimental/Iris/Core/Private/Iris/ReplicationSystem/ReplicationWriter.cpp
@@ -2188,8 +2188,6 @@ FReplicationWriter::EWriteObjectStatus FReplicationWriter::WriteObjectAndSubObje
 			EWriteObjectStatus SubObjectWriteStatus = WriteObjectAndSubObjects(Context, SubObjectInternalIndex, WriteObjectFlags, OutBatchInfo);
 			if (!IsWriteObjectSuccess(SubObjectWriteStatus))
 			{
-				// Need to remove the batch entry from BatchInfo.
-				OutBatchInfo.ObjectInfos.RemoveAt(OutBatchInfo.ObjectInfos.Num() - 1);
 				return SubObjectWriteStatus;
 			}
 
@@ -2355,6 +2353,9 @@ int FReplicationWriter::PrepareAndSendHugeObjectPayload(FNetSerializationContext
 		// If we cannot fit the object in the largest supported buffer then we will never fit the object.
 		if (WriteHugeObjectStatus == EWriteObjectStatus::BitStreamOverflow)
 		{
+			// Cleanup data from batch
+			HandleObjectBatchFailure(WriteHugeObjectStatus, BatchInfo, WriteBitStreamInfo);
+
 			UE_LOG(LogIris, Error, TEXT("Unable to fit object ( InternalIndex: %u ) in maximum combined payload of %u bytes. Connection %u will be disconnected."), InternalIndex, MaxHugeObjectPayLoadBytes, Context.GetLocalConnectionId());
 			Context.SetError(NetError_ObjectStateTooLarge);
 			return -1;
@@ -2363,7 +2364,7 @@ int FReplicationWriter::PrepareAndSendHugeObjectPayload(FNetSerializationContext
 		// If we encounter some other error we can try sending a smaller object in the meantime.
 		if (!IsWriteObjectSuccess(WriteHugeObjectStatus))
 		{
-			// Need to call this to cleanup data from batch
+			// Cleanup data from batch
 			HandleObjectBatchFailure(WriteHugeObjectStatus, BatchInfo, WriteBitStreamInfo);
 
 			UE_LOG(LogIris, Verbose, TEXT("Problem writing huge object ( InternalIndex: %u ). WriteObjectStatus: %u. Trying smaller object."), InternalIndex, unsigned(WriteHugeObjectStatus));
