Iris
* Poll configs are now set using Frequency (updates per second) instead of frames between polls.
* This allows the poll configs to work with any type of server tick rate
* Added ability to recover a given object's poll frequency via GetPollFrequencyOfRootObject
* ActorReplicationBridge will now listen to changes in server tick rates and reinitalize the poll frame updates of all objects via ReinitPollFrequency()

#rnx

diff --git a/Engine/Config/BaseEngine.ini b/Engine/Config/BaseEngine.ini
--- a/Engine/Config/BaseEngine.ini
+++ b/Engine/Config/BaseEngine.ini
@@ -1255,7 +1255,7 @@ MaxPartCount=4096
 
 [/Script/IrisCore.ObjectReplicationBridgeConfig]
 ; It is possible to override how often replicated objects are Polled for dirtyness if they are not using pushed based dirtiness, normally based on NetUpdateFrequency
-;+PollConfigs=(ClassName=/Script/Example.Pawn, PollFramePeriod=0, bIncludeSubclasses=true)
+;+PollConfigs=(ClassName=/Script/Example.Pawn, PollFrequency=0, bIncludeSubclasses=true)
 
 ; Filters
 DefaultSpatialFilterName=Spatial
diff --git a/Engine/Plugins/Runtime/ReplicationSystemTestPlugin/Source/Private/Tests/ReplicationSystem/ReplicatedTestObject.cpp b/Engine/Plugins/Runtime/ReplicationSystemTestPlugin/Source/Private/Tests/ReplicationSystem/ReplicatedTestObject.cpp
--- a/Engine/Plugins/Runtime/ReplicationSystemTestPlugin/Source/Private/Tests/ReplicationSystem/ReplicatedTestObject.cpp
+++ b/Engine/Plugins/Runtime/ReplicationSystemTestPlugin/Source/Private/Tests/ReplicationSystem/ReplicatedTestObject.cpp
@@ -218,23 +218,6 @@ bool UReplicatedTestObjectBridge::IsAllowedToDestroyInstance(const UObject* Inst
 	return true;
 }
 
-void UReplicatedTestObjectBridge::SetPollFramePeriod(UReplicatedTestObject* Instance, uint8 FramePeriod)
-{
-	using namespace UE::Net;
-	using namespace UE::Net::Private;
-
-	if (Instance == nullptr)
-	{
-		return;
-	}
-	const FNetRefHandleManager* LocalNetRefHandleManager = &GetReplicationSystem()->GetReplicationSystemInternal()->GetNetRefHandleManager();
-	const uint32 InternalObjectIndex = LocalNetRefHandleManager->GetInternalIndex(Instance->NetRefHandle);
-	if (InternalObjectIndex != FNetRefHandleManager::InvalidInternalIndex)
-	{
-		UObjectReplicationBridge::SetPollFramePeriod(InternalObjectIndex, FramePeriod);
-	}
-}
-
 void UReplicatedTestObjectBridge::SetExternalWorldLocationUpdateFunctor(TFunction<void(FNetRefHandle NetHandle, const UObject* ReplicatedObject, FVector& OutLocation, float& OutCullDistance)> LocUpdateFunctor)
 {
 	SetInstanceGetWorldObjectInfoFunction(LocUpdateFunctor);
diff --git a/Engine/Plugins/Runtime/ReplicationSystemTestPlugin/Source/Private/Tests/ReplicationSystem/ReplicatedTestObject.h b/Engine/Plugins/Runtime/ReplicationSystemTestPlugin/Source/Private/Tests/ReplicationSystem/ReplicatedTestObject.h
--- a/Engine/Plugins/Runtime/ReplicationSystemTestPlugin/Source/Private/Tests/ReplicationSystem/ReplicatedTestObject.h
+++ b/Engine/Plugins/Runtime/ReplicationSystemTestPlugin/Source/Private/Tests/ReplicationSystem/ReplicatedTestObject.h
@@ -446,12 +446,12 @@ public:
 	// For testing we expose some things that normally are not accessible
 	const UE::Net::FReplicationInstanceProtocol* GetReplicationInstanceProtocol(FNetRefHandle Handle) const;
 
-	void SetPollFramePeriod(UReplicatedTestObject* Instance, uint8 FramePeriod);
-
 	void SetExternalWorldLocationUpdateFunctor(TFunction<void(FNetRefHandle NetHandle, const UObject* ReplicatedObject, FVector& OutLocation, float& OutCullDistance)> LocUpdateFunctor);
 
 	void SetExternalPreUpdateFunctor(TFunction<void(FNetRefHandle, UObject*, const UReplicationBridge*)> PreUpdateFunctor);
 
+	float GetMaxTickRate() const { return Super::GetMaxTickRate(); }
+
 protected:
 	// Type specifics for serializing creation data this will most likely be made into a separate interface to support different types of header data for different types
 	// But if we can avoid having custom data per type for instantiating remote objects as we would like to be able to fully express the state of a replicated object using the define protocol alone. 
diff --git a/Engine/Plugins/Runtime/ReplicationSystemTestPlugin/Source/Private/Tests/ReplicationSystem/ReplicationSystemServerClientTestFixture.cpp b/Engine/Plugins/Runtime/ReplicationSystemTestPlugin/Source/Private/Tests/ReplicationSystem/ReplicationSystemServerClientTestFixture.cpp
--- a/Engine/Plugins/Runtime/ReplicationSystemTestPlugin/Source/Private/Tests/ReplicationSystem/ReplicationSystemServerClientTestFixture.cpp
+++ b/Engine/Plugins/Runtime/ReplicationSystemTestPlugin/Source/Private/Tests/ReplicationSystem/ReplicationSystemServerClientTestFixture.cpp
@@ -326,6 +326,12 @@ uint32 FReplicationSystemTestNode::GetReplicationSystemId() const
 	return ReplicationSystem ? ReplicationSystem->GetId() : uint32(~0U);
 }
 
+float FReplicationSystemTestNode::ConvertPollPeriodIntoFrequency(uint32 PollPeriod) const
+{
+	const float PollFrequency = ReplicationBridge->GetMaxTickRate() / (float)PollPeriod;
+	return PollFrequency;
+}
+
 // FReplicationSystemTestClient implementation
 FReplicationSystemTestClient::FReplicationSystemTestClient(const TCHAR* Name)
 : FReplicationSystemTestNode(false, Name)
diff --git a/Engine/Plugins/Runtime/ReplicationSystemTestPlugin/Source/Private/Tests/ReplicationSystem/ReplicationSystemServerClientTestFixture.h b/Engine/Plugins/Runtime/ReplicationSystemTestPlugin/Source/Private/Tests/ReplicationSystem/ReplicationSystemServerClientTestFixture.h
--- a/Engine/Plugins/Runtime/ReplicationSystemTestPlugin/Source/Private/Tests/ReplicationSystem/ReplicationSystemServerClientTestFixture.h
+++ b/Engine/Plugins/Runtime/ReplicationSystemTestPlugin/Source/Private/Tests/ReplicationSystem/ReplicationSystemServerClientTestFixture.h
@@ -148,6 +148,8 @@ public:
 
 	uint32 GetNetTraceId() const;
 
+	float ConvertPollPeriodIntoFrequency(uint32 PollPeriod) const;
+
 public:
 	UReplicationSystem* ReplicationSystem;
 	UReplicatedTestObjectBridge* ReplicationBridge;
diff --git a/Engine/Plugins/Runtime/ReplicationSystemTestPlugin/Source/Private/Tests/ReplicationSystem/TestDependentObjects.cpp b/Engine/Plugins/Runtime/ReplicationSystemTestPlugin/Source/Private/Tests/ReplicationSystem/TestDependentObjects.cpp
--- a/Engine/Plugins/Runtime/ReplicationSystemTestPlugin/Source/Private/Tests/ReplicationSystem/TestDependentObjects.cpp
+++ b/Engine/Plugins/Runtime/ReplicationSystemTestPlugin/Source/Private/Tests/ReplicationSystem/TestDependentObjects.cpp
@@ -188,8 +188,9 @@ UE_NET_TEST_FIXTURE(FReplicationSystemServerClientTestFixture, TestDependentObje
 
 	// Spawn second object on server that later will be added as a dependent object
 	// With high PollFramePeriod so that it will not replicate in a while unless it is a dependent
-	UTestReplicatedIrisObject* ServerDependentObject = Server->CreateObject(0, 0);
-	Server->ReplicationBridge->SetPollFramePeriod(ServerDependentObject, 255U);
+	UObjectReplicationBridge::FCreateNetRefHandleParams Params = Server->GetReplicationBridge()->DefaultCreateNetRefHandleParams;
+	Params.PollFrequency = Server->ConvertPollPeriodIntoFrequency(255U);
+	UTestReplicatedIrisObject* ServerDependentObject = Server->CreateObject(Params);
 
 	// Send and deliver packet
 	Server->PreSendUpdate();
@@ -254,16 +255,20 @@ UE_NET_TEST_FIXTURE(FReplicationSystemServerClientTestFixture, TestDependentObje
 	FReplicationSystemTestClient* Client = CreateClient();
 
 	// Spawn object on server
-	UTestReplicatedIrisObject* ServerObject0 = Server->CreateObject(0, 0);
-	UTestReplicatedIrisObject* ServerObject1 = Server->CreateObject(0, 0);
+	UObjectReplicationBridge::FCreateNetRefHandleParams Params = Server->GetReplicationBridge()->DefaultCreateNetRefHandleParams;
+
+	// Setup different poll frequencies for the objects
+	Params.PollFrequency = Server->ConvertPollPeriodIntoFrequency(10U);
+	UTestReplicatedIrisObject* ServerObject0 = Server->CreateObject(Params);
+
+	Params.PollFrequency = Server->ConvertPollPeriodIntoFrequency(20U);
+	UTestReplicatedIrisObject* ServerObject1 = Server->CreateObject(Params);
 
 	// Spawn second object on server that later will be added as a dependent object
-	UTestReplicatedIrisObject* ServerDependentObject = Server->CreateObject(0, 0);
+	Params.PollFrequency = Server->ConvertPollPeriodIntoFrequency(40U);
+	UTestReplicatedIrisObject* ServerDependentObject = Server->CreateObject(Params);
+
 
-	// Setup different pollframeperiods for the objects
-	Server->ReplicationBridge->SetPollFramePeriod(ServerObject0, 10);
-	Server->ReplicationBridge->SetPollFramePeriod(ServerObject1, 20);
-	Server->ReplicationBridge->SetPollFramePeriod(ServerDependentObject, 40U);
 
 	// Add dependent object to both server objects
 	Server->ReplicationBridge->AddDependentObject(ServerObject0->NetRefHandle, ServerDependentObject->NetRefHandle);
@@ -329,12 +334,13 @@ UE_NET_TEST_FIXTURE(FReplicationSystemServerClientTestFixture, TestDependentObje
 	FReplicationSystemTestClient* Client = CreateClient();
 
 	// Spawn object on server
-	UTestReplicatedIrisObject* ServerObject = Server->CreateObject(0, 0);
-	Server->ReplicationBridge->SetPollFramePeriod(ServerObject, 14U);
+	UObjectReplicationBridge::FCreateNetRefHandleParams Params = Server->GetReplicationBridge()->DefaultCreateNetRefHandleParams;
+	Params.PollFrequency = Server->ConvertPollPeriodIntoFrequency(14U);
+	UTestReplicatedIrisObject* ServerObject = Server->CreateObject(Params);
 
 	// Spawn second object add it as a dependency and bump poll period
-	UTestReplicatedIrisObject* ServerDependentObject = Server->CreateObject(0, 0);
-	Server->ReplicationBridge->SetPollFramePeriod(ServerDependentObject, 30U);
+	Params.PollFrequency = Server->ConvertPollPeriodIntoFrequency(30U);
+	UTestReplicatedIrisObject* ServerDependentObject = Server->CreateObject(Params);
 	Server->ReplicationBridge->AddDependentObject(ServerObject->NetRefHandle, ServerDependentObject->NetRefHandle);
 
 	// Send and deliver packet
diff --git a/Engine/Plugins/Runtime/ReplicationSystemTestPlugin/Source/Private/Tests/ReplicationSystem/TestDirtyNetObjectTracker.cpp b/Engine/Plugins/Runtime/ReplicationSystemTestPlugin/Source/Private/Tests/ReplicationSystem/TestDirtyNetObjectTracker.cpp
--- a/Engine/Plugins/Runtime/ReplicationSystemTestPlugin/Source/Private/Tests/ReplicationSystem/TestDirtyNetObjectTracker.cpp
+++ b/Engine/Plugins/Runtime/ReplicationSystemTestPlugin/Source/Private/Tests/ReplicationSystem/TestDirtyNetObjectTracker.cpp
@@ -143,7 +143,9 @@ UE_NET_TEST_FIXTURE(FReplicationSystemServerClientTestFixture, GlobalDirtyTracke
 
 	// Spawn object on server that is polled only every 3 frames
 	UObjectReplicationBridge::FCreateNetRefHandleParams Params;
-	Params.PollFramePeriod = 2;
+	const uint32 PollPeriod = 2;
+	const float PollFrequency = Server->ConvertPollPeriodIntoFrequency(PollPeriod);
+	Params.PollFrequency = PollFrequency;
 	Params.bCanReceive = true;
 	Params.bAllowDynamicFilter = true;
 	Params.bNeedsPreUpdate = true;
@@ -186,7 +188,9 @@ UE_NET_TEST_FIXTURE(FReplicationSystemServerClientTestFixture, DirtyInsidePreUpd
 
 	// Spawn object on server that is polled late in order to test ForceNetUpdate
 	UObjectReplicationBridge::FCreateNetRefHandleParams Params;
-	Params.PollFramePeriod = 10;
+	const uint32 PollPeriod = 10;
+	const float PollFrequency = Server->ConvertPollPeriodIntoFrequency(PollPeriod);
+	Params.PollFrequency = PollFrequency;
 	Params.bCanReceive = true;
 	Params.bAllowDynamicFilter = true;	
 	Params.bNeedsPreUpdate = true;
diff --git a/Engine/Plugins/Runtime/ReplicationSystemTestPlugin/Source/Private/Tests/ReplicationSystem/TestObjectPollFrequencyLimiter.cpp b/Engine/Plugins/Runtime/ReplicationSystemTestPlugin/Source/Private/Tests/ReplicationSystem/TestObjectPollFrequencyLimiter.cpp
--- a/Engine/Plugins/Runtime/ReplicationSystemTestPlugin/Source/Private/Tests/ReplicationSystem/TestObjectPollFrequencyLimiter.cpp
+++ b/Engine/Plugins/Runtime/ReplicationSystemTestPlugin/Source/Private/Tests/ReplicationSystem/TestObjectPollFrequencyLimiter.cpp
@@ -119,9 +119,10 @@ UE_NET_TEST_FIXTURE(FReplicationSystemServerClientTestFixture, TestSubObjectsPol
 	FReplicationSystemTestClient* Client = CreateClient();
 
 	// Spawn object on server polled every 3 frames
-	const uint32 PollPeriod = 2;
+	const uint32 PollPeriod = 3;
+	const float PollFrequency = Server->ConvertPollPeriodIntoFrequency(PollPeriod);
 	UObjectReplicationBridge::FCreateNetRefHandleParams Params;
-	Params.PollFramePeriod = PollPeriod;
+	Params.PollFrequency = PollFrequency;
 	Params.bCanReceive = true;
 	Params.bAllowDynamicFilter = true;
 	Params.bNeedsPreUpdate = true;
diff --git a/Engine/Plugins/Runtime/ReplicationSystemTestPlugin/Source/Private/Tests/ReplicationSystem/TestReplicationSystem.cpp b/Engine/Plugins/Runtime/ReplicationSystemTestPlugin/Source/Private/Tests/ReplicationSystem/TestReplicationSystem.cpp
--- a/Engine/Plugins/Runtime/ReplicationSystemTestPlugin/Source/Private/Tests/ReplicationSystem/TestReplicationSystem.cpp
+++ b/Engine/Plugins/Runtime/ReplicationSystemTestPlugin/Source/Private/Tests/ReplicationSystem/TestReplicationSystem.cpp
@@ -1658,9 +1658,10 @@ UE_NET_TEST_FIXTURE(FReplicationSystemServerClientTestFixture, TestObjectPollFra
 	// Spawn object on server
 	UTestReplicatedIrisObject* ServerObject = Server->CreateObject(0, 0);
 
-	// Spawn second object on server that later will be addedded as a dependent object
-	UTestReplicatedIrisObject* ServerObjectPolledEveryOtherFrame = Server->CreateObject(0, 0);
-	Server->ReplicationBridge->SetPollFramePeriod(ServerObjectPolledEveryOtherFrame, 1U);
+	// Spawn second object on server that later will be added as a dependent object
+	UObjectReplicationBridge::FCreateNetRefHandleParams Params = Server->GetReplicationBridge()->DefaultCreateNetRefHandleParams;
+	Params.PollFrequency = Server->ConvertPollPeriodIntoFrequency(1U);
+	UTestReplicatedIrisObject* ServerObjectPolledEveryOtherFrame = Server->CreateObject(Params);
 
 	// Send and deliver packet
 	Server->PreSendUpdate();
diff --git a/Engine/Source/Runtime/Engine/Private/Net/Iris/ReplicationSystem/ActorReplicationBridge.cpp b/Engine/Source/Runtime/Engine/Private/Net/Iris/ReplicationSystem/ActorReplicationBridge.cpp
--- a/Engine/Source/Runtime/Engine/Private/Net/Iris/ReplicationSystem/ActorReplicationBridge.cpp
+++ b/Engine/Source/Runtime/Engine/Private/Net/Iris/ReplicationSystem/ActorReplicationBridge.cpp
@@ -53,18 +53,6 @@ bool ShouldUseIrisReplication(const UObject* Object)
 
 namespace UE::Net::Private
 {
-// This is not pretty but the alternative is exposing the PollFrequencyMultiplier in a public header file.
-// It's really only useful for this particular class and in order to ensure
-// NetUpdateFrequencies are ok.
-extern IRISCORE_API float PollFrequencyMultiplier;
-extern IRISCORE_API const int MaxPollFramePeriod;
-
-static bool bIrisEnableLowNetUpdateFrequencyEnsure = false;
-static FAutoConsoleVariableRef CVarEnableLowNetUpdateFrequencyEnsure(
-		TEXT("net.Iris.EnableLowNetUpdateFrequencyEnsure"),
-		bIrisEnableLowNetUpdateFrequencyEnsure,
-		TEXT("Whether to ensure when NetUpdateFrequency is so low that the poll frame period is clamped. Default is true.")
-		);
 
 bool IsActorValidForIrisReplication(const AActor* Actor)
 {
@@ -94,12 +82,11 @@ void ActorReplicationBridgeGetActorWorldObjectInfo(FNetRefHandle Handle, const U
 	}
 }
 
-}
+} // end namespace UE::Net::Private
 
 UActorReplicationBridge::UActorReplicationBridge()
 : UObjectReplicationBridge()
 , NetDriver(nullptr)
-, MaxPollFrequency(0.0f)
 , ObjectReferencePackageMap(nullptr)
 , SpawnInfoFlags(0U)
 {
@@ -257,15 +244,9 @@ UE::Net::FNetRefHandle UActorReplicationBridge::BeginReplication(AActor* Actor,
 	CreateNetRefHandleParams.bNeedsPreUpdate = 1U;
 	CreateNetRefHandleParams.bNeedsWorldLocationUpdate = 1U;
 	CreateNetRefHandleParams.StaticPriority = (Actor->bAlwaysRelevant || Actor->bOnlyRelevantToOwner) ? Actor->NetPriority : 0.0f;
-	const uint32 UnclampedPollPeriod = GetPollFramePeriod(Actor->NetUpdateFrequency);
-	const uint32 ClampedPollPeriod = FMath::Clamp<uint32>(UnclampedPollPeriod, 1U, static_cast<uint32>(std::numeric_limits<uint8>::max()) + 1U);
-	CreateNetRefHandleParams.PollFramePeriod = (ClampedPollPeriod - 1U) & 255U;
+	CreateNetRefHandleParams.PollFrequency = Actor->NetUpdateFrequency;
+
 #if !UE_BUILD_SHIPPING
-	if ((ClampedPollPeriod < UnclampedPollPeriod))
-	{
-		ensureAlwaysMsgf(!UE::Net::Private::bIrisEnableLowNetUpdateFrequencyEnsure, TEXT("Very low NetUpdateFrequency %f for Actor %s. Suggest setting it to %f or higher."), Actor->NetUpdateFrequency, ToCStr(Actor->GetName()), GetMinSupportedNetUpdateFrequency());
-		UE_CLOG(!UE::Net::Private::bIrisEnableLowNetUpdateFrequencyEnsure, LogIrisBridge, Warning, TEXT("Very low NetUpdateFrequency %f for Actor %s. Suggest setting it to %f or higher."), Actor->NetUpdateFrequency, ToCStr(Actor->GetName()), GetMinSupportedNetUpdateFrequency());
-	}
 	ensureAlwaysMsgf(!(Actor->bAlwaysRelevant || Actor->bOnlyRelevantToOwner) || CreateNetRefHandleParams.StaticPriority >= 1.0f, TEXT("Very low NetPriority %.02f for always relevant or owner relevant Actor %s. Set it to 1.0f or higher."), Actor->NetPriority, ToCStr(Actor->GetName()));
 #endif
 
@@ -901,7 +882,7 @@ void UActorReplicationBridge::SetNetDriver(UNetDriver* const InNetDriver)
 	NetDriver = InNetDriver;
 	if (InNetDriver != nullptr)
 	{
-		MaxPollFrequency = static_cast<float>(FPlatformMath::Max(InNetDriver->GetNetServerMaxTickRate(), 0));
+		SetMaxTickRate(static_cast<float>(FPlatformMath::Max(InNetDriver->GetNetServerMaxTickRate(), 0)));
 
 		InNetDriver->OnNetServerMaxTickRateChanged.AddUObject(this, &UActorReplicationBridge::OnMaxTickRateChanged);
 
@@ -921,9 +902,9 @@ void UActorReplicationBridge::SetNetDriver(UNetDriver* const InNetDriver)
 
 void UActorReplicationBridge::OnMaxTickRateChanged(UNetDriver* InNetDriver, int32 NewMaxTickRate, int32 OldMaxTickRate)
 {
-	MaxPollFrequency = static_cast<float>(FPlatformMath::Max(InNetDriver->GetNetServerMaxTickRate(), 0));
+	SetMaxTickRate(static_cast<float>(FPlatformMath::Max(InNetDriver->GetNetServerMaxTickRate(), 0)));
 
-	//TODO: Reset poll frequencies for all objects
+	ReinitPollFrequency();
 }
 
 void UActorReplicationBridge::GetActorCreationHeader(const AActor* Actor, UE::Net::Private::FActorCreationHeader& Header) const
@@ -1086,17 +1067,17 @@ UActorReplicationBridge* UActorReplicationBridge::Create(UNetDriver* NetDriver)
 	return Bridge;
 }
 
-uint32 UActorReplicationBridge::GetPollFramePeriod(float PollFrequency) const
+float UActorReplicationBridge::GetPollFrequencyOfRootObject(const UObject* ReplicatedObject) const
 {
-	const uint32 FramesBetweenUpdatesForObject = static_cast<uint32>(MaxPollFrequency/FPlatformMath::Max(0.001f, UE::Net::Private::PollFrequencyMultiplier*PollFrequency));
-	return FramesBetweenUpdatesForObject;
+	const AActor* ReplicatedActor = CastChecked<AActor>(ReplicatedObject);
+	float PollFrequency = ReplicatedActor->NetUpdateFrequency;
+	GetClassPollFrequency(ReplicatedActor->GetClass(), PollFrequency);
+	return PollFrequency;
 }
 
-float UActorReplicationBridge::GetMinSupportedNetUpdateFrequency() const
-{
-	return MaxPollFrequency/(UE::Net::Private::MaxPollFramePeriod*UE::Net::Private::PollFrequencyMultiplier);
-}
-#else
+#else //!UE_WITH_IRIS
+
 UActorReplicationBridge::UActorReplicationBridge() = default;
 UActorReplicationBridge::~UActorReplicationBridge() = default;
+
 #endif
diff --git a/Engine/Source/Runtime/Engine/Public/Net/Iris/ReplicationSystem/ActorReplicationBridge.h b/Engine/Source/Runtime/Engine/Public/Net/Iris/ReplicationSystem/ActorReplicationBridge.h
--- a/Engine/Source/Runtime/Engine/Public/Net/Iris/ReplicationSystem/ActorReplicationBridge.h
+++ b/Engine/Source/Runtime/Engine/Public/Net/Iris/ReplicationSystem/ActorReplicationBridge.h
@@ -93,19 +93,16 @@ protected:
 	virtual bool ObjectLevelHasFinishedLoading(UObject* Object) const override;
 	virtual bool IsAllowedToDestroyInstance(const UObject* Instance) const override;
 
+	virtual float GetPollFrequencyOfRootObject(const UObject* ReplicatedObject) const override;
+
 private:
 	void GetActorCreationHeader(const AActor* Actor, UE::Net::Private::FActorCreationHeader& Header) const;
 	void GetSubObjectCreationHeader(const UObject* Object, UE::Net::Private::FSubObjectCreationHeader& Header) const;
 
-	uint32 GetPollFramePeriod(float PollFrequency) const;
-	float GetMinSupportedNetUpdateFrequency() const;
-
 	void OnMaxTickRateChanged(UNetDriver* InNetDriver, int32 NewMaxTickRate, int32 OldMaxTickRate);
 
 private:
-
 	UNetDriver* NetDriver;
-	float MaxPollFrequency;
 
 	UIrisObjectReferencePackageMap* ObjectReferencePackageMap;
 
diff --git a/Engine/Source/Runtime/Experimental/Iris/Core/Private/Iris/ReplicationSystem/ObjectPollFrequencyLimiter.cpp b/Engine/Source/Runtime/Experimental/Iris/Core/Private/Iris/ReplicationSystem/ObjectPollFrequencyLimiter.cpp
--- a/Engine/Source/Runtime/Experimental/Iris/Core/Private/Iris/ReplicationSystem/ObjectPollFrequencyLimiter.cpp
+++ b/Engine/Source/Runtime/Experimental/Iris/Core/Private/Iris/ReplicationSystem/ObjectPollFrequencyLimiter.cpp
@@ -8,16 +8,6 @@
 namespace UE::Net::Private
 {
 
-// Storing poll frame period as uint8
-const int MaxPollFramePeriod = 256;
-float PollFrequencyMultiplier = 1.0f;
-
-static FAutoConsoleVariableRef CVarPollFrequencyMultiplier(
-		TEXT("net.Iris.PollFrequencyMultiplier"),
-		PollFrequencyMultiplier,
-		TEXT("Multiplied with the NetUpdateFrequency to decide how often PreReplication is called on an Actor. Default factor is 1.0.")
-		);
-
 FObjectPollFrequencyLimiter::FObjectPollFrequencyLimiter()
 {
 }
diff --git a/Engine/Source/Runtime/Experimental/Iris/Core/Private/Iris/ReplicationSystem/ObjectPollFrequencyLimiter.h b/Engine/Source/Runtime/Experimental/Iris/Core/Private/Iris/ReplicationSystem/ObjectPollFrequencyLimiter.h
--- a/Engine/Source/Runtime/Experimental/Iris/Core/Private/Iris/ReplicationSystem/ObjectPollFrequencyLimiter.h
+++ b/Engine/Source/Runtime/Experimental/Iris/Core/Private/Iris/ReplicationSystem/ObjectPollFrequencyLimiter.h
@@ -14,9 +14,6 @@ namespace UE::Net::Private
 namespace UE::Net::Private
 {
 
-extern IRISCORE_API float PollFrequencyMultiplier;
-extern IRISCORE_API const int MaxPollFramePeriod;
-
 class FObjectPollFrequencyLimiter
 {
 public:
@@ -35,8 +32,13 @@ public:
 	*/
 	void Update(const FNetBitArrayView& RelevantObjects, const FNetBitArrayView& DirtyObjects, FNetBitArrayView& OutObjectsToPoll);
 
+	/** We use a uint8 to track frames, so the limit is 255 frames.*/
+	static constexpr uint32 GetMaxPollingFrames()
+	{		
+		return static_cast<uint32>(std::numeric_limits<uint8>::max());
+	}
+
 private:
-	uint32 GetPollFramePeriodForFrequency(float PollFrequency) const;
 
 	uint32 MaxInternalHandle = 0;
 	uint32 FrameIndex = 0;
diff --git a/Engine/Source/Runtime/Experimental/Iris/Core/Private/Iris/ReplicationSystem/ObjectReplicationBridge.cpp b/Engine/Source/Runtime/Experimental/Iris/Core/Private/Iris/ReplicationSystem/ObjectReplicationBridge.cpp
--- a/Engine/Source/Runtime/Experimental/Iris/Core/Private/Iris/ReplicationSystem/ObjectReplicationBridge.cpp
+++ b/Engine/Source/Runtime/Experimental/Iris/Core/Private/Iris/ReplicationSystem/ObjectReplicationBridge.cpp
@@ -307,9 +307,11 @@ UE::Net::FNetRefHandle UObjectReplicationBridge::BeginReplication(UObject* Insta
 			ObjectsWithObjectReferences.SetBitValue(InternalReplicationIndex, EnumHasAnyFlags(InstanceProtocol->InstanceTraits, EReplicationInstanceProtocolTraits::HasObjectReference));
 
 			// Set poll frame period
-			uint8 PollFramePeriod = Params.PollFramePeriod;
-			FindPollInfo(Instance->GetClass(), PollFramePeriod);
-			SetPollFramePeriod(InternalReplicationIndex, PollFramePeriod);
+			float PollFrequency = Params.PollFrequency;
+			FindOrCachePollFrequency(Instance->GetClass(), PollFrequency);
+			
+			uint8 PollFramePeriod = ConvertPollFrequencyIntoFrames(Params.PollFrequency);
+			PollFrequencyLimiter->SetPollFramePeriod(InternalReplicationIndex, PollFramePeriod);
 
 			UE_LOG_OBJECTREPLICATIONBRIDGE(Verbose, TEXT("BeginReplication Created %s with ProtocolId:0x%" UINT64_x_FMT " for Object named %s"), *RefHandle.ToString(), ReplicationProtocol->ProtocolIdentifier, *Instance->GetName());
 
@@ -913,11 +915,6 @@ void UObjectReplicationBridge::UpdateInstancesWorldLocation()
 
 }
 
-void UObjectReplicationBridge::SetPollFramePeriod(UE::Net::Private::FInternalNetRefIndex InternalReplicationIndex, uint8 PollFramePeriod)
-{
-	PollFrequencyLimiter->SetPollFramePeriod(InternalReplicationIndex, PollFramePeriod);
-}
-
 void UObjectReplicationBridge::SetPollWithObject(FNetRefHandle ObjectToPollWithHandle, FNetRefHandle ObjectHandle)
 {
 	const UE::Net::Private::FInternalNetRefIndex PollWithInternalReplicationIndex = NetRefHandleManager->GetInternalIndex(ObjectToPollWithHandle);
@@ -991,7 +988,67 @@ void UObjectReplicationBridge::SetNetPushIdOnInstance(UE::Net::FReplicationInsta
 #endif
 }
 
-bool UObjectReplicationBridge::FindPollInfo(const UClass* Class, uint8& OutPollPeriod)
+bool UObjectReplicationBridge::GetClassPollFrequency(const UClass* Class, float& OutPollFrequency) const
+{
+	if (!(bAllowPollPeriodOverrides & bHasPollOverrides))
+	{
+		return false;
+	}
+
+	const FName ClassName = Class->GetFName();
+	if (const FPollInfo* PollInfo = ClassesWithPollPeriodOverride.Find(ClassName))
+	{
+		OutPollFrequency = PollInfo->PollFrequency;
+		return true;
+	}
+
+	if (ClassesWithoutPollPeriodOverride.Find(ClassName))
+	{
+		return false;
+	}
+
+	bool bFoundOverride = false;
+	const UClass* SuperclassWithPollInfo = nullptr;
+	FPollInfo SuperclassPollInfo;
+	for (const auto& ClassNameAndPollInfo : ClassHierarchyPollPeriodOverrides)
+	{
+		const UClass* ClassWithPollInfo = ClassNameAndPollInfo.Value.Class.Get();
+		if (ClassWithPollInfo == nullptr)
+		{
+			continue;
+		}
+
+		if (Class->IsChildOf(ClassWithPollInfo))
+		{
+			// If we've already found a superclass with a config, see which one is closer in the hierarchy.
+			if (SuperclassWithPollInfo != nullptr)
+			{
+				if (ClassWithPollInfo->IsChildOf(SuperclassWithPollInfo))
+				{
+					bFoundOverride = true;
+					SuperclassWithPollInfo = ClassWithPollInfo;
+					SuperclassPollInfo = ClassNameAndPollInfo.Value;
+				}
+			}
+			else
+			{
+				bFoundOverride = true;
+				SuperclassWithPollInfo = ClassWithPollInfo;
+				SuperclassPollInfo = ClassNameAndPollInfo.Value;
+			}
+		}
+	}
+
+	if (bFoundOverride)
+	{
+		OutPollFrequency = SuperclassPollInfo.PollFrequency;
+		return true;
+	}
+
+	return false;
+}
+
+bool UObjectReplicationBridge::FindOrCachePollFrequency(const UClass* Class, float& OutPollFrequency)
 {
 	if (!(bAllowPollPeriodOverrides & bHasPollOverrides))
 	{
@@ -1001,7 +1058,7 @@ bool UObjectReplicationBridge::FindPollInfo(const UClass* Class, uint8& OutPollP
 	const FName ClassName = Class->GetFName();
 	if (const FPollInfo* PollInfo = ClassesWithPollPeriodOverride.Find(ClassName))
 	{
-		OutPollPeriod = PollInfo->PollFramePeriod;
+		OutPollFrequency = PollInfo->PollFrequency;
 		return true;
 	}
 
@@ -1069,7 +1126,7 @@ bool UObjectReplicationBridge::FindPollInfo(const UClass* Class, uint8& OutPollP
 			ClassesWithPollPeriodOverride.FindOrAdd(ClassToAddName, SuperclassPollInfo);
 		}
 
-		OutPollPeriod = SuperclassPollInfo.PollFramePeriod;
+		OutPollFrequency = SuperclassPollInfo.PollFrequency;
 		return true;
 	}
 	else
@@ -1246,7 +1303,7 @@ void UObjectReplicationBridge::LoadConfig()
 		bHasPollOverrides = true;
 
 		FPollInfo PollInfo;
-		PollInfo.PollFramePeriod = FPlatformMath::Min(PollOverride.PollFramePeriod, 255U) & 255U;
+		PollInfo.PollFrequency = FPlatformMath::Max(PollOverride.PollFrequency, 0.0f);
 		if (PollOverride.bIncludeSubclasses)
 		{
 			bHasDirtyClassesInPollPeriodOverrides = true;
@@ -1371,3 +1428,53 @@ void UObjectReplicationBridge::InitConditionalPropertyDelegates()
 		}
 	});
 }
+
+uint8 UObjectReplicationBridge::ConvertPollFrequencyIntoFrames(float PollFrequency) const
+{
+	if (PollFrequency <= 0.0f)
+	{
+		return 0U;
+	}
+
+	const uint32 FramesBetweenUpdatesForObject = static_cast<uint32>(MaxTickRate / FPlatformMath::Max(0.001f, PollFrequency));
+	return static_cast<uint8>(FMath::Clamp<uint32>(FramesBetweenUpdatesForObject, 0U, UE::Net::Private::FObjectPollFrequencyLimiter::GetMaxPollingFrames()));
+}
+
+float UObjectReplicationBridge::GetPollFrequencyOfRootObject(const UObject* ReplicatedObject) const
+{
+	float PollFrequency = 0.0f;
+	GetClassPollFrequency(ReplicatedObject->GetClass(), PollFrequency);
+
+	return PollFrequency;
+}
+
+void UObjectReplicationBridge::ReinitPollFrequency()
+{
+	using namespace UE::Net;
+	using namespace UE::Net::Private;
+
+	FReplicationSystemInternal* ReplicationSystemInternal = GetReplicationSystem()->GetReplicationSystemInternal();
+	const FNetRefHandleManager& LocalNetRefHandleManager = ReplicationSystemInternal->GetNetRefHandleManager();
+
+	auto UpdatePollFrequency = [this, &LocalNetRefHandleManager](uint32 RootObjectIndex)
+	{
+		if (UObject* RootObjectInstance = LocalNetRefHandleManager.GetReplicatedObjectInstance(RootObjectIndex))
+		{
+			const float PollFrequency = GetPollFrequencyOfRootObject(RootObjectInstance);
+			const uint8 PollFramePeriod = ConvertPollFrequencyIntoFrames(PollFrequency);
+
+			PollFrequencyLimiter->SetPollFramePeriod(RootObjectIndex, PollFramePeriod);
+
+			// Set the subobjects of the object
+			for (const FInternalNetRefIndex SubObjectIndex : LocalNetRefHandleManager.GetSubObjects(RootObjectIndex))
+			{
+				PollFrequencyLimiter->SetPollFramePeriod(SubObjectIndex, PollFramePeriod);
+			}
+		}
+	};
+
+	const FNetBitArrayView RootObjects = LocalNetRefHandleManager.GetScopableInternalIndicesView();
+	const FNetBitArrayView SubObjects = MakeNetBitArrayView(LocalNetRefHandleManager.GetSubObjectInternalIndices());
+
+	FNetBitArrayView::ForAllSetBits(RootObjects, SubObjects, FNetBitArrayView::AndNotOp, UpdatePollFrequency);
+}
diff --git a/Engine/Source/Runtime/Experimental/Iris/Core/Public/Iris/ReplicationSystem/ObjectReplicationBridge.h b/Engine/Source/Runtime/Experimental/Iris/Core/Public/Iris/ReplicationSystem/ObjectReplicationBridge.h
--- a/Engine/Source/Runtime/Experimental/Iris/Core/Public/Iris/ReplicationSystem/ObjectReplicationBridge.h
+++ b/Engine/Source/Runtime/Experimental/Iris/Core/Public/Iris/ReplicationSystem/ObjectReplicationBridge.h
@@ -53,16 +53,16 @@ public:
 		bool bAllowDynamicFilter = false;
 
 		/**
-		  * If StaticPriority is > 0 the ReplicationSystem will use that as priority when scheduling objects.
-		  * If it's <= 0.0f one will look for a world location support and then use the default spatial prioritizer.
-		  */
+		 * If StaticPriority is > 0 the ReplicationSystem will use that as priority when scheduling objects. 
+		 * If it's <= 0.0f one will look for a world location support and then use the default spatial prioritizer.
+		 */
 		float StaticPriority = 0.0f;
 
 		/**
-		  * How often the object should be polled for dirtiness, including calling the InstancePreUpdate function.
-		  * The period is zero based so 0 means poll every frame and 255 means poll every 256th frame.
-		  */
-		uint8 PollFramePeriod = 0;
+		 * How often per second the object should be polled for dirtiness, including calling the InstancePreUpdate function. 
+		 * When set to zero it will be polled every frame.
+		 */
+		float PollFrequency = 0.0f;
 	};
 
 	IRISCORE_API static FCreateNetRefHandleParams DefaultCreateNetRefHandleParams;
@@ -207,12 +207,18 @@ protected:
 	
 	// Poll frequency support
 
-	/** Set poll frame period for a specific object. Zero based period where 0 means every frame and 255 means every 256th frame. */
-	IRISCORE_API void SetPollFramePeriod(UE::Net::Private::FInternalNetRefIndex InternalReplicationIndex, uint8 FramePeriod);
-
 	/** Force polling of Object when ObjectToPollWith is polled. */
 	IRISCORE_API void SetPollWithObject(FNetRefHandle ObjectToPollWith, FNetRefHandle Object);
 
+	/** Transform a poll frequency (in updates per second) to an equivalent number of frames. */
+	IRISCORE_API uint8 ConvertPollFrequencyIntoFrames(float PollFrequency) const;
+
+	/** Returns the poll frequency of a specific root object. */
+	IRISCORE_API virtual float GetPollFrequencyOfRootObject(const UObject* ReplicatedObject) const;
+
+	/** Re-initialize the poll frequency of all replicated root objects. */
+	IRISCORE_API void ReinitPollFrequency();
+
 	/** Re-initialize config-driven parameters found in ObjectReplicationBridgeConfig. */
 	IRISCORE_API void LoadConfig();
 
@@ -225,6 +231,22 @@ protected:
 	/** Set the function used to determine whether two classes are to be considered equal when it comes to filtering. Used on subclasses. */
 	IRISCORE_API void SetShouldSubclassUseSameFilterFunction(TFunction<bool(const UClass* Class, const UClass* Subclass)>);
 
+	/** Finds the final poll frequency for a class and cache it for future lookups. OutPollPeriod will only be modified if this method returns true. */
+	IRISCORE_API bool FindOrCachePollFrequency(const UClass* Class, float& OutPollPeriod);
+
+	/**
+	 * Find the poll period of a class if it was configured with an override. 
+	 * Use this only if all class configs have been properly cached.
+	 * OutPollPeriod will only be modified if this method returns true.
+	 */
+	IRISCORE_API bool GetClassPollFrequency(const UClass* Class, float& OutPollPeriod) const;
+
+	/** Current max tick rate set by the engine */
+	float GetMaxTickRate() const { return MaxTickRate; }
+
+	/** Change the max tick rate to match the one from the engine */
+	void SetMaxTickRate(float InMaxTickRate) { MaxTickRate = InMaxTickRate; }
+
 private:
 
 	/** Forcibly poll a single replicated object */
@@ -240,8 +262,6 @@ private:
 
 	void SetNetPushIdOnInstance(UE::Net::FReplicationInstanceProtocol* InstanceProtocol, FNetHandle NetHandle);
 
-	/** Finds a poll period override if it exists. OutPollPeriod will only be modified if this method returns true. */
-	bool FindPollInfo(const UClass* Class, uint8& OutPollPeriod);
 
 	/** Tries to load the classes used in poll period overrides. */
 	void FindClassesInPollPeriodOverrides();
@@ -277,13 +297,18 @@ private:
 	};
 
 	// Polling
+
+	/** The maximum tick per second of the engine. Default is to use 30hz */
+	float MaxTickRate = 30.0f;
+
 	struct FPollInfo
 	{
-		uint8 PollFramePeriod = 0;
+		float PollFrequency = 0.0f;
 		TWeakObjectPtr<const UClass> Class;
 	};
 	UE::Net::Private::FObjectPollFrequencyLimiter* PollFrequencyLimiter;
 
+	//$IRIS TODO: The poll class config management code should be moved into it's own class. Maybe in a class that handles any type of per-class settings.
 	// Class hierarchies with poll period overrides
 	TMap<FName, FPollInfo> ClassHierarchyPollPeriodOverrides;
 	// Exact classes with poll period overrides
diff --git a/Engine/Source/Runtime/Experimental/Iris/Core/Public/Iris/ReplicationSystem/ObjectReplicationBridgeConfig.h b/Engine/Source/Runtime/Experimental/Iris/Core/Public/Iris/ReplicationSystem/ObjectReplicationBridgeConfig.h
--- a/Engine/Source/Runtime/Experimental/Iris/Core/Public/Iris/ReplicationSystem/ObjectReplicationBridgeConfig.h
+++ b/Engine/Source/Runtime/Experimental/Iris/Core/Public/Iris/ReplicationSystem/ObjectReplicationBridgeConfig.h
@@ -20,12 +20,14 @@ struct FObjectReplicationBridgePollConfig
 	FName ClassName;
 
 	/**
-	 * How many frames between polling for modified replicated properties.
-	 * The value will be clamped to what the runtime supports, currently [0,255].
-	 * 0 means poll every frame and 255 means poll every 256th frame.
+	 * How many times per second should we poll for modified replicated properties.
+	 * The value will be converted into a frame count based on the current TickRate up to a maximum of 255 frames
+	 * This means the slowest poll frequency is 255*MaxTickRate (ex: 8.5secs at 30hz)
+	 * If set to 0 it means we poll the object every frame.
 	 */
 	UPROPERTY()
-	uint32 PollFramePeriod = 0;
+	float PollFrequency = 0.0f;
+	
 
 	/** Whether instances of subclasses should also use this poll period. */
 	UPROPERTY()
