Iris:
* Make sure objects polled for the first time always ignore pushmodel dirtyness

#rnx

diff --git a/Engine/Source/Runtime/Experimental/Iris/Core/Private/Iris/ReplicationSystem/NetRefHandleManager.cpp b/Engine/Source/Runtime/Experimental/Iris/Core/Private/Iris/ReplicationSystem/NetRefHandleManager.cpp
--- a/Engine/Source/Runtime/Experimental/Iris/Core/Private/Iris/ReplicationSystem/NetRefHandleManager.cpp
+++ b/Engine/Source/Runtime/Experimental/Iris/Core/Private/Iris/ReplicationSystem/NetRefHandleManager.cpp
@@ -119,6 +119,9 @@ FInternalNetRefIndex FNetRefHandleManager::InternalCreateNetObject(const FNetRef
 		// Need a full copy if set, normally only needed for new objects.
 		Data.bNeedsFullCopyAndQuantize = 1U;
 
+		// Make sure we do a full poll of all properties the first time the object gets polled.
+		Data.bWantsFullPoll = 1U;
+
 		ReplicatedObjectRefCount[InternalIndex] = 0;
 
 		return InternalIndex;
diff --git a/Engine/Source/Runtime/Experimental/Iris/Core/Private/Iris/ReplicationSystem/NetRefHandleManager.h b/Engine/Source/Runtime/Experimental/Iris/Core/Private/Iris/ReplicationSystem/NetRefHandleManager.h
--- a/Engine/Source/Runtime/Experimental/Iris/Core/Private/Iris/ReplicationSystem/NetRefHandleManager.h
+++ b/Engine/Source/Runtime/Experimental/Iris/Core/Private/Iris/ReplicationSystem/NetRefHandleManager.h
@@ -98,8 +98,7 @@ public:
 				uint32 bHasDependentObjects : 1U;
 				uint32 bAllowDestroyInstanceFromRemote : 1U;
 				uint32 bNeedsFullCopyAndQuantize : 1U;
-				// Padding. Adjust when adding or removing flags.
-				uint32 Padding : 25U;
+				uint32 bWantsFullPoll : 1U;
 			};
 		};
 	
diff --git a/Engine/Source/Runtime/Experimental/Iris/Core/Private/Iris/ReplicationSystem/Polling/ObjectPoller.cpp b/Engine/Source/Runtime/Experimental/Iris/Core/Private/Iris/ReplicationSystem/Polling/ObjectPoller.cpp
--- a/Engine/Source/Runtime/Experimental/Iris/Core/Private/Iris/ReplicationSystem/Polling/ObjectPoller.cpp
+++ b/Engine/Source/Runtime/Experimental/Iris/Core/Private/Iris/ReplicationSystem/Polling/ObjectPoller.cpp
@@ -23,7 +23,6 @@ FObjectPoller::FObjectPoller(const FInitParams& InitParams)
 	, LocalNetRefHandleManager(ReplicationSystemInternal->GetNetRefHandleManager())
 	, ReplicatedInstances(LocalNetRefHandleManager.GetReplicatedInstances())
 	, AccumulatedDirtyObjects(ReplicationSystemInternal->GetDirtyNetObjectTracker().GetAccumulatedDirtyNetObjects())
-	, PrevFrameGlobalScopeObjects(MakeNetBitArrayView(LocalNetRefHandleManager.GetPrevFrameScopableInternalIndices()))
 {
 	GarbageCollectionAffectedObjects = MakeNetBitArrayView(ObjectReplicationBridge->GarbageCollectionAffectedObjects);
 
@@ -78,7 +77,7 @@ void FObjectPoller::PollSingleObject(FNetRefHandle Handle)
 
 void FObjectPoller::ForcePollObject(FInternalNetRefIndex ObjectIndex)
 {
-	const FNetRefHandleManager::FReplicatedObjectData& ObjectData = LocalNetRefHandleManager.GetReplicatedObjectDataNoCheck(ObjectIndex);
+	FNetRefHandleManager::FReplicatedObjectData& ObjectData = LocalNetRefHandleManager.GetReplicatedObjectDataNoCheck(ObjectIndex);
 	if (ObjectData.InstanceProtocol == nullptr)
 	{
 		return;
@@ -86,6 +85,9 @@ void FObjectPoller::ForcePollObject(FInternalNetRefIndex ObjectIndex)
 
 	IRIS_PROFILER_PROTOCOL_NAME(ObjectData.Protocol->DebugName->Name);
 
+	// We always poll all states here.
+	ObjectData.bWantsFullPoll = 0U;
+
 	// Call per-instance PreUpdate function
 	if (ObjectReplicationBridge->PreUpdateInstanceFunction && EnumHasAnyFlags(ObjectData.InstanceProtocol->InstanceTraits, EReplicationInstanceProtocolTraits::NeedsPreSendUpdate))
 	{
@@ -117,7 +119,7 @@ void FObjectPoller::ForcePollObject(FInternalNetRefIndex ObjectIndex)
 
 void FObjectPoller::PushModelPollObject(FInternalNetRefIndex ObjectIndex)
 {
-	const FNetRefHandleManager::FReplicatedObjectData& ObjectData = LocalNetRefHandleManager.GetReplicatedObjectDataNoCheck(ObjectIndex);
+	FNetRefHandleManager::FReplicatedObjectData& ObjectData = LocalNetRefHandleManager.GetReplicatedObjectDataNoCheck(ObjectIndex);
 	if (ObjectData.InstanceProtocol == nullptr)
 	{
 		return;
@@ -165,15 +167,15 @@ void FObjectPoller::PushModelPollObject(FInternalNetRefIndex ObjectIndex)
 
 	IRIS_PROFILER_SCOPE_VERBOSE(PollPushBased);
 
-	const bool bIsNewInScope = !PrevFrameGlobalScopeObjects.GetBit(ObjectIndex);
-
+	// Does the object need to poll all states once.
+	const bool bWantsFullPoll = ObjectData.bWantsFullPoll;
+	ObjectData.bWantsFullPoll = 0U;
 	
-
 	// If the object is fully push model we only need to poll it if it's dirty, unless it's a new object or was garbage collected.
 	bool bWasMarkedDirty = false;
 	if (EnumHasAnyFlags(InstanceTraits, EReplicationInstanceProtocolTraits::HasFullPushBasedDirtiness))
 	{
-		if (bIsDirtyObject || bIsNewInScope)
+		if (bIsDirtyObject || bWantsFullPoll)
 		{
 			// We need to do a poll if object is marked as dirty
 			EReplicationFragmentPollFlags PollOptions = EReplicationFragmentPollFlags::PollAllState;
@@ -194,7 +196,7 @@ void FObjectPoller::PushModelPollObject(FInternalNetRefIndex ObjectIndex)
 		// If the object has pushed based fragments, and not is marked dirty and object is affected by GC we need to make sure that we refresh cached references for all push based fragments
 		const bool bIsPushBasedObject = EnumHasAnyFlags(InstanceTraits, EReplicationInstanceProtocolTraits::HasPartialPushBasedDirtiness | EReplicationInstanceProtocolTraits::HasFullPushBasedDirtiness);
 		const bool bHasObjectReferences = EnumHasAnyFlags(InstanceTraits, EReplicationInstanceProtocolTraits::HasObjectReference);
-		const bool bNeedsRefreshOfCachedObjectReferences = ((!(bIsNewInScope | bIsDirtyObject)) & bIsGCAffectedObject & bIsPushBasedObject & bHasObjectReferences);
+		const bool bNeedsRefreshOfCachedObjectReferences = ((!(bWantsFullPoll | bIsDirtyObject)) & bIsGCAffectedObject & bIsPushBasedObject & bHasObjectReferences);
 		if (bNeedsRefreshOfCachedObjectReferences)
 		{
 			// Only states which has push based dirtiness need to be updated as the other states will be polled in full anyway.
@@ -208,7 +210,7 @@ void FObjectPoller::PushModelPollObject(FInternalNetRefIndex ObjectIndex)
 		PollOptions |= bIsGCAffectedObject ? EReplicationFragmentPollFlags::ForceRefreshCachedObjectReferencesAfterGC : EReplicationFragmentPollFlags::None;
 
 		// If the object is not new or dirty at this point we only need to poll non-push based fragments as we know that pushed based states have not been modified
-		const EReplicationFragmentTraits ExcludeTraits = (bIsDirtyObject | bIsNewInScope) ? EReplicationFragmentTraits::None : EReplicationFragmentTraits::HasPushBasedDirtiness;
+		const EReplicationFragmentTraits ExcludeTraits = (bIsDirtyObject || bWantsFullPoll) ? EReplicationFragmentTraits::None : EReplicationFragmentTraits::HasPushBasedDirtiness;
 		bWasMarkedDirty |= FReplicationInstanceOperations::PollAndRefreshCachedPropertyData(InstanceProtocol, ExcludeTraits, PollOptions);
 		++PollStats.PolledObjectCount;
 	}
diff --git a/Engine/Source/Runtime/Experimental/Iris/Core/Private/Iris/ReplicationSystem/Polling/ObjectPoller.h b/Engine/Source/Runtime/Experimental/Iris/Core/Private/Iris/ReplicationSystem/Polling/ObjectPoller.h
--- a/Engine/Source/Runtime/Experimental/Iris/Core/Private/Iris/ReplicationSystem/Polling/ObjectPoller.h
+++ b/Engine/Source/Runtime/Experimental/Iris/Core/Private/Iris/ReplicationSystem/Polling/ObjectPoller.h
@@ -67,18 +67,15 @@ private:
 	UObjectReplicationBridge* ObjectReplicationBridge;
 	FReplicationSystemInternal* ReplicationSystemInternal;
 
-	const FNetRefHandleManager& LocalNetRefHandleManager;
+	FNetRefHandleManager& LocalNetRefHandleManager;
 	const TArray<UObject*>& ReplicatedInstances;
 
 	const FNetBitArrayView AccumulatedDirtyObjects;
-	const FNetBitArrayView PrevFrameGlobalScopeObjects;
 
 	FNetBitArrayView DirtyObjectsThisFrame;
-
 	FNetBitArrayView GarbageCollectionAffectedObjects;
 
 	FPreUpdateAndPollStats PollStats;
-
 };
 
 } // end namespace UE::Net::Private
