cluster unions should properly release islands that contain auxiliary particles

- non-accelerated cluster unions should properly check the component (not the actor) for response against a collision channel

diff --git a/Engine/Source/Runtime/Engine/Private/PhysicsEngine/ClusterUnionComponent.cpp b/Engine/Source/Runtime/Engine/Private/PhysicsEngine/ClusterUnionComponent.cpp
--- a/Engine/Source/Runtime/Engine/Private/PhysicsEngine/ClusterUnionComponent.cpp
+++ b/Engine/Source/Runtime/Engine/Private/PhysicsEngine/ClusterUnionComponent.cpp
@@ -197,6 +197,8 @@ void UClusterUnionComponent::RemoveComponentFromCluster(UPrimitiveComponent* InC
 				AccelerationStructure->RemoveElement(Payload);
 			}
 		}
+
+		PendingComponentSync.Remove(InComponent);
 	}
 
 	TSet<Chaos::FPhysicsObjectHandle> PhysicsObjectsToRemove;
@@ -358,6 +360,10 @@ void UClusterUnionComponent::OnCreatePhysicsState()
 	InitData.ActorId = GetOwner()->GetUniqueID();
 	InitData.ComponentId = GetUniqueID();
 	InitData.bNeedsClusterXRInitialization = GetOwner()->HasAuthority();
+
+	// Only need to check connectivity on the server and have the client rely on replication to get the memo on when to release from cluster union.
+	InitData.bCheckConnectivity = GetOwner()->HasAuthority();
+
 	bHasReceivedTransform = false;
 	PhysicsProxy = new Chaos::FClusterUnionPhysicsProxy{ this, Parameters, InitData };
 	PhysicsProxy->Initialize_External();
@@ -1136,11 +1142,11 @@ void UClusterUnionComponent::VisitAllCurrentChildComponentsForCollision(ECollisi
 					{
 						return true;
 					}
+				}
 
-					if (Owner->GetComponentsCollisionResponseToChannel(TraceChannel) == ECollisionResponse::ECR_Ignore)
-					{
-						return true;
-					}
+				if (Component->GetCollisionResponseToChannel(TraceChannel) == ECollisionResponse::ECR_Ignore)
+				{
+					return true;
 				}
 
 				Lambda(Component);
diff --git a/Engine/Source/Runtime/Experimental/Chaos/Private/Chaos/ClusterUnionManager.cpp b/Engine/Source/Runtime/Experimental/Chaos/Private/Chaos/ClusterUnionManager.cpp
--- a/Engine/Source/Runtime/Experimental/Chaos/Private/Chaos/ClusterUnionManager.cpp
+++ b/Engine/Source/Runtime/Experimental/Chaos/Private/Chaos/ClusterUnionManager.cpp
@@ -453,6 +453,13 @@ namespace Chaos
 			{
 				Cluster->PendingConnectivityOperations.Add({Particle, EClusterUnionConnectivityOperation ::Add});
 			}
+
+			if (!Cluster->ChildProperties.Contains(Particle))
+			{
+				FClusterUnionParticleProperties Properties;
+				Properties.bIsAuxiliaryParticle = false;
+				Cluster->ChildProperties.Add(Particle, Properties);
+			}
 		}
 
 		MClustering.AddParticlesToCluster(Cluster->InternalCluster, FinalParticlesToAdd, ChildToParentMap);
diff --git a/Engine/Source/Runtime/Experimental/Chaos/Private/Chaos/PBDRigidClustering.cpp b/Engine/Source/Runtime/Experimental/Chaos/Private/Chaos/PBDRigidClustering.cpp
--- a/Engine/Source/Runtime/Experimental/Chaos/Private/Chaos/PBDRigidClustering.cpp
+++ b/Engine/Source/Runtime/Experimental/Chaos/Private/Chaos/PBDRigidClustering.cpp
@@ -901,6 +901,7 @@ namespace Chaos
 		// Grab cluster union parent if there is one
 		FPBDRigidParticleHandle* ParentRigid = ClusteredParticle->ClusterIds().Id;
 		FPBDRigidClusteredParticleHandle* Parent = ParentRigid ? ParentRigid->CastToClustered() : nullptr;
+		TSet<FClusterUnionIndex> ClusterUnionsToConsiderForConnectivity;
 
 		for (int32 ChildIdx = Children.Num() - 1; ChildIdx >= 0; --ChildIdx)
 		{
@@ -933,7 +934,14 @@ namespace Chaos
 						}
 					}
 
-					ClusterUnionManager.HandleRemoveOperationWithClusterLookup({ ClusteredParticle }, EClusterUnionOperationTiming::Defer);
+					const FClusterUnionIndex ClusterUnionIndex = ClusterUnionManager.FindClusterUnionIndexFromParticle(ClusteredParticle);
+					if (ClusterUnionIndex != INDEX_NONE)
+					{
+						// Remove node connections here immediately just in case we need to manage connectivity on the cluster union.
+						RemoveNodeConnections(ClusteredParticle);
+						ClusterUnionsToConsiderForConnectivity.Add(ClusterUnionIndex);
+						ClusterUnionManager.HandleRemoveOperationWithClusterLookup({ ClusteredParticle }, EClusterUnionOperationTiming::Defer);
+					}
 					bFoundFirstRelease = true;
 				}
 
@@ -943,6 +951,7 @@ namespace Chaos
 
 				if (bIsInClusterUnion)
 				{
+					ClusterUnionsToConsiderForConnectivity.Add(ClusterUnionIndex);
 					DeferredRemoveFromClusterUnion.Add(Child);
 				}
 				else
@@ -1038,6 +1047,18 @@ namespace Chaos
 					RemoveNodeConnections(Child);
 				}
 				ActivatedChildren.Append(HandleConnectivityOnReleaseClusterParticle(ClusteredParticle, bCreateNewClusters));
+
+				// Every removal from a cluster union can trigger connectivity changes as well.
+				if (!bIsClusterUnion && bCheckForInterclusterEdgesOnRelease)
+				{
+					for (FClusterUnionIndex UnionIndex : ClusterUnionsToConsiderForConnectivity)
+					{
+						if (FClusterUnion* ClusterUnion = ClusterUnionManager.FindClusterUnion(UnionIndex))
+						{
+							ActivatedChildren.Append(HandleConnectivityOnReleaseClusterParticle(ClusterUnion->InternalCluster, false));
+						}
+					}
+				}
 			}
 
 			for (FPBDRigidParticleHandle* Child : ActivatedChildren)
@@ -1178,7 +1199,7 @@ namespace Chaos
 					ActivatedChildren.Append(MoveTemp(NewClusters));
 				}
 			}
-			else if (bCheckForInterclusterEdgesOnRelease)
+			else if (bCheckForInterclusterEdgesOnRelease && ParentClusterUnion->bCheckConnectivity)
 			{
 				// We know we're in an cluster union. There are pieces that we consider to be the "main body". There are certain pieces that we consider to be auxiliary as well.
 				// If an island is only made up of auxiliary pieces, then those pieces should fall off. Connectivity of main pieces should be handled by the GT.
diff --git a/Engine/Source/Runtime/Experimental/Chaos/Private/PhysicsProxy/ClusterUnionPhysicsProxy.cpp b/Engine/Source/Runtime/Experimental/Chaos/Private/PhysicsProxy/ClusterUnionPhysicsProxy.cpp
--- a/Engine/Source/Runtime/Experimental/Chaos/Private/PhysicsProxy/ClusterUnionPhysicsProxy.cpp
+++ b/Engine/Source/Runtime/Experimental/Chaos/Private/PhysicsProxy/ClusterUnionPhysicsProxy.cpp
@@ -129,6 +129,7 @@ namespace Chaos
 
 			// On the client, we'd rather wait for the server to initialize the particle properly.
 			ClusterUnion->bNeedsXRInitialization = InitData.bNeedsClusterXRInitialization;
+			ClusterUnion->bCheckConnectivity = InitData.bCheckConnectivity;
 		}
 	}
 
diff --git a/Engine/Source/Runtime/Experimental/Chaos/Public/Chaos/ClusterUnionManager.h b/Engine/Source/Runtime/Experimental/Chaos/Public/Chaos/ClusterUnionManager.h
--- a/Engine/Source/Runtime/Experimental/Chaos/Public/Chaos/ClusterUnionManager.h
+++ b/Engine/Source/Runtime/Experimental/Chaos/Public/Chaos/ClusterUnionManager.h
@@ -96,6 +96,9 @@ namespace Chaos
 		// Whether or not the anchor got set by the GT and we shouldn't try to recompute its value.
 		bool bAnchorLock = false;
 
+		// Whether or not we need to check for connectivity for this cluster union.
+		bool bCheckConnectivity = false;
+
 		// Pending particles that need to be added into the connectivity graph.
 		TArray<TPair<FPBDRigidParticleHandle*, EClusterUnionConnectivityOperation>> PendingConnectivityOperations;
 	};
diff --git a/Engine/Source/Runtime/Experimental/Chaos/Public/PhysicsProxy/ClusterUnionPhysicsProxy.h b/Engine/Source/Runtime/Experimental/Chaos/Public/PhysicsProxy/ClusterUnionPhysicsProxy.h
--- a/Engine/Source/Runtime/Experimental/Chaos/Public/PhysicsProxy/ClusterUnionPhysicsProxy.h
+++ b/Engine/Source/Runtime/Experimental/Chaos/Public/PhysicsProxy/ClusterUnionPhysicsProxy.h
@@ -29,6 +29,7 @@ namespace Chaos
 		uint32 ActorId = INDEX_NONE;
 		uint32 ComponentId = INDEX_NONE;
 		bool bNeedsClusterXRInitialization = true;
+		bool bCheckConnectivity = true;
 	};
 
 	/**
