geometry collection pull from physics state should pull from *next* pull data if it exists for non-interpolatable data

diff --git a/Engine/Source/Runtime/Experimental/Chaos/Private/PhysicsProxy/GeometryCollectionPhysicsProxy.cpp b/Engine/Source/Runtime/Experimental/Chaos/Private/PhysicsProxy/GeometryCollectionPhysicsProxy.cpp
--- a/Engine/Source/Runtime/Experimental/Chaos/Private/PhysicsProxy/GeometryCollectionPhysicsProxy.cpp
+++ b/Engine/Source/Runtime/Experimental/Chaos/Private/PhysicsProxy/GeometryCollectionPhysicsProxy.cpp
@@ -3398,7 +3398,7 @@ bool FGeometryCollectionPhysicsProxy::PullFromPhysicsState(const Chaos::FDirtyGe
 		return false;
 	}
 
-	const FGeometryCollectionResults& Results = PullData.Results();
+	const FGeometryCollectionResults& CurrentResults = NextPullData ? NextPullData->Results() : PullData.Results();
 
 	const int32 NumTransforms = GameThreadCollection.Transform.Num();
 	const bool bNeedInterpolation = (NextPullData != nullptr);
@@ -3418,9 +3418,9 @@ bool FGeometryCollectionPhysicsProxy::PullFromPhysicsState(const Chaos::FDirtyGe
 		const TManagedArray<bool>* AnimationsActive = GameThreadCollection.FindAttribute<bool>(AnimateTransformAttributeName, FGeometryCollection::TransformGroup);
 
 		// first step : process the non values that do not need to be interpolated 
-		for (int32 EntryIndex = 0; EntryIndex < Results.GetNumEntries(); EntryIndex++)
+		for (int32 EntryIndex = 0; EntryIndex < CurrentResults.GetNumEntries(); EntryIndex++)
 		{
-			const FGeometryCollectionResults::FStateData& StateData = Results.GetState(EntryIndex);
+			const FGeometryCollectionResults::FStateData& StateData = CurrentResults.GetState(EntryIndex);
 			const int32 TransformGroupIndex = StateData.TransformIndex;
 
 			if (GTParticles[TransformGroupIndex] == nullptr)
@@ -3468,19 +3468,19 @@ bool FGeometryCollectionPhysicsProxy::PullFromPhysicsState(const Chaos::FDirtyGe
 				const bool bAnimatingWhileDisabled = AnimationsActive ? (*AnimationsActive)[TransformGroupIndex] : false;
 				if (bIsActive || bAnimatingWhileDisabled)
 				{
-					const FGeometryCollectionResults::FPositionData& PositionData = Results.GetPositions(EntryIndex);
+					const FGeometryCollectionResults::FPositionData& PositionData = CurrentResults.GetPositions(EntryIndex);
 
 					const bool XRModified = UpdateGTParticleXR(GTParticle, PositionData.ParticleX, PositionData.ParticleR);
 					bIsCollectionDirty |= XRModified;
 
 					if (LinearVelocities && AngularVelocities)
 					{
-						const FGeometryCollectionResults::FVelocityData& VelocityData = Results.GetVelocities(EntryIndex);
+						const FGeometryCollectionResults::FVelocityData& VelocityData = CurrentResults.GetVelocities(EntryIndex);
 						bIsCollectionDirty |= UpdateValue((*LinearVelocities)[TransformGroupIndex], VelocityData.ParticleV);
 						bIsCollectionDirty |= UpdateValue((*AngularVelocities)[TransformGroupIndex], VelocityData.ParticleW);
 					}
 
-					const FTransform& NewTransform = Results.GetTransform(EntryIndex);
+					const FTransform& NewTransform = CurrentResults.GetTransform(EntryIndex);
 					bIsCollectionDirty |= UpdateTransform(GameThreadCollection.Transform[TransformGroupIndex], NewTransform);
 				}
 			}
@@ -3491,7 +3491,7 @@ bool FGeometryCollectionPhysicsProxy::PullFromPhysicsState(const Chaos::FDirtyGe
 		{
 			for (int32 TransformGroupIndex = 0; TransformGroupIndex < NumTransforms; ++TransformGroupIndex)
 			{
-				const FGeometryCollectionResults::FDamageData& DamageData = Results.GetDamages(TransformGroupIndex);
+				const FGeometryCollectionResults::FDamageData& DamageData = CurrentResults.GetDamages(TransformGroupIndex);
 				Collector->SampleDamage(TransformGroupIndex, DamageData.Damage, DamageData.DamageThreshold);
 			}
 		}
@@ -3500,6 +3500,7 @@ bool FGeometryCollectionPhysicsProxy::PullFromPhysicsState(const Chaos::FDirtyGe
 		// second : interpolate-able ones
 		if (bNeedInterpolation)
 		{
+			const FGeometryCollectionResults& PrevResults = PullData.Results();
 			const FGeometryCollectionResults& NextResults = NextPullData->Results();
 
 			// for that case we cannot just go through the list of entries since Results and NextResults may have different number of entries that don't always match
@@ -3513,7 +3514,7 @@ bool FGeometryCollectionPhysicsProxy::PullFromPhysicsState(const Chaos::FDirtyGe
 				}
 
 				FParticle& GTParticle = *GTParticles[TransformGroupIndex];
-				const FResultInterpolator ResultInterpolator(TransformGroupIndex, Results, NextResults, GameThreadCollection.Transform, GTParticle, *Alpha);
+				const FResultInterpolator ResultInterpolator(TransformGroupIndex, PrevResults, NextResults, GameThreadCollection.Transform, GTParticle, *Alpha);
 				if (ResultInterpolator.HasNoEntry())
 				{
 					continue;
@@ -3552,10 +3553,10 @@ bool FGeometryCollectionPhysicsProxy::PullFromPhysicsState(const Chaos::FDirtyGe
 		InternalClusterUniqueIdxToChildrenTransformIndices.Reset();
 		for (int32 TransformGroupIndex = 0; TransformGroupIndex < NumTransforms; ++TransformGroupIndex)
 		{
-			const FGeometryCollectionResults::FEntryIndex EntryIndex = Results.GetEntryIndexByTransformIndex(TransformGroupIndex);
+			const FGeometryCollectionResults::FEntryIndex EntryIndex = CurrentResults.GetEntryIndexByTransformIndex(TransformGroupIndex);
 			if (EntryIndex != INDEX_NONE)
 			{
-				const FGeometryCollectionResults::FStateData& StateData = Results.GetState(EntryIndex);
+				const FGeometryCollectionResults::FStateData& StateData = CurrentResults.GetState(EntryIndex);
 
 				if (StateData.InternalClusterUniqueIdx > INDEX_NONE)
 				{
