[Chaos VD] Implemented Collision Data serialization

- Created Data wrappers for MidPhases, Constraints and Contact Points
- Created a "Data processor" interface, so we can implement serialization code for each binary data type as separate classes. This also allows register a "Data Processor" from outside of the CVD tool.
The API design that will allow to extend the CVD tool is not flesh out yet, but in the future the idea is allow to users to define and trace their own binary data types, and then load that data into the tool.
- Added a new member to the CVD Particle Data actor to hold the recorded collision data. This will change in the future when I settle on the how the recorded data will be referenced in each particle actor. Right now the data is copied each step, which is not efficient
- Added new CVD Macros to trace Midphases and Constraints

diff --git a/Engine/Plugins/Experimental/ChaosVD/Source/ChaosVD/Private/ChaosVDModule.cpp b/Engine/Plugins/Experimental/ChaosVD/Source/ChaosVD/Private/ChaosVDModule.cpp
--- a/Engine/Plugins/Experimental/ChaosVD/Source/ChaosVD/Private/ChaosVDModule.cpp
+++ b/Engine/Plugins/Experimental/ChaosVD/Source/ChaosVD/Private/ChaosVDModule.cpp
@@ -16,6 +16,8 @@
 
 #define LOCTEXT_NAMESPACE "ChaosVisualDebugger"
 
+DEFINE_LOG_CATEGORY(LogChaosVDEditor);
+
 FChaosVDModule& FChaosVDModule::Get()
 {
 	return FModuleManager::Get().LoadModuleChecked<FChaosVDModule>(TEXT("ChaosVD"));
diff --git a/Engine/Plugins/Experimental/ChaosVD/Source/ChaosVD/Private/ChaosVDParticleActor.cpp b/Engine/Plugins/Experimental/ChaosVD/Source/ChaosVD/Private/ChaosVDParticleActor.cpp
--- a/Engine/Plugins/Experimental/ChaosVD/Source/ChaosVD/Private/ChaosVDParticleActor.cpp
+++ b/Engine/Plugins/Experimental/ChaosVD/Source/ChaosVD/Private/ChaosVDParticleActor.cpp
@@ -38,6 +38,30 @@ void AChaosVDParticleActor::UpdateFromRecordedParticleData(const FChaosVDParticl
 	ParticleDataViewer = InRecordedData;
 }
 
+void AChaosVDParticleActor::UpdateCollisionData(const TArray<TSharedPtr<FChaosVDParticlePairMidPhase>>& InRecordedMidPhases)
+{
+	// TODO: We should store a ptr to the data and in our custom details panel draw it
+	ParticleDataViewer.ParticleMidPhases.Reserve(InRecordedMidPhases.Num());
+	ParticleDataViewer.ParticleMidPhases.Reset(InRecordedMidPhases.Num());
+
+	for (const TSharedPtr<FChaosVDParticlePairMidPhase>& MidPhase : InRecordedMidPhases)
+	{
+		ParticleDataViewer.ParticleMidPhases.Emplace(*MidPhase.Get());
+	}
+}
+
+void AChaosVDParticleActor::UpdateCollisionData(const TArray<FChaosVDConstraint>& InRecordedConstraints)
+{
+	// TODO: We should store a ptr to the data and in our custom details panel draw it
+	ParticleDataViewer.ParticleConstraints.Reserve(InRecordedConstraints.Num());
+	ParticleDataViewer.ParticleConstraints.Reset(InRecordedConstraints.Num());
+
+	for (const FChaosVDConstraint& Constraint : InRecordedConstraints)
+	{
+		ParticleDataViewer.ParticleConstraints.Emplace(Constraint);
+	}
+}
+
 void AChaosVDParticleActor::UpdateGeometry(const TSharedPtr<const Chaos::FImplicitObject>& ImplicitObject, EChaosVDActorGeometryUpdateFlags OptionsFlags)
 {
 	if (EnumHasAnyFlags(OptionsFlags, EChaosVDActorGeometryUpdateFlags::ForceUpdate))
diff --git a/Engine/Plugins/Experimental/ChaosVD/Source/ChaosVD/Private/ChaosVDScene.cpp b/Engine/Plugins/Experimental/ChaosVD/Source/ChaosVD/Private/ChaosVDScene.cpp
--- a/Engine/Plugins/Experimental/ChaosVD/Source/ChaosVD/Private/ChaosVDScene.cpp
+++ b/Engine/Plugins/Experimental/ChaosVD/Source/ChaosVD/Private/ChaosVDScene.cpp
@@ -124,6 +124,8 @@ void FChaosVDScene::UpdateFromRecordedStepData(const int32 SolverID, const FStri
 		}
 	}
 
+	UpdateParticlesCollisionData(InRecordedStepData, SolverID);
+
 	// TODO: This will not work once the recording has what changes between solver steps
 	// So it needs to be update to remove particles based on a Particles Removed event
 	int32 AmountRemoved = 0;
@@ -150,6 +152,35 @@ void FChaosVDScene::UpdateFromRecordedStepData(const int32 SolverID, const FStri
 	OnSceneUpdated().Broadcast();
 }
 
+
+void FChaosVDScene::UpdateParticlesCollisionData(const FChaosVDStepData& InRecordedStepData, int32 SolverID)
+{
+	FChaosVDParticlesByIDMap& SolverParticlesByID = ParticlesBySolverID.FindChecked(SolverID);
+	
+	for (const TPair<int32, TArray<TSharedPtr<FChaosVDParticlePairMidPhase>>>& ParticleIDMidPhasePair : InRecordedStepData.RecordedMidPhasesByParticleID)
+	{
+		if (AChaosVDParticleActor** ExistingParticleVDInstancePtrPtr = SolverParticlesByID.Find(ParticleIDMidPhasePair.Key))
+		{
+			if (AChaosVDParticleActor* ExistingParticleVDInstancePtr = *ExistingParticleVDInstancePtrPtr)
+			{
+				ExistingParticleVDInstancePtr->UpdateCollisionData(ParticleIDMidPhasePair.Value);
+			}
+		}
+	}
+
+	for (const TPair<int32, TArray<FChaosVDConstraint>>& ParticleIDMidPhasePair : InRecordedStepData.RecordedConstraintsByParticleID)
+	{
+		if (AChaosVDParticleActor** ExistingParticleVDInstancePtrPtr = SolverParticlesByID.Find(ParticleIDMidPhasePair.Key))
+		{
+			if (AChaosVDParticleActor* ExistingParticleVDInstancePtr = *ExistingParticleVDInstancePtrPtr)
+			{
+				ExistingParticleVDInstancePtr->UpdateCollisionData(ParticleIDMidPhasePair.Value);
+			}
+		}
+	}
+}
+
+
 void FChaosVDScene::HandleNewGeometryData(const TSharedPtr<const Chaos::FImplicitObject>& GeometryData, const uint32 GeometryID) const
 {
 	NewGeometryAvailableDelegate.Broadcast(GeometryData, GeometryID);
diff --git a/Engine/Plugins/Experimental/ChaosVD/Source/ChaosVD/Private/Trace/ChaosVDTraceAnalyzer.cpp b/Engine/Plugins/Experimental/ChaosVD/Source/ChaosVD/Private/Trace/ChaosVDTraceAnalyzer.cpp
--- a/Engine/Plugins/Experimental/ChaosVD/Source/ChaosVD/Private/Trace/ChaosVDTraceAnalyzer.cpp
+++ b/Engine/Plugins/Experimental/ChaosVD/Source/ChaosVD/Private/Trace/ChaosVDTraceAnalyzer.cpp
@@ -2,7 +2,7 @@
 
 #include "Trace/ChaosVDTraceAnalyzer.h"
 
-#include "Chaos/GeometryParticlesfwd.h"
+#include "ChaosVDModule.h"
 #include "Trace/ChaosVDTraceProvider.h"
 #include "TraceServices/Model/AnalysisSession.h"
 #include "TraceServices/Model/Frames.h"
@@ -156,7 +156,11 @@ bool FChaosVDTraceAnalyzer::OnEvent(uint16 RouteId, EStyle Style, const FOnEvent
 		{
 			const int32 DataID = EventData.GetValue<int32>("DataID");
 
-			ChaosVDTraceProvider->SetBinaryDataReadyToUse(DataID);
+			if (!ensure(ChaosVDTraceProvider->ProcessBinaryData(DataID)))
+			{
+				UE_LOG(LogChaosVDEditor, Error, TEXT("[%s] FailedToProcess Binary Data with ID [%d]"), ANSI_TO_TCHAR(__FUNCTION__), DataID);
+			}
+
 			break;
 		}
 	case RouteId_ChaosVDSolverSimulationSpace:
diff --git a/Engine/Plugins/Experimental/ChaosVD/Source/ChaosVD/Private/Trace/ChaosVDTraceProvider.cpp b/Engine/Plugins/Experimental/ChaosVD/Source/ChaosVD/Private/Trace/ChaosVDTraceProvider.cpp
--- a/Engine/Plugins/Experimental/ChaosVD/Source/ChaosVD/Private/Trace/ChaosVDTraceProvider.cpp
+++ b/Engine/Plugins/Experimental/ChaosVD/Source/ChaosVD/Private/Trace/ChaosVDTraceProvider.cpp
@@ -2,20 +2,17 @@
 
 #include "Trace/ChaosVDTraceProvider.h"
 
+#include "ChaosVDModule.h"
 #include "ChaosVDRecording.h"
 
 #include "Chaos/ChaosArchive.h"
-#include "Chaos/ImplicitObject.h"
 
 #include "Compression/OodleDataCompressionUtil.h"
 #include "Serialization/MemoryReader.h"
-
-#if WITH_CHAOS_VISUAL_DEBUGGER
-#include "DataWrappers/ChaosVDImplicitObjectDataWrapper.h"
-#include "DataWrappers/ChaosVDParticleDataWrapper.h"
-
-using FChaosVDImplicitObjectWrapper = FChaosVDImplicitObjectDataWrapper<Chaos::TSerializablePtr<Chaos::FImplicitObject>, Chaos::FChaosArchive>;
-#endif
+#include "Trace/DataProcessors/ChaosVDConstraintDataProcessor.h"
+#include "Trace/DataProcessors/ChaosVDMidPhaseDataProcessor.h"
+#include "Trace/DataProcessors/ChaosVDTraceImplicitObjectProcessor.h"
+#include "Trace/DataProcessors/ChaosVDTraceParticleDataProcessor.h"
 
 FName FChaosVDTraceProvider::ProviderName("ChaosVDProvider");
 
@@ -114,8 +111,10 @@ FChaosVDBinaryDataContainer& FChaosVDTraceProvider::FindOrAddUnprocessedData(con
 	}
 }
 
-void FChaosVDTraceProvider::SetBinaryDataReadyToUse(const int32 DataID)
+bool FChaosVDTraceProvider::ProcessBinaryData(const int32 DataID)
 {
+	RegisterDefaultDataProcessorsIfNeeded();
+
 	if (const TSharedPtr<FChaosVDBinaryDataContainer>* UnprocessedDataPtr = UnprocessedDataByID.Find(DataID))
 	{
 		const TSharedPtr<FChaosVDBinaryDataContainer> UnprocessedData = *UnprocessedDataPtr;
@@ -123,8 +122,6 @@ void FChaosVDTraceProvider::SetBinaryDataReadyToUse(const int32 DataID)
 		{
 			UnprocessedData->bIsReady = true;
 
-			BinaryDataReadyDelegate.Broadcast(*UnprocessedDataPtr);
-
 			const TArray<uint8>* RawData = nullptr;
 			TArray<uint8> UncompressedData;
 			if (UnprocessedData->bIsCompressed)
@@ -138,47 +135,58 @@ void FChaosVDTraceProvider::SetBinaryDataReadyToUse(const int32 DataID)
 				RawData = &UnprocessedData->RawData;
 			}
 
-#if WITH_CHAOS_VISUAL_DEBUGGER
-			// TODO: Create a system to register external "Data handlers" with the logic on how serialize each type
-			// This should not be done here but as this is the only type we have is ok for now
-			if (UnprocessedData->TypeName == TEXT("FChaosVDImplicitObjectDataWrapper"))
+			if (TSharedPtr<IChaosVDDataProcessor>* DataProcessorPtrPtr = RegisteredDataProcessors.Find(UnprocessedData->TypeName))
 			{
-				FMemoryReader MemeReader(*RawData);
-				Chaos::FChaosArchive Ar(MemeReader);
-
-				FChaosVDImplicitObjectWrapper WrappedGeometryData;
-				WrappedGeometryData.Serialize(Ar);
-
-				InternalRecording->AddImplicitObject(WrappedGeometryData.Hash, WrappedGeometryData.ImplicitObject.Get());
-			}
-
-			if (UnprocessedData->TypeName == TEXT("FChaosVDParticleDataWrapper"))
-			{	
-				FMemoryReader MemeReader(*RawData);
-				MemeReader.SetUseUnversionedPropertySerialization(true);
-
-				FChaosVDParticleDataWrapper ParticleData;
-				MemeReader << ParticleData;
-
-				// This can be null if the recording started Mid-Frame. In this case we just discard the data for now
-				if (FChaosVDSolverFrameData* FrameData = GetLastSolverFrame(ParticleData.SolverID))
+				if (TSharedPtr<IChaosVDDataProcessor> DataProcessorPtr = *DataProcessorPtrPtr)
 				{
-					if (ensureMsgf(FrameData->SolverSteps.Num() > 0, TEXT("A particle was traced without a valid step scope")))
+					if (ensure(DataProcessorPtr->ProcessRawData(*RawData)))
 					{
-						FrameData->SolverSteps.Last().RecordedParticlesData.Add(MoveTemp(ParticleData));
+						return true;
 					}
 				}
 			}
-#endif
+			else
+			{
+				UE_LOG(LogChaosVDEditor, Warning, TEXT("[%s] Data processor for type [%s] not found"), ANSI_TO_TCHAR(__FUNCTION__), *UnprocessedData->TypeName);
+			}
 		}
 	}
-	else
-	{
-		ensure(false);	
-	}
+
+	return false;
 }
 
 TSharedPtr<FChaosVDRecording> FChaosVDTraceProvider::GetRecordingForSession() const
 {
 	return InternalRecording;
 }
+
+void FChaosVDTraceProvider::RegisterDataProcessor(TSharedPtr<IChaosVDDataProcessor> InDataProcessor)
+{
+	RegisteredDataProcessors.Add(InDataProcessor->GetCompatibleTypeName(), InDataProcessor);
+}
+
+void FChaosVDTraceProvider::RegisterDefaultDataProcessorsIfNeeded()
+{
+	if (bDefaultDataProcessorsRegistered)
+	{
+		return;
+	}
+	
+	TSharedPtr<FChaosVDTraceImplicitObjectProcessor> ImplicitObjectProcessor = MakeShared<FChaosVDTraceImplicitObjectProcessor>();
+	ImplicitObjectProcessor->SetTraceProvider(AsShared());
+	RegisterDataProcessor(ImplicitObjectProcessor);
+
+	TSharedPtr<FChaosVDTraceParticleDataProcessor> ParticleDataProcessor = MakeShared<FChaosVDTraceParticleDataProcessor>();
+	ParticleDataProcessor->SetTraceProvider(AsShared());
+	RegisterDataProcessor(ParticleDataProcessor);
+
+	TSharedPtr<FChaosVDMidPhaseDataProcessor> MidPhaseDataProcessor = MakeShared<FChaosVDMidPhaseDataProcessor>();
+	MidPhaseDataProcessor->SetTraceProvider(AsShared());
+	RegisterDataProcessor(MidPhaseDataProcessor);
+
+	TSharedPtr<FChaosVDConstraintDataProcessor> ConstraintDataProcessor = MakeShared<FChaosVDConstraintDataProcessor>();
+	ConstraintDataProcessor->SetTraceProvider(AsShared());
+	RegisterDataProcessor(ConstraintDataProcessor);
+
+	bDefaultDataProcessorsRegistered = true;
+}
diff --git a/Engine/Plugins/Experimental/ChaosVD/Source/ChaosVD/Private/Trace/DataProcessors/ChaosVDConstraintDataProcessor.cpp b/Engine/Plugins/Experimental/ChaosVD/Source/ChaosVD/Private/Trace/DataProcessors/ChaosVDConstraintDataProcessor.cpp
new file mode 100644
--- /dev/null
+++ b/Engine/Plugins/Experimental/ChaosVD/Source/ChaosVD/Private/Trace/DataProcessors/ChaosVDConstraintDataProcessor.cpp
@@ -0,0 +1,49 @@
+﻿// Copyright Epic Games, Inc. All Rights Reserved.
+
+#include "Trace/DataProcessors/ChaosVDConstraintDataProcessor.h"
+
+#include "ChaosVDRecording.h"
+#include "DataWrappers/ChaosVDCollisionDataWrappers.h"
+#include "Serialization/MemoryReader.h"
+#include "Trace/ChaosVDTraceProvider.h"
+
+FChaosVDConstraintDataProcessor::FChaosVDConstraintDataProcessor() : IChaosVDDataProcessor(FChaosVDConstraint::WrapperTypeName)
+{
+}
+
+bool FChaosVDConstraintDataProcessor::ProcessRawData(const TArray<uint8>& InData)
+{
+	TSharedPtr<FChaosVDTraceProvider> ProviderSharedPtr = TraceProvider.Pin();
+	if (!ensure(ProviderSharedPtr.IsValid()))
+	{
+		return false;
+	}
+
+	FChaosVDConstraint RecordedConstraint;
+	FMemoryReader MemReader(InData);
+	RecordedConstraint.Serialize(MemReader);
+
+	// This can be null if the recording started Mid-Frame. In this case we just discard the data for now
+	if (FChaosVDSolverFrameData* FrameData = ProviderSharedPtr->GetLastSolverFrame(RecordedConstraint.SolverID))
+	{
+		if (ensureMsgf(FrameData->SolverSteps.Num() > 0, TEXT("A MidPhase was traced without a valid step scope")))
+		{
+			AddConstraintToParticleIDMap(RecordedConstraint, RecordedConstraint.Particle0Index, *FrameData);
+			AddConstraintToParticleIDMap(RecordedConstraint, RecordedConstraint.Particle1Index, *FrameData);
+		}
+	}
+
+	return !MemReader.IsError() && !MemReader.IsCriticalError();
+}
+
+void FChaosVDConstraintDataProcessor::AddConstraintToParticleIDMap(const FChaosVDConstraint& InConstraintData, int32 ParticleID, FChaosVDSolverFrameData& InFrameData)
+{
+	if (TArray<FChaosVDConstraint>* ParticleConstraints = InFrameData.SolverSteps.Last().RecordedConstraintsByParticleID.Find(ParticleID))
+	{
+		ParticleConstraints->Add(InConstraintData);
+	}
+	else
+	{
+		InFrameData.SolverSteps.Last().RecordedConstraintsByParticleID.Add(ParticleID, { InConstraintData });
+	}
+}
diff --git a/Engine/Plugins/Experimental/ChaosVD/Source/ChaosVD/Private/Trace/DataProcessors/ChaosVDMidPhaseDataProcessor.cpp b/Engine/Plugins/Experimental/ChaosVD/Source/ChaosVD/Private/Trace/DataProcessors/ChaosVDMidPhaseDataProcessor.cpp
new file mode 100644
--- /dev/null
+++ b/Engine/Plugins/Experimental/ChaosVD/Source/ChaosVD/Private/Trace/DataProcessors/ChaosVDMidPhaseDataProcessor.cpp
@@ -0,0 +1,51 @@
+﻿// Copyright Epic Games, Inc. All Rights Reserved.
+
+#include "Trace/DataProcessors/ChaosVDMidPhaseDataProcessor.h"
+#include "DataWrappers/ChaosVDCollisionDataWrappers.h"
+#include "ChaosVDRecording.h"
+#include "Serialization/MemoryReader.h"
+#include "Trace/ChaosVDTraceProvider.h"
+
+FChaosVDMidPhaseDataProcessor::FChaosVDMidPhaseDataProcessor() : IChaosVDDataProcessor(FChaosVDParticlePairMidPhase::WrapperTypeName)
+{
+}
+
+bool FChaosVDMidPhaseDataProcessor::ProcessRawData(const TArray<uint8>& InData)
+{
+	TSharedPtr<FChaosVDTraceProvider> ProviderSharedPtr = TraceProvider.Pin();
+	if (!ensure(ProviderSharedPtr.IsValid()))
+	{
+		return false;
+	}
+
+	TSharedPtr<FChaosVDParticlePairMidPhase> MidPhase = MakeShared<FChaosVDParticlePairMidPhase>();
+	FMemoryReader MemReader(InData);
+
+	MidPhase->Serialize(MemReader);
+
+	// This can be null if the recording started Mid-Frame. In this case we just discard the data for now
+	if (FChaosVDSolverFrameData* FrameData = ProviderSharedPtr->GetLastSolverFrame(MidPhase->SolverID))
+	{
+		if (ensureMsgf(FrameData->SolverSteps.Num() > 0, TEXT("A MidPhase was traced without a valid step scope")))
+		{
+			FrameData->SolverSteps.Last().RecordedMidPhases.Add(MidPhase);
+
+			AddMidPhaseToParticleIDMap(MidPhase, MidPhase->Particle0Idx, *FrameData);
+			AddMidPhaseToParticleIDMap(MidPhase, MidPhase->Particle1Idx, *FrameData);
+		}
+	}
+
+	return !MemReader.IsError() && !MemReader.IsCriticalError();
+}
+
+void FChaosVDMidPhaseDataProcessor::AddMidPhaseToParticleIDMap(const TSharedPtr<FChaosVDParticlePairMidPhase>& MidPhaseData, int32 ParticleID, FChaosVDSolverFrameData& InFrameData)
+{
+	if (TArray<TSharedPtr<FChaosVDParticlePairMidPhase>>* ParticleMidPhases = InFrameData.SolverSteps.Last().RecordedMidPhasesByParticleID.Find(ParticleID))
+	{
+		ParticleMidPhases->Add(MidPhaseData);
+	}
+	else
+	{
+		InFrameData.SolverSteps.Last().RecordedMidPhasesByParticleID.Add(ParticleID, { MidPhaseData });
+	}
+}
diff --git a/Engine/Plugins/Experimental/ChaosVD/Source/ChaosVD/Private/Trace/DataProcessors/ChaosVDTraceImplicitObjectProcessor.cpp b/Engine/Plugins/Experimental/ChaosVD/Source/ChaosVD/Private/Trace/DataProcessors/ChaosVDTraceImplicitObjectProcessor.cpp
new file mode 100644
--- /dev/null
+++ b/Engine/Plugins/Experimental/ChaosVD/Source/ChaosVD/Private/Trace/DataProcessors/ChaosVDTraceImplicitObjectProcessor.cpp
@@ -0,0 +1,38 @@
+﻿// Copyright Epic Games, Inc. All Rights Reserved.
+
+#include "Trace/DataProcessors/ChaosVDTraceImplicitObjectProcessor.h"
+
+#include "Chaos/ImplicitObject.h"
+#include "ChaosVDRecording.h"
+#include "ChaosVisualDebugger/ChaosVisualDebuggerTrace.h"
+#include "DataWrappers/ChaosVDImplicitObjectDataWrapper.h"
+#include "Serialization/MemoryReader.h"
+#include "Trace/ChaosVDTraceProvider.h"
+
+using FChaosVDImplicitObjectWrapper = FChaosVDImplicitObjectDataWrapper<Chaos::TSerializablePtr<Chaos::FImplicitObject>, Chaos::FChaosArchive>;
+
+FChaosVDTraceImplicitObjectProcessor::FChaosVDTraceImplicitObjectProcessor(): IChaosVDDataProcessor(FChaosVDImplicitObjectWrapper::WrapperTypeName)
+{
+}
+
+bool FChaosVDTraceImplicitObjectProcessor::ProcessRawData(const TArray<uint8>& InData)
+{
+	TSharedPtr<FChaosVDTraceProvider> ProviderSharedPtr = TraceProvider.Pin();
+	if (!ensure(ProviderSharedPtr.IsValid()))
+	{
+		return false;
+	}
+
+	FMemoryReader MemeReader(InData);
+	Chaos::FChaosArchive Ar(MemeReader);
+
+	FChaosVDImplicitObjectWrapper WrappedGeometryData;
+	WrappedGeometryData.Serialize(Ar);
+
+	if (TSharedPtr<FChaosVDRecording> Recording = ProviderSharedPtr->GetRecordingForSession())
+	{
+		Recording->AddImplicitObject(WrappedGeometryData.Hash, WrappedGeometryData.ImplicitObject.Get());
+	}
+
+	return !Ar.IsError() && !Ar.IsCriticalError();;
+}
diff --git a/Engine/Plugins/Experimental/ChaosVD/Source/ChaosVD/Private/Trace/DataProcessors/ChaosVDTraceParticleDataProcessor.cpp b/Engine/Plugins/Experimental/ChaosVD/Source/ChaosVD/Private/Trace/DataProcessors/ChaosVDTraceParticleDataProcessor.cpp
new file mode 100644
--- /dev/null
+++ b/Engine/Plugins/Experimental/ChaosVD/Source/ChaosVD/Private/Trace/DataProcessors/ChaosVDTraceParticleDataProcessor.cpp
@@ -0,0 +1,40 @@
+﻿// Copyright Epic Games, Inc. All Rights Reserved.
+
+#include "Trace/DataProcessors/ChaosVDTraceParticleDataProcessor.h"
+
+#include "ChaosVDRecording.h"
+#include "DataWrappers/ChaosVDParticleDataWrapper.h"
+#include "Serialization/MemoryReader.h"
+#include "Trace/ChaosVDTraceProvider.h"
+
+FChaosVDTraceParticleDataProcessor::FChaosVDTraceParticleDataProcessor(): IChaosVDDataProcessor(FChaosVDParticleDataWrapper::WrapperTypeName)
+{
+}
+
+bool FChaosVDTraceParticleDataProcessor::ProcessRawData(const TArray<uint8>& InData)
+{
+	TSharedPtr<FChaosVDTraceProvider> ProviderSharedPtr = TraceProvider.Pin();
+	if (!ensure(ProviderSharedPtr.IsValid()))
+	{
+		return false;
+	}
+
+	FMemoryReader MemeReader(InData);
+	MemeReader.SetUseUnversionedPropertySerialization(true);
+
+	Chaos::FChaosArchive Ar(MemeReader);
+
+	FChaosVDParticleDataWrapper ParticleData;
+	ParticleData.Serialize(Ar);
+
+	// This can be null if the recording started Mid-Frame. In this case we just discard the data for now
+	if (FChaosVDSolverFrameData* FrameData = ProviderSharedPtr->GetLastSolverFrame(ParticleData.SolverID))
+	{
+		if (ensureMsgf(FrameData->SolverSteps.Num() > 0, TEXT("A particle was traced without a valid step scope")))
+		{
+			FrameData->SolverSteps.Last().RecordedParticlesData.Add(MoveTemp(ParticleData));
+		}
+	}
+
+	return !Ar.IsError() && !Ar.IsCriticalError();
+}
diff --git a/Engine/Plugins/Experimental/ChaosVD/Source/ChaosVD/Public/ChaosVDModule.h b/Engine/Plugins/Experimental/ChaosVD/Source/ChaosVD/Public/ChaosVDModule.h
--- a/Engine/Plugins/Experimental/ChaosVD/Source/ChaosVD/Public/ChaosVDModule.h
+++ b/Engine/Plugins/Experimental/ChaosVD/Source/ChaosVD/Public/ChaosVDModule.h
@@ -12,6 +12,7 @@ class SDockTab;
 class FSpawnTabArgs;
 struct FGuid;
 
+DECLARE_LOG_CATEGORY_EXTERN(LogChaosVDEditor, Log, Log);
 class FChaosVDModule : public IModuleInterface
 {
 public:
diff --git a/Engine/Plugins/Experimental/ChaosVD/Source/ChaosVD/Public/ChaosVDParticleActor.h b/Engine/Plugins/Experimental/ChaosVD/Source/ChaosVD/Public/ChaosVDParticleActor.h
--- a/Engine/Plugins/Experimental/ChaosVD/Source/ChaosVD/Public/ChaosVDParticleActor.h
+++ b/Engine/Plugins/Experimental/ChaosVD/Source/ChaosVD/Public/ChaosVDParticleActor.h
@@ -3,6 +3,7 @@
 #pragma once
 #include "Chaos/Core.h"
 #include "DataWrappers/ChaosVDParticleDataWrapper.h"
+#include "DataWrappers/ChaosVDCollisionDataWrappers.h"
 #include "GameFramework/Actor.h"
 
 #include "ChaosVDParticleActor.generated.h"
@@ -37,6 +38,9 @@ public:
 
 	void UpdateFromRecordedParticleData(const FChaosVDParticleDataWrapper& InRecordedData, const Chaos::FRigidTransform3& SimulationTransform);
 
+	void UpdateCollisionData(const TArray<TSharedPtr<FChaosVDParticlePairMidPhase>>& InRecordedMidPhases);
+	void UpdateCollisionData(const TArray<FChaosVDConstraint>& InRecordedConstraints);
+
 	void UpdateGeometry(const TSharedPtr<const Chaos::FImplicitObject>& ImplicitObject, EChaosVDActorGeometryUpdateFlags OptionsFlags = EChaosVDActorGeometryUpdateFlags::None);
 	void UpdateGeometry(uint32 NewGeometryHash, EChaosVDActorGeometryUpdateFlags OptionsFlags = EChaosVDActorGeometryUpdateFlags::None);
 
diff --git a/Engine/Plugins/Experimental/ChaosVD/Source/ChaosVD/Public/ChaosVDScene.h b/Engine/Plugins/Experimental/ChaosVD/Source/ChaosVD/Public/ChaosVDScene.h
--- a/Engine/Plugins/Experimental/ChaosVD/Source/ChaosVD/Public/ChaosVDScene.h
+++ b/Engine/Plugins/Experimental/ChaosVD/Source/ChaosVD/Public/ChaosVDScene.h
@@ -44,6 +44,8 @@ public:
 	/** Updates, Adds and Remove actors to match the provided Step Data */
 	void UpdateFromRecordedStepData(const int32 SolverID, const FString& SolverName, const FChaosVDStepData& InRecordedStepData, const FChaosVDSolverFrameData& InFrameData);
 
+	void UpdateParticlesCollisionData(const FChaosVDStepData& InRecordedStepData, int32 SolverID);
+
 	void HandleNewGeometryData(const TSharedPtr<const Chaos::FImplicitObject>&, const uint32 GeometryID) const;
 
 	void HandleEnterNewGameFrame(int32 FrameNumber, const TArray<int32>& AvailableSolversIds);
diff --git a/Engine/Plugins/Experimental/ChaosVD/Source/ChaosVD/Public/Trace/ChaosVDTraceProvider.h b/Engine/Plugins/Experimental/ChaosVD/Source/ChaosVD/Public/Trace/ChaosVDTraceProvider.h
--- a/Engine/Plugins/Experimental/ChaosVD/Source/ChaosVD/Public/Trace/ChaosVDTraceProvider.h
+++ b/Engine/Plugins/Experimental/ChaosVD/Source/ChaosVD/Public/Trace/ChaosVDTraceProvider.h
@@ -9,6 +9,7 @@
 #include "Containers/Array.h"
 
 #include "Chaos/ChaosArchive.h"
+#include "DataProcessors/IChaosVDDataProcessor.h"
 
 struct FChaosVDGameFrameData;
 class FChaosVDEngine;
@@ -37,13 +38,11 @@ struct FChaosVDTraceSessionData
 	TMap<int32, TSharedPtr<FChaosVDBinaryDataContainer>> UnprocessedDataByID;
 };
 
-DECLARE_MULTICAST_DELEGATE_OneParam(FOnBinaryDataReady, TSharedPtr<FChaosVDBinaryDataContainer>)
-
 /** Provider class for Chaos VD trace recordings.
  * It stores and handles rebuilt recorded frame data from Trace events
  * dispatched by the Chaos VD Trace analyzer
  */
-class FChaosVDTraceProvider : public TraceServices::IProvider
+class FChaosVDTraceProvider : public TraceServices::IProvider, public TSharedFromThis<FChaosVDTraceProvider>
 {
 public:
 	
@@ -63,18 +62,23 @@ public:
 
 	FChaosVDBinaryDataContainer& FindOrAddUnprocessedData(const int32 DataID);
 
-	void SetBinaryDataReadyToUse(const int32 DataID);
+	bool ProcessBinaryData(const int32 DataID);
 
 	TSharedPtr<FChaosVDRecording> GetRecordingForSession() const;
-	
-	FOnBinaryDataReady& OnBinaryDataReady() { return BinaryDataReadyDelegate; }
+
+	void RegisterDataProcessor(TSharedPtr<IChaosVDDataProcessor> InDataProcessor);
 
 private:
+
+	void RegisterDefaultDataProcessorsIfNeeded();
+	
 	TraceServices::IAnalysisSession& Session;
 
 	TSharedPtr<FChaosVDRecording> InternalRecording;
 
 	TMap<int32, TSharedPtr<FChaosVDBinaryDataContainer>> UnprocessedDataByID;
-	
-	FOnBinaryDataReady BinaryDataReadyDelegate;
+
+	TMap<FStringView, TSharedPtr<IChaosVDDataProcessor>> RegisteredDataProcessors;
+
+	bool bDefaultDataProcessorsRegistered = false;
 };
diff --git a/Engine/Plugins/Experimental/ChaosVD/Source/ChaosVD/Public/Trace/DataProcessors/ChaosVDConstraintDataProcessor.h b/Engine/Plugins/Experimental/ChaosVD/Source/ChaosVD/Public/Trace/DataProcessors/ChaosVDConstraintDataProcessor.h
new file mode 100644
--- /dev/null
+++ b/Engine/Plugins/Experimental/ChaosVD/Source/ChaosVD/Public/Trace/DataProcessors/ChaosVDConstraintDataProcessor.h
@@ -0,0 +1,20 @@
+﻿// Copyright Epic Games, Inc. All Rights Reserved.
+
+#pragma once
+#include "IChaosVDDataProcessor.h"
+
+struct FChaosVDSolverFrameData;
+struct FChaosVDConstraint;
+
+/**
+ * Data processor implementation that is able to deserialize traced Constraints
+ */
+class FChaosVDConstraintDataProcessor final : public IChaosVDDataProcessor
+{
+public:
+	explicit FChaosVDConstraintDataProcessor();
+	
+	virtual bool ProcessRawData(const TArray<uint8>& InData) override;
+
+	void AddConstraintToParticleIDMap(const FChaosVDConstraint& InConstraintData, int32 ParticleID, FChaosVDSolverFrameData& InFrameData);
+};
diff --git a/Engine/Plugins/Experimental/ChaosVD/Source/ChaosVD/Public/Trace/DataProcessors/ChaosVDMidPhaseDataProcessor.h b/Engine/Plugins/Experimental/ChaosVD/Source/ChaosVD/Public/Trace/DataProcessors/ChaosVDMidPhaseDataProcessor.h
new file mode 100644
--- /dev/null
+++ b/Engine/Plugins/Experimental/ChaosVD/Source/ChaosVD/Public/Trace/DataProcessors/ChaosVDMidPhaseDataProcessor.h
@@ -0,0 +1,21 @@
+﻿// Copyright Epic Games, Inc. All Rights Reserved.
+
+#pragma once
+#include "IChaosVDDataProcessor.h"
+#include "Templates/SharedPointer.h"
+
+struct FChaosVDSolverFrameData;
+struct FChaosVDParticlePairMidPhase;
+
+/**
+ * Data processor implementation that is able to deserialize traced MidPhases 
+ */
+class FChaosVDMidPhaseDataProcessor final : public IChaosVDDataProcessor
+{
+public:
+	explicit FChaosVDMidPhaseDataProcessor();
+
+	virtual bool ProcessRawData(const TArray<uint8>& InData) override;
+
+	void AddMidPhaseToParticleIDMap(const TSharedPtr<FChaosVDParticlePairMidPhase>& MidPhaseData, int32 ParticleID, FChaosVDSolverFrameData& InFrameData);
+};
diff --git a/Engine/Plugins/Experimental/ChaosVD/Source/ChaosVD/Public/Trace/DataProcessors/ChaosVDTraceImplicitObjectProcessor.h b/Engine/Plugins/Experimental/ChaosVD/Source/ChaosVD/Public/Trace/DataProcessors/ChaosVDTraceImplicitObjectProcessor.h
new file mode 100644
--- /dev/null
+++ b/Engine/Plugins/Experimental/ChaosVD/Source/ChaosVD/Public/Trace/DataProcessors/ChaosVDTraceImplicitObjectProcessor.h
@@ -0,0 +1,16 @@
+﻿// Copyright Epic Games, Inc. All Rights Reserved.
+
+#pragma once
+
+#include "IChaosVDDataProcessor.h"
+
+/**
+ * Data processor implementation that is able to deserialize traced Implicit objects
+ */
+class FChaosVDTraceImplicitObjectProcessor final : public IChaosVDDataProcessor
+{
+public:
+	explicit FChaosVDTraceImplicitObjectProcessor();
+
+	virtual bool ProcessRawData(const TArray<uint8>& InData) override;
+};
diff --git a/Engine/Plugins/Experimental/ChaosVD/Source/ChaosVD/Public/Trace/DataProcessors/ChaosVDTraceParticleDataProcessor.h b/Engine/Plugins/Experimental/ChaosVD/Source/ChaosVD/Public/Trace/DataProcessors/ChaosVDTraceParticleDataProcessor.h
new file mode 100644
--- /dev/null
+++ b/Engine/Plugins/Experimental/ChaosVD/Source/ChaosVD/Public/Trace/DataProcessors/ChaosVDTraceParticleDataProcessor.h
@@ -0,0 +1,16 @@
+﻿// Copyright Epic Games, Inc. All Rights Reserved.
+
+#pragma once
+
+#include "Trace/DataProcessors/IChaosVDDataProcessor.h"
+
+/**
+ * Data processor implementation that is able to deserialize traced Particles data
+ */
+class FChaosVDTraceParticleDataProcessor final : public IChaosVDDataProcessor
+{
+public:
+	explicit FChaosVDTraceParticleDataProcessor();
+
+	virtual bool ProcessRawData(const TArray<uint8>& InData) override;
+};
diff --git a/Engine/Plugins/Experimental/ChaosVD/Source/ChaosVD/Public/Trace/DataProcessors/IChaosVDDataProcessor.h b/Engine/Plugins/Experimental/ChaosVD/Source/ChaosVD/Public/Trace/DataProcessors/IChaosVDDataProcessor.h
new file mode 100644
--- /dev/null
+++ b/Engine/Plugins/Experimental/ChaosVD/Source/ChaosVD/Public/Trace/DataProcessors/IChaosVDDataProcessor.h
@@ -0,0 +1,32 @@
+﻿// Copyright Epic Games, Inc. All Rights Reserved.
+
+#pragma once
+#include "Containers/Array.h"
+#include "Containers/StringFwd.h"
+#include "Hal/Platform.h"
+#include "Templates/SharedPointer.h"
+
+class FChaosVDTraceProvider;
+
+/** Interface for all used for any class that is able to process traced Chaos Visual Debugger binary data */
+class IChaosVDDataProcessor
+{
+public:
+	virtual ~IChaosVDDataProcessor() = default;
+
+	explicit IChaosVDDataProcessor(FStringView InCompatibleType) : CompatibleType(InCompatibleType)
+	{
+	}
+
+	/** Type name this data processor can interpret */
+	FStringView GetCompatibleTypeName() const { return CompatibleType; }
+	/** Called with the raw serialized data to be processed */
+	virtual bool ProcessRawData(const TArray<uint8>& InData) { return false; }
+
+	/** Sets the Trace Provider that is storing the data being analyzed */
+	void SetTraceProvider(const TSharedPtr<FChaosVDTraceProvider>& InProvider) { TraceProvider = InProvider; }
+
+protected:
+	TWeakPtr<FChaosVDTraceProvider> TraceProvider;
+	FStringView CompatibleType;
+};
diff --git a/Engine/Source/Runtime/Experimental/Chaos/Private/Chaos/PBDRigidsEvolutionGBF.cpp b/Engine/Source/Runtime/Experimental/Chaos/Private/Chaos/PBDRigidsEvolutionGBF.cpp
--- a/Engine/Source/Runtime/Experimental/Chaos/Private/Chaos/PBDRigidsEvolutionGBF.cpp
+++ b/Engine/Source/Runtime/Experimental/Chaos/Private/Chaos/PBDRigidsEvolutionGBF.cpp
@@ -387,7 +387,10 @@ void FPBDRigidsEvolutionGBF::AdvanceOneTimeStepImpl(const FReal Dt, const FSubSt
 		CSV_SCOPED_TIMING_STAT(PhysicsVerbose, StepSolver_DetectCollisions);
 		CollisionDetector.GetBroadPhase().SetSpatialAcceleration(InternalAcceleration);
 
-		CollisionDetector.RunBroadPhase(Dt, GetCurrentStepResimCache());
+		{
+			CVD_SCOPE_TRACE_SOLVER_STEP(TEXT("Collision Detection Broad Phase"))
+			CollisionDetector.RunBroadPhase(Dt, GetCurrentStepResimCache());
+		}
 
 		if (MidPhaseModifiers)
 		{
@@ -395,7 +398,11 @@ void FPBDRigidsEvolutionGBF::AdvanceOneTimeStepImpl(const FReal Dt, const FSubSt
 			CollisionConstraints.ApplyMidPhaseModifier(*MidPhaseModifiers, Dt);
 		}
 
-		CollisionDetector.RunNarrowPhase(Dt, GetCurrentStepResimCache());
+		{
+			CVD_SCOPE_TRACE_SOLVER_STEP(TEXT("Collision Detection Narrow Phase"))
+
+			CollisionDetector.RunNarrowPhase(Dt, GetCurrentStepResimCache());
+		}
 	}
 
 	if (PostDetectCollisionsCallback != nullptr)
diff --git a/Engine/Source/Runtime/Experimental/Chaos/Private/ChaosVisualDebugger/ChaosVDDataWrapperUtils.cpp b/Engine/Source/Runtime/Experimental/Chaos/Private/ChaosVisualDebugger/ChaosVDDataWrapperUtils.cpp
--- a/Engine/Source/Runtime/Experimental/Chaos/Private/ChaosVisualDebugger/ChaosVDDataWrapperUtils.cpp
+++ b/Engine/Source/Runtime/Experimental/Chaos/Private/ChaosVisualDebugger/ChaosVDDataWrapperUtils.cpp
@@ -3,8 +3,39 @@
 #include "ChaosVisualDebugger/ChaosVDDataWrapperUtils.h"
 
 #if WITH_CHAOS_VISUAL_DEBUGGER
+
+#include "Algo/Copy.h"
+#include "Chaos/Collision/ParticlePairMidPhase.h"
 #include "Chaos/ParticleHandle.h"
+#include "DataWrappers/ChaosVDCollisionDataWrappers.h"
 #include "DataWrappers/ChaosVDParticleDataWrapper.h"
+#include "Chaos/Collision/CollisionConstraintAllocator.h"
+
+void FChaosVDDataWrapperUtils::CopyManifoldPointsToDataWrapper(const Chaos::FManifoldPoint& InCopyFrom, FChaosVDManifoldPoint& OutCopyTo)
+{
+	OutCopyTo.bDisabled = InCopyFrom.Flags.bDisabled;
+	OutCopyTo.bWasRestored = InCopyFrom.Flags.bWasRestored;
+	OutCopyTo.bWasReplaced = InCopyFrom.Flags.bWasReplaced;
+	OutCopyTo.bHasStaticFrictionAnchor = InCopyFrom.Flags.bHasStaticFrictionAnchor;
+	OutCopyTo.TargetPhi = InCopyFrom.TargetPhi;
+
+	Algo::Transform(InCopyFrom.ShapeAnchorPoints, OutCopyTo.ShapeAnchorPoints, &FChaosVDDataWrapperUtils::ConvertToFVector);
+	Algo::Transform(InCopyFrom.InitialShapeContactPoints, OutCopyTo.InitialShapeContactPoints, &FChaosVDDataWrapperUtils::ConvertToFVector);
+	Algo::Transform(InCopyFrom.ContactPoint.ShapeContactPoints, OutCopyTo.ContactPoint.ShapeContactPoints, &FChaosVDDataWrapperUtils::ConvertToFVector);
+
+	OutCopyTo.ContactPoint.ShapeContactNormal = FVector(InCopyFrom.ContactPoint.ShapeContactNormal);
+	OutCopyTo.ContactPoint.Phi = InCopyFrom.ContactPoint.Phi;
+	OutCopyTo.ContactPoint.FaceIndex = InCopyFrom.ContactPoint.FaceIndex;
+	OutCopyTo.ContactPoint.ContactType = static_cast<EChaosVDContactPointType>(InCopyFrom.ContactPoint.ContactType);
+}
+
+void FChaosVDDataWrapperUtils::CopyManifoldPointResultsToDataWrapper(const Chaos::FManifoldPointResult& InCopyFrom, FChaosVDManifoldPoint& OutCopyTo)
+{
+	OutCopyTo.NetPushOut =  FVector(InCopyFrom.NetPushOut);
+	OutCopyTo.NetImpulse =  FVector(InCopyFrom.NetImpulse);
+	OutCopyTo.bIsValid =  InCopyFrom.bIsValid;
+	OutCopyTo.bInsideStaticFrictionCone =  InCopyFrom.bInsideStaticFrictionCone;
+}
 
 FChaosVDParticleDataWrapper FChaosVDDataWrapperUtils::BuildParticleDataWrapperFromParticle(const Chaos::FGeometryParticleHandle* ParticleHandlePtr)
 {
@@ -34,4 +65,97 @@ FChaosVDParticleDataWrapper FChaosVDDataWrapperUtils::BuildParticleDataWrapperFr
 	}
 	return MoveTemp(WrappedParticleData);
 }
+
+FChaosVDConstraint FChaosVDDataWrapperUtils::BuildConstraintDataWrapperFromConstraint(const Chaos::FPBDCollisionConstraint& InConstraint)
+{
+	FChaosVDConstraint WrappedConstraintData;
+	
+	WrappedConstraintData.bDisabled = InConstraint.Flags.bDisabled;
+	WrappedConstraintData.bUseManifold = InConstraint.Flags.bUseManifold;
+	WrappedConstraintData.bUseIncrementalManifold = InConstraint.Flags.bUseIncrementalManifold;
+	WrappedConstraintData.bCanRestoreManifold = InConstraint.Flags.bCanRestoreManifold;
+	WrappedConstraintData.bWasManifoldRestored = InConstraint.Flags.bWasManifoldRestored;
+	WrappedConstraintData.bIsQuadratic0 = InConstraint.Flags.bIsQuadratic0;
+	WrappedConstraintData.bIsQuadratic1 = InConstraint.Flags.bIsQuadratic1;
+	WrappedConstraintData.bIsProbeUnmodified = InConstraint.Flags.bIsProbeUnmodified;
+	WrappedConstraintData.bIsProbe = InConstraint.Flags.bIsProbe;
+	WrappedConstraintData.bCCDEnabled = InConstraint.Flags.bCCDEnabled;
+	WrappedConstraintData.bCCDSweepEnabled = InConstraint.Flags.bCCDSweepEnabled;
+	WrappedConstraintData.bModifierApplied = InConstraint.Flags.bModifierApplied;
+	WrappedConstraintData.bMaterialSet = InConstraint.Flags.bMaterialSet;
+	WrappedConstraintData.ShapesType = static_cast<EChaosVDContactShapesType>(InConstraint.ShapesType);
+	WrappedConstraintData.CullDistance = InConstraint.CullDistance;
+	WrappedConstraintData.CollisionTolerance = InConstraint.CollisionTolerance;
+	WrappedConstraintData.ClosestManifoldPointIndex = InConstraint.ClosestManifoldPointIndex;
+	WrappedConstraintData.ExpectedNumManifoldPoints = InConstraint.ExpectedNumManifoldPoints;
+	WrappedConstraintData.Stiffness = InConstraint.Stiffness;
+	WrappedConstraintData.CCDTimeOfImpact = InConstraint.CCDTimeOfImpact;
+	WrappedConstraintData.CCDEnablePenetration = InConstraint.CCDEnablePenetration;
+	WrappedConstraintData.CCDTargetPenetration = InConstraint.CCDTargetPenetration;
+	
+	WrappedConstraintData.AccumulatedImpulse = FVector(InConstraint.AccumulatedImpulse);
+
+	WrappedConstraintData.Particle0Index = InConstraint.GetParticle0()->UniqueIdx().Idx;
+	WrappedConstraintData.Particle1Index = InConstraint.GetParticle1()->UniqueIdx().Idx;
+
+	Algo::Copy(InConstraint.ShapeWorldTransforms, WrappedConstraintData.ShapeWorldTransforms);
+
+	WrappedConstraintData.CollisionMargins = TArray(InConstraint.CollisionMargins, std::size(InConstraint.CollisionMargins));
+	WrappedConstraintData.LastShapeWorldPositionDelta = FVector(InConstraint.LastShapeWorldPositionDelta);
+	WrappedConstraintData.LastShapeWorldRotationDelta = FQuat(InConstraint.LastShapeWorldRotationDelta);
+
+	WrappedConstraintData.ManifoldPoints.Reserve(Chaos::FPBDCollisionConstraint::MaxManifoldPoints);
+	WrappedConstraintData.ManifoldPoints.SetNum(Chaos::FPBDCollisionConstraint::MaxManifoldPoints);
+
+	for (int32 PointIndex = 0; PointIndex < Chaos::FPBDCollisionConstraint::MaxManifoldPoints; PointIndex++)
+	{
+		FChaosVDManifoldPoint& CurrentCVDMainFoldPoint = WrappedConstraintData.ManifoldPoints[PointIndex];
+
+		if (PointIndex < InConstraint.SavedManifoldPoints.Num())
+		{
+			const Chaos::FSavedManifoldPoint& CurrentChaosSavedManifoldPoint = InConstraint.SavedManifoldPoints[PointIndex];
+
+			Algo::Transform(CurrentChaosSavedManifoldPoint.ShapeContactPoints, CurrentCVDMainFoldPoint.ShapeContactPoints, &FChaosVDDataWrapperUtils::ConvertToFVector);
+		}
+
+		if (PointIndex < InConstraint.ManifoldPoints.Num())
+		{
+			const Chaos::FManifoldPoint& CurrentChaosMainFoldPoint = InConstraint.ManifoldPoints[PointIndex];
+			CopyManifoldPointsToDataWrapper(CurrentChaosMainFoldPoint, CurrentCVDMainFoldPoint);
+		}
+		
+
+		if (PointIndex < InConstraint.ManifoldPointResults.Num())
+		{
+			const Chaos::FManifoldPointResult& CurrentChaosMainFoldPointResult = InConstraint.ManifoldPointResults[PointIndex];
+			CopyManifoldPointResultsToDataWrapper(CurrentChaosMainFoldPointResult, CurrentCVDMainFoldPoint);
+		}
+	}
+
+	return MoveTemp(WrappedConstraintData);
+}
+
+FChaosVDParticlePairMidPhase FChaosVDDataWrapperUtils::BuildMidPhaseDataWrapperFromMidPhase(const Chaos::FParticlePairMidPhase& InMidPhase)
+{
+	FChaosVDParticlePairMidPhase WrappedMidPhaseData;
+	
+	WrappedMidPhaseData.bIsActive = InMidPhase.Flags.bIsActive;
+	WrappedMidPhaseData.bIsCCD = InMidPhase.Flags.bIsCCD;
+	WrappedMidPhaseData.bIsCCDActive = InMidPhase.Flags.bIsCCDActive;
+	WrappedMidPhaseData.bIsSleeping = InMidPhase.Flags.bIsSleeping;
+	WrappedMidPhaseData.bIsModified = InMidPhase.Flags.bIsModified;
+	WrappedMidPhaseData.LastUsedEpoch = InMidPhase.LastUsedEpoch;
+
+	WrappedMidPhaseData.Particle0Idx = InMidPhase.Particle0->UniqueIdx().Idx;
+	WrappedMidPhaseData.Particle1Idx = InMidPhase.Particle1->UniqueIdx().Idx;
+
+	InMidPhase.VisitConstCollisions([&WrappedMidPhaseData](const Chaos::FPBDCollisionConstraint& Constraint)
+	{
+		FChaosVDConstraint WrappedConstraintData = FChaosVDDataWrapperUtils::BuildConstraintDataWrapperFromConstraint(Constraint);
+		WrappedMidPhaseData.Constraints.Add(MoveTemp(WrappedConstraintData));
+		return Chaos::ECollisionVisitorResult::Continue;
+	}, false);
+
+	return MoveTemp(WrappedMidPhaseData);
+}
 #endif //WITH_CHAOS_VISUAL_DEBUGGER
diff --git a/Engine/Source/Runtime/Experimental/Chaos/Private/ChaosVisualDebugger/ChaosVisualDebuggerTrace.cpp b/Engine/Source/Runtime/Experimental/Chaos/Private/ChaosVisualDebugger/ChaosVisualDebuggerTrace.cpp
--- a/Engine/Source/Runtime/Experimental/Chaos/Private/ChaosVisualDebugger/ChaosVisualDebuggerTrace.cpp
+++ b/Engine/Source/Runtime/Experimental/Chaos/Private/ChaosVisualDebugger/ChaosVisualDebuggerTrace.cpp
@@ -1,15 +1,15 @@
 // Copyright Epic Games, Inc. All Rights Reserved.
 #include "ChaosVisualDebugger/ChaosVisualDebuggerTrace.h"
 
-#include "ChaosVisualDebugger/ChaosVDDataWrapperUtils.h"
-#include "Chaos/PBDRigidsSOAs.h"
-#include "Tasks/Task.h"
-
 #if CHAOS_VISUAL_DEBUGGER_ENABLED
-#include "Chaos/ParticleHandle.h"
+
 #include "Chaos/Framework/PhysicsSolverBase.h"
 #include "Chaos/ImplicitObject.h"
+#include "Chaos/ParticleHandle.h"
+#include "Chaos/PBDRigidsSOAs.h"
+#include "ChaosVisualDebugger/ChaosVDDataWrapperUtils.h"
 #include "Compression/OodleDataCompressionUtil.h"
+#include "DataWrappers/ChaosVDCollisionDataWrappers.h"
 #include "DataWrappers/ChaosVDImplicitObjectDataWrapper.h"
 #include "DataWrappers/ChaosVDParticleDataWrapper.h"
 #include "HAL/CriticalSection.h"
@@ -146,7 +146,7 @@ void FChaosVisualDebuggerTrace::TraceParticle(Chaos::FGeometryParticleHandle* Pa
 
 		ParticleDataWrapper.Serialize(MemWriterAr);
 
-		TraceBinaryData(TLSDataBuffer.BufferRef, TEXT("FChaosVDParticleDataWrapper"));
+		TraceBinaryData(TLSDataBuffer.BufferRef, FChaosVDParticleDataWrapper::WrapperTypeName);
 	}
 }
 
@@ -243,6 +243,74 @@ bool FChaosVisualDebuggerTrace::ShouldPerformFullCapture(int32 SolverID)
 	return FoundSolverID == nullptr;
 }
 
+void FChaosVisualDebuggerTrace::TraceMidPhase(const Chaos::FParticlePairMidPhase* MidPhase)
+{
+	if (!IsTracing())
+	{
+		return;
+	}
+
+	const FChaosVDContext* CVDContextData = FChaosVDThreadContext::Get().GetCurrentContext();
+	if (!ensure(CVDContextData))
+	{
+		return;
+	}
+
+	FChaosVDParticlePairMidPhase CVDMidPhase = FChaosVDDataWrapperUtils::BuildMidPhaseDataWrapperFromMidPhase(*MidPhase);
+	CVDMidPhase.SolverID = CVDContextData->Id;
+
+	FChaosVDScopedTLSBufferAccessor CVDBuffer;
+	FMemoryWriter MemWriter(CVDBuffer.BufferRef);
+		
+	CVDMidPhase.Serialize(MemWriter);
+
+	TraceBinaryData(CVDBuffer.BufferRef, FChaosVDParticlePairMidPhase::WrapperTypeName);
+}
+
+void FChaosVisualDebuggerTrace::TraceCollisionConstraint(const Chaos::FPBDCollisionConstraint* CollisionConstraint)
+{
+	if (!IsTracing())
+	{
+		return;
+	}
+
+	const FChaosVDContext* CVDContextData = FChaosVDThreadContext::Get().GetCurrentContext();
+	if (!ensure(CVDContextData))
+	{
+		return;
+	}
+
+	FChaosVDConstraint CVDConstraint = FChaosVDDataWrapperUtils::BuildConstraintDataWrapperFromConstraint(*CollisionConstraint);
+	CVDConstraint.SolverID = CVDContextData->Id;
+
+	FChaosVDScopedTLSBufferAccessor CVDBuffer;
+	FMemoryWriter MemWriter(CVDBuffer.BufferRef);
+		
+	CVDConstraint.Serialize(MemWriter);
+
+	TraceBinaryData(CVDBuffer.BufferRef, FChaosVDConstraint::WrapperTypeName);
+}
+
+void FChaosVisualDebuggerTrace::TraceCollisionConstraintView(TArrayView<Chaos::FPBDCollisionConstraint* const> CollisionConstraintView)
+{
+	if (!IsTracing())
+	{
+		return;
+	}
+
+	const FChaosVDContext* CVDContextData = FChaosVDThreadContext::Get().GetCurrentContext();
+	if (!ensure(CVDContextData))
+	{
+		return;
+	}
+
+	ParallelFor(CollisionConstraintView.Num(), [&CollisionConstraintView, CopyContext = *CVDContextData](int32 ConstraintIndex)
+	{
+		CVD_SCOPE_CONTEXT(CopyContext);
+		TraceCollisionConstraint(CollisionConstraintView[ConstraintIndex]);
+	});
+}
+
 void FChaosVisualDebuggerTrace::TraceSolverFrameStart(const FChaosVDContext& ContextData, const FString& InDebugName)
 {
 	if (!IsTracing())
@@ -352,7 +420,7 @@ void FChaosVisualDebuggerTrace::TraceSolverSimulationSpace(const Chaos::FRigidTr
 		<< CVD_TRACE_ROTATOR_ON_EVENT(ChaosVDSolverSimulationSpace, Rotation, Transform.GetRotation());
 }
 
-void FChaosVisualDebuggerTrace::TraceBinaryData(const TArray<uint8>& InData, const FString& TypeName)
+void FChaosVisualDebuggerTrace::TraceBinaryData(const TArray<uint8>& InData, FStringView TypeName)
 {
 	if (!IsTracing())
 	{
@@ -386,7 +454,7 @@ void FChaosVisualDebuggerTrace::TraceBinaryData(const TArray<uint8>& InData, con
 
 	UE_TRACE_LOG(ChaosVDLogger, ChaosVDBinaryDataStart, ChaosVDChannel)
 		<< ChaosVDBinaryDataStart.Cycle(FPlatformTime::Cycles64())
-		<< ChaosVDBinaryDataStart.TypeName(*TypeName, TypeName.Len())
+		<< ChaosVDBinaryDataStart.TypeName(TypeName.GetData(), TypeName.Len())
 		<< ChaosVDBinaryDataStart.DataID(DataID)
 		<< ChaosVDBinaryDataStart.DataSize(DataSize)
 		<< ChaosVDBinaryDataStart.OriginalSize(InData.Num())
@@ -443,7 +511,7 @@ void FChaosVisualDebuggerTrace::TraceImplicitObject(FChaosVDImplicitObjectWrappe
 
 	WrappedGeometryData.Serialize(Ar);
 
-	TraceBinaryData(TLSDataBuffer.BufferRef, TEXT("FChaosVDImplicitObjectDataWrapper"));
+	TraceBinaryData(TLSDataBuffer.BufferRef, FChaosVDImplicitObjectWrapper::WrapperTypeName);
 }	
 
 void FChaosVisualDebuggerTrace::RegisterEventHandlers()
diff --git a/Engine/Source/Runtime/Experimental/Chaos/Public/Chaos/Collision/PBDCollisionConstraint.h b/Engine/Source/Runtime/Experimental/Chaos/Public/Chaos/Collision/PBDCollisionConstraint.h
--- a/Engine/Source/Runtime/Experimental/Chaos/Public/Chaos/Collision/PBDCollisionConstraint.h
+++ b/Engine/Source/Runtime/Experimental/Chaos/Public/Chaos/Collision/PBDCollisionConstraint.h
@@ -15,6 +15,8 @@
 #include "Chaos/ParticleHandle.h"
 #include "Chaos/Vector.h"
 
+class FChaosVDDataWrapperUtils;
+
 namespace Chaos
 {
 	namespace Private
@@ -857,6 +859,8 @@ namespace Chaos
 		// The penetration we leave behind when rolling back to a CCD time of impact. Should be less than or equal to CCDEnablePenetration.
 		// Calculated from particle pair properties when constraint is created.
 		FRealSingle CCDTargetPenetration;
+
+		friend class ::FChaosVDDataWrapperUtils;
 	};
 
 
diff --git a/Engine/Source/Runtime/Experimental/Chaos/Public/Chaos/Collision/ParticlePairMidPhase.h b/Engine/Source/Runtime/Experimental/Chaos/Public/Chaos/Collision/ParticlePairMidPhase.h
--- a/Engine/Source/Runtime/Experimental/Chaos/Public/Chaos/Collision/ParticlePairMidPhase.h
+++ b/Engine/Source/Runtime/Experimental/Chaos/Public/Chaos/Collision/ParticlePairMidPhase.h
@@ -11,6 +11,8 @@
 #include "ProfilingDebugging/CsvProfiler.h"
 
 
+class FChaosVDDataWrapperUtils;
+
 namespace Chaos
 {
 	namespace Private
@@ -435,6 +437,8 @@ namespace Chaos
 		// Indices into the arrays of collisions on the particles. This is a cookie for use by FParticleCollisions
 		int32 ParticleCollisionsIndex0;					// 4 bytes
 		int32 ParticleCollisionsIndex1;					// 4 bytes
+		
+		friend ::FChaosVDDataWrapperUtils;
 	};
 
 
diff --git a/Engine/Source/Runtime/Experimental/Chaos/Public/Chaos/Collision/SpatialAccelerationBroadPhase.h b/Engine/Source/Runtime/Experimental/Chaos/Public/Chaos/Collision/SpatialAccelerationBroadPhase.h
--- a/Engine/Source/Runtime/Experimental/Chaos/Public/Chaos/Collision/SpatialAccelerationBroadPhase.h
+++ b/Engine/Source/Runtime/Experimental/Chaos/Public/Chaos/Collision/SpatialAccelerationBroadPhase.h
@@ -338,9 +338,13 @@ namespace Chaos
 				SCOPE_CYCLE_COUNTER(STAT_Collisions_MidPhase);
 				CSV_SCOPED_TIMING_STAT(PhysicsVerbose, DetectCollisions_AssignMidPhases);
 
+				FChaosVDContextWrapper CVDContext;
+				CVD_GET_WRAPPED_CURRENT_CONTEXT(CVDContext);
+
 				// Find or assign a midphase to each overlapping particle pair
-				const auto& AssignMidPhasesWorker = [this, Dt](const int32 ContextIndex)
+				const auto& AssignMidPhasesWorker = [this, Dt, &CVDContext](const int32 ContextIndex)
 				{
+					CVD_SCOPE_CONTEXT(CVDContext.Context);
 					AssignMidPhases(BroadphaseContexts[ContextIndex]);
 				};
 				PhysicsParallelFor(NumActiveBroadphaseContexts, AssignMidPhasesWorker, bDisableCollisionParallelFor);
@@ -361,8 +365,13 @@ namespace Chaos
 		{
 			SCOPE_CYCLE_COUNTER(STAT_Collisions_NarrowPhase);
 			CSV_SCOPED_TIMING_STAT(PhysicsVerbose, DetectCollisions_MidPhase);
-			const auto& ProcessMidPhasesWorker = [this, Dt](const int32 ContextIndex)
+
+			FChaosVDContextWrapper CVDContext;
+			CVD_GET_WRAPPED_CURRENT_CONTEXT(CVDContext);
+	
+			const auto& ProcessMidPhasesWorker = [this, Dt, &CVDContext](const int32 ContextIndex)
 			{
+				CVD_SCOPE_CONTEXT(CVDContext.Context);
 				ProcessMidPhases(Dt, BroadphaseContexts[ContextIndex]);
 			};
 			PhysicsParallelFor(NumActiveBroadphaseContexts, ProcessMidPhasesWorker, bDisableCollisionParallelFor);
@@ -570,6 +579,9 @@ namespace Chaos
 					// Get the midphase for this pair
 					FParticlePairMidPhase* MidPhase = ContextAllocator->GetMidPhase(Overlap.Particles[0], Overlap.Particles[1], Overlap.Particles[Overlap.SearchParticleIndex], BroadphaseContext.CollisionContext);
 					BroadphaseContext.MidPhases[MidPhaseIndex] = MidPhase;
+
+					CVD_TRACE_MID_PHASE(MidPhase);
+
 					++MidPhaseIndex;
 				}
 			}
@@ -600,6 +612,8 @@ namespace Chaos
 
 				// Run MidPhase + NarrowPhase
 				BroadphaseContext.MidPhases[Index]->GenerateCollisions(BroadphaseContext.CollisionContext.GetSettings().BoundsExpansion, Dt, BroadphaseContext.CollisionContext);
+
+				CVD_TRACE_MID_PHASE(BroadphaseContext.MidPhases[Index]);
 			}
 
 			PHYSICS_CSV_CUSTOM_EXPENSIVE(PhysicsCounters, NumFromBroadphase, NumPotentials, ECsvCustomStatOp::Accumulate);
diff --git a/Engine/Source/Runtime/Experimental/Chaos/Public/ChaosVisualDebugger/ChaosVDDataWrapperUtils.h b/Engine/Source/Runtime/Experimental/Chaos/Public/ChaosVisualDebugger/ChaosVDDataWrapperUtils.h
--- a/Engine/Source/Runtime/Experimental/Chaos/Public/ChaosVisualDebugger/ChaosVDDataWrapperUtils.h
+++ b/Engine/Source/Runtime/Experimental/Chaos/Public/ChaosVisualDebugger/ChaosVDDataWrapperUtils.h
@@ -3,9 +3,19 @@
 #pragma once
 
 #if WITH_CHAOS_VISUAL_DEBUGGER
+#include "Chaos/Collision/ContactPoint.h"
 #include "Chaos/ParticleHandleFwd.h"
+#include "DataWrappers/ChaosVDCollisionDataWrappers.h"
+
+namespace Chaos
+{
+	class FParticlePairMidPhase;
+	class FPBDCollisionConstraint;
+}
 
 class FChaosVisualDebuggerTrace;
+struct FChaosVDConstraint;
+struct FChaosVDParticlePairMidPhase;
 struct FChaosVDParticleDataWrapper;
 
 /**
@@ -18,9 +28,25 @@ struct FChaosVDParticleDataWrapper;
 class FChaosVDDataWrapperUtils
 {
 private:
-	/** Creates and populates a FChaosVDParticleDataWrapper with the data of the provided ParticleHandlePtr */
+
+	/** Takes a FManifoldPoint and copies the relevant data to the CVD counterpart */
+	static void CopyManifoldPointsToDataWrapper(const Chaos::FManifoldPoint& InCopyFrom, FChaosVDManifoldPoint& OutCopyTo);
+
+	/** Takes a FManifoldPointResult and copies the relevant data to the CVD counterpart */
+	static void CopyManifoldPointResultsToDataWrapper(const Chaos::FManifoldPointResult& InCopyFrom, FChaosVDManifoldPoint& OutCopyTo);
+
+	/** Creates and populates a FChaosVDParticleDataWrapper with the data of the provided FGeometryParticleHandle */
 	static FChaosVDParticleDataWrapper BuildParticleDataWrapperFromParticle(const Chaos::FGeometryParticleHandle* ParticleHandlePtr);
 
+	/** Creates and populates a FChaosVDConstraint with the data of the provided FPBDCollisionConstraint */
+	static FChaosVDConstraint BuildConstraintDataWrapperFromConstraint(const Chaos::FPBDCollisionConstraint& InConstraint);
+
+	/** Creates and populates a FChaosVDParticlePairMidPhase with the data of the provided FParticlePairMidPhase */
+	static FChaosVDParticlePairMidPhase BuildMidPhaseDataWrapperFromMidPhase(const Chaos::FParticlePairMidPhase& InMidPhase);
+
+	/** Converts a Chaos::FVec3 to a FVector. It is worth notice that FVector is double precision and FVec3 is single */
+	static FVector ConvertToFVector(const Chaos::FVec3f& VectorRef) { return FVector(VectorRef); }
+
 	friend FChaosVisualDebuggerTrace;
 };
 #endif //WITH_CHAOS_VISUAL_DEBUGGER
diff --git a/Engine/Source/Runtime/Experimental/Chaos/Public/ChaosVisualDebugger/ChaosVisualDebuggerTrace.h b/Engine/Source/Runtime/Experimental/Chaos/Public/ChaosVisualDebugger/ChaosVisualDebuggerTrace.h
--- a/Engine/Source/Runtime/Experimental/Chaos/Public/ChaosVisualDebugger/ChaosVisualDebuggerTrace.h
+++ b/Engine/Source/Runtime/Experimental/Chaos/Public/ChaosVisualDebugger/ChaosVisualDebuggerTrace.h
@@ -64,6 +64,20 @@
 #ifndef CVD_TRACE_PARTICLE_DESTROYED
 	#define CVD_TRACE_PARTICLE_DESTROYED(DestroyedParticleHandle)
 #endif
+
+#ifndef CVD_TRACE_MID_PHASE
+	#define CVD_TRACE_MID_PHASE(MidPhase)
+#endif
+#ifndef CVD_TRACE_COLLISION_CONSTRAINT
+	#define CVD_TRACE_COLLISION_CONSTRAINT(Constraint)
+#endif
+#ifndef CVD_TRACE_COLLISION_CONSTRAINT_VIEW
+	#define CVD_TRACE_COLLISION_CONSTRAINT_VIEW(ConstraintView)
+#endif
+
+#ifndef CVD_TRACE_PARTICLES_VIEW
+	#define CVD_TRACE_PARTICLES_VIEW(ParticleHandlesView)
+#endif
 #else
 
 #include "ChaosVDRuntimeModule.h"
@@ -180,6 +194,11 @@ UE_TRACE_EVENT_END()
 		FChaosVisualDebuggerTrace::TraceParticles(ParticleHandles);
 #endif
 
+#ifndef CVD_TRACE_PARTICLES_VIEW
+	#define CVD_TRACE_PARTICLES_VIEW(ParticleHandlesView) \
+	FChaosVisualDebuggerTrace::TraceParticlesView(ParticleHandlesView);
+#endif
+
 #ifndef CVD_TRACE_PARTICLES_SOA
 	#define CVD_TRACE_PARTICLES_SOA(ParticleSoA) \
 	FChaosVisualDebuggerTrace::TraceParticlesSoA(ParticleSoA);
@@ -234,6 +253,20 @@ UE_TRACE_EVENT_END()
 	FChaosVisualDebuggerTrace::TraceParticleDestroyed(DestroyedParticleHandle);
 #endif
 
+#ifndef CVD_TRACE_MID_PHASE
+	#define CVD_TRACE_MID_PHASE(MidPhase) \
+	FChaosVisualDebuggerTrace::TraceMidPhase(MidPhase);
+#endif
+
+#ifndef CVD_TRACE_COLLISION_CONSTRAINT
+	#define CVD_TRACE_COLLISION_CONSTRAINT(Constraint) \
+	FChaosVisualDebuggerTrace::TraceCollisionConstraint(Constraint);
+#endif
+
+#ifndef CVD_TRACE_COLLISION_CONSTRAINT_VIEW
+	#define CVD_TRACE_COLLISION_CONSTRAINT_VIEW(ConstraintView) \
+	FChaosVisualDebuggerTrace::TraceCollisionConstraintView(ConstraintView);
+#endif
 
 struct FChaosVDContext;
 
@@ -244,6 +277,9 @@ namespace Chaos
 	class FPhysicsSolverBase;
 	template <typename T, int d>
 	class TGeometryParticleHandles;
+
+	class FPBDCollisionConstraint;
+	class FParticlePairMidPhase;
 }
 
 using FChaosVDImplicitObjectWrapper = FChaosVDImplicitObjectDataWrapper<Chaos::TSerializablePtr<Chaos::FImplicitObject>, Chaos::FChaosArchive>;
@@ -288,6 +324,15 @@ public:
 	template<typename ParticleType>
 	static void TraceParticlesView(const Chaos::TParticleView<ParticleType>& ParticlesView);
 
+	/** Traces a Particle pair MidPhase as binary data */
+	static CHAOS_API void TraceMidPhase(const Chaos::FParticlePairMidPhase* MidPhase);
+
+	/** Traces a Particle pair MidPhase as binary data */
+	static CHAOS_API void TraceCollisionConstraint(const Chaos::FPBDCollisionConstraint* CollisionConstraint);
+
+	/** Traces a Particle pair MidPhase as binary data in parallel */
+	static CHAOS_API void TraceCollisionConstraintView(TArrayView<Chaos::FPBDCollisionConstraint* const> CollisionConstraintView);
+
 	/** Traces the start of a solver frame and it pushes its context data to the CVD TLS context stack */
 	static CHAOS_API void TraceSolverFrameStart(const FChaosVDContext& ContextData, const FString& InDebugName);
 	
@@ -311,7 +356,7 @@ public:
 	 * @param InData Data to trace
 	 * @param TypeName Type name the data represents. It is used during Trace Analysis serialize it back (this is not automatic)
 	 */
-	static CHAOS_API void TraceBinaryData(const TArray<uint8>& InData, const FString& TypeName);
+	static CHAOS_API void TraceBinaryData(const TArray<uint8>& InData, FStringView TypeName);
 
 	/**
 	 * Serializes the implicit object contained in the wrapper and trace its it as binary data
diff --git a/Engine/Source/Runtime/Experimental/ChaosVDData/Public/ChaosVDRecording.h b/Engine/Source/Runtime/Experimental/ChaosVDData/Public/ChaosVDRecording.h
--- a/Engine/Source/Runtime/Experimental/ChaosVDData/Public/ChaosVDRecording.h
+++ b/Engine/Source/Runtime/Experimental/ChaosVDData/Public/ChaosVDRecording.h
@@ -3,6 +3,8 @@
 #pragma once
 
 #include "Chaos/Core.h"
+#include "Containers/UnrealString.h"
+#include "DataWrappers/ChaosVDCollisionDataWrappers.h"
 #include "DataWrappers/ChaosVDParticleDataWrapper.h"
 #include "UObject/ObjectMacros.h"
 #include "Containers/UnrealString.h"
@@ -19,6 +21,9 @@ struct FChaosVDStepData
 {
 	FString StepName;
 	TArray<FChaosVDParticleDataWrapper> RecordedParticlesData;
+	TArray<TSharedPtr<FChaosVDParticlePairMidPhase>> RecordedMidPhases;
+	TMap<int32, TArray<FChaosVDConstraint>> RecordedConstraintsByParticleID;
+	TMap<int32, TArray<TSharedPtr<FChaosVDParticlePairMidPhase>>> RecordedMidPhasesByParticleID;
 	TSet<int32> ParticlesDestroyedIDs;
 };
 
@@ -185,4 +190,5 @@ protected:
 	TMap<uint32, TSharedPtr<const Chaos::FImplicitObject>> ImplicitObjects;
 
 	friend class FChaosVDTraceProvider;
+	friend class FChaosVDTraceImplicitObjectProcessor;
 };
diff --git a/Engine/Source/Runtime/Experimental/ChaosVisualDebugger/Private/DataWrappers/ChaosVDCollisionDataWrappers.cpp b/Engine/Source/Runtime/Experimental/ChaosVisualDebugger/Private/DataWrappers/ChaosVDCollisionDataWrappers.cpp
new file mode 100644
--- /dev/null
+++ b/Engine/Source/Runtime/Experimental/ChaosVisualDebugger/Private/DataWrappers/ChaosVDCollisionDataWrappers.cpp
@@ -0,0 +1,95 @@
+﻿// Copyright Epic Games, Inc. All Rights Reserved.
+
+#include "DataWrappers/ChaosVDCollisionDataWrappers.h"
+
+bool FChaosVDContactPoint::Serialize(FArchive& Ar)
+{
+	Ar << ShapeContactPoints;
+	Ar << ShapeContactNormal;
+	Ar << Phi;
+	Ar << FaceIndex;
+	Ar << ContactType;
+
+	return true;
+}
+
+bool FChaosVDManifoldPoint::Serialize(FArchive& Ar)
+{
+	FArchive_Serialize_BitfieldBool(Ar, bDisabled);
+	FArchive_Serialize_BitfieldBool(Ar, bWasRestored);
+	FArchive_Serialize_BitfieldBool(Ar, bWasReplaced);
+
+	FArchive_Serialize_BitfieldBool(Ar, bHasStaticFrictionAnchor);
+	FArchive_Serialize_BitfieldBool(Ar, bIsValid);
+	FArchive_Serialize_BitfieldBool(Ar, bInsideStaticFrictionCone);
+
+	Ar << NetPushOut;
+	Ar << NetImpulse;
+
+	Ar << TargetPhi;
+	Ar << ShapeAnchorPoints;
+	Ar << InitialShapeContactPoints;
+	Ar << ContactPoint;
+	Ar << ShapeContactPoints;
+
+	return true;
+}
+
+bool FChaosVDConstraint::Serialize(FArchive& Ar)
+{
+	FArchive_Serialize_BitfieldBool(Ar, bDisabled);
+	FArchive_Serialize_BitfieldBool(Ar, bUseManifold);
+	FArchive_Serialize_BitfieldBool(Ar, bUseIncrementalManifold);
+	FArchive_Serialize_BitfieldBool(Ar, bCanRestoreManifold);
+	FArchive_Serialize_BitfieldBool(Ar, bWasManifoldRestored);
+	FArchive_Serialize_BitfieldBool(Ar, bIsQuadratic0);
+	FArchive_Serialize_BitfieldBool(Ar, bIsQuadratic1);
+	FArchive_Serialize_BitfieldBool(Ar, bIsProbeUnmodified);
+	FArchive_Serialize_BitfieldBool(Ar, bIsProbe);
+	FArchive_Serialize_BitfieldBool(Ar, bCCDEnabled);
+	FArchive_Serialize_BitfieldBool(Ar, bCCDSweepEnabled);
+	FArchive_Serialize_BitfieldBool(Ar, bModifierApplied);
+	FArchive_Serialize_BitfieldBool(Ar, bMaterialSet);
+
+	Ar << AccumulatedImpulse;
+	Ar << ShapesType;
+	Ar << ShapeWorldTransforms;
+	Ar << CullDistance;
+	Ar << CollisionMargins;	
+	Ar << CollisionTolerance;	
+	Ar << ClosestManifoldPointIndex;	
+	Ar << ExpectedNumManifoldPoints;	
+	Ar << LastShapeWorldPositionDelta;
+	Ar << LastShapeWorldRotationDelta;
+	Ar << Stiffness;
+	Ar << CCDTimeOfImpact;
+	Ar << CCDEnablePenetration;
+	Ar << CCDTargetPenetration;
+	Ar << ManifoldPoints;
+	Ar << SolverID;
+	Ar << Particle0Index;
+	Ar << Particle1Index;
+
+	return true;
+}
+
+
+bool FChaosVDParticlePairMidPhase::Serialize(FArchive& Ar)
+{
+	Ar << SolverID;
+	
+	FArchive_Serialize_BitfieldBool(Ar, bIsActive);
+	FArchive_Serialize_BitfieldBool(Ar, bIsCCD);
+	FArchive_Serialize_BitfieldBool(Ar, bIsCCDActive);
+	FArchive_Serialize_BitfieldBool(Ar, bIsSleeping);
+	FArchive_Serialize_BitfieldBool(Ar, bIsModified);
+		
+	Ar << LastUsedEpoch;
+
+	Ar << Particle0Idx;
+	Ar << Particle1Idx;
+
+	Ar << Constraints;
+
+	return true;
+}
diff --git a/Engine/Source/Runtime/Experimental/ChaosVisualDebugger/Public/DataWrappers/ChaosVDCollisionDataWrappers.h b/Engine/Source/Runtime/Experimental/ChaosVisualDebugger/Public/DataWrappers/ChaosVDCollisionDataWrappers.h
new file mode 100644
--- /dev/null
+++ b/Engine/Source/Runtime/Experimental/ChaosVisualDebugger/Public/DataWrappers/ChaosVDCollisionDataWrappers.h
@@ -0,0 +1,304 @@
+﻿// Copyright Epic Games, Inc. All Rights Reserved.
+
+#pragma once
+#include "HAL/Platform.h"
+
+#include "ChaosVDCollisionDataWrappers.generated.h"
+
+namespace Chaos
+{
+	enum class EParticleType : uint8;
+	class FChaosArchive;
+}
+
+UENUM()
+enum class EChaosVDContactShapesType
+{
+	Unknown,
+	SphereSphere,
+	SphereCapsule,
+	SphereBox,
+	SphereConvex,
+	SphereTriMesh,
+	SphereHeightField,
+	SpherePlane,
+	CapsuleCapsule,
+	CapsuleBox,
+	CapsuleConvex,
+	CapsuleTriMesh,
+	CapsuleHeightField,
+	BoxBox,
+	BoxConvex,
+	BoxTriMesh,
+	BoxHeightField,
+	BoxPlane,
+	ConvexConvex,
+	ConvexTriMesh,
+	ConvexHeightField,
+	GenericConvexConvex,
+	LevelSetLevelSet,
+
+	NumShapesTypes
+};
+
+UENUM()
+enum class EChaosVDContactPointType : int8
+{
+	Unknown,
+	VertexPlane,
+	EdgeEdge,
+	PlaneVertex,
+	VertexVertex,
+};
+
+USTRUCT()
+struct CHAOSVDRUNTIME_API FChaosVDContactPoint
+{
+	GENERATED_BODY()
+
+	// Shape-space contact points on the two bodies
+	UPROPERTY(EditAnywhere, Category=Contact)
+	TArray<FVector> ShapeContactPoints;
+
+	// Shape-space contact normal on the second shape with direction that points away from shape 1
+	UPROPERTY(EditAnywhere, Category=Contact)
+	FVector  ShapeContactNormal;
+
+	// Contact separation (negative for overlap)
+	UPROPERTY(EditAnywhere, Category=Contact)
+	float Phi;
+
+	// Face index of the shape we hit. Only valid for Heightfield and Trimesh contact points, otherwise INDEX_NONE
+	UPROPERTY(EditAnywhere, Category=Contact)
+	int32 FaceIndex;
+
+	// Whether this is a vertex-plane contact, edge-edge contact etc.
+	UPROPERTY(EditAnywhere, Category=Contact)
+	EChaosVDContactPointType ContactType;
+
+	bool Serialize(FArchive& Ar);
+};
+
+template<>
+struct TStructOpsTypeTraits<FChaosVDContactPoint> : public TStructOpsTypeTraitsBase2<FChaosVDContactPoint>
+{
+	enum
+	{
+		WithSerializer = true,
+	};
+};
+
+inline FArchive& operator<<(FArchive& Ar, FChaosVDContactPoint& Data)
+{
+	Data.Serialize(Ar);
+	return Ar;
+}
+
+USTRUCT()
+struct CHAOSVDRUNTIME_API FChaosVDManifoldPoint
+{
+	GENERATED_BODY()
+
+	UPROPERTY(EditAnywhere, Category=ManifoldPointFalgs)
+	uint8 bDisabled:1 = false;
+	UPROPERTY(EditAnywhere, Category=ManifoldPointFalgs)
+	uint8 bWasRestored:1 = false;
+	UPROPERTY(EditAnywhere, Category=ManifoldPointFalgs)
+	uint8 bWasReplaced:1 = false;
+	UPROPERTY(EditAnywhere, Category=ManifoldPointFalgs)
+	uint8 bHasStaticFrictionAnchor:1 = false;
+
+	UPROPERTY(EditAnywhere, Category=ManifoldPointResult)
+	uint8 bIsValid:1 = false;
+	UPROPERTY(EditAnywhere, Category=ManifoldPointResult)
+	uint8 bInsideStaticFrictionCone:1 = false;
+
+	UPROPERTY(EditAnywhere, Category=ManifoldPointResult)
+	FVector NetPushOut;
+	UPROPERTY(EditAnywhere, Category=ManifoldPointResult)
+	FVector NetImpulse;
+
+	UPROPERTY(EditAnywhere, Category=ManifoldPoint)
+	float TargetPhi;
+	UPROPERTY(EditAnywhere, Category=ManifoldPoint)
+	TArray<FVector> ShapeAnchorPoints;
+	UPROPERTY(EditAnywhere, Category=ManifoldPoint)
+	TArray<FVector> InitialShapeContactPoints;
+	UPROPERTY(EditAnywhere, Category=ManifoldPoint)
+	FChaosVDContactPoint ContactPoint;
+
+	UPROPERTY(EditAnywhere, Category=SavedManifoldPoint)
+	TArray<FVector> ShapeContactPoints;
+
+	bool Serialize(FArchive& Ar);
+};
+
+template<>
+struct TStructOpsTypeTraits<FChaosVDManifoldPoint> : public TStructOpsTypeTraitsBase2<FChaosVDManifoldPoint>
+{
+	enum
+	{
+		WithSerializer = true,
+	};
+};
+
+inline FArchive& operator<<(FArchive& Ar, FChaosVDManifoldPoint& Data)
+{
+	Data.Serialize(Ar);
+	return Ar;
+}
+
+USTRUCT()
+struct CHAOSVDRUNTIME_API FChaosVDConstraint
+{
+	GENERATED_BODY()
+
+	inline static FStringView WrapperTypeName = TEXT("FChaosVDConstraint");
+
+	UPROPERTY(EditAnywhere, Category=Falgs)
+	uint8 bDisabled:1 = false;
+	UPROPERTY(EditAnywhere, Category=Falgs)
+	uint8 bUseManifold:1 = false;
+	UPROPERTY(EditAnywhere, Category=Falgs)
+	uint8 bUseIncrementalManifold:1 = false;
+	UPROPERTY(EditAnywhere, Category=Falgs)
+	uint8 bCanRestoreManifold:1 = false;
+	UPROPERTY(EditAnywhere, Category=Falgs)
+	uint8 bWasManifoldRestored:1 = false;
+	UPROPERTY(EditAnywhere, Category=Falgs)
+	uint8 bIsQuadratic0:1 = false;
+	UPROPERTY(EditAnywhere, Category=Falgs)
+	uint8 bIsQuadratic1:1 = false;
+	UPROPERTY(EditAnywhere, Category=Falgs)
+	uint8 bIsProbeUnmodified:1 = false;
+	UPROPERTY(EditAnywhere, Category=Falgs)
+	uint8 bIsProbe:1 = false;
+	UPROPERTY(EditAnywhere, Category=Falgs)
+	uint8 bCCDEnabled:1 = false;
+	UPROPERTY(EditAnywhere, Category=Falgs)
+	uint8 bCCDSweepEnabled:1 = false;
+	UPROPERTY(EditAnywhere, Category=Falgs)
+	uint8 bModifierApplied:1 = false;
+	UPROPERTY(EditAnywhere, Category=Falgs)
+	uint8 bMaterialSet:1 = false;
+	
+	UPROPERTY(EditAnywhere, Category=Collision)
+	FVector AccumulatedImpulse;
+	
+	UPROPERTY(EditAnywhere, Category=Collision)
+	EChaosVDContactShapesType ShapesType;
+	
+	UPROPERTY(EditAnywhere, Category=Collision)
+	TArray<FTransform> ShapeWorldTransforms;
+	
+	UPROPERTY(EditAnywhere, Category=Collision)
+	float CullDistance;
+	
+	UPROPERTY(EditAnywhere, Category=Collision)
+	TArray<float> CollisionMargins;
+	
+	UPROPERTY(EditAnywhere, Category=Collision)
+	float CollisionTolerance;
+	
+	UPROPERTY(EditAnywhere, Category=Collision)
+	int32 ClosestManifoldPointIndex;
+	
+	UPROPERTY(EditAnywhere, Category=Collision)
+	int32 ExpectedNumManifoldPoints;
+	
+	UPROPERTY(EditAnywhere, Category=Collision)
+	FVector LastShapeWorldPositionDelta;
+	
+	UPROPERTY(EditAnywhere, Category=Collision)
+	FQuat LastShapeWorldRotationDelta;
+	
+	UPROPERTY(EditAnywhere, Category=Collision)
+	float Stiffness;
+	
+	UPROPERTY(EditAnywhere, Category=Collision)
+	float CCDTimeOfImpact;
+	
+	UPROPERTY(EditAnywhere, Category=Collision)
+	float CCDEnablePenetration;
+	
+	UPROPERTY(EditAnywhere, Category=Collision)
+	float CCDTargetPenetration;
+
+	UPROPERTY(EditAnywhere, Category=Collision)
+	TArray<FChaosVDManifoldPoint> ManifoldPoints;
+
+	UPROPERTY(VisibleAnywhere, Category=Collision)
+	int32 Particle0Index = INDEX_NONE;
+	UPROPERTY(VisibleAnywhere, Category=Collision)
+	int32 Particle1Index = INDEX_NONE;
+
+	int32 SolverID = INDEX_NONE;
+
+	bool Serialize(FArchive& Ar);
+};
+
+template<>
+struct TStructOpsTypeTraits<FChaosVDConstraint> : public TStructOpsTypeTraitsBase2<FChaosVDConstraint>
+{
+	enum
+	{
+		WithSerializer = true,
+	};
+};
+
+inline FArchive& operator<<(FArchive& Ar, FChaosVDConstraint& Data)
+{
+	Data.Serialize(Ar);
+	return Ar;
+}
+
+USTRUCT()
+struct CHAOSVDRUNTIME_API FChaosVDParticlePairMidPhase
+{
+	GENERATED_BODY()
+
+	inline static FStringView WrapperTypeName = TEXT("FChaosVDParticlePairMidPhase");
+
+	UPROPERTY(VisibleAnywhere, Category=Solver)
+	int32 SolverID = INDEX_NONE;
+
+	UPROPERTY(EditAnywhere, Category=Falgs)
+	uint8 bIsActive:1 = false;
+	UPROPERTY(EditAnywhere, Category=Falgs)
+	uint8 bIsCCD:1 = false;
+	UPROPERTY(EditAnywhere, Category=Falgs)
+	uint8 bIsCCDActive:1 = false;
+	UPROPERTY(EditAnywhere, Category=Falgs)
+	uint8 bIsSleeping:1 = false;
+	UPROPERTY(EditAnywhere, Category=Falgs)
+	uint8 bIsModified:1 = false;
+
+	UPROPERTY(EditAnywhere, Category=Misc)
+	int32 LastUsedEpoch = 0;
+
+	UPROPERTY(EditAnywhere, Category=Particle)
+	int32 Particle0Idx = 0;
+	UPROPERTY(EditAnywhere, Category=Particle)
+	int32 Particle1Idx = 0;
+
+	UPROPERTY(EditAnywhere, Category=Particle)
+	TArray<FChaosVDConstraint> Constraints;
+
+	bool Serialize(FArchive& Ar);
+};
+
+template<>
+struct TStructOpsTypeTraits<FChaosVDParticlePairMidPhase> : public TStructOpsTypeTraitsBase2<FChaosVDParticlePairMidPhase>
+{
+	enum
+	{
+		WithSerializer = true,
+	};
+};
+
+
+inline FArchive& operator<<(FArchive& Ar, FChaosVDParticlePairMidPhase& Data)
+{
+	Data.Serialize(Ar);
+	return Ar;
+}
diff --git a/Engine/Source/Runtime/Experimental/ChaosVisualDebugger/Public/DataWrappers/ChaosVDImplicitObjectDataWrapper.h b/Engine/Source/Runtime/Experimental/ChaosVisualDebugger/Public/DataWrappers/ChaosVDImplicitObjectDataWrapper.h
--- a/Engine/Source/Runtime/Experimental/ChaosVisualDebugger/Public/DataWrappers/ChaosVDImplicitObjectDataWrapper.h
+++ b/Engine/Source/Runtime/Experimental/ChaosVisualDebugger/Public/DataWrappers/ChaosVDImplicitObjectDataWrapper.h
@@ -2,6 +2,7 @@
 
 #pragma once
 #include "Hal/Platform.h"
+#include "Containers/StringFwd.h"
 
 /** Data Wrapper for used to Trace Implicit objects
  * @note This should not be templated. But making it so for now so this can live in the ChaosVDRuntime module without referencing Chaos types directly.
@@ -10,6 +11,8 @@
 template<class SerializableImplicitType, class ArchiveType>
 struct FChaosVDImplicitObjectDataWrapper
 {
+	inline static FStringView WrapperTypeName = TEXT("FChaosVDImplicitObjectDataWrapper");
+
 	uint32 Hash;
 	SerializableImplicitType ImplicitObject;
 
diff --git a/Engine/Source/Runtime/Experimental/ChaosVisualDebugger/Public/DataWrappers/ChaosVDParticleDataWrapper.h b/Engine/Source/Runtime/Experimental/ChaosVisualDebugger/Public/DataWrappers/ChaosVDParticleDataWrapper.h
--- a/Engine/Source/Runtime/Experimental/ChaosVisualDebugger/Public/DataWrappers/ChaosVDParticleDataWrapper.h
+++ b/Engine/Source/Runtime/Experimental/ChaosVisualDebugger/Public/DataWrappers/ChaosVDParticleDataWrapper.h
@@ -4,6 +4,7 @@
 
 #include "UObject/ObjectMacros.h"
 #include "HAL/Platform.h"
+#include "DataWrappers/ChaosVDCollisionDataWrappers.h"
 
 #include "ChaosVDParticleDataWrapper.generated.h"
 
@@ -400,6 +401,8 @@ struct FChaosVDParticleDataWrapper : public FChaosVDParticleDataBase
 {
 	virtual ~FChaosVDParticleDataWrapper() override = default;
 
+	inline static FStringView WrapperTypeName = TEXT("FChaosVDParticleDataWrapper");
+
 	GENERATED_BODY()
 
 	FChaosVDParticleDataWrapper()
@@ -436,6 +439,12 @@ struct FChaosVDParticleDataWrapper : public FChaosVDParticleDataBase
 	UPROPERTY(EditAnywhere, Category= "Particle Particle Mass Props")
 	FChaosVDParticleMassProps ParticleMassProps;
 
+	UPROPERTY(EditAnywhere, Category= "Particle Collision")
+	TArray<FChaosVDParticlePairMidPhase> ParticleMidPhases;
+
+	UPROPERTY(EditAnywhere, Category= "Particle Collision")
+	TArray<FChaosVDConstraint> ParticleConstraints;
+	
 	/** Only used during recording */
 	TSharedPtr<FString> DebugNamePtr;
 	bool bHasDebugName = false;
