[ChaosVD] Fix CIS issue caused by case error on include file on ChaosVDImplicitObjectDataWrapper.h

diff --git a/Engine/Source/Runtime/Experimental/ChaosVisualDebugger/Public/DataWrappers/ChaosVDImplicitObjectDataWrapper.h b/Engine/Source/Runtime/Experimental/ChaosVisualDebugger/Public/DataWrappers/ChaosVDImplicitObjectDataWrapper.h
--- a/Engine/Source/Runtime/Experimental/ChaosVisualDebugger/Public/DataWrappers/ChaosVDImplicitObjectDataWrapper.h
+++ b/Engine/Source/Runtime/Experimental/ChaosVisualDebugger/Public/DataWrappers/ChaosVDImplicitObjectDataWrapper.h
@@ -1,7 +1,7 @@
 ﻿// Copyright Epic Games, Inc. All Rights Reserved.
 
 #pragma once
-#include "Hal/Platform.h"
+#include "HAL/Platform.h"
 #include "Containers/StringFwd.h"
 
 /** Data Wrapper for used to Trace Implicit objects
