Chaos - add IsValid to SerializablePtr to match shared and unique ptr API

diff --git a/Engine/Source/Runtime/Experimental/Chaos/Public/Chaos/Serializable.h b/Engine/Source/Runtime/Experimental/Chaos/Public/Chaos/Serializable.h
--- a/Engine/Source/Runtime/Experimental/Chaos/Public/Chaos/Serializable.h
+++ b/Engine/Source/Runtime/Experimental/Chaos/Public/Chaos/Serializable.h
@@ -28,6 +28,7 @@ public:
 	bool operator==(const TSerializablePtr<T>& Serializable) const { return Ptr == Serializable.Ptr; }
 	bool operator!=(const TSerializablePtr<T>& Serializable) const { return Ptr != Serializable.Ptr; }
 	operator bool() const { return Ptr != nullptr; }
+	bool IsValid() const { return Ptr != nullptr; }
 
 	template <typename R>
 	operator TSerializablePtr<R>() const
