Chaos - use zero bounds rather than empty bounds for particles without geometry

diff --git a/Engine/Source/Runtime/Experimental/Chaos/Private/Chaos/ImplicitObjectUnion.cpp b/Engine/Source/Runtime/Experimental/Chaos/Private/Chaos/ImplicitObjectUnion.cpp
--- a/Engine/Source/Runtime/Experimental/Chaos/Private/Chaos/ImplicitObjectUnion.cpp
+++ b/Engine/Source/Runtime/Experimental/Chaos/Private/Chaos/ImplicitObjectUnion.cpp
@@ -21,9 +21,25 @@ namespace Chaos
 		FAutoConsoleVariableRef CVarChaosUnionBVHEnabled(TEXT("p.Chaos.Collision.UnionBVH.Enabled"), bChaosUnionBVHEnabled, TEXT("Set to false to disable use of BVH during collision detection (without affecting creations and serialization)"));
 	}
 
+inline FAABB3 CalculateObjectsBounds(const TArrayView<TUniquePtr<FImplicitObject>>& Objects)
+{
+	if (Objects.IsEmpty())
+	{
+		// No geometry, so a point particle at our local origin
+		return FAABB3::ZeroAABB();
+	}
+
+	FAABB3 Bounds = Objects[0]->BoundingBox();
+	for (int32 ObjectIndex = 1; ObjectIndex < Objects.Num(); ++ObjectIndex)
+	{
+		Bounds.GrowToInclude(Objects[ObjectIndex]->BoundingBox());
+	}
+	return Bounds;
+}
+
 FImplicitObjectUnion::FImplicitObjectUnion() 
 	: FImplicitObject(EImplicitObject::HasBoundingBox, ImplicitObjectType::Union)
-	, MLocalBoundingBox()
+	, MLocalBoundingBox(FAABB3::ZeroAABB())
 	, NumLeafObjects(0)
 	, Flags()
 {
@@ -32,22 +48,12 @@ FImplicitObjectUnion::FImplicitObjectUnion()
 FImplicitObjectUnion::FImplicitObjectUnion(TArray<TUniquePtr<FImplicitObject>>&& Objects)
 	: FImplicitObject(EImplicitObject::HasBoundingBox, ImplicitObjectType::Union)
 	, MObjects(MoveTemp(Objects))
-	, MLocalBoundingBox()
 	, NumLeafObjects(0)
 	, Flags()
 {
 	ensure(MObjects.Num());
-	for (int32 i = 0; i < MObjects.Num(); ++i)
-	{
-		if (i > 0)
-		{
-			MLocalBoundingBox.GrowToInclude(MObjects[i]->BoundingBox());
-		}
-		else
-		{
-			MLocalBoundingBox = MObjects[i]->BoundingBox();
-		}
-	}
+
+	MLocalBoundingBox = CalculateObjectsBounds(MakeArrayView(MObjects));
 
 	SetNumLeafObjects(Private::FImplicitBVH::CountLeafObjects(MakeArrayView(MObjects)));
 }
@@ -68,9 +74,14 @@ void FImplicitObjectUnion::Combine(TArray<TUniquePtr<FImplicitObject>>& OtherObj
 {
 	ensure(MObjects.Num());
 
-	for (int32 i = 0; i < OtherObjects.Num(); ++i)
+	const FAABB3 OtherBounds = CalculateObjectsBounds(MakeArrayView(OtherObjects));
+	if (MObjects.Num() > 0)
 	{
-		MLocalBoundingBox.GrowToInclude(OtherObjects[i]->BoundingBox());
+		MLocalBoundingBox.GrowToInclude(OtherBounds);
+	}
+	else
+	{
+		MLocalBoundingBox = OtherBounds;
 	}
 
 	MObjects.Reserve(MObjects.Num() + OtherObjects.Num());
@@ -94,18 +105,7 @@ void FImplicitObjectUnion::RemoveAt(int32 RemoveIndex)
 		MObjects.RemoveAt(RemoveIndex);
 	}
 
-	MLocalBoundingBox = FAABB3::EmptyAABB();
-	for (int32 i = 0; i < MObjects.Num(); ++i)
-	{
-		if (i > 0)
-		{
-			MLocalBoundingBox.GrowToInclude(MObjects[i]->BoundingBox());
-		}
-		else
-		{
-			MLocalBoundingBox = MObjects[i]->BoundingBox();
-		}
-	}
+	MLocalBoundingBox = CalculateObjectsBounds(MakeArrayView(MObjects));
 
 	RebuildBVH();
 }
@@ -389,6 +389,16 @@ void FImplicitObjectUnion::Serialize(FChaosArchive& Ar)
 			}
 		}
 	}
+
+	if (Ar.IsLoading())
+	{
+		// We used to use empty bounds for particles with no geometry, but
+		// now it is a point bounds at the local origin.
+		if (MLocalBoundingBox.IsEmpty())
+		{
+			MLocalBoundingBox = FAABB3::ZeroAABB();
+		}
+	}
 }
 
 void FImplicitObjectUnion::LegacySerializeBVH(FChaosArchive& Ar)
diff --git a/Engine/Source/Runtime/Experimental/Chaos/Private/Chaos/ShapeInstance.cpp b/Engine/Source/Runtime/Experimental/Chaos/Private/Chaos/ShapeInstance.cpp
--- a/Engine/Source/Runtime/Experimental/Chaos/Private/Chaos/ShapeInstance.cpp
+++ b/Engine/Source/Runtime/Experimental/Chaos/Private/Chaos/ShapeInstance.cpp
@@ -333,6 +333,10 @@ namespace Chaos
 		{
 			WorldSpaceInflatedShapeBounds = Geometry->CalculateTransformedBounds(WorldTransform).ThickenSymmetrically(BoundsExpansion);
 		}
+		else
+		{
+			WorldSpaceInflatedShapeBounds = FAABB3(WorldTransform.GetLocation(), WorldTransform.GetLocation()).ThickenSymmetrically(BoundsExpansion);
+		}
 	}
 
 	void FShapeInstanceProxy::UpdateWorldSpaceState(const FRigidTransform3& WorldTransform, const FVec3& BoundsExpansion)
@@ -393,6 +397,10 @@ namespace Chaos
 		{
 			WorldSpaceInflatedShapeBounds = Geometry->CalculateTransformedBounds(WorldTransform).ThickenSymmetrically(BoundsExpansion);
 		}
+		else
+		{
+			WorldSpaceInflatedShapeBounds = FAABB3(WorldTransform.GetLocation(), WorldTransform.GetLocation()).ThickenSymmetrically(BoundsExpansion);
+		}
 	}
 
 	void FShapeInstance::UpdateWorldSpaceState(const FRigidTransform3& WorldTransform, const FVec3& BoundsExpansion)
@@ -418,6 +426,10 @@ namespace Chaos
 		{
 			WorldSpaceInflatedShapeBounds = LeafGeometry->CalculateTransformedBounds(LeafWorldTransform).ThickenSymmetrically(BoundsExpansion);;
 		}
+		else
+		{
+			WorldSpaceInflatedShapeBounds = FAABB3(WorldTransform.GetLocation(), WorldTransform.GetLocation()).ThickenSymmetrically(BoundsExpansion);
+		}
 	}
 
 	const FImplicitObject* FShapeInstance::GetLeafGeometry() const
diff --git a/Engine/Source/Runtime/Experimental/Chaos/Public/Chaos/GeometryParticles.h b/Engine/Source/Runtime/Experimental/Chaos/Public/Chaos/GeometryParticles.h
--- a/Engine/Source/Runtime/Experimental/Chaos/Public/Chaos/GeometryParticles.h
+++ b/Engine/Source/Runtime/Experimental/Chaos/Public/Chaos/GeometryParticles.h
@@ -377,56 +377,69 @@ namespace Chaos
 			{
 				MLocalBounds[Index] = TAABB<T, d>(InGeometry->BoundingBox());
 
-				if (CVars::CCDAxisThresholdMode == 0)
-				{
-					// Use object extents as CCD axis threshold
-					MCCDAxisThreshold[Index] = MLocalBounds[Index].Extents();
-				}
-				else if (CVars::CCDAxisThresholdMode == 1)
-				{
-					// Use thinnest object extents as all axis CCD thresholds
-					MCCDAxisThreshold[Index] = FVec3(MLocalBounds[Index].Extents().GetMin());
-				}
-				else
+				// Update the threshold we use to determine when to enable CCD. This is based on the bounds.
+				UpdateCCDAxisThreshold(Index);
+
+				// Update the world-space stat of all the shapes - must be called after UpdateShapesArray
+				// world space inflated bounds needs to take expansion into account - this is done in integrate for dynamics anyway, so
+				// this computation is mainly for statics
+				UpdateWorldSpaceState(Index, TRigidTransform<FReal, 3>(X(Index), R(Index)), FVec3(0));
+			}
+		}
+
+		void UpdateCCDAxisThreshold(const int32 Index)
+		{
+			// NOTE: We get empty bounds (as opposed to no bounds) if we have Geometry that is an empty Union
+			if (!MHasBounds[Index] || MLocalBounds[Index].IsEmpty())
+			{
+				MCCDAxisThreshold[Index] = FVec3(0);
+				return;
+			}
+
+			if (CVars::CCDAxisThresholdMode == 0)
+			{
+				// Use object extents as CCD axis threshold
+				MCCDAxisThreshold[Index] = MLocalBounds[Index].Extents();
+			}
+			else if (CVars::CCDAxisThresholdMode == 1)
+			{
+				// Use thinnest object extents as all axis CCD thresholds
+				MCCDAxisThreshold[Index] = FVec3(MLocalBounds[Index].Extents().GetMin());
+			}
+			else
+			{
+				// Find minimum shape bounds thickness on each axis
+				FVec3 ThinnestBoundsPerAxis = MLocalBounds[Index].Extents();
+				for (const TUniquePtr<FPerShapeData>& Shape : ShapesArray(Index))
 				{
-					// Find minimum shape bounds thickness on each axis
-					FVec3 ThinnestBoundsPerAxis = MLocalBounds[Index].Extents();
-					for (const TUniquePtr<FPerShapeData>& Shape : ShapesArray(Index))
+					// Only sim-enabled shapes should ever be swept with CCD, so make sure the
+					// sim-enabled flag is on for each shape before considering it's min bounds
+					// for CCD extents.
+					if (Shape->GetSimEnabled() && (CVars::bCCDAxisThresholdUsesProbeShapes || !Shape->GetIsProbe()))
 					{
-						// Only sim-enabled shapes should ever be swept with CCD, so make sure the
-						// sim-enabled flag is on for each shape before considering it's min bounds
-						// for CCD extents.
-						if (Shape->GetSimEnabled() && (CVars::bCCDAxisThresholdUsesProbeShapes || !Shape->GetIsProbe()))
+						const TSerializablePtr<FImplicitObject> Geometry = Shape->GetGeometry();
+						if (Geometry->HasBoundingBox())
 						{
-							const TSerializablePtr<FImplicitObject> Geometry = Shape->GetGeometry();
-							if (Geometry->HasBoundingBox())
+							const TVector<T, d> ShapeExtents = Geometry->BoundingBox().Extents();
+							TVector<T, d>& CCDAxisThreshold = MCCDAxisThreshold[Index];
+							for (int32 AxisIndex = 0; AxisIndex < d; ++AxisIndex)
 							{
-								const TVector<T, d> ShapeExtents = Geometry->BoundingBox().Extents();
-								TVector<T, d>& CCDAxisThreshold = MCCDAxisThreshold[Index];
-								for (int32 AxisIndex = 0; AxisIndex < d; ++AxisIndex)
-								{
-									ThinnestBoundsPerAxis[AxisIndex] = FMath::Min(ShapeExtents[AxisIndex], ThinnestBoundsPerAxis[AxisIndex]);
-								}
+								ThinnestBoundsPerAxis[AxisIndex] = FMath::Min(ShapeExtents[AxisIndex], ThinnestBoundsPerAxis[AxisIndex]);
 							}
 						}
 					}
-
-					if (CVars::CCDAxisThresholdMode == 2)
-					{
-						// On each axis, use the thinnest shape bound on that axis
-						MCCDAxisThreshold[Index] = ThinnestBoundsPerAxis;
-					}
-					else if (CVars::CCDAxisThresholdMode == 3)
-					{
-						// Find the thinnest shape bound on any axis and use this for all axes
-						MCCDAxisThreshold[Index] = FVec3(ThinnestBoundsPerAxis.GetMin());
-					}
 				}
 
-				// Update the world-space stat of all the shapes - must be called after UpdateShapesArray
-				// world space inflated bounds needs to take expansion into account - this is done in integrate for dynamics anyway, so
-				// this computation is mainly for statics
-				UpdateWorldSpaceState(Index, TRigidTransform<FReal, 3>(X(Index), R(Index)), FVec3(0));
+				if (CVars::CCDAxisThresholdMode == 2)
+				{
+					// On each axis, use the thinnest shape bound on that axis
+					MCCDAxisThreshold[Index] = ThinnestBoundsPerAxis;
+				}
+				else if (CVars::CCDAxisThresholdMode == 3)
+				{
+					// Find the thinnest shape bound on any axis and use this for all axes
+					MCCDAxisThreshold[Index] = FVec3(ThinnestBoundsPerAxis.GetMin());
+				}
 			}
 		}
 	public:
@@ -485,11 +498,11 @@ namespace Chaos
 
 		CHAOS_API void UpdateWorldSpaceState(const int32 Index, const FRigidTransform3& WorldTransform, const FVec3& BoundsExpansion)
 		{
-			// NOTE: Particle bounds are expanded for use by the spatial partitioning and broad phase but individual
-			// shape bounds are not. If expanded Shape bounds are required, the expansion should be done at the calling end.
-			FAABB3 WorldBounds = FAABB3::EmptyAABB();
-
 			const FShapesArray& Shapes = ShapesArray(Index);
+
+			// If we have no shapes we are have a point bounds at our local origin
+			FAABB3 WorldBounds = (Shapes.Num() > 0) ? FAABB3::EmptyAABB() : FAABB3::ZeroAABB();
+
 			for (const auto& Shape : Shapes)
 			{
 				Shape->UpdateWorldSpaceState(WorldTransform, BoundsExpansion);
@@ -501,6 +514,7 @@ namespace Chaos
 
 		void UpdateWorldSpaceStateSwept(const int32 Index, const FRigidTransform3& EndWorldTransform, const FVec3& BoundsExpansion, const FVec3& DeltaX)
 		{
+			// NOTE: Individual shape bounds are not expanded by the DeltaX (velocity term). Maybe they should be...
 			UpdateWorldSpaceState(Index, EndWorldTransform, BoundsExpansion);
 			MWorldSpaceInflatedBounds[Index].GrowByVector(DeltaX);
 		}
diff --git a/Engine/Source/Runtime/Experimental/ChaosCore/Public/Chaos/Core.h b/Engine/Source/Runtime/Experimental/ChaosCore/Public/Chaos/Core.h
--- a/Engine/Source/Runtime/Experimental/ChaosCore/Public/Chaos/Core.h
+++ b/Engine/Source/Runtime/Experimental/ChaosCore/Public/Chaos/Core.h
@@ -51,7 +51,7 @@ namespace Chaos
 	// NOTE: if you get a merge conflict on the GUID, you must replace it with a new GUID - do not accept the source or target
 	// or you will likely get DDC version conflicts resulting in crashes during load.
 	// Core version string for Chaos data. Any DDC builder dependent on Chaos for serialization should depend on this version
-	inline const TCHAR* const ChaosVersionGUID = TEXT("F960A384-9C87-4377-AE3A-1E9D861E34AB");
+	inline const TCHAR* const ChaosVersionGUID = TEXT("5E9A3DDA-B200-415C-9C6D-A5C15B2AA011");
 
 	inline FString GetChaosVersionStringInner()
 	{
