fix crash on dataflow selection by size or volume
+ add note about un-needed copy

diff --git a/Engine/Source/Runtime/Experimental/Chaos/Private/GeometryCollection/Facades/CollectionTransformSelectionFacade.cpp b/Engine/Source/Runtime/Experimental/Chaos/Private/GeometryCollection/Facades/CollectionTransformSelectionFacade.cpp
--- a/Engine/Source/Runtime/Experimental/Chaos/Private/GeometryCollection/Facades/CollectionTransformSelectionFacade.cpp
+++ b/Engine/Source/Runtime/Experimental/Chaos/Private/GeometryCollection/Facades/CollectionTransformSelectionFacade.cpp
@@ -632,12 +632,17 @@ namespace GeometryCollection::Facades
 	{
 		TArray<int32> OutSelection;
 
+		// TODO: (See also SelectByVolume) We should add a method to get the volumes without modifying the collection, and then remove this full geometrycollection copy
 		if (TUniquePtr<FGeometryCollection> TempGeomCollection = TUniquePtr<FGeometryCollection>(ConstCollection.NewCopy<FGeometryCollection>()))
 		{
 			FGeometryCollectionConvexUtility::SetVolumeAttributes(TempGeomCollection.Get());
 
-			const TManagedArrayAccessor<float> SizeAttribute((FManagedArrayCollection)*TempGeomCollection, "Size", FTransformCollection::TransformGroup);
-			const TManagedArray<float>& Sizes = SizeAttribute.Get();
+			const TManagedArray<float>* SizesPtr = TempGeomCollection->FindAttribute<float>("Size", FTransformCollection::TransformGroup);
+			if (!ensure(SizesPtr))
+			{
+				return OutSelection;
+			}
+			const TManagedArray<float>& Sizes = *SizesPtr;
 
 			for (int32 BoneIdx = 0; BoneIdx < Sizes.Num(); ++BoneIdx)
 			{
@@ -687,12 +692,17 @@ namespace GeometryCollection::Facades
 	{
 		TArray<int32> OutSelection;
 
+		// TODO: (See also SelectBySize) We should add a method to get the volumes without modifying the collection, and then remove this full geometrycollection copy
 		if (TUniquePtr<FGeometryCollection> TempGeomCollection = TUniquePtr<FGeometryCollection>(ConstCollection.NewCopy<FGeometryCollection>()))
 		{
 			FGeometryCollectionConvexUtility::SetVolumeAttributes(TempGeomCollection.Get());
 
-			const TManagedArrayAccessor<float> VolumeAttribute((FManagedArrayCollection)*TempGeomCollection, "Volume", FTransformCollection::TransformGroup);
-			const TManagedArray<float>& Volumes = VolumeAttribute.Get();
+			const TManagedArray<float>* VolumesPtr = TempGeomCollection->FindAttribute<float>("Volume", FTransformCollection::TransformGroup);
+			if (!ensure(VolumesPtr))
+			{
+				return OutSelection;
+			}
+			const TManagedArray<float>& Volumes = *VolumesPtr;
 
 			for (int32 BoneIdx = 0; BoneIdx < Volumes.Num(); ++BoneIdx)
 			{
