Split sending of an activities recv-done callback into a separate state

#rnx

diff --git a/Engine/Source/Runtime/Experimental/IoStoreOnDemand/Private/CoreHttp/Http.cpp b/Engine/Source/Runtime/Experimental/IoStoreOnDemand/Private/CoreHttp/Http.cpp
--- a/Engine/Source/Runtime/Experimental/IoStoreOnDemand/Private/CoreHttp/Http.cpp
+++ b/Engine/Source/Runtime/Experimental/IoStoreOnDemand/Private/CoreHttp/Http.cpp
@@ -754,6 +754,7 @@ struct alignas(16) FActivity
 		RecvMessage,
 		RecvContent,
 		RecvStream,
+		RecvDone,
 		Completed,
 		Cancelled,
 		Failed,
@@ -1010,7 +1011,7 @@ FTicketStatus::EId FTicketStatus::GetId() const
 	switch (Activity->State)
 	{
 	case FActivity::EState::RecvMessage:	return EId::Response;
-	case FActivity::EState::RecvContent:	return EId::Content;
+	case FActivity::EState::RecvDone:		return EId::Content;
 	case FActivity::EState::Cancelled:		return EId::Cancelled;
 	case FActivity::EState::Failed:			return EId::Error;
 	default:								check(false);
@@ -1117,6 +1118,7 @@ public:
 	static int32	DoRecvMessage(FActivity* Activity);
 	static int32	DoRecvContent(FActivity* Activity);
 	static int32	DoRecvStream(FActivity* Activity);
+	static int32	DoRecvDone(FActivity* Activity);
 	static void		Cancel(FActivity* Activity);
 };
 
@@ -1631,6 +1633,13 @@ int32 FEventLoopInternal::DoRecvContent(FActivity* Activity)
 		Activity->StateParam += Result;
 	}
 
+	Activity->State = FActivity::EState::RecvDone;
+	return 0;
+}
+
+////////////////////////////////////////////////////////////////////////////////
+int32 FEventLoopInternal::DoRecvDone(FActivity* Activity)
+{
 	// Notify the user we've received everything
 	FTicketStatus& SinkArg = *(FTicketStatus*)Activity;
 	Activity->Sink(SinkArg);
@@ -1913,6 +1922,11 @@ uint32 FEventLoop::Tick()
 			if (Result)
 				break;
 
+		case FActivity::EState::RecvDone:
+			Result = FEventLoopInternal::DoRecvDone(Activity);
+			if (Result)
+				break;
+
 		case FActivity::EState::Completed:
 			--BusyCount;
 			break;
