Split out size parsing and parameter lookup
#rnx

diff --git a/Engine/Source/Runtime/Experimental/IoStoreOnDemand/Private/IoStoreOnDemand.cpp b/Engine/Source/Runtime/Experimental/IoStoreOnDemand/Private/IoStoreOnDemand.cpp
--- a/Engine/Source/Runtime/Experimental/IoStoreOnDemand/Private/IoStoreOnDemand.cpp
+++ b/Engine/Source/Runtime/Experimental/IoStoreOnDemand/Private/IoStoreOnDemand.cpp
@@ -21,27 +21,31 @@ namespace UE::IO::Private
 {
 
 ////////////////////////////////////////////////////////////////////////////////
-static uint64 ParseSizeParam(const TCHAR* Param)
+static int64 ParseSizeParam(FStringView Value)
 {
-	const TCHAR* CommandLine = FCommandLine::Get();
+	Value = Value.TrimStartAndEnd();
 
-	FString ParamValue;
-	if (!FParse::Value(CommandLine, Param, ParamValue))
+	int64 Size = -1;
+	LexFromString(Size, Value.GetData());
+	if (Size >= 0)
 	{
-		return 0;
+		if (Value.EndsWith(TEXT("GB"))) return Size << 30;
+		if (Value.EndsWith(TEXT("MB"))) return Size << 20;
+		if (Value.EndsWith(TEXT("KB"))) return Size << 10;
 	}
+	return Size;
+}
 
-	uint64 Size = 0;
-	if (!FParse::Value(CommandLine, Param, Size))
+////////////////////////////////////////////////////////////////////////////////
+static int64 ParseSizeParam(const TCHAR* CommandLine, const TCHAR* Param)
+{
+	FString ParamValue;
+	if (!FParse::Value(CommandLine, Param, ParamValue))
 	{
-		return 0;
+		return -1;
 	}
 
-	if (ParamValue.EndsWith(TEXT("GB"))) return Size << 30;
-	if (ParamValue.EndsWith(TEXT("MB"))) return Size << 20;
-	if (ParamValue.EndsWith(TEXT("KB"))) return Size << 10;
-
-	return Size;
+	return ParseSizeParam(ParamValue);
 }
 
 ////////////////////////////////////////////////////////////////////////////////
