Added drop configuration option to IAS' cache
#rnx

diff --git a/Engine/Source/Runtime/Experimental/IoStoreOnDemand/Private/FileCache.cpp b/Engine/Source/Runtime/Experimental/IoStoreOnDemand/Private/FileCache.cpp
--- a/Engine/Source/Runtime/Experimental/IoStoreOnDemand/Private/FileCache.cpp
+++ b/Engine/Source/Runtime/Experimental/IoStoreOnDemand/Private/FileCache.cpp
@@ -668,6 +668,13 @@ void FFileIoCache::Initialize()
 	WriteCursorPos = 0;
 
 	IFileManager& FileMgr = IFileManager::Get();
+
+	if (CacheConfig.DropCache)
+	{
+		FileMgr.Delete(*CacheTocPath);
+		FileMgr.Delete(*CacheFilePath);
+		return;
+	}
 	
 	FIoStatus CacheStatus(EIoErrorCode::Unknown);
 	if (FileMgr.FileExists(*CacheTocPath))
diff --git a/Engine/Source/Runtime/Experimental/IoStoreOnDemand/Private/FileCache.h b/Engine/Source/Runtime/Experimental/IoStoreOnDemand/Private/FileCache.h
--- a/Engine/Source/Runtime/Experimental/IoStoreOnDemand/Private/FileCache.h
+++ b/Engine/Source/Runtime/Experimental/IoStoreOnDemand/Private/FileCache.h
@@ -45,6 +45,7 @@ struct FFileIoCacheConfig
 	uint64 DiskStorageSize;
 	uint64 MemoryStorageSize;
 	FRate WriteRate;
+	bool DropCache = false;
 };
 
 TUniquePtr<IIoCache> MakeFileIoCache(const FFileIoCacheConfig& Config);
