Homogenised all outward-facing nomenclature to "Ias".

#rnx

diff --git a/Engine/Source/Runtime/Experimental/IoStoreOnDemand/Private/IoStoreOnDemand.cpp b/Engine/Source/Runtime/Experimental/IoStoreOnDemand/Private/IoStoreOnDemand.cpp
--- a/Engine/Source/Runtime/Experimental/IoStoreOnDemand/Private/IoStoreOnDemand.cpp
+++ b/Engine/Source/Runtime/Experimental/IoStoreOnDemand/Private/IoStoreOnDemand.cpp
@@ -317,7 +317,7 @@ void FIoStoreOnDemandModule::StartupModule()
 	UE::FOnDemandEndpoint Endpoint;
 	
 	FString UrlParam;
-	if (FParse::Value(FCommandLine::Get(), TEXT("IoStoreOnDemand="), UrlParam))
+	if (FParse::Value(FCommandLine::Get(), TEXT("Ias.TocUrl="), UrlParam))
 	{
 		FStringView UrlView(UrlParam);
 		if (UrlView.StartsWith(TEXTVIEW("http://")) && UrlView.EndsWith(TEXTVIEW(".iochunktoc")))
@@ -333,7 +333,7 @@ void FIoStoreOnDemandModule::StartupModule()
 
 	{
 		FString EncryptionKey;
-		if (FParse::Value(FCommandLine::Get(), TEXT("OnDemandEncryptionKey="), EncryptionKey))
+		if (FParse::Value(FCommandLine::Get(), TEXT("Ias.EncryptionKey="), EncryptionKey))
 		{
 			FGuid KeyGuid;
 			FAES::FAESKey Key;
@@ -392,9 +392,9 @@ void FIoStoreOnDemandModule::StartupModule()
 	{
 		TSharedPtr<IIoCache> Cache;
 
-		if (uint64 DiskSize = ParseSizeParam(TEXT("OnDemandFileCache=")); DiskSize > 0)
+		if (uint64 DiskSize = ParseSizeParam(TEXT("Ias.FileCache=")); DiskSize > 0)
 		{
-			uint64 MemorySize = ParseSizeParam(TEXT("OnDemandFileCacheQueueSize="));
+			uint64 MemorySize = ParseSizeParam(TEXT("Ias.FileCacheQueueSize="));
 			FFileIoCacheConfig FileCacheConfig;
 			FileCacheConfig.DiskStorageSize = DiskSize;
 			FileCacheConfig.MemoryStorageSize = MemorySize > 0 ? MemorySize : (16ull << 20);
diff --git a/Engine/Source/Runtime/Experimental/IoStoreOnDemand/Private/OnDemandIoDispatcherBackend.cpp b/Engine/Source/Runtime/Experimental/IoStoreOnDemand/Private/OnDemandIoDispatcherBackend.cpp
--- a/Engine/Source/Runtime/Experimental/IoStoreOnDemand/Private/OnDemandIoDispatcherBackend.cpp
+++ b/Engine/Source/Runtime/Experimental/IoStoreOnDemand/Private/OnDemandIoDispatcherBackend.cpp
@@ -58,52 +58,52 @@ using namespace UE::Tasks;
 ///////////////////////////////////////////////////////////////////////////////
 // TRACE STATS
 // iorequest stats
-TRACE_DECLARE_INT_COUNTER(IoRequestsMade, TEXT("OnDemandIoBackend/IoRequestsMade"));
-TRACE_DECLARE_INT_COUNTER(IoRequestsCompleted, TEXT("OnDemandIoBackend/IoRequestsCompleted"));
-TRACE_DECLARE_MEMORY_COUNTER(IoRequestsCompletedSize, TEXT("OnDemandIoBackend/Size/IoRequestsCompletedSize"));
-TRACE_DECLARE_INT_COUNTER(IoRequestsCancelled, TEXT("OnDemandIoBackend/IoRequestsCancelled"));
-TRACE_DECLARE_INT_COUNTER(IoRequestsFailed, TEXT("OnDemandIoBackend/IoRequestsFailed"));
+TRACE_DECLARE_INT_COUNTER(IoRequestsMade, TEXT("Ias/IoRequestsMade"));
+TRACE_DECLARE_INT_COUNTER(IoRequestsCompleted, TEXT("Ias/IoRequestsCompleted"));
+TRACE_DECLARE_MEMORY_COUNTER(IoRequestsCompletedSize, TEXT("Ias/Size/IoRequestsCompletedSize"));
+TRACE_DECLARE_INT_COUNTER(IoRequestsCancelled, TEXT("Ias/IoRequestsCancelled"));
+TRACE_DECLARE_INT_COUNTER(IoRequestsFailed, TEXT("Ias/IoRequestsFailed"));
 // chunkrequest stats
-TRACE_DECLARE_INT_COUNTER(ReadRequestsCreated, TEXT("OnDemandIoBackend/ReadRequestsCreated"));
-TRACE_DECLARE_INT_COUNTER(ReadRequestsRemoved, TEXT("OnDemandIoBackend/ReadRequestsRemoved"));
+TRACE_DECLARE_INT_COUNTER(ReadRequestsCreated, TEXT("Ias/ReadRequestsCreated"));
+TRACE_DECLARE_INT_COUNTER(ReadRequestsRemoved, TEXT("Ias/ReadRequestsRemoved"));
 // cache stats
-TRACE_DECLARE_INT_COUNTER(CacheHits, TEXT("OnDemandIoBackend/CacheHits"));
-TRACE_DECLARE_MEMORY_COUNTER(CacheHitsSize, TEXT("OnDemandIoBackend/Size/CacheHitsSize"));
-TRACE_DECLARE_INT_COUNTER(CachePuts, TEXT("OnDemandIoBackend/CachePuts"));
-TRACE_DECLARE_MEMORY_COUNTER(CachePutsSize, TEXT("OnDemandIoBackend/Size/CachePutsSize"));
-TRACE_DECLARE_INT_COUNTER(CacheRejects, TEXT("OnDemandIoBackend/CacheRejects"));
-TRACE_DECLARE_MEMORY_COUNTER(CacheRejectsSize, TEXT("OnDemandIoBackend/Size/CacheRejectsSize"));
+TRACE_DECLARE_INT_COUNTER(CacheHits, TEXT("Ias/CacheHits"));
+TRACE_DECLARE_MEMORY_COUNTER(CacheHitsSize, TEXT("Ias/Size/CacheHitsSize"));
+TRACE_DECLARE_INT_COUNTER(CachePuts, TEXT("Ias/CachePuts"));
+TRACE_DECLARE_MEMORY_COUNTER(CachePutsSize, TEXT("Ias/Size/CachePutsSize"));
+TRACE_DECLARE_INT_COUNTER(CacheRejects, TEXT("Ias/CacheRejects"));
+TRACE_DECLARE_MEMORY_COUNTER(CacheRejectsSize, TEXT("Ias/Size/CacheRejectsSize"));
 // http stats
-TRACE_DECLARE_INT_COUNTER(HttpRequestsCompleted, TEXT("OnDemandIoBackend/HttpRequestsCompleted"));
-TRACE_DECLARE_INT_COUNTER(HttpRequestsFailed, TEXT("OnDemandIoBackend/HttpRequestsFailed"));
-TRACE_DECLARE_INT_COUNTER(HttpRequestsPending, TEXT("OnDemandIoBackend/HttpRequestsPending"));
-TRACE_DECLARE_INT_COUNTER(HttpRequestsInflight, TEXT("OnDemandIoBackend/HttpRequestsInflight"));
-TRACE_DECLARE_MEMORY_COUNTER(HttpRequestsCompletedSize, TEXT("OnDemandIoBackend/Size/HttpRequestsCompletedSize"));
+TRACE_DECLARE_INT_COUNTER(HttpRequestsCompleted, TEXT("Ias/HttpRequestsCompleted"));
+TRACE_DECLARE_INT_COUNTER(HttpRequestsFailed, TEXT("Ias/HttpRequestsFailed"));
+TRACE_DECLARE_INT_COUNTER(HttpRequestsPending, TEXT("Ias/HttpRequestsPending"));
+TRACE_DECLARE_INT_COUNTER(HttpRequestsInflight, TEXT("Ias/HttpRequestsInflight"));
+TRACE_DECLARE_MEMORY_COUNTER(HttpRequestsCompletedSize, TEXT("Ias/Size/HttpRequestsCompletedSize"));
 ///////////////////////////////////////////////////////////////////////////////
 // CSV STATS
-CSV_DEFINE_CATEGORY(IoStoreOnDemandBackend, true);
+CSV_DEFINE_CATEGORY(Ias, true);
 // iorequest stats
-CSV_DEFINE_STAT(IoStoreOnDemandBackend, FrameIoRequestsMade);
-CSV_DEFINE_STAT(IoStoreOnDemandBackend, FrameIoRequestsCompleted);
-CSV_DEFINE_STAT(IoStoreOnDemandBackend, FrameIoRequestsCompletedSize);
-CSV_DEFINE_STAT(IoStoreOnDemandBackend, FrameIoRequestsCancelled);
-CSV_DEFINE_STAT(IoStoreOnDemandBackend, FrameIoRequestsFailed);
+CSV_DEFINE_STAT(Ias, FrameIoRequestsMade);
+CSV_DEFINE_STAT(Ias, FrameIoRequestsCompleted);
+CSV_DEFINE_STAT(Ias, FrameIoRequestsCompletedSize);
+CSV_DEFINE_STAT(Ias, FrameIoRequestsCancelled);
+CSV_DEFINE_STAT(Ias, FrameIoRequestsFailed);
 // chunkrequest stats
-CSV_DEFINE_STAT(IoStoreOnDemandBackend, FrameReadRequestsCreated);
-CSV_DEFINE_STAT(IoStoreOnDemandBackend, FrameReadRequestsRemoved);
+CSV_DEFINE_STAT(Ias, FrameReadRequestsCreated);
+CSV_DEFINE_STAT(Ias, FrameReadRequestsRemoved);
 // cache stats
-CSV_DEFINE_STAT(IoStoreOnDemandBackend, FrameCacheHits);
-CSV_DEFINE_STAT(IoStoreOnDemandBackend, FrameCacheHitsSize);
-CSV_DEFINE_STAT(IoStoreOnDemandBackend, FrameCachePuts);
-CSV_DEFINE_STAT(IoStoreOnDemandBackend, FrameCachePutsSize);
-CSV_DEFINE_STAT(IoStoreOnDemandBackend, FrameCacheRejects);
-CSV_DEFINE_STAT(IoStoreOnDemandBackend, FrameCacheRejectsSize);
+CSV_DEFINE_STAT(Ias, FrameCacheHits);
+CSV_DEFINE_STAT(Ias, FrameCacheHitsSize);
+CSV_DEFINE_STAT(Ias, FrameCachePuts);
+CSV_DEFINE_STAT(Ias, FrameCachePutsSize);
+CSV_DEFINE_STAT(Ias, FrameCacheRejects);
+CSV_DEFINE_STAT(Ias, FrameCacheRejectsSize);
 // http stats
-CSV_DEFINE_STAT(IoStoreOnDemandBackend, FrameHttpRequestsCompleted);
-CSV_DEFINE_STAT(IoStoreOnDemandBackend, FrameHttpRequestsFailed);
-CSV_DEFINE_STAT(IoStoreOnDemandBackend, FrameHttpRequestsPending);
-CSV_DEFINE_STAT(IoStoreOnDemandBackend, FrameHttpRequestsInflight);
-CSV_DEFINE_STAT(IoStoreOnDemandBackend, FrameHttpRequestsCompletedSize);
+CSV_DEFINE_STAT(Ias, FrameHttpRequestsCompleted);
+CSV_DEFINE_STAT(Ias, FrameHttpRequestsFailed);
+CSV_DEFINE_STAT(Ias, FrameHttpRequestsPending);
+CSV_DEFINE_STAT(Ias, FrameHttpRequestsInflight);
+CSV_DEFINE_STAT(Ias, FrameHttpRequestsCompletedSize);
 
 #if (COUNTERSTRACE_ENABLED || CSV_PROFILER)
 class FOnDemandIoBackendStats
