Early out if there's no endpoint to proceed with

#rnx

diff --git a/Engine/Source/Runtime/Experimental/IoStoreOnDemand/Private/IoStoreOnDemand.cpp b/Engine/Source/Runtime/Experimental/IoStoreOnDemand/Private/IoStoreOnDemand.cpp
--- a/Engine/Source/Runtime/Experimental/IoStoreOnDemand/Private/IoStoreOnDemand.cpp
+++ b/Engine/Source/Runtime/Experimental/IoStoreOnDemand/Private/IoStoreOnDemand.cpp
@@ -388,26 +388,28 @@ void FIoStoreOnDemandModule::StartupModule()
 		}
 	}
 
-	if (Endpoint.IsValid())
+	if (!Endpoint.IsValid())
 	{
-		TSharedPtr<IIoCache> Cache;
+		return;
+	}
 
-		if (uint64 DiskSize = ParseSizeParam(TEXT("Ias.FileCache=")); DiskSize > 0)
-		{
-			uint64 MemorySize = ParseSizeParam(TEXT("Ias.FileCacheQueueSize="));
-			FFileIoCacheConfig FileCacheConfig;
-			FileCacheConfig.DiskStorageSize = DiskSize;
-			FileCacheConfig.MemoryStorageSize = MemorySize > 0 ? MemorySize : (16ull << 20);
+	TSharedPtr<IIoCache> Cache;
 
-			UE_LOG(LogIoStoreOnDemand, Log, TEXT("Using %lluB file cache"), DiskSize);
-			Cache = MakeShareable(MakeFileIoCache(FileCacheConfig).Release());
-		}
-		TSharedPtr<UE::IOnDemandIoDispatcherBackend> Backend = UE::MakeOnDemandIoDispatcherBackend(Cache);
+	if (uint64 DiskSize = ParseSizeParam(TEXT("Ias.FileCache=")); DiskSize > 0)
+	{
+		uint64 MemorySize = ParseSizeParam(TEXT("Ias.FileCacheQueueSize="));
+		FFileIoCacheConfig FileCacheConfig;
+		FileCacheConfig.DiskStorageSize = DiskSize;
+		FileCacheConfig.MemoryStorageSize = MemorySize > 0 ? MemorySize : (16ull << 20);
 
-		Endpoint.EndpointType = UE::EOnDemandEndpointType::CDN;
-		Backend->Mount(Endpoint);
-		FIoDispatcher::Get().Mount(Backend.ToSharedRef(), -10);
+		UE_LOG(LogIoStoreOnDemand, Log, TEXT("Using %lluB file cache"), DiskSize);
+		Cache = MakeShareable(MakeFileIoCache(FileCacheConfig).Release());
 	}
+	TSharedPtr<UE::IOnDemandIoDispatcherBackend> Backend = UE::MakeOnDemandIoDispatcherBackend(Cache);
+
+	Endpoint.EndpointType = UE::EOnDemandEndpointType::CDN;
+	Backend->Mount(Endpoint);
+	FIoDispatcher::Get().Mount(Backend.ToSharedRef(), -10);
 #endif // !WITH_EDITOR
 }
 
