Configure IAS' cache with .ini files and/or the command line

#rnx

diff --git a/Engine/Source/Runtime/Experimental/IoStoreOnDemand/Private/FileCache.cpp b/Engine/Source/Runtime/Experimental/IoStoreOnDemand/Private/FileCache.cpp
--- a/Engine/Source/Runtime/Experimental/IoStoreOnDemand/Private/FileCache.cpp
+++ b/Engine/Source/Runtime/Experimental/IoStoreOnDemand/Private/FileCache.cpp
@@ -555,7 +555,8 @@ FFileIoCache::FFileIoCache(const FFileIoCacheConfig& Config)
 {
 	Governor.Set(Config.WriteRate.Allowance, Config.WriteRate.Ops, Config.WriteRate.Seconds);
 
-	CacheMap.SetCacheLimits(Config.MemoryStorageSize, Config.DiskStorageSize);
+	CacheMap.SetCacheLimits(Config.MemoryQuota, Config.DiskQuota);
+
 	Initialize();
 }
 
@@ -660,7 +661,7 @@ FIoStatus FFileIoCache::Put(const FIoHash& Key, FIoBuffer& Data)
 void FFileIoCache::Initialize()
 {
 	UE_LOG(LogIoCache, Log,
-		TEXT("Initializing file I/O cache, disk size %lluB, memory size %lluB"), CacheConfig.DiskStorageSize, CacheConfig.MemoryStorageSize);
+		TEXT("Initializing file I/O cache, disk size %lluB, memory size %lluB"), CacheConfig.DiskQuota, CacheConfig.MemoryQuota);
 
 	const FString CacheDir = FPaths::ProjectPersistentDownloadDir() / TEXT("IoCache");
 	const FString CacheTocPath = CacheDir / TEXT("cache.utoc");
@@ -779,7 +780,7 @@ void FFileIoCache::FileWriterThreadInner()
 	}
 
 	// Wrap the write cursor if we'd end up writing over the end
-	if (WriteCursorPos + PendingSize > CacheConfig.DiskStorageSize)
+	if (WriteCursorPos + PendingSize > CacheConfig.DiskQuota)
 	{
 		WriteCursorPos = 0;
 	}
diff --git a/Engine/Source/Runtime/Experimental/IoStoreOnDemand/Private/FileCache.h b/Engine/Source/Runtime/Experimental/IoStoreOnDemand/Private/FileCache.h
--- a/Engine/Source/Runtime/Experimental/IoStoreOnDemand/Private/FileCache.h
+++ b/Engine/Source/Runtime/Experimental/IoStoreOnDemand/Private/FileCache.h
@@ -38,12 +38,12 @@ struct FFileIoCacheConfig
 	struct FRate
 	{
 		uint32 Allowance = ~0u;
-		uint32 Ops = 1;
+		uint32 Ops = 256;
 		uint32 Seconds = 1;
 	};
 
-	uint64 DiskStorageSize;
-	uint64 MemoryStorageSize;
+	uint64 DiskQuota = 1ull << 30;
+	uint32 MemoryQuota = 3 << 20;
 	FRate WriteRate;
 	bool DropCache = false;
 };
diff --git a/Engine/Source/Runtime/Experimental/IoStoreOnDemand/Private/IoStoreOnDemand.cpp b/Engine/Source/Runtime/Experimental/IoStoreOnDemand/Private/IoStoreOnDemand.cpp
--- a/Engine/Source/Runtime/Experimental/IoStoreOnDemand/Private/IoStoreOnDemand.cpp
+++ b/Engine/Source/Runtime/Experimental/IoStoreOnDemand/Private/IoStoreOnDemand.cpp
@@ -70,6 +70,53 @@ static bool ParseEncryptionKeyParam(const FString& Param, FGuid& OutKeyGuid, FAE
 	return false;
 }
 
+////////////////////////////////////////////////////////////////////////////////
+static FFileIoCacheConfig GetFileCacheConfig(const TCHAR* CommandLine)
+{
+	FFileIoCacheConfig Ret;
+
+	// Fetch values from .ini files
+	auto GetConfigIntImpl = [CommandLine] (const TCHAR* ConfigKey, const TCHAR* ParamName, auto& Out)
+	{
+		int64 Value = -1;
+		if (FString Temp; GConfig->GetString(TEXT("Ias"), ConfigKey, Temp, GEngineIni))
+		{
+			Value = ParseSizeParam(Temp);
+		}
+#if !UE_BUILD_SHIPPING
+		if (int64 Override = ParseSizeParam(CommandLine, ParamName); Override >= 0)
+		{
+			Value = Override;
+		}
+#endif
+
+		if (Value >= 0)
+		{
+			Out = decltype(Out)(Value);
+		}
+
+		return true;
+	};
+
+#define GetConfigInt(Name, Dest) \
+	do { GetConfigIntImpl(TEXT("FileCache.") Name, TEXT("Ias.FileCache.") Name, Dest); } while (false)
+	GetConfigInt(TEXT("WritePeriodSeconds"),	Ret.WriteRate.Seconds);
+	GetConfigInt(TEXT("WriteOpsPerPeriod"),		Ret.WriteRate.Ops);
+	GetConfigInt(TEXT("WriteBytesPerPeriod"),	Ret.WriteRate.Allowance);
+	GetConfigInt(TEXT("DiskQuota"),				Ret.DiskQuota);
+	GetConfigInt(TEXT("MemoryQuota"),			Ret.MemoryQuota);
+#undef GetConfigInt
+
+#if !UE_BUILD_SHIPPING
+	if (FParse::Param(CommandLine, TEXT("Ias.DropCache")))
+	{
+		Ret.DropCache = true;
+	}
+#endif
+
+	return Ret;
+}
+
 } // namespace UE::IO::Private
 
 
@@ -314,10 +361,12 @@ void FIoStoreOnDemandModule::StartupModule()
 	using namespace UE::IO::Private;
 
 #if !WITH_EDITOR
+	const TCHAR* CommandLine = FCommandLine::Get();
+
 	UE::FOnDemandEndpoint Endpoint;
 	
 	FString UrlParam;
-	if (FParse::Value(FCommandLine::Get(), TEXT("Ias.TocUrl="), UrlParam))
+	if (FParse::Value(CommandLine, TEXT("Ias.TocUrl="), UrlParam))
 	{
 		FStringView UrlView(UrlParam);
 		if (UrlView.StartsWith(TEXTVIEW("http://")) && UrlView.EndsWith(TEXTVIEW(".iochunktoc")))
@@ -333,7 +382,7 @@ void FIoStoreOnDemandModule::StartupModule()
 
 	{
 		FString EncryptionKey;
-		if (FParse::Value(FCommandLine::Get(), TEXT("Ias.EncryptionKey="), EncryptionKey))
+		if (FParse::Value(CommandLine, TEXT("Ias.EncryptionKey="), EncryptionKey))
 		{
 			FGuid KeyGuid;
 			FAES::FAESKey Key;
@@ -393,21 +442,15 @@ void FIoStoreOnDemandModule::StartupModule()
 		return;
 	}
 
-	TSharedPtr<IIoCache> Cache;
+	Endpoint.EndpointType = UE::EOnDemandEndpointType::CDN;
 
-	if (uint64 DiskSize = ParseSizeParam(TEXT("Ias.FileCache=")); DiskSize > 0)
+	TSharedPtr<IIoCache> Cache;
+	if (FFileIoCacheConfig Config = GetFileCacheConfig(CommandLine); Config.DiskQuota > 0)
 	{
-		uint64 MemorySize = ParseSizeParam(TEXT("Ias.FileCacheQueueSize="));
-		FFileIoCacheConfig FileCacheConfig;
-		FileCacheConfig.DiskStorageSize = DiskSize;
-		FileCacheConfig.MemoryStorageSize = MemorySize > 0 ? MemorySize : (16ull << 20);
-
-		UE_LOG(LogIoStoreOnDemand, Log, TEXT("Using %lluB file cache"), DiskSize);
-		Cache = MakeShareable(MakeFileIoCache(FileCacheConfig).Release());
+		Cache = MakeShareable(MakeFileIoCache(Config).Release());
 	}
-	TSharedPtr<UE::IOnDemandIoDispatcherBackend> Backend = UE::MakeOnDemandIoDispatcherBackend(Cache);
 
-	Endpoint.EndpointType = UE::EOnDemandEndpointType::CDN;
+	TSharedPtr<UE::IOnDemandIoDispatcherBackend> Backend = UE::MakeOnDemandIoDispatcherBackend(Cache);
 	Backend->Mount(Endpoint);
 	FIoDispatcher::Get().Mount(Backend.ToSharedRef(), -10);
 #endif // !WITH_EDITOR
