[Chaos]
* Moved dll export from types to methods/staticvar

diff --git a/Engine/Source/Runtime/Experimental/Chaos/Private/Chaos/GeometryParticles.cpp b/Engine/Source/Runtime/Experimental/Chaos/Private/Chaos/GeometryParticles.cpp
--- a/Engine/Source/Runtime/Experimental/Chaos/Private/Chaos/GeometryParticles.cpp
+++ b/Engine/Source/Runtime/Experimental/Chaos/Private/Chaos/GeometryParticles.cpp
@@ -65,7 +65,7 @@ namespace Chaos
 	}
 
 	template<class T, int d, EGeometryParticlesSimType SimType>
-	CHAOS_API TGeometryParticlesImp<T, d, SimType>* TGeometryParticlesImp<T, d, SimType>::SerializationFactory(FChaosArchive& Ar, TGeometryParticlesImp<T,d,SimType>* Particles)
+	TGeometryParticlesImp<T, d, SimType>* TGeometryParticlesImp<T, d, SimType>::SerializationFactory(FChaosArchive& Ar, TGeometryParticlesImp<T,d,SimType>* Particles)
 	{
 		int8 ParticleType = Ar.IsLoading() ? 0 : (int8)Particles->ParticleType();
 		Ar << ParticleType;
diff --git a/Engine/Source/Runtime/Experimental/Chaos/Private/Field/FieldSystemNodes.cpp b/Engine/Source/Runtime/Experimental/Chaos/Private/Field/FieldSystemNodes.cpp
--- a/Engine/Source/Runtime/Experimental/Chaos/Private/Field/FieldSystemNodes.cpp
+++ b/Engine/Source/Runtime/Experimental/Chaos/Private/Field/FieldSystemNodes.cpp
@@ -1809,9 +1809,9 @@ void FCullingField<T>::ComputeFieldBounds(FVector& MinBounds, FVector& MaxBounds
 	}
 }
 
-template class CHAOS_API FCullingField<int32>;
-template class CHAOS_API FCullingField<float>;
-template class CHAOS_API FCullingField<FVector>;
+template class FCullingField<int32>;
+template class FCullingField<float>;
+template class FCullingField<FVector>;
 
 
 /**
@@ -1848,6 +1848,6 @@ bool FReturnResultsTerminal<T>::operator==(const FFieldNodeBase& Node)
 	}
 	return false;
 }
-template class CHAOS_API FReturnResultsTerminal<int32>;
-template class CHAOS_API FReturnResultsTerminal<float>;
-template class CHAOS_API FReturnResultsTerminal<FVector>;
+template class FReturnResultsTerminal<int32>;
+template class FReturnResultsTerminal<float>;
+template class FReturnResultsTerminal<FVector>;
diff --git a/Engine/Source/Runtime/Experimental/Chaos/Public/Chaos/AABBTree.h b/Engine/Source/Runtime/Experimental/Chaos/Public/Chaos/AABBTree.h
--- a/Engine/Source/Runtime/Experimental/Chaos/Public/Chaos/AABBTree.h
+++ b/Engine/Source/Runtime/Experimental/Chaos/Public/Chaos/AABBTree.h
@@ -23,40 +23,40 @@
 
 CSV_DECLARE_CATEGORY_EXTERN(ChaosPhysicsTimers);
 
-struct CHAOS_API FAABBTreeCVars
+struct FAABBTreeCVars
 {
-	static int32 UpdateDirtyElementPayloadData;
-	static FAutoConsoleVariableRef CVarUpdateDirtyElementPayloadData;
+	static CHAOS_API int32 UpdateDirtyElementPayloadData;
+	static CHAOS_API FAutoConsoleVariableRef CVarUpdateDirtyElementPayloadData;
 
-	static int32 SplitAtAverageCenter;
-	static FAutoConsoleVariableRef CVarSplitAtAverageCenter;
+	static CHAOS_API int32 SplitAtAverageCenter;
+	static CHAOS_API FAutoConsoleVariableRef CVarSplitAtAverageCenter;
 
-	static int32 SplitOnVarianceAxis;
-	static FAutoConsoleVariableRef CVarSplitOnVarianceAxis;
+	static CHAOS_API int32 SplitOnVarianceAxis;
+	static CHAOS_API FAutoConsoleVariableRef CVarSplitOnVarianceAxis;
 
-	static float DynamicTreeBoundingBoxPadding;
-	static FAutoConsoleVariableRef CVarDynamicTreeBoundingBoxPadding;
+	static CHAOS_API float DynamicTreeBoundingBoxPadding;
+	static CHAOS_API FAutoConsoleVariableRef CVarDynamicTreeBoundingBoxPadding;
 
-	static int32 DynamicTreeLeafCapacity;
-	static FAutoConsoleVariableRef CVarDynamicTreeLeafCapacity;
+	static CHAOS_API int32 DynamicTreeLeafCapacity;
+	static CHAOS_API FAutoConsoleVariableRef CVarDynamicTreeLeafCapacity;
 
-	static bool DynamicTreeSkipCheckAuntOnRotate;
-	static FAutoConsoleVariableRef CVarDynamicTreeSkipCheckAuntOnRotate;
+	static CHAOS_API bool DynamicTreeSkipCheckAuntOnRotate;
+	static CHAOS_API FAutoConsoleVariableRef CVarDynamicTreeSkipCheckAuntOnRotate;
 };
 
-struct CHAOS_API FAABBTreeDirtyGridCVars
+struct FAABBTreeDirtyGridCVars
 {
-	static int32 DirtyElementGridCellSize;
-	static FAutoConsoleVariableRef CVarDirtyElementGridCellSize;
+	static CHAOS_API int32 DirtyElementGridCellSize;
+	static CHAOS_API FAutoConsoleVariableRef CVarDirtyElementGridCellSize;
 
-	static int32 DirtyElementMaxGridCellQueryCount;
-	static FAutoConsoleVariableRef CVarDirtyElementMaxGridCellQueryCount;
+	static CHAOS_API int32 DirtyElementMaxGridCellQueryCount;
+	static CHAOS_API FAutoConsoleVariableRef CVarDirtyElementMaxGridCellQueryCount;
 
-	static int32 DirtyElementMaxPhysicalSizeInCells;
-	static FAutoConsoleVariableRef CVarDirtyElementMaxPhysicalSizeInCells;
+	static CHAOS_API int32 DirtyElementMaxPhysicalSizeInCells;
+	static CHAOS_API FAutoConsoleVariableRef CVarDirtyElementMaxPhysicalSizeInCells;
 
-	static int32 DirtyElementMaxCellCapacity;
-	static FAutoConsoleVariableRef CVarDirtyElementMaxCellCapacity;
+	static CHAOS_API int32 DirtyElementMaxCellCapacity;
+	static CHAOS_API FAutoConsoleVariableRef CVarDirtyElementMaxCellCapacity;
 };
 
 namespace Chaos
diff --git a/Engine/Source/Runtime/Experimental/Chaos/Public/Chaos/BVHParticles.h b/Engine/Source/Runtime/Experimental/Chaos/Public/Chaos/BVHParticles.h
--- a/Engine/Source/Runtime/Experimental/Chaos/Public/Chaos/BVHParticles.h
+++ b/Engine/Source/Runtime/Experimental/Chaos/Public/Chaos/BVHParticles.h
@@ -35,7 +35,7 @@ namespace Chaos
 	    CHAOS_API FBVHParticles& operator=(const FBVHParticles& Other);
 	    CHAOS_API FBVHParticles& operator=(FBVHParticles&& Other);
 
-		CHAOS_API FBVHParticles* NewCopy() const
+		FBVHParticles* NewCopy() const
 		{
 			return new FBVHParticles(*this);
 		}
@@ -50,7 +50,7 @@ namespace Chaos
 
 		CHAOS_API void Serialize(FChaosArchive& Ar);
 
-		CHAOS_API void Serialize(FArchive& Ar)
+		void Serialize(FArchive& Ar)
 		{
 			check(false); //Aggregate simplicial require FChaosArchive - check false by default
 		}
diff --git a/Engine/Source/Runtime/Experimental/Chaos/Public/Chaos/Box.h b/Engine/Source/Runtime/Experimental/Chaos/Public/Chaos/Box.h
--- a/Engine/Source/Runtime/Experimental/Chaos/Public/Chaos/Box.h
+++ b/Engine/Source/Runtime/Experimental/Chaos/Public/Chaos/Box.h
@@ -179,7 +179,7 @@ namespace Chaos
 			return TAABB<T, d>(InMin, InMax).RaycastFast(StartPoint, Dir, InvDir, bParallel, Length, InvLength, OutTime, OutPosition);
 		}
 
-		virtual bool CHAOS_API Raycast(const FVec3& StartPoint, const FVec3& Dir, const FReal Length, const FReal Thickness, FReal& OutTime, FVec3& OutPosition, FVec3& OutNormal, int32& OutFaceIndex) const override
+		virtual bool Raycast(const FVec3& StartPoint, const FVec3& Dir, const FReal Length, const FReal Thickness, FReal& OutTime, FVec3& OutPosition, FVec3& OutNormal, int32& OutFaceIndex) const override
 		{
 			if (AABB.Raycast(StartPoint, Dir, Length, Thickness, OutTime, OutPosition, OutNormal, OutFaceIndex))
 			{
diff --git a/Engine/Source/Runtime/Experimental/Chaos/Public/Chaos/CCDModification.h b/Engine/Source/Runtime/Experimental/Chaos/Public/Chaos/CCDModification.h
--- a/Engine/Source/Runtime/Experimental/Chaos/Public/Chaos/CCDModification.h
+++ b/Engine/Source/Runtime/Experimental/Chaos/Public/Chaos/CCDModification.h
@@ -27,7 +27,7 @@ namespace Chaos
 	 * 
 	 * @see FMidPhaseModifier, FContactPairModifier
 	 */
-	class CHAOS_API FCCDModifier
+	class FCCDModifier
 	{
 	public:
 		FCCDModifier()
@@ -44,34 +44,34 @@ namespace Chaos
 		//
 
 		// Whether the two particle actually hit each other in the sweep
-		bool IsSweepHit() const;
+		CHAOS_API bool IsSweepHit() const;
 
 		// Get the time of impact. This is as a fraction of total movement and will be in the range [0,1] if a hit was found
-		FReal GetSweepHitTOI() const;
+		CHAOS_API FReal GetSweepHitTOI() const;
 
 		// Get the sweep impact position for the specified body
 		// NOTE: See class comments.
-		FVec3 GetWorldSweepHitLocation(const int32 ParticleIndex) const;
+		CHAOS_API FVec3 GetWorldSweepHitLocation(const int32 ParticleIndex) const;
 
 		// Get the sweep impact normal
 		// NOTE: See class comments.
-		FVec3 GetWorldSweepHitNormal() const;
+		CHAOS_API FVec3 GetWorldSweepHitNormal() const;
 
-		const FGeometryParticleHandle* GetParticle(const int32 ParticleIndex) const;
-		const FGeometryParticleHandle* GetOtherParticle(const FGeometryParticleHandle* InParticle) const;
+		CHAOS_API const FGeometryParticleHandle* GetParticle(const int32 ParticleIndex) const;
+		CHAOS_API const FGeometryParticleHandle* GetOtherParticle(const FGeometryParticleHandle* InParticle) const;
 
 		//
 		// Modifying functions
 		//
 
 		// Re-enable this contact (if it was previously disabled)
-		void Enable();
+		CHAOS_API void Enable();
 
 		// Disable this contact
-		void Disable();
+		CHAOS_API void Disable();
 
 		// Convert this contact to a probe
-		void ConvertToProbe();
+		CHAOS_API void ConvertToProbe();
 
 	private:
 		FCCDModifier(
@@ -96,7 +96,7 @@ namespace Chaos
 	/*
 	 * Class for iterating over ccd results involving a specific particle
 	 */
-	class CHAOS_API FCCDModifierParticleIterator
+	class FCCDModifierParticleIterator
 	{
 	public:
 		FCCDModifier& operator*()
@@ -136,9 +136,9 @@ namespace Chaos
 		{ 
 		}
 
-		FCCDModifierParticleIterator& MoveToBegin();
-		FCCDModifierParticleIterator& MoveToEnd();
-		FCCDModifierParticleIterator& MoveToNext();
+		CHAOS_API FCCDModifierParticleIterator& MoveToBegin();
+		CHAOS_API FCCDModifierParticleIterator& MoveToEnd();
+		CHAOS_API FCCDModifierParticleIterator& MoveToNext();
 
 		FCCDModifierParticleRange* Range;
 		int32 ConstraintIndex;
@@ -150,22 +150,22 @@ namespace Chaos
 	/*
 	 * Represents all the swept constraints on a particle to allow for iteration.
 	 */
-	class CHAOS_API FCCDModifierParticleRange
+	class FCCDModifierParticleRange
 	{
 	public:
-		FCCDModifierParticleRange(
+		CHAOS_API FCCDModifierParticleRange(
 			FCCDModifierAccessor* InAccessor, 
 			FGeometryParticleHandle* InParticle);
 
-		FCCDModifierParticleIterator begin();
-		FCCDModifierParticleIterator end();
+		CHAOS_API FCCDModifierParticleIterator begin();
+		CHAOS_API FCCDModifierParticleIterator end();
 
 	private:
 		FCCDModifierAccessor* GetAccessor() const { return Accessor; }
 		FGeometryParticleHandle* GetParticle() const { return Particle; }
 		int32 GetNumConstraints() const { return Constraints.Num(); }
 		FPBDCollisionConstraint* GetConstraint(const int32 Index) const { return Constraints[Index]; }
-		void CollectConstraints();
+		CHAOS_API void CollectConstraints();
 
 		FCCDModifierAccessor* Accessor;
 		FGeometryParticleHandle* Particle;
@@ -177,13 +177,13 @@ namespace Chaos
 	/*
 	 * Provides access to CCD modifiers
 	 */
-	class CHAOS_API FCCDModifierAccessor
+	class FCCDModifierAccessor
 	{
 	public:
-		FCCDModifierAccessor(const FReal InDt);
+		CHAOS_API FCCDModifierAccessor(const FReal InDt);
 
 		// Get an object which allows for range iteration over the CCD modifiers for a particle
-		FCCDModifierParticleRange GetModifiers(FGeometryParticleHandle* Particle);
+		CHAOS_API FCCDModifierParticleRange GetModifiers(FGeometryParticleHandle* Particle);
 
 		FReal GetDt() const
 		{
diff --git a/Engine/Source/Runtime/Experimental/Chaos/Public/Chaos/ChaosArchive.h b/Engine/Source/Runtime/Experimental/Chaos/Public/Chaos/ChaosArchive.h
--- a/Engine/Source/Runtime/Experimental/Chaos/Public/Chaos/ChaosArchive.h
+++ b/Engine/Source/Runtime/Experimental/Chaos/Public/Chaos/ChaosArchive.h
@@ -10,7 +10,7 @@
 
 namespace Chaos
 {
-class CHAOS_API FImplicitObject; //needed for legacy serializer
+class FImplicitObject; //needed for legacy serializer
 
 
 #if CHAOS_MEMORY_TRACKING
@@ -68,7 +68,7 @@ class FChaosArchive;
 
 
 
-class CHAOS_API FChaosArchiveContext
+class FChaosArchiveContext
 #if CHAOS_MEMORY_TRACKING
 	: public FChaosArchiveMemoryTrackingContext
 #endif
@@ -80,9 +80,9 @@ public:
 	int32 TagCount;
 
 
-	FChaosArchiveContext();
+	CHAOS_API FChaosArchiveContext();
 
-	~FChaosArchiveContext();
+	CHAOS_API ~FChaosArchiveContext();
 
 	template <typename T, ESPMode Mode>
 	TSharedPtr<T, Mode>& ToSharedPointerHelper(TSerializablePtr<T>& Obj)
@@ -133,7 +133,7 @@ private:
 
 };
 
-class CHAOS_API FChaosArchive : public FArchiveProxy
+class FChaosArchive : public FArchiveProxy
 {
 public:
 	FChaosArchive(FArchive& ArIn)
@@ -281,7 +281,7 @@ private:
 		check(false);
 	}
 
-	void SerializeLegacy(TUniquePtr<FImplicitObject>& Obj);
+	CHAOS_API void SerializeLegacy(TUniquePtr<FImplicitObject>& Obj);
 
 	template <typename T>
 	void StaticSerialize(TSerializablePtr<T>& Serializable)
@@ -309,7 +309,7 @@ private:
 #endif
 };
 
-class CHAOS_API FChaosArchiveScopedMemory
+class FChaosArchiveScopedMemory
 {
 public:
 	FChaosArchiveScopedMemory(FChaosArchive& ArIn, const FName& SectionName, const bool bAbsorbChildren = true)
diff --git a/Engine/Source/Runtime/Experimental/Chaos/Public/Chaos/ChaosDebugDraw.h b/Engine/Source/Runtime/Experimental/Chaos/Public/Chaos/ChaosDebugDraw.h
--- a/Engine/Source/Runtime/Experimental/Chaos/Public/Chaos/ChaosDebugDraw.h
+++ b/Engine/Source/Runtime/Experimental/Chaos/Public/Chaos/ChaosDebugDraw.h
@@ -25,9 +25,9 @@ namespace Chaos
 	{
 #if CHAOS_DEBUG_DRAW
 
-		struct CHAOS_API FChaosDebugDrawColorsByState
+		struct FChaosDebugDrawColorsByState
 		{
-			FChaosDebugDrawColorsByState(
+			CHAOS_API FChaosDebugDrawColorsByState(
 				FColor InDynamicColor,
 				FColor InSleepingColor,
 				FColor InKinematicColor,
@@ -39,12 +39,12 @@ namespace Chaos
 			FColor KinematicColor;
 			FColor StaticColor;
 
-			FColor GetColorFromState(EObjectStateType State) const;
+			CHAOS_API FColor GetColorFromState(EObjectStateType State) const;
 		};
 
-		struct CHAOS_API FChaosDebugDrawColorsByShapeType
+		struct FChaosDebugDrawColorsByShapeType
 		{
-			FChaosDebugDrawColorsByShapeType(
+			CHAOS_API FChaosDebugDrawColorsByShapeType(
 				FColor InSimpleTypeColor,
 				FColor InConvexColor,
 				FColor InHeightFieldColor,
@@ -59,10 +59,10 @@ namespace Chaos
 			FColor TriangleMeshColor;
 			FColor LevelSetColor;
 
-			FColor GetColorFromShapeType(EImplicitObjectType ShapeType) const;
+			CHAOS_API FColor GetColorFromShapeType(EImplicitObjectType ShapeType) const;
 		};
 
-		struct CHAOS_API FChaosDebugDrawSettings
+		struct FChaosDebugDrawSettings
 		{
 		public:
 			FChaosDebugDrawSettings(
diff --git a/Engine/Source/Runtime/Experimental/Chaos/Public/Chaos/ChaosMarshallingManager.h b/Engine/Source/Runtime/Experimental/Chaos/Public/Chaos/ChaosMarshallingManager.h
--- a/Engine/Source/Runtime/Experimental/Chaos/Public/Chaos/ChaosMarshallingManager.h
+++ b/Engine/Source/Runtime/Experimental/Chaos/Public/Chaos/ChaosMarshallingManager.h
@@ -263,11 +263,11 @@ struct FPushPhysicsData
 
 /** Manages data that gets marshaled from GT to PT using a timestamp
 */
-class CHAOS_API FChaosMarshallingManager
+class FChaosMarshallingManager
 {
 public:
-	FChaosMarshallingManager();
-	~FChaosMarshallingManager();
+	CHAOS_API FChaosMarshallingManager();
+	CHAOS_API ~FChaosMarshallingManager();
 
 	/** Grabs the producer data to write into. Should only be called by external thread */
 	FPushPhysicsData* GetProducerData_External()
@@ -296,19 +296,19 @@ public:
 		GetProducerData_External()->SimCallbackInputs.Add(FSimCallbackInputAndObject{ SimCallbackObject, InputData });
 	}
 	/** Step forward using the external delta time. Should only be called by external thread */
-	void Step_External(FReal ExternalDT, const int32 NumSteps = 1, bool bSolverSubstepped = false);
+	CHAOS_API void Step_External(FReal ExternalDT, const int32 NumSteps = 1, bool bSolverSubstepped = false);
 
 	/** Step the internal time forward if possible*/
-	FPushPhysicsData* StepInternalTime_External();
+	CHAOS_API FPushPhysicsData* StepInternalTime_External();
 
 	/** Frees the push data back into the pool. Internal thread should call this when finished processing data*/
-	void FreeData_Internal(FPushPhysicsData* PushData);
+	CHAOS_API void FreeData_Internal(FPushPhysicsData* PushData);
 
 	/** May record data to history, or may free immediately depending on rewind needs. Either way you should assume data is gone after calling this */
-	void FreeDataToHistory_Internal(FPushPhysicsData* PushData);
+	CHAOS_API void FreeDataToHistory_Internal(FPushPhysicsData* PushData);
 
 	/** Frees the pull data back into the pool. External thread should call this when finished processing data*/
-	void FreePullData_External(FPullPhysicsData* PullData);
+	CHAOS_API void FreePullData_External(FPullPhysicsData* PullData);
 
 	/** Returns the timestamp associated with inputs enqueued. */
 	int32 GetExternalTimestamp_External() const { return ExternalTimestamp_External; }
@@ -326,7 +326,7 @@ public:
 	FPullPhysicsData* GetCurrentPullData_Internal() { return CurPullData; }
 
 	/** Hands pull data off to external thread */
-	void FinalizePullData_Internal(int32 LatestExternalTimestampConsumed, FReal SimStartTime, FReal DeltaTime);
+	CHAOS_API void FinalizePullData_Internal(int32 LatestExternalTimestampConsumed, FReal SimStartTime, FReal DeltaTime);
 
 	/** Pops and returns the earliest pull data available. nullptr means results are not ready or no work is pending */
 	FPullPhysicsData* PopPullData_External()
@@ -336,10 +336,10 @@ public:
 		return Result;
 	}
 
-	void SetHistoryLength_Internal(int32 InHistoryLength);
+	CHAOS_API void SetHistoryLength_Internal(int32 InHistoryLength);
 
 	/** Returns the history buffer of the latest NumFrames. The history that comes before these frames is discarded*/
-	TArray<FPushPhysicsData*> StealHistory_Internal(int32 NumFrames);
+	CHAOS_API TArray<FPushPhysicsData*> StealHistory_Internal(int32 NumFrames);
 
 	/** Return the size of the history queue */
 	int32 GetNumHistory_Internal() const {return HistoryQueue_Internal.Num();}
@@ -367,7 +367,7 @@ private:
 
 	int32 HistoryLength;	//how long to keep push data for
 
-	void PrepareExternalQueue_External();
-	void PreparePullData();
+	CHAOS_API void PrepareExternalQueue_External();
+	CHAOS_API void PreparePullData();
 };
 }; // namespace Chaos
diff --git a/Engine/Source/Runtime/Experimental/Chaos/Public/Chaos/ChaosPerfTest.h b/Engine/Source/Runtime/Experimental/Chaos/Public/Chaos/ChaosPerfTest.h
--- a/Engine/Source/Runtime/Experimental/Chaos/Public/Chaos/ChaosPerfTest.h
+++ b/Engine/Source/Runtime/Experimental/Chaos/Public/Chaos/ChaosPerfTest.h
@@ -16,7 +16,7 @@ enum class EChaosPerfUnits
 	Num
 };
 
-class CHAOS_API FChaosScopedDurationTimeLogger
+class FChaosScopedDurationTimeLogger
 {
 public:
 	explicit FChaosScopedDurationTimeLogger(const TCHAR* InLabel)
@@ -40,9 +40,9 @@ public:
 		}
 	}
 
-	static const TCHAR* GlobalLabel;
+	static CHAOS_API const TCHAR* GlobalLabel;
 
-	static EChaosPerfUnits GlobalUnits;
+	static CHAOS_API EChaosPerfUnits GlobalUnits;
 
 	const TCHAR*   Label;
 	FOutputDevice* Device;
@@ -65,4 +65,4 @@ struct FScopedChaosPerfTest
 #else
 #define CHAOS_PERF_TEST(x, units)
 #define CHAOS_SCOPED_TIMER(x)
-#endif
\ No newline at end of file
+#endif
diff --git a/Engine/Source/Runtime/Experimental/Chaos/Public/Chaos/Character/CharacterGroundConstraint.h b/Engine/Source/Runtime/Experimental/Chaos/Public/Chaos/Character/CharacterGroundConstraint.h
--- a/Engine/Source/Runtime/Experimental/Chaos/Public/Chaos/Character/CharacterGroundConstraint.h
+++ b/Engine/Source/Runtime/Experimental/Chaos/Public/Chaos/Character/CharacterGroundConstraint.h
@@ -15,18 +15,18 @@ namespace Chaos
 	/// The FCharacterGroundConstraintProxy class, accessible via
 	/// GetProxy(), knows about both representations and marshalls
 	/// the data between them
-	class CHAOS_API FCharacterGroundConstraint : public FConstraintBase
+	class FCharacterGroundConstraint : public FConstraintBase
 	{
 	public:
 		using Base = FConstraintBase;
 		friend class FCharacterGroundConstraintProxy; // For setting the output data
 
-		FCharacterGroundConstraint();
+		CHAOS_API FCharacterGroundConstraint();
 		virtual ~FCharacterGroundConstraint() override {}
 
 		// The constraint requires a character body, which is fixed, but the ground body is
 		// optional and can be changed dynamically using SetGroundParticleProxy
-		void Init(FSingleParticlePhysicsProxy* InCharacterProxy);
+		CHAOS_API void Init(FSingleParticlePhysicsProxy* InCharacterProxy);
 
 		//////////////////////////////////////////////////////////////////////////
 		// Particle proxy
@@ -296,7 +296,7 @@ namespace Chaos
 
 		//////////////////////////////////////////////////////////////////////////
 		// FConstraintBase implementation
-		virtual void SyncRemoteDataImp(FDirtyPropertiesManager& Manager, int32 DataIdx, FDirtyChaosProperties& RemoteData) override;
+		CHAOS_API virtual void SyncRemoteDataImp(FDirtyPropertiesManager& Manager, int32 DataIdx, FDirtyChaosProperties& RemoteData) override;
 
 		TChaosProperty<FParticleProxyProperty, EChaosProperty::CharacterParticleProxy> CharacterProxy;
 		TChaosProperty<FParticleProxyProperty, EChaosProperty::GroundParticleProxy> GroundProxy;
@@ -308,4 +308,4 @@ namespace Chaos
 		FVector SolverAppliedTorque;
 	};
 
-} // namespace Chaos
\ No newline at end of file
+} // namespace Chaos
diff --git a/Engine/Source/Runtime/Experimental/Chaos/Public/Chaos/Character/CharacterGroundConstraintContainer.h b/Engine/Source/Runtime/Experimental/Chaos/Public/Chaos/Character/CharacterGroundConstraintContainer.h
--- a/Engine/Source/Runtime/Experimental/Chaos/Public/Chaos/Character/CharacterGroundConstraintContainer.h
+++ b/Engine/Source/Runtime/Experimental/Chaos/Public/Chaos/Character/CharacterGroundConstraintContainer.h
@@ -25,7 +25,7 @@ namespace Chaos
 	/// Physics thread side representation of a character ground constraint
 	/// Data accessed through this class is synced from the game thread
 	/// representation FCharacterGroundConstraint
-	class CHAOS_API FCharacterGroundConstraintHandle final : public TIntrusiveConstraintHandle<FCharacterGroundConstraintHandle>
+	class FCharacterGroundConstraintHandle final : public TIntrusiveConstraintHandle<FCharacterGroundConstraintHandle>
 	{
 	public:
 		using Base = TIntrusiveConstraintHandle<FCharacterGroundConstraintHandle>;
@@ -102,25 +102,25 @@ namespace Chaos
 	};
 
 	/// Container class for all character ground constraints on the physics thread
-	class CHAOS_API FCharacterGroundConstraintContainer : public FPBDConstraintContainer
+	class FCharacterGroundConstraintContainer : public FPBDConstraintContainer
 	{
 	public:
 		using Base = FPBDConstraintContainer;
 		using FConstraints = TArrayView<FCharacterGroundConstraintHandle* const>;
 		using FConstConstraints = TArrayView<const FCharacterGroundConstraintHandle* const>;
 
-		FCharacterGroundConstraintContainer();
-		virtual ~FCharacterGroundConstraintContainer();
+		CHAOS_API FCharacterGroundConstraintContainer();
+		CHAOS_API virtual ~FCharacterGroundConstraintContainer();
 
 		int32 NumConstraints() const { return Constraints.Num(); }
 
-		FCharacterGroundConstraintHandle* AddConstraint(
+		CHAOS_API FCharacterGroundConstraintHandle* AddConstraint(
 			const FCharacterGroundConstraintSettings& InConstraintSettings,
 			const FCharacterGroundConstraintDynamicData& InConstraintData,
 			FGeometryParticleHandle* CharacterParticle,
 			FGeometryParticleHandle* GroundParticle = nullptr);
 
-		void RemoveConstraint(FCharacterGroundConstraintHandle* Constraint);
+		CHAOS_API void RemoveConstraint(FCharacterGroundConstraintHandle* Constraint);
 
 		FConstraints GetConstraints() { return MakeArrayView(Constraints); }
 		FConstConstraints GetConstConstraints() const { return MakeArrayView(Constraints); }
@@ -130,24 +130,24 @@ namespace Chaos
 
 		//////////////////////////////////////////////////////////////////////////
 		// FConstraintContainer Implementation
-		virtual TUniquePtr<FConstraintContainerSolver> CreateSceneSolver(const int32 Priority) override final;
-		virtual TUniquePtr<FConstraintContainerSolver> CreateGroupSolver(const int32 Priority) override final;
+		CHAOS_API virtual TUniquePtr<FConstraintContainerSolver> CreateSceneSolver(const int32 Priority) override final;
+		CHAOS_API virtual TUniquePtr<FConstraintContainerSolver> CreateGroupSolver(const int32 Priority) override final;
 		virtual int32 GetNumConstraints() const override final { return Constraints.Num(); }
 		virtual void ResetConstraints() override final {}
-		virtual void AddConstraintsToGraph(Private::FPBDIslandManager& IslandManager) override final;
-		virtual void PrepareTick() override final;
-		virtual void UnprepareTick() override final;
-		virtual void DisconnectConstraints(const TSet<TGeometryParticleHandle<FReal, 3>*>& RemovedParticles) override final;
+		CHAOS_API virtual void AddConstraintsToGraph(Private::FPBDIslandManager& IslandManager) override final;
+		CHAOS_API virtual void PrepareTick() override final;
+		CHAOS_API virtual void UnprepareTick() override final;
+		CHAOS_API virtual void DisconnectConstraints(const TSet<TGeometryParticleHandle<FReal, 3>*>& RemovedParticles) override final;
 
 		//////////////////////////////////////////////////////////////////////////
 		// Required API from FConstraintContainer
-		void SetConstraintEnabled(int32 ConstraintIndex, bool bEnabled);
+		CHAOS_API void SetConstraintEnabled(int32 ConstraintIndex, bool bEnabled);
 		bool IsConstraintEnabled(int32 ConstraintIndex) const { check(ConstraintIndex < NumConstraints()); return Constraints[ConstraintIndex]->IsEnabled(); }
 
 	private:
-		bool CanEvaluate(const FCharacterGroundConstraintHandle* Constraint) const;
+		CHAOS_API bool CanEvaluate(const FCharacterGroundConstraintHandle* Constraint) const;
 
 		TObjectPool<FCharacterGroundConstraintHandle> ConstraintPool;
 		TArray<FCharacterGroundConstraintHandle*> Constraints;
 	};
-}
\ No newline at end of file
+}
diff --git a/Engine/Source/Runtime/Experimental/Chaos/Public/Chaos/Character/CharacterGroundConstraintSettings.h b/Engine/Source/Runtime/Experimental/Chaos/Public/Chaos/Character/CharacterGroundConstraintSettings.h
--- a/Engine/Source/Runtime/Experimental/Chaos/Public/Chaos/Character/CharacterGroundConstraintSettings.h
+++ b/Engine/Source/Runtime/Experimental/Chaos/Public/Chaos/Character/CharacterGroundConstraintSettings.h
@@ -5,7 +5,7 @@
 
 namespace Chaos
 {
-	class CHAOS_API FCharacterGroundConstraintSettings
+	class FCharacterGroundConstraintSettings
 	{
 	public:
 		bool operator==(const FCharacterGroundConstraintSettings& Other) const
@@ -24,7 +24,7 @@ namespace Chaos
 		void* UserData = nullptr;
 	};
 
-	class CHAOS_API FCharacterGroundConstraintDynamicData
+	class FCharacterGroundConstraintDynamicData
 	{
 	public:
 		bool operator==(const FCharacterGroundConstraintDynamicData& Other) const
@@ -39,4 +39,4 @@ namespace Chaos
 		FReal CosMaxWalkableSlopeAngle = 0.633;			/// Override for max walkable slope angle
 	};
 
-} // namespace Chaos
\ No newline at end of file
+} // namespace Chaos
diff --git a/Engine/Source/Runtime/Experimental/Chaos/Public/Chaos/ClusterCreationParameters.h b/Engine/Source/Runtime/Experimental/Chaos/Public/Chaos/ClusterCreationParameters.h
--- a/Engine/Source/Runtime/Experimental/Chaos/Public/Chaos/ClusterCreationParameters.h
+++ b/Engine/Source/Runtime/Experimental/Chaos/Public/Chaos/ClusterCreationParameters.h
@@ -10,7 +10,7 @@ namespace Chaos
 {
 	class FBVHParticles;
 
-	struct CHAOS_API FClusterCreationParameters
+	struct FClusterCreationParameters
 	{
 		enum EConnectionMethod
 		{
diff --git a/Engine/Source/Runtime/Experimental/Chaos/Public/Chaos/ClusterUnionManager.h b/Engine/Source/Runtime/Experimental/Chaos/Public/Chaos/ClusterUnionManager.h
--- a/Engine/Source/Runtime/Experimental/Chaos/Public/Chaos/ClusterUnionManager.h
+++ b/Engine/Source/Runtime/Experimental/Chaos/Public/Chaos/ClusterUnionManager.h
@@ -50,7 +50,7 @@ namespace Chaos
 		Remove
 	};
 
-	struct CHAOS_API FClusterUnionCreationParameters
+	struct FClusterUnionCreationParameters
 	{
 		FClusterUnionExplicitIndex ExplicitIndex = INDEX_NONE;
 		const FUniqueIdx* UniqueIndex = nullptr;
@@ -58,13 +58,13 @@ namespace Chaos
 		uint32 ComponentId = 0;
 	};
 
-	struct CHAOS_API FClusterUnionParticleProperties
+	struct FClusterUnionParticleProperties
 	{
 		// An auxiliary particle will be removed from the cluster union if FRigidClustering::HandleConnectivityOnReleaseClusterParticle detects an island made up of only auxiliary particles.
 		bool bIsAuxiliaryParticle = false;
 	};
 
-	struct CHAOS_API FClusterUnion
+	struct FClusterUnion
 	{
 		// The root cluster particle that we created internally to represent the cluster.
 		FPBDRigidClusteredParticleHandle* InternalCluster;
@@ -114,7 +114,7 @@ namespace Chaos
 	 * cause one or more clusters to simulate together as a single rigid
 	 * particle.
 	 */
-	class CHAOS_API FClusterUnionManager
+	class FClusterUnionManager
 	{
 	public:
 
@@ -133,56 +133,56 @@ namespace Chaos
 			return TransformedChildGeometry;
 		}
 
-		FClusterUnionManager(FRigidClustering& InClustering, FPBDRigidsEvolutionGBF& InEvolution);
+		CHAOS_API FClusterUnionManager(FRigidClustering& InClustering, FPBDRigidsEvolutionGBF& InEvolution);
 
 		// Creates a new cluster union with an automatically assigned cluster union index.
-		FClusterUnionIndex CreateNewClusterUnion(const FClusterCreationParameters& Parameters, const FClusterUnionCreationParameters& ClusterUnionParameters = FClusterUnionCreationParameters{});
+		CHAOS_API FClusterUnionIndex CreateNewClusterUnion(const FClusterCreationParameters& Parameters, const FClusterUnionCreationParameters& ClusterUnionParameters = FClusterUnionCreationParameters{});
 
 		// Destroy a given cluster union.
-		void DestroyClusterUnion(FClusterUnionIndex Index);
+		CHAOS_API void DestroyClusterUnion(FClusterUnionIndex Index);
 
 		// Add a new operation to the queue. Note that we currently only support the pending/flush operations for explicit operations. The behavior is legacy anyway so this should be fine.
-		void AddPendingExplicitIndexOperation(FClusterUnionExplicitIndex Index, EClusterUnionOperation Op, const TArray<FPBDRigidParticleHandle*>& Particles);
-		void AddPendingClusterIndexOperation(FClusterUnionIndex Index, EClusterUnionOperation Op, const TArray<FPBDRigidParticleHandle*>& Particles);
+		CHAOS_API void AddPendingExplicitIndexOperation(FClusterUnionExplicitIndex Index, EClusterUnionOperation Op, const TArray<FPBDRigidParticleHandle*>& Particles);
+		CHAOS_API void AddPendingClusterIndexOperation(FClusterUnionIndex Index, EClusterUnionOperation Op, const TArray<FPBDRigidParticleHandle*>& Particles);
 
 		// Actually performs the change specified in the FClusterUnionOperationData structure.
-		void HandleAddOperation(FClusterUnionIndex ClusterIndex, const TArray<FPBDRigidParticleHandle*>& Particles, bool bReleaseClustersFirst);
+		CHAOS_API void HandleAddOperation(FClusterUnionIndex ClusterIndex, const TArray<FPBDRigidParticleHandle*>& Particles, bool bReleaseClustersFirst);
 
 		// Removes the specified particles from the specified cluster.
-		void HandleRemoveOperation(FClusterUnionIndex ClusterIndex, const TArray<FPBDRigidParticleHandle*>& Particles, EClusterUnionOperationTiming UpdateClusterPropertiesTiming);
+		CHAOS_API void HandleRemoveOperation(FClusterUnionIndex ClusterIndex, const TArray<FPBDRigidParticleHandle*>& Particles, EClusterUnionOperationTiming UpdateClusterPropertiesTiming);
 
 		// Helper function to remove particles given only the particle handle. This will consult the lookup table to find which cluster the particle is in before doing a normal remove operation.
-		void HandleRemoveOperationWithClusterLookup(const TArray<FPBDRigidParticleHandle*>& InParticles, EClusterUnionOperationTiming UpdateClusterPropertiesTiming);
+		CHAOS_API void HandleRemoveOperationWithClusterLookup(const TArray<FPBDRigidParticleHandle*>& InParticles, EClusterUnionOperationTiming UpdateClusterPropertiesTiming);
 
 		// Performs the queued child to parent modifications.
-		void HandleUpdateChildToParentOperation(FClusterUnionIndex ClusterIndex, const TArray<FPBDRigidParticleHandle*>& Particles);
+		CHAOS_API void HandleUpdateChildToParentOperation(FClusterUnionIndex ClusterIndex, const TArray<FPBDRigidParticleHandle*>& Particles);
 
 		// Will be called at the beginning of every time step to ensure that all the expected cluster unions have been modified.
-		void FlushPendingOperations();
+		CHAOS_API void FlushPendingOperations();
 
 		// Update cluster union properties here if they were deferred. This can be called manually but will otherwise also be handled in FlushPendingOperations.
-		void HandleDeferredClusterUnionUpdateProperties();
+		CHAOS_API void HandleDeferredClusterUnionUpdateProperties();
 
 		// Access the cluster union externally.
-		FClusterUnion* FindClusterUnionFromExplicitIndex(FClusterUnionExplicitIndex Index);
-		FClusterUnion* FindClusterUnion(FClusterUnionIndex Index);
-		FClusterUnion* FindClusterUnionFromParticle(FPBDRigidParticleHandle* Particle);
-		FClusterUnionIndex FindClusterUnionIndexFromParticle(FPBDRigidParticleHandle* Particle);
+		CHAOS_API FClusterUnion* FindClusterUnionFromExplicitIndex(FClusterUnionExplicitIndex Index);
+		CHAOS_API FClusterUnion* FindClusterUnion(FClusterUnionIndex Index);
+		CHAOS_API FClusterUnion* FindClusterUnionFromParticle(FPBDRigidParticleHandle* Particle);
+		CHAOS_API FClusterUnionIndex FindClusterUnionIndexFromParticle(FPBDRigidParticleHandle* Particle);
 
 		// An extension to FindClusterUnionIndexFromParticle to check whether or not the given particle is the cluster union particle itself.
-		bool IsClusterUnionParticle(FPBDRigidClusteredParticleHandle* Particle);
+		CHAOS_API bool IsClusterUnionParticle(FPBDRigidClusteredParticleHandle* Particle);
 
 		// Changes the ChildToParent of a number of particles in a cluster union. bLock will prevent other functions from overriding the ChildToParent (aside from another call to UpdateClusterUnionParticlesChildToParent).
-		void UpdateClusterUnionParticlesChildToParent(FClusterUnionIndex Index, const TArray<FPBDRigidParticleHandle*>& Particles, const TArray<FTransform>& ChildToParent, bool bLock);
+		CHAOS_API void UpdateClusterUnionParticlesChildToParent(FClusterUnionIndex Index, const TArray<FPBDRigidParticleHandle*>& Particles, const TArray<FTransform>& ChildToParent, bool bLock);
 
 		// Update the cluster union's properties after its set of particle changes.
-		void UpdateAllClusterUnionProperties(FClusterUnion& ClusterUnion, EUpdateClusterUnionPropertiesFlags Flags = EUpdateClusterUnionPropertiesFlags::All);
+		CHAOS_API void UpdateAllClusterUnionProperties(FClusterUnion& ClusterUnion, EUpdateClusterUnionPropertiesFlags Flags = EUpdateClusterUnionPropertiesFlags::All);
 
 		// Returns all cluster unions. Really meant only to be used for debugging.
 		const TMap<FClusterUnionIndex, FClusterUnion>& GetAllClusterUnions() const { return ClusterUnions; }
 
 		// Put in a deferred request to update the properties of the specified cluster union. Flags can limit the scope of the update.
-		void RequestDeferredClusterPropertiesUpdate(FClusterUnionIndex ClusterIndex, EUpdateClusterUnionPropertiesFlags Flags);
+		CHAOS_API void RequestDeferredClusterPropertiesUpdate(FClusterUnionIndex ClusterIndex, EUpdateClusterUnionPropertiesFlags Flags);
 	private:
 		FRigidClustering& MClustering;
 		FPBDRigidsEvolutionGBF& MEvolution;
@@ -216,7 +216,7 @@ namespace Chaos
 		// - If a cluster union gets released/destroyed, that index can be reused.
 		// - Otherwise, we use the NextAvailableUnionIndex which is just the max index we've seen + 1.
 		//
-		FClusterUnionIndex ClaimNextUnionIndex();
+		CHAOS_API FClusterUnionIndex ClaimNextUnionIndex();
 		TArray<FClusterUnionIndex> ReusableIndices;
 		FClusterUnionIndex NextAvailableUnionIndex = 1;
 
@@ -225,19 +225,19 @@ namespace Chaos
 		// (i.e. an index that comes in via FClusterUnionOperationData for the first time) to an automatically
 		// generated index (i.e one that would returned via CreateNewClusterUnion).
 		TMap<FClusterUnionExplicitIndex, FClusterUnionIndex> ExplicitIndexMap;
-		FClusterCreationParameters DefaultClusterCreationParameters() const;
+		CHAOS_API FClusterCreationParameters DefaultClusterCreationParameters() const;
 
 		// If no cluster index is set but an explicit index is set, map the explicit index to a regular index.
-		FClusterUnionIndex GetOrCreateClusterUnionIndexFromExplicitIndex(FClusterUnionExplicitIndex InIndex);
+		CHAOS_API FClusterUnionIndex GetOrCreateClusterUnionIndexFromExplicitIndex(FClusterUnionExplicitIndex InIndex);
 
 		// Forcefully recreate the shared geometry on a cluster. Potentially expensive so ideally should be used rarely.
-		TSharedPtr<FImplicitObject, ESPMode::ThreadSafe> ForceRecreateClusterUnionSharedGeometry(const FClusterUnion& Union);
+		CHAOS_API TSharedPtr<FImplicitObject, ESPMode::ThreadSafe> ForceRecreateClusterUnionSharedGeometry(const FClusterUnion& Union);
 
 		// Handles updating the cluster union.
-		void DeferredClusterUnionUpdate(FClusterUnion& Union, EUpdateClusterUnionPropertiesFlags Flags);
+		CHAOS_API void DeferredClusterUnionUpdate(FClusterUnion& Union, EUpdateClusterUnionPropertiesFlags Flags);
 
 		// Flush the cluster union's incremental connectivity operations
-		void FlushIncrementalConnectivityGraphOperations(FClusterUnion& ClusterUnion);
+		CHAOS_API void FlushIncrementalConnectivityGraphOperations(FClusterUnion& ClusterUnion);
 	};
 
-}
\ No newline at end of file
+}
diff --git a/Engine/Source/Runtime/Experimental/Chaos/Public/Chaos/CollectionPropertyFacade.h b/Engine/Source/Runtime/Experimental/Chaos/Public/Chaos/CollectionPropertyFacade.h
--- a/Engine/Source/Runtime/Experimental/Chaos/Public/Chaos/CollectionPropertyFacade.h
+++ b/Engine/Source/Runtime/Experimental/Chaos/Public/Chaos/CollectionPropertyFacade.h
@@ -45,10 +45,10 @@ namespace Chaos::Softs
 	 * This is mainly used for the cloth simulation properties to provide
 	 * weighted values that works in conjunction with weight maps.
 	 */
-	class CHAOS_API FCollectionPropertyConstFacade
+	class FCollectionPropertyConstFacade
 	{
 	public:
-		explicit FCollectionPropertyConstFacade(const TSharedPtr<const FManagedArrayCollection>& InManagedArrayCollection);
+		CHAOS_API explicit FCollectionPropertyConstFacade(const TSharedPtr<const FManagedArrayCollection>& InManagedArrayCollection);
 		virtual ~FCollectionPropertyConstFacade() = default;
 
 		FCollectionPropertyConstFacade() = delete;
@@ -59,7 +59,7 @@ namespace Chaos::Softs
 		FCollectionPropertyConstFacade& operator=(FCollectionPropertyConstFacade&&) = default;
 
 		/** Return whether the facade is defined on the collection. */
-		bool IsValid() const;
+		CHAOS_API bool IsValid() const;
 
 		/** Return the number of properties in this collection. */
 		int32 Num() const { return KeyIndices.Num(); }
@@ -159,13 +159,13 @@ PRAGMA_ENABLE_DEPRECATION_WARNINGS
 
 	protected:
 		// No init constructor for FCollectionPropertyFacade
-		FCollectionPropertyConstFacade(const TSharedPtr<const FManagedArrayCollection>& InManagedArrayCollection, ENoInit);
+		CHAOS_API FCollectionPropertyConstFacade(const TSharedPtr<const FManagedArrayCollection>& InManagedArrayCollection, ENoInit);
 
 		// Update the array views
-		void UpdateArrays();
+		CHAOS_API void UpdateArrays();
 
 		// Rebuild the search map
-		void RebuildKeyIndices();
+		CHAOS_API void RebuildKeyIndices();
 
 		template<typename CallableType, typename ReturnType>
 		ReturnType SafeGet(const FString& Key, CallableType&& Callable, ReturnType Default, int32* OutKeyIndex) const
@@ -187,12 +187,12 @@ PRAGMA_ENABLE_DEPRECATION_WARNINGS
 		bool HasAnyFlags(int32 KeyIndex, ECollectionPropertyFlags Flags) const { return EnumHasAnyFlags(GetValue<ECollectionPropertyFlags, uint8>(KeyIndex, FlagsArray), Flags); }
 
 		// Attribute groups, predefined data member of the collection
-		static const FName PropertyGroup;
-		static const FName KeyName;  // Property key, name to look for
-		static const FName LowValueName;  // Boolean, 24 bit integer (max 16777215), float, or vector value, or value of the lowest weight on the weight map if any
-		static const FName HighValueName;  // Boolean, 24 bit integer (max 16777215), float, or vector value of the highest weight on the weight map if any
-		static const FName StringValueName;  // String value, or weight map name, ...etc.
-		static const FName FlagsName;  // Whether this property is enabled, animatable, ...etc.
+		static CHAOS_API const FName PropertyGroup;
+		static CHAOS_API const FName KeyName;  // Property key, name to look for
+		static CHAOS_API const FName LowValueName;  // Boolean, 24 bit integer (max 16777215), float, or vector value, or value of the lowest weight on the weight map if any
+		static CHAOS_API const FName HighValueName;  // Boolean, 24 bit integer (max 16777215), float, or vector value of the highest weight on the weight map if any
+		static CHAOS_API const FName StringValueName;  // String value, or weight map name, ...etc.
+		static CHAOS_API const FName FlagsName;  // Whether this property is enabled, animatable, ...etc.
 
 		// Property Group array views
 		TConstArrayView<FString> KeyArray;
@@ -215,10 +215,10 @@ PRAGMA_ENABLE_DEPRECATION_WARNINGS
 	 * 
 	 * Note: Int property values are stored as float and therefore limited to 24 bits (-16777215 to 16777215).
 	 */
-	class CHAOS_API FCollectionPropertyFacade : public FCollectionPropertyConstFacade
+	class FCollectionPropertyFacade : public FCollectionPropertyConstFacade
 	{
 	public:
-		explicit FCollectionPropertyFacade(const TSharedPtr<FManagedArrayCollection>& InManagedArrayCollection);
+		CHAOS_API explicit FCollectionPropertyFacade(const TSharedPtr<FManagedArrayCollection>& InManagedArrayCollection);
 		virtual ~FCollectionPropertyFacade() = default;
 
 		FCollectionPropertyFacade() = delete;
@@ -326,11 +326,11 @@ PRAGMA_ENABLE_DEPRECATION_WARNINGS
 			return SafeSet(Key, [this](int32 KeyIndex) { SetStringDirty(KeyIndex); });
 		}
 
-		void ClearDirtyFlags();
+		CHAOS_API void ClearDirtyFlags();
 
 	protected:
 		// No init constructor for FCollectionPropertyMutableFacade
-		FCollectionPropertyFacade(const TSharedPtr<FManagedArrayCollection>& InManagedArrayCollection, ENoInit);
+		CHAOS_API FCollectionPropertyFacade(const TSharedPtr<FManagedArrayCollection>& InManagedArrayCollection, ENoInit);
 
 		// Access to a writeable ManagedArrayCollection is protected, use an FPropertyCollectionMutableAdapter if needed to get a non const pointer
 		TSharedPtr<FManagedArrayCollection> GetManagedArrayCollection() { return ConstCastSharedPtr<FManagedArrayCollection>(ManagedArrayCollection); }
@@ -356,7 +356,7 @@ PRAGMA_ENABLE_DEPRECATION_WARNINGS
 		template<typename T>
 		inline void SetValue(int32 KeyIndex, const TArrayView<T>& ValueArray, const T& Value);
 
-		void EnableFlags(int32 KeyIndex, ECollectionPropertyFlags Flags, bool bEnable);
+		CHAOS_API void EnableFlags(int32 KeyIndex, ECollectionPropertyFlags Flags, bool bEnable);
 	};
 
 	/**
@@ -364,10 +364,10 @@ PRAGMA_ENABLE_DEPRECATION_WARNINGS
 	 * This is mainly used for the cloth simulation properties to provide
 	 * weighted values that works in conjunction with weight maps.
 	 */
-	class CHAOS_API FCollectionPropertyMutableFacade final : public FCollectionPropertyFacade
+	class FCollectionPropertyMutableFacade final : public FCollectionPropertyFacade
 	{
 	public:
-		explicit FCollectionPropertyMutableFacade(const TSharedPtr<FManagedArrayCollection>& InManagedArrayCollection);
+		CHAOS_API explicit FCollectionPropertyMutableFacade(const TSharedPtr<FManagedArrayCollection>& InManagedArrayCollection);
 		virtual ~FCollectionPropertyMutableFacade() = default;
 
 		FCollectionPropertyMutableFacade() = delete;
@@ -378,38 +378,38 @@ PRAGMA_ENABLE_DEPRECATION_WARNINGS
 		FCollectionPropertyMutableFacade& operator=(FCollectionPropertyMutableFacade&&) = default;
 
 		/** Create this facade's groups and attributes. */
-		void DefineSchema();
+		CHAOS_API void DefineSchema();
 
 		/** Add a single property, and return its index. */
-		int32 AddProperty(const FString& Key, bool bEnabled = true, bool bAnimatable = false);
+		CHAOS_API int32 AddProperty(const FString& Key, bool bEnabled = true, bool bAnimatable = false);
 
 		/** Add new properties, and return the index of the first added property. */
-		int32 AddProperties(const TArray<FString>& Keys, bool bEnabled = true, bool bAnimatable = false);
+		CHAOS_API int32 AddProperties(const TArray<FString>& Keys, bool bEnabled = true, bool bAnimatable = false);
 
 		/**
 		 * Append all properties and values from an existing collection to this property collection.
 		 * This won't copy any other groups, only data from PropertyGroup.
 		 * Any pre-existing data will be preserved.
 		 */
-		void Append(const FManagedArrayCollection& InManagedArrayCollection);
+		CHAOS_API void Append(const FManagedArrayCollection& InManagedArrayCollection);
 
 		/**
 		 * Copy all properties and values from an existing collection to this property collection.
 		 * This won't copy any other groups, only data from PropertyGroup.
 		 * Any pre-xisting data will be removed/replaced.
 		 */
-		void Copy(const FManagedArrayCollection& InManagedArrayCollection);
+		CHAOS_API void Copy(const FManagedArrayCollection& InManagedArrayCollection);
 
 		//~ Add values
 		template<typename T, TEMPLATE_REQUIRES(TIsWeightedType<T>::Value)>
 		inline int32 AddWeightedValue(const FString& Key, const T& LowValue, const T& HighValue, bool bEnabled = true, bool bAnimatable = false);
 
-		int32 AddWeightedFloatValue(const FString& Key, const FVector2f& Value, bool bEnabled, bool bAnimatable);
+		CHAOS_API int32 AddWeightedFloatValue(const FString& Key, const FVector2f& Value, bool bEnabled, bool bAnimatable);
 
 		template<typename T, TEMPLATE_REQUIRES(TIsWeightedType<T>::Value)>
 		int32 AddValue(const FString& Key, const T& Value, bool bEnabled = true, bool bAnimatable = false) { return AddWeightedValue(Key, Value, Value, bEnabled, bAnimatable); }
 
-		int32 AddStringValue(const FString& Key, const FString& Value, bool bEnabled = true, bool bAnimatable = false);
+		CHAOS_API int32 AddStringValue(const FString& Key, const FString& Value, bool bEnabled = true, bool bAnimatable = false);
 	};
 
 	template<typename T>
diff --git a/Engine/Source/Runtime/Experimental/Chaos/Public/Chaos/Collision/BasicBroadPhase.h b/Engine/Source/Runtime/Experimental/Chaos/Public/Chaos/Collision/BasicBroadPhase.h
--- a/Engine/Source/Runtime/Experimental/Chaos/Public/Chaos/Collision/BasicBroadPhase.h
+++ b/Engine/Source/Runtime/Experimental/Chaos/Public/Chaos/Collision/BasicBroadPhase.h
@@ -16,7 +16,7 @@ namespace Chaos
 	 *
 	 * No spatial acceleration, and the order is assumed to be already optimized for cache efficiency.
 	 */
-	class CHAOS_API FBasicBroadPhase
+	class FBasicBroadPhase
 	{
 	public:
 		using FParticleHandle = TGeometryParticleHandle<FReal, 3>;
diff --git a/Engine/Source/Runtime/Experimental/Chaos/Public/Chaos/Collision/BasicCollisionDetector.h b/Engine/Source/Runtime/Experimental/Chaos/Public/Chaos/Collision/BasicCollisionDetector.h
--- a/Engine/Source/Runtime/Experimental/Chaos/Public/Chaos/Collision/BasicCollisionDetector.h
+++ b/Engine/Source/Runtime/Experimental/Chaos/Public/Chaos/Collision/BasicCollisionDetector.h
@@ -11,7 +11,7 @@
 
 namespace Chaos
 {
-	class CHAOS_API FBasicCollisionDetector : public FCollisionDetector
+	class FBasicCollisionDetector : public FCollisionDetector
 	{
 	public:
 		FBasicCollisionDetector(FBasicBroadPhase& InBroadPhase, FPBDCollisionConstraints& InCollisionContainer)
@@ -45,4 +45,4 @@ namespace Chaos
 		FBasicBroadPhase& BroadPhase;
 	};
 
-}
\ No newline at end of file
+}
diff --git a/Engine/Source/Runtime/Experimental/Chaos/Public/Chaos/Collision/CollisionConstraintAllocator.h b/Engine/Source/Runtime/Experimental/Chaos/Public/Chaos/Collision/CollisionConstraintAllocator.h
--- a/Engine/Source/Runtime/Experimental/Chaos/Public/Chaos/Collision/CollisionConstraintAllocator.h
+++ b/Engine/Source/Runtime/Experimental/Chaos/Public/Chaos/Collision/CollisionConstraintAllocator.h
@@ -23,7 +23,7 @@ namespace Chaos
 		 * Container the storage for the FCollisionConstraintAllocator, as well as the API to create new midphases and collision constraints.
 		 * We have one of these objects per thread on which collisions detection is performed to get lock-free allocations and lists.
 		 */
-		class CHAOS_API FCollisionContextAllocator
+		class FCollisionContextAllocator
 		{
 		public:
 			FCollisionContextAllocator(FPBDCollisionConstraints* InCollisionContainer, const int32 InNumCollisionsPerBlock, const int32 InCurrentEpoch)
@@ -49,7 +49,7 @@ namespace Chaos
 			/**
 			 * Create a constraint (called by the MidPhase)
 			 */
-			FPBDCollisionConstraintPtr CreateConstraint(
+			CHAOS_API FPBDCollisionConstraintPtr CreateConstraint(
 				FGeometryParticleHandle* Particle0,
 				const FImplicitObject* Implicit0,
 				const FPerShapeData* Shape0,
@@ -150,10 +150,10 @@ namespace Chaos
 			}
 
 			// Process all the new midphases from the parallel collision detection
-			void ProcessNewMidPhases(FCollisionConstraintAllocator* Allocator);
+			CHAOS_API void ProcessNewMidPhases(FCollisionConstraintAllocator* Allocator);
 
 			// Process all the activated collisions from the parallel collision detection
-			void ProcessNewConstraints(FCollisionConstraintAllocator* Allocator);
+			CHAOS_API void ProcessNewConstraints(FCollisionConstraintAllocator* Allocator);
 
 			// Find the midphase for the particle pair if it exists. Every particle holds a list of its midphases. We search "SearchParticle" which should be
 			// ideally be the one with fewer midphases on it.
@@ -230,7 +230,7 @@ namespace Chaos
 		 * NOTE: To reduce RBAN memory use, we do not create any collision blocks until the first call to CreateCollisionConstraint
 		 * (see ConstraintPool initialization)
 		*/
-		class CHAOS_API FCollisionConstraintAllocator
+		class FCollisionConstraintAllocator
 		{
 		public:
 			UE_NONCOPYABLE(FCollisionConstraintAllocator);
@@ -476,7 +476,7 @@ namespace Chaos
 			 * @brief Destroy all collision and caches involving the particle
 			 * Called when a particle is destroyed or disabled (not sleeping).
 			*/
-			void RemoveParticle(FGeometryParticleHandle* Particle);
+			CHAOS_API void RemoveParticle(FGeometryParticleHandle* Particle);
 
 			/**
 			 * @brief Iterate over all collisions (read-only), including sleeping ones
@@ -559,7 +559,7 @@ namespace Chaos
 			}
 
 			// Find all midphases that have not been updated this tick (i.e., bounds no longer overlap) and destroy them
-			void PruneExpiredMidPhases();
+			CHAOS_API void PruneExpiredMidPhases();
 
 			// Collect all the constraints activated in the collision tasks and register them (calls ActivateConstraintImp on each)
 			void ProcessNewConstraints()
@@ -589,7 +589,7 @@ namespace Chaos
 				Cookie.LastUsedEpoch = CurrentEpoch;
 			}
 
-			void RemoveActiveConstraint(FPBDCollisionConstraint& Constraint);
+			CHAOS_API void RemoveActiveConstraint(FPBDCollisionConstraint& Constraint);
 
 			// The container that owns the allocator (only needed because new constraints need to know)
 			FPBDCollisionConstraints* CollisionContainer;
diff --git a/Engine/Source/Runtime/Experimental/Chaos/Public/Chaos/Collision/CollisionConstraintFlags.h b/Engine/Source/Runtime/Experimental/Chaos/Public/Chaos/Collision/CollisionConstraintFlags.h
--- a/Engine/Source/Runtime/Experimental/Chaos/Public/Chaos/Collision/CollisionConstraintFlags.h
+++ b/Engine/Source/Runtime/Experimental/Chaos/Public/Chaos/Collision/CollisionConstraintFlags.h
@@ -29,7 +29,7 @@ namespace Chaos
 	// @todo(chaos): this should be broken out into FIgnoreCollisionManager (PT-only) and FIgnoreCollisionManagerProxy (interop).
 	// E.g., The interop could store sets of particle proxy pairs to enable/disable and call the ParticleHandle API on the PT object. 
 	// Or it could use the solver to map IDs to particle handles like we currently do.
-	class CHAOS_API FIgnoreCollisionManager
+	class FIgnoreCollisionManager
 	{
 	public:
 		using FHandleID = FUniqueIdx;
@@ -63,27 +63,27 @@ namespace Chaos
 		//
 		//
 
-		bool ContainsHandle(FHandleID Body0) const;
+		CHAOS_API bool ContainsHandle(FHandleID Body0) const;
 
-		bool IgnoresCollision(FHandleID Body0, FHandleID Body1) const;
+		CHAOS_API bool IgnoresCollision(FHandleID Body0, FHandleID Body1) const;
 
-		int32 NumIgnoredCollision(FHandleID Body0) const;
+		CHAOS_API int32 NumIgnoredCollision(FHandleID Body0) const;
 
 		/**
 		 * Return true if we should ignore collisions between the two particles
 		 */
-		bool IgnoresCollision(const FGeometryParticleHandle* Particle0, const FGeometryParticleHandle* Particle1) const;
+		CHAOS_API bool IgnoresCollision(const FGeometryParticleHandle* Particle0, const FGeometryParticleHandle* Particle1) const;
 
 		/**
 		 * Add an ignore entry for collisions of Particle0->Particle1.
 		 * Note, this will add an entry for both particles against the other if both particles are dynamic or kinematic.
 		 */
-		void AddIgnoreCollisions(FGeometryParticleHandle* Particle0, FGeometryParticleHandle* Particle1);
+		CHAOS_API void AddIgnoreCollisions(FGeometryParticleHandle* Particle0, FGeometryParticleHandle* Particle1);
 
 		/**
 		 * Remove an ignore entry for collisions of Particle0->Particle1 (in both directions if applicable)
 		 */
-		void RemoveIgnoreCollisions(FGeometryParticleHandle* Particle0, FGeometryParticleHandle* Particle1);
+		CHAOS_API void RemoveIgnoreCollisions(FGeometryParticleHandle* Particle0, FGeometryParticleHandle* Particle1);
 
 		//
 		//
@@ -138,12 +138,12 @@ namespace Chaos
 		/*
 		*
 		*/
-		void ProcessPendingQueues(FPBDRigidsSolver& Solver);
+		CHAOS_API void ProcessPendingQueues(FPBDRigidsSolver& Solver);
 
 		/*
 		*
 		*/
-		void PopStorageData_Internal(int32 ExternalTimestamp);
+		CHAOS_API void PopStorageData_Internal(int32 ExternalTimestamp);
 
 
 		// DEPRECATED
@@ -154,10 +154,10 @@ namespace Chaos
 
 	private:
 
-		FGeometryParticleHandle* GetParticleHandle(FHandleID Body, FPBDRigidsSolver& Solver);
-		void SetIgnoreCollisionFlag(FPBDRigidParticleHandle* Rigid, const bool bUsesIgnoreCollisionManager);
-		void AddIgnoreCollisionsImpl(FHandleID Body0, FHandleID Body1);
-		int32 RemoveIgnoreCollisionsImpl(FHandleID Body0, FHandleID Body1);
+		CHAOS_API FGeometryParticleHandle* GetParticleHandle(FHandleID Body, FPBDRigidsSolver& Solver);
+		CHAOS_API void SetIgnoreCollisionFlag(FPBDRigidParticleHandle* Rigid, const bool bUsesIgnoreCollisionManager);
+		CHAOS_API void AddIgnoreCollisionsImpl(FHandleID Body0, FHandleID Body1);
+		CHAOS_API int32 RemoveIgnoreCollisionsImpl(FHandleID Body0, FHandleID Body1);
 
 		// Multiple sources can request an ignore pair, so we handle a simple count of
 		// ignore requests to ensure we don't prematurely remove an ignore pair
diff --git a/Engine/Source/Runtime/Experimental/Chaos/Public/Chaos/Collision/CollisionDetector.h b/Engine/Source/Runtime/Experimental/Chaos/Public/Chaos/Collision/CollisionDetector.h
--- a/Engine/Source/Runtime/Experimental/Chaos/Public/Chaos/Collision/CollisionDetector.h
+++ b/Engine/Source/Runtime/Experimental/Chaos/Public/Chaos/Collision/CollisionDetector.h
@@ -10,7 +10,7 @@ namespace Chaos
 	class FPBDCollisionConstraints;
 	class FEvolutionResimCache;
 
-	class CHAOS_API FCollisionDetector
+	class FCollisionDetector
 	{
 	public:
 		FCollisionDetector(FPBDCollisionConstraints& InCollisionContainer)
@@ -21,13 +21,13 @@ namespace Chaos
 		virtual ~FCollisionDetector() {}
 
 		UE_DEPRECATED(5.2, "Moved to FPBDCollisionConstraints")
-		const FCollisionDetectorSettings& GetSettings() const;
+		CHAOS_API const FCollisionDetectorSettings& GetSettings() const;
 
 		UE_DEPRECATED(5.2, "Moved to FPBDCollisionConstraints")
-		void SetSettings(const FCollisionDetectorSettings& InSettings);
+		CHAOS_API void SetSettings(const FCollisionDetectorSettings& InSettings);
 
 		UE_DEPRECATED(5.2, "Moved to FPBDCollisionConstraints and renamed to SetCullDistance")
-		void SetBoundsExpansion(const FReal InBoundsExpansion);
+		CHAOS_API void SetBoundsExpansion(const FReal InBoundsExpansion);
 
 		UE_DEPRECATED(5.2, "No longer supported")
 		void SetBoundsVelocityInflation(const FReal InBoundsVelocityInflation) {}
diff --git a/Engine/Source/Runtime/Experimental/Chaos/Public/Chaos/Collision/ContactPoint.h b/Engine/Source/Runtime/Experimental/Chaos/Public/Chaos/Collision/ContactPoint.h
--- a/Engine/Source/Runtime/Experimental/Chaos/Public/Chaos/Collision/ContactPoint.h
+++ b/Engine/Source/Runtime/Experimental/Chaos/Public/Chaos/Collision/ContactPoint.h
@@ -38,7 +38,7 @@ namespace Chaos
 	 * @see FContactPoint, FContactPointf
 	*/
 	template<typename T>
-	class CHAOS_API TContactPoint
+	class TContactPoint
 	{
 	public:
 		using FRealType = T;
@@ -93,7 +93,7 @@ namespace Chaos
 	 * @brief A single point in a contact manifold.
 	 * Each Collision Constraint will have up to 4 of these.
 	*/
-	class CHAOS_API FManifoldPoint
+	class FManifoldPoint
 	{
 	public:
 		union FFlags
@@ -138,7 +138,7 @@ namespace Chaos
 	/**
 	 * World-space contact point data
 	 */
-	class CHAOS_API FWorldContactPoint
+	class FWorldContactPoint
 	{
 	public:
 		// World-space contact point relative to each particle's center of mass
@@ -162,7 +162,7 @@ namespace Chaos
 	 * @brief The friction data for a manifold point
 	 * This is the information that needs to be stored between ticks to implement static friction.
 	*/
-	class CHAOS_API FSavedManifoldPoint
+	class FSavedManifoldPoint
 	{
 	public:
 		FSavedManifoldPoint()
@@ -178,7 +178,7 @@ namespace Chaos
 		FVec3f ShapeContactPoints[2];			// Contact anchor points for friction
 	};
 
-	class CHAOS_API FManifoldPointResult
+	class FManifoldPointResult
 	{
 	public:
 		FManifoldPointResult()
@@ -200,4 +200,4 @@ namespace Chaos
 		uint8 bInsideStaticFrictionCone : 1;
 	};
 
-}
\ No newline at end of file
+}
diff --git a/Engine/Source/Runtime/Experimental/Chaos/Public/Chaos/Collision/ContactTriangles.h b/Engine/Source/Runtime/Experimental/Chaos/Public/Chaos/Collision/ContactTriangles.h
--- a/Engine/Source/Runtime/Experimental/Chaos/Public/Chaos/Collision/ContactTriangles.h
+++ b/Engine/Source/Runtime/Experimental/Chaos/Public/Chaos/Collision/ContactTriangles.h
@@ -12,7 +12,7 @@ namespace Chaos
 	/**
 	 * @brief Data held alongside contact points when generating contacts against a (likely non-convex) mesh of triangles
 	*/
-	class CHAOS_API FContactTriangle
+	class FContactTriangle
 	{
 	public:
 		// Triangle data
@@ -308,4 +308,4 @@ namespace Chaos
 		FRigidTransform3 ConvexTransform;
 	};
 
-}
\ No newline at end of file
+}
diff --git a/Engine/Source/Runtime/Experimental/Chaos/Public/Chaos/Collision/PBDCollisionConstraint.h b/Engine/Source/Runtime/Experimental/Chaos/Public/Chaos/Collision/PBDCollisionConstraint.h
--- a/Engine/Source/Runtime/Experimental/Chaos/Public/Chaos/Collision/PBDCollisionConstraint.h
+++ b/Engine/Source/Runtime/Experimental/Chaos/Public/Chaos/Collision/PBDCollisionConstraint.h
@@ -43,7 +43,7 @@ namespace Chaos
 	/*
 	 * @brief Material properties for a collision constraint
 	*/
-	class CHAOS_API FPBDCollisionConstraintMaterial
+	class FPBDCollisionConstraintMaterial
 	{
 	public:
 		FPBDCollisionConstraintMaterial()
@@ -165,7 +165,7 @@ namespace Chaos
 	 * They use intrusive handles to reduce unnecessary indirection.
 	 * 
 	*/
-	class CHAOS_API FPBDCollisionConstraint final : public FPBDCollisionConstraintHandle
+	class FPBDCollisionConstraint final : public FPBDCollisionConstraintHandle
 	{
 		friend class Private::FCollisionConstraintAllocator;
 		friend class Private::FCollisionContextAllocator;
@@ -189,7 +189,7 @@ namespace Chaos
 		 * Initializes a constraint stored inline in an object. Only intended to be called once right after construction.
 		 * Does not reinitialize all data so not intended to reset a constraint for reuse with different particles etc.
 		*/
-		static void Make(
+		static CHAOS_API void Make(
 			FGeometryParticleHandle* Particle0,
 			const FImplicitObject* Implicit0,
 			const FPerShapeData* Shape0,
@@ -209,23 +209,23 @@ namespace Chaos
 		 * @brief For use by the tri mesh and heighfield collision detection as a temporary measure
 		 * @see FHeightField::ContactManifoldImp, FTriangleMeshImplicitObject::ContactManifoldImp
 		*/
-		static FPBDCollisionConstraint MakeTriangle(const FImplicitObject* Implicit0);
+		static CHAOS_API FPBDCollisionConstraint MakeTriangle(const FImplicitObject* Implicit0);
 
 		/**
 		 * @brief Return a constraint copied from the Source constraint, for use in the Resim Cache or other system
 		 * @note Unlike the other factory method, this version returns a constraint by value for emplacing into an array (ther others are by pointer)
 		*/
-		static FPBDCollisionConstraint MakeCopy(const FPBDCollisionConstraint& Source);
+		static CHAOS_API FPBDCollisionConstraint MakeCopy(const FPBDCollisionConstraint& Source);
 
 		/**
 		 * Restore the properties of a collision from the properties of another. Used by the rewind/resim system when restoring
 		 * a collision constraint that was saved (and called MakeCopy).
 		 * This takes care not to overwrite any data maintained for the use of other systems (the ContainerCookie, Graph data, etc)
 		 */
-		void RestoreFrom(const FPBDCollisionConstraint& Source);
+		CHAOS_API void RestoreFrom(const FPBDCollisionConstraint& Source);
 
 		FPBDCollisionConstraint();
-		virtual ~FPBDCollisionConstraint();
+		CHAOS_API virtual ~FPBDCollisionConstraint();
 
 		/**
 		 * Whether CCD is enabled for this collision.
@@ -311,13 +311,13 @@ namespace Chaos
 		/** 
 		 * Called each frame when the constraint is active after primary collision detection (but not per incremental collision detection call if enabled) 
 		 */
-		void Activate();
+		CHAOS_API void Activate();
 
 		UE_DEPRECATED(5.2, "Removed parameter")
 		void Activate(const FReal Dt) { Activate(); }
 
 		// When a particle is moved under user control, we need to update some cached state to prevent friction from undoing the move
-		void UpdateParticleTransform(FGeometryParticleHandle* InParticle);
+		CHAOS_API void UpdateParticleTransform(FGeometryParticleHandle* InParticle);
 
 		// @todo(chaos): half of this API is wrong for the new multi-point manifold constraints. Remove it
 
@@ -331,16 +331,16 @@ namespace Chaos
 		bool GetIsProbe() const { return Flags.bIsProbe; }
 
 		virtual bool SupportsSleeping() const override final { return true; }
-		virtual bool IsSleeping() const override final;
-		virtual void SetIsSleeping(const bool bInIsSleeping) override final;
+		CHAOS_API virtual bool IsSleeping() const override final;
+		CHAOS_API virtual void SetIsSleeping(const bool bInIsSleeping) override final;
 
 		// Get the world-space normal of the closest manifold point
 		// @todo(chaos): remove (used by legacy RBAN collision solver)
-		FVec3 CalculateWorldContactNormal() const;
+		CHAOS_API FVec3 CalculateWorldContactNormal() const;
 
 		// Get the world-space contact location of the closest manifold point
 		// @todo(chaos): remove (used by legacy RBAN collision solver)
-		FVec3 CalculateWorldContactLocation() const;
+		CHAOS_API FVec3 CalculateWorldContactLocation() const;
 
 		// Called to force the contact to recollect its material properties (e.g., when materials are modified)
 		void ClearMaterialProperties()
@@ -392,7 +392,7 @@ namespace Chaos
 
 		EContactShapesType GetShapesType() const { return ShapesType; }
 
-		FString ToString() const;
+		CHAOS_API FString ToString() const;
 
 		FReal GetCullDistance() const { return CullDistance; }
 		void SetCullDistance(FReal InCullDistance) { CullDistance = FRealSingle(InCullDistance); }
@@ -426,7 +426,7 @@ namespace Chaos
 		/**
 		 * @brief Clear the current and previous manifolds
 		*/
-		void ResetManifold();
+		CHAOS_API void ResetManifold();
 
 		// @todo(chaos): remove array view and provide per-point accessor
 		TArrayView<FManifoldPoint> GetManifoldPoints() { return MakeArrayView(ManifoldPoints.begin(), ManifoldPoints.Num()); }
@@ -478,7 +478,7 @@ namespace Chaos
 			}
 		}
 
-		void AddIncrementalManifoldContact(const FContactPoint& ContactPoint);
+		CHAOS_API void AddIncrementalManifoldContact(const FContactPoint& ContactPoint);
 
 		// @todo(chaos): remove this and use SetOneShotManifoldContacts
 		inline void AddOneshotManifoldContact(const FContactPoint& ContactPoint)
@@ -532,7 +532,7 @@ namespace Chaos
 
 		}
 
-		void UpdateManifoldContacts();
+		CHAOS_API void UpdateManifoldContacts();
 
 		// Particle-relative transform of each collision shape in the constraint
 		const FRigidTransform3& GetShapeRelativeTransform0() const { return ImplicitTransform[0]; }
@@ -558,10 +558,10 @@ namespace Chaos
 		}
 
 		bool GetCanRestoreManifold() const { return Flags.bCanRestoreManifold; }
-		bool TryRestoreManifold();
-		void ResetActiveManifoldContacts();
-		bool TryAddManifoldContact(const FContactPoint& ContactPoint);
-		bool TryInsertManifoldContact(const FContactPoint& ContactPoint);
+		CHAOS_API bool TryRestoreManifold();
+		CHAOS_API void ResetActiveManifoldContacts();
+		CHAOS_API bool TryAddManifoldContact(const FContactPoint& ContactPoint);
+		CHAOS_API bool TryInsertManifoldContact(const FContactPoint& ContactPoint);
 
 		//@ todo(chaos): These are for the collision forwarding system - this should use the collision modifier system (which should be extended to support adding collisions)
 		void SetManifoldPoints(const TArray<FManifoldPoint>& InManifoldPoints)
@@ -594,7 +594,7 @@ namespace Chaos
 		 * Determine the constraint direction based on Normal and Phi.
 		 * This function assumes that the constraint is update-to-date.
 		 */
-		ECollisionConstraintDirection GetConstraintDirection(const FReal Dt) const;
+		CHAOS_API ECollisionConstraintDirection GetConstraintDirection(const FReal Dt) const;
 
 		/**
 		 * @brief Clear the saved manifold points. This effectively resets friction anchors.
@@ -720,7 +720,7 @@ namespace Chaos
 			const FBVHParticles* Simplicial1);
 
 		// Set all the data not initialized in the constructor
-		void Setup(
+		CHAOS_API void Setup(
 			const ECollisionCCDType InCCDType,
 			const EContactShapesType InShapesType,
 			const FRigidTransform3& InImplicitTransform0,
@@ -732,11 +732,11 @@ namespace Chaos
 		const FPBDCollisionConstraintContainerCookie& GetContainerCookie() const { return ContainerCookie; }
 		FPBDCollisionConstraintContainerCookie& GetContainerCookie() { return ContainerCookie; }
 
-		bool AreMatchingContactPoints(const FContactPoint& A, const FContactPoint& B, FReal& OutScore) const;
-		int32 FindManifoldPoint(const FContactPoint& ContactPoint) const;
+		CHAOS_API bool AreMatchingContactPoints(const FContactPoint& A, const FContactPoint& B, FReal& OutScore) const;
+		CHAOS_API int32 FindManifoldPoint(const FContactPoint& ContactPoint) const;
 
-		int32 FindSavedManifoldPoint(const int32 ManifoldPointIndex, TCArray<int32, MaxManifoldPoints>& InOutAllowedSavedPointIndices) const;
-		void AssignSavedManifoldPoints();
+		CHAOS_API int32 FindSavedManifoldPoint(const int32 ManifoldPointIndex, TCArray<int32, MaxManifoldPoints>& InOutAllowedSavedPointIndices) const;
+		CHAOS_API void AssignSavedManifoldPoints();
 
 		inline void InitManifoldPoint(const int32 ManifoldPointIndex, const FContactPoint& ContactPoint)
 		{
@@ -759,16 +759,16 @@ namespace Chaos
 
 		// Update the store Phi for the manifold point based on current world-space shape transforms
 		// @todo(chaos): Only intended for use by the legacy solvers - remove it
-		void UpdateManifoldPointPhi(const int32 ManifoldPointIndex);
+		CHAOS_API void UpdateManifoldPointPhi(const int32 ManifoldPointIndex);
 
-		void InitMarginsAndTolerances(const EImplicitObjectType ImplicitType0, const EImplicitObjectType ImplicitType1, const FRealSingle Margin0, const FRealSingle Margin1);
+		CHAOS_API void InitMarginsAndTolerances(const EImplicitObjectType ImplicitType0, const EImplicitObjectType ImplicitType1, const FRealSingle Margin0, const FRealSingle Margin1);
 
-		void InitCCDThreshold();
+		CHAOS_API void InitCCDThreshold();
 
-		void UpdateMaterialPropertiesImpl();
+		CHAOS_API void UpdateMaterialPropertiesImpl();
 
 	private:
-		FReal CalculateSavedManifoldPointScore(const FSavedManifoldPoint& SavedManifoldPoint, const FManifoldPoint& ManifoldPoint, const FReal DistanceToleranceSq) const;
+		CHAOS_API FReal CalculateSavedManifoldPointScore(const FSavedManifoldPoint& SavedManifoldPoint, const FManifoldPoint& ManifoldPoint, const FReal DistanceToleranceSq) const;
 
 		union FFlags
 		{
diff --git a/Engine/Source/Runtime/Experimental/Chaos/Public/Chaos/Collision/PBDCollisionConstraintHandle.h b/Engine/Source/Runtime/Experimental/Chaos/Public/Chaos/Collision/PBDCollisionConstraintHandle.h
--- a/Engine/Source/Runtime/Experimental/Chaos/Public/Chaos/Collision/PBDCollisionConstraintHandle.h
+++ b/Engine/Source/Runtime/Experimental/Chaos/Public/Chaos/Collision/PBDCollisionConstraintHandle.h
@@ -45,7 +45,7 @@ namespace Chaos
 	 * @note This is an intrusive handle, so you can use a contact pointer as a handle.
 	 * @todo(chaos): remove this class - it can just be a "using FPBDCollisionConstraintHandle = FPBDCollisionConstraint"
 	*/
-	class CHAOS_API FPBDCollisionConstraintHandle : public TIntrusiveConstraintHandle<FPBDCollisionConstraint>
+	class FPBDCollisionConstraintHandle : public TIntrusiveConstraintHandle<FPBDCollisionConstraint>
 	{
 	public:
 		using Base = TIntrusiveConstraintHandle<FPBDCollisionConstraint>;
@@ -58,8 +58,8 @@ namespace Chaos
 		{
 		}
 
-		const FPBDCollisionConstraint& GetContact() const;
-		FPBDCollisionConstraint& GetContact();
+		CHAOS_API const FPBDCollisionConstraint& GetContact() const;
+		CHAOS_API FPBDCollisionConstraint& GetContact();
 
 		UE_DEPRECATED(4.27, "Use GetContact()")
 		const FPBDCollisionConstraint& GetPointContact() const { return GetContact(); }
@@ -73,21 +73,21 @@ namespace Chaos
 		UE_DEPRECATED(4.27, "Use GetContact()")
 		FPBDCollisionConstraint& GetSweptPointContact() { return GetContact(); }
 
-		bool GetCCDEnabled() const;
+		CHAOS_API bool GetCCDEnabled() const;
 
-		virtual void SetEnabled(bool InEnabled) override final;
+		CHAOS_API virtual void SetEnabled(bool InEnabled) override final;
 
-		virtual bool IsEnabled() const override final;
+		CHAOS_API virtual bool IsEnabled() const override final;
 
-		virtual bool IsProbe() const override final;
+		CHAOS_API virtual bool IsProbe() const override final;
 
-		FVec3 GetAccumulatedImpulse() const;
+		CHAOS_API FVec3 GetAccumulatedImpulse() const;
 
 		// Declared final so that TPBDConstraintGraphRuleImpl::AddToGraph() does not need to hit vtable
-		virtual FParticlePair GetConstrainedParticles() const override final;
+		CHAOS_API virtual FParticlePair GetConstrainedParticles() const override final;
 
-		const FPBDCollisionConstraints* ConcreteContainer() const;
-		FPBDCollisionConstraints* ConcreteContainer();
+		CHAOS_API const FPBDCollisionConstraints* ConcreteContainer() const;
+		CHAOS_API FPBDCollisionConstraints* ConcreteContainer();
 
 		static const FConstraintHandleTypeID& StaticType()
 		{
diff --git a/Engine/Source/Runtime/Experimental/Chaos/Public/Chaos/Collision/PBDCollisionSolverSimd.h b/Engine/Source/Runtime/Experimental/Chaos/Public/Chaos/Collision/PBDCollisionSolverSimd.h
--- a/Engine/Source/Runtime/Experimental/Chaos/Public/Chaos/Collision/PBDCollisionSolverSimd.h
+++ b/Engine/Source/Runtime/Experimental/Chaos/Public/Chaos/Collision/PBDCollisionSolverSimd.h
@@ -1092,7 +1092,7 @@ namespace Chaos
 		{
 		public:
 			template<int TNumLanes>
-			static CHAOS_API void SolvePositionNoFriction(
+			static void SolvePositionNoFriction(
 				const TArrayView<TPBDCollisionSolverSimd<TNumLanes>>& Solvers,
 				const TArrayView<TPBDCollisionSolverManifoldPointsSimd<TNumLanes>>& ManifoldPoints,
 				const TArrayView<TSolverBodyPtrPairSimd<TNumLanes>>& SolverBodies,
@@ -1100,7 +1100,7 @@ namespace Chaos
 				const FSolverReal MaxPushOut);
 
 			template<int TNumLanes>
-			static CHAOS_API void SolvePositionWithFriction(
+			static void SolvePositionWithFriction(
 				const TArrayView<TPBDCollisionSolverSimd<TNumLanes>>& Solvers,
 				const TArrayView<TPBDCollisionSolverManifoldPointsSimd<TNumLanes>>& ManifoldPoints,
 				const TArrayView<TSolverBodyPtrPairSimd<TNumLanes>>& SolverBodies,
@@ -1108,14 +1108,14 @@ namespace Chaos
 				const FSolverReal MaxPushOut);
 
 			template<int TNumLanes>
-			static CHAOS_API void SolveVelocityNoFriction(
+			static void SolveVelocityNoFriction(
 				const TArrayView<TPBDCollisionSolverSimd<TNumLanes>>& Solvers,
 				const TArrayView<TPBDCollisionSolverManifoldPointsSimd<TNumLanes>>& ManifoldPoints,
 				const TArrayView<TSolverBodyPtrPairSimd<TNumLanes>>& SolverBodies,
 				const FSolverReal Dt);
 
 			template<int TNumLanes>
-			static CHAOS_API void SolveVelocityWithFriction(
+			static void SolveVelocityWithFriction(
 				const TArrayView<TPBDCollisionSolverSimd<TNumLanes>>& Solvers,
 				const TArrayView<TPBDCollisionSolverManifoldPointsSimd<TNumLanes>>& ManifoldPoints,
 				const TArrayView<TSolverBodyPtrPairSimd<TNumLanes>>& SolverBodies,
diff --git a/Engine/Source/Runtime/Experimental/Chaos/Public/Chaos/Collision/ParticleCollisions.h b/Engine/Source/Runtime/Experimental/Chaos/Public/Chaos/Collision/ParticleCollisions.h
--- a/Engine/Source/Runtime/Experimental/Chaos/Public/Chaos/Collision/ParticleCollisions.h
+++ b/Engine/Source/Runtime/Experimental/Chaos/Public/Chaos/Collision/ParticleCollisions.h
@@ -17,7 +17,7 @@ namespace Chaos
 	 * when Islands are woken to restore the collisions.
 	 * 
 	*/
-	class CHAOS_API FParticleCollisions
+	class FParticleCollisions
 	{
 	public:
 		// In a mostly stationary scene the choice of container doesn't matter much. In a highly dynamic
@@ -29,8 +29,8 @@ namespace Chaos
 		// that should not have too many contacts. 
 		using FContainerType = TArray<TPair<uint64, FParticlePairMidPhase*>>;
 
-		FParticleCollisions();
-		~FParticleCollisions();
+		CHAOS_API FParticleCollisions();
+		CHAOS_API ~FParticleCollisions();
 
 		inline int32 Num() const 
 		{ 
@@ -51,12 +51,12 @@ namespace Chaos
 		 * on the midphase that we want to retrieve, and it has one cookie per particle. 
 		 * This could probably be cleaned up a bit...
 		*/
-		void AddMidPhase(FGeometryParticleHandle* InParticle, FParticlePairMidPhase* InMidPhase);
+		CHAOS_API void AddMidPhase(FGeometryParticleHandle* InParticle, FParticlePairMidPhase* InMidPhase);
 
 		/**
 		 * @brief Remove a mid phase
 		*/
-		void RemoveMidPhase(FGeometryParticleHandle* InParticle, FParticlePairMidPhase* InMidPhase);
+		CHAOS_API void RemoveMidPhase(FGeometryParticleHandle* InParticle, FParticlePairMidPhase* InMidPhase);
 
 		/**
 		 * @brief Get a midphase by its index
@@ -127,4 +127,4 @@ namespace Chaos
 		FContainerType MidPhases;
 	};
 
-}
\ No newline at end of file
+}
diff --git a/Engine/Source/Runtime/Experimental/Chaos/Public/Chaos/Collision/ParticlePairMidPhase.h b/Engine/Source/Runtime/Experimental/Chaos/Public/Chaos/Collision/ParticlePairMidPhase.h
--- a/Engine/Source/Runtime/Experimental/Chaos/Public/Chaos/Collision/ParticlePairMidPhase.h
+++ b/Engine/Source/Runtime/Experimental/Chaos/Public/Chaos/Collision/ParticlePairMidPhase.h
@@ -40,22 +40,22 @@ namespace Chaos
 	 * 
 	 * This is used by FShapePairParticlePairMidPhase
 	 */
-	class CHAOS_API FSingleShapePairCollisionDetector
+	class FSingleShapePairCollisionDetector
 	{
 	public:
 		using FCollisionsArray = TArray<FPBDCollisionConstraint*, TInlineAllocator<1>>;
 
-		FSingleShapePairCollisionDetector(
+		CHAOS_API FSingleShapePairCollisionDetector(
 			FGeometryParticleHandle* InParticle0,
 			const FPerShapeData* InShape0,
 			FGeometryParticleHandle* InParticle1,
 			const FPerShapeData* InShape1,
 			const EContactShapesType InShapePairType, 
 			FParticlePairMidPhase& MidPhase);
-		FSingleShapePairCollisionDetector(FSingleShapePairCollisionDetector&& R);
+		CHAOS_API FSingleShapePairCollisionDetector(FSingleShapePairCollisionDetector&& R);
 		FSingleShapePairCollisionDetector(const FSingleShapePairCollisionDetector& R) = delete;
 		FSingleShapePairCollisionDetector& operator=(const FSingleShapePairCollisionDetector& R) = delete;
-		~FSingleShapePairCollisionDetector();
+		CHAOS_API ~FSingleShapePairCollisionDetector();
 
 		const FPBDCollisionConstraint* GetConstraint() const{ return Constraint.Get(); }
 		FPBDCollisionConstraint* GetConstraint() { return Constraint.Get(); }
@@ -78,7 +78,7 @@ namespace Chaos
 		 * @brief Perform a bounds check and run the narrow phase if necessary
 		 * @return The number of collisions constraints that were activated
 		*/
-		int32 GenerateCollision(
+		CHAOS_API int32 GenerateCollision(
 			const FReal CullDistance,
 			const FReal Dt,
 			const FCollisionContext& Context);
@@ -87,7 +87,7 @@ namespace Chaos
 		 * @brief Generate a SweptConstraint as long as AABBs overlap
 		 * @return The number of collisions constraints that were activated
 		*/
-		int32 GenerateCollisionCCD(
+		CHAOS_API int32 GenerateCollisionCCD(
 			const bool bEnableCCDSweep,
 			const FReal CullDistance,
 			const FReal Dt,
@@ -98,27 +98,27 @@ namespace Chaos
 		 * @parame SleepEpoch The tick on which the particle went to sleep.
 		 * Only constraints that were active when the particle went to sleep should be reactivated.
 		*/
-		void WakeCollision(const int32 SleepEpoch, const int32 CurrentEpoch);
+		CHAOS_API void WakeCollision(const int32 SleepEpoch, const int32 CurrentEpoch);
 
 		/**
 		 * @brief Set the collision from the parameter and activate it
 		 * This is used by the Resim restore functionality
 		*/
-		void SetCollision(const FPBDCollisionConstraint& Constraint, const FCollisionContext& Context);
+		CHAOS_API void SetCollision(const FPBDCollisionConstraint& Constraint, const FCollisionContext& Context);
 
 	private:
-		int32 GenerateCollisionImpl(
+		CHAOS_API int32 GenerateCollisionImpl(
 			const FReal CullDistance, 
 			const FReal Dt,
 			const FCollisionContext& Context);
 			
-		int32 GenerateCollisionCCDImpl(
+		CHAOS_API int32 GenerateCollisionCCDImpl(
 			const bool bEnableCCDSweep,
 			const FReal CullDistance,
 			const FReal Dt,
 			const FCollisionContext& Context);
 
-		int32 GenerateCollisionProbeImpl(
+		CHAOS_API int32 GenerateCollisionProbeImpl(
 			const FReal CullDistance, 
 			const FReal Dt,
 			const FCollisionContext& Context);
@@ -127,12 +127,12 @@ namespace Chaos
 		 * @brief Whether the two shapes are separated by less than CullDistance (i.e., we should run the narrow phase).
 		 * Also returns true if bounds checking is disabled (globally or for this pair)
 		*/
-		bool DoBoundsOverlap(const FReal CullDistance, const int32 CurrentEpoch);
+		CHAOS_API bool DoBoundsOverlap(const FReal CullDistance, const int32 CurrentEpoch);
 
 		/**
 		 * @brief Create a constraint
 		*/
-		void CreateConstraint(const FReal CullDistance, const FCollisionContext& Context);
+		CHAOS_API void CreateConstraint(const FReal CullDistance, const FCollisionContext& Context);
 
 		FParticlePairMidPhase& MidPhase;
 		FPBDCollisionConstraintPtr Constraint;
@@ -172,16 +172,16 @@ namespace Chaos
 	 * @see FShapePairParticlePairMidPhase, FGenericParticlePairMidPhase
 	 * 
 	 */
-	class CHAOS_API FParticlePairMidPhase
+	class FParticlePairMidPhase
 	{
 	public:
-		static EParticlePairMidPhaseType CalculateMidPhaseType(FGeometryParticleHandle* InParticle0, FGeometryParticleHandle* InParticle1);
-		static FParticlePairMidPhase* Make(FGeometryParticleHandle* InParticle0, FGeometryParticleHandle* InParticle1);
+		static CHAOS_API EParticlePairMidPhaseType CalculateMidPhaseType(FGeometryParticleHandle* InParticle0, FGeometryParticleHandle* InParticle1);
+		static CHAOS_API FParticlePairMidPhase* Make(FGeometryParticleHandle* InParticle0, FGeometryParticleHandle* InParticle1);
 
 		UE_NONCOPYABLE(FParticlePairMidPhase);
 
-		FParticlePairMidPhase(const EParticlePairMidPhaseType InMidPhaseType);
-		virtual ~FParticlePairMidPhase();
+		CHAOS_API FParticlePairMidPhase(const EParticlePairMidPhaseType InMidPhaseType);
+		CHAOS_API virtual ~FParticlePairMidPhase();
 
 		EParticlePairMidPhaseType GetMidPhaseType() const
 		{
@@ -194,7 +194,7 @@ namespace Chaos
 		 * the constructor so that we can reduce the time that the lock is held when allocating
 		 * new MidPhases.
 		*/
-		void Init(
+		CHAOS_API void Init(
 			FGeometryParticleHandle* InParticle0,
 			FGeometryParticleHandle* InParticle1,
 			const FCollisionParticlePairKey& InKey,
@@ -233,31 +233,31 @@ namespace Chaos
 		 * If this switches the state to Awake, it will reactivate any collisions between the particle pair that
 		 * were active when they went to sleep.
 		*/
-		void SetIsSleeping(const bool bInIsSleeping, const int32 CurrentEpoch);
+		CHAOS_API void SetIsSleeping(const bool bInIsSleeping, const int32 CurrentEpoch);
 
-		bool IsInConstraintGraph() const;
+		CHAOS_API bool IsInConstraintGraph() const;
 
 		/**
 		 * @brief Destroy all collisions and prevent this midphasae from being used any more. Called when one of its particles is destoyed.
 		 * It will be culled at the next Prune in the CollisionConstraintAllocator. We don't delete it immediately so that we don't
 		 * have to remove it from either Particle's ParticleCollisions array (which is O(N) and unnecessary when the particles are being destroyed)
 		*/
-		void DetachParticle(FGeometryParticleHandle* Particle);
+		CHAOS_API void DetachParticle(FGeometryParticleHandle* Particle);
 
 		/**
 		 * @brief Delete all cached data and collisions. Should be called when a particle changes its shapes
 		*/
-		void Reset();
+		CHAOS_API void Reset();
 
 		/**
 		 *
 		 */
-		void ResetModifications();
+		CHAOS_API void ResetModifications();
 
 		/**
 		 * @brief Create collision constraints for all colliding shape pairs
 		*/
-		void GenerateCollisions(
+		CHAOS_API void GenerateCollisions(
 			const FReal CullDistance,
 			const FReal Dt,
 			const FCollisionContext& Context);
@@ -267,7 +267,7 @@ namespace Chaos
 		 * This is used by the Resim system to restore saved colisions. If there is already a matching constraint
 		 * it will be overwritten, otherwise a new constraint will be added.
 		*/
-		void InjectCollision(const FPBDCollisionConstraint& Constraint, const FCollisionContext& Context);
+		CHAOS_API void InjectCollision(const FPBDCollisionConstraint& Constraint, const FCollisionContext& Context);
 
 		/**
 		 * @brief Call a lambda on each active collision constraint
@@ -400,9 +400,9 @@ namespace Chaos
 		 * @return true if CCD is enabled this tick, false otherwise
 		 * This may return false, even for collisions on CCD-enabled bodies when the bodies are moving slowly
 		 */
-		bool ShouldEnableCCD(const FReal Dt);
+		CHAOS_API bool ShouldEnableCCD(const FReal Dt);
 
-		void InitThresholds();
+		CHAOS_API void InitThresholds();
 
 
 		union FFlags
@@ -448,28 +448,28 @@ namespace Chaos
 	 * does not contain a complex hierarchy. This midphase (compared to FGenericParticlePairMidPhase)
 	 * caches various data like the results of the collision filtering, the shapr pair types, etc.
 	 */
-	class CHAOS_API FShapePairParticlePairMidPhase : public FParticlePairMidPhase
+	class FShapePairParticlePairMidPhase : public FParticlePairMidPhase
 	{
 	public:
 		friend class FParticlePairMidPhase;
 
-		FShapePairParticlePairMidPhase();
+		CHAOS_API FShapePairParticlePairMidPhase();
 
-		virtual void ResetImpl() override final;
-		virtual void BuildDetectorsImpl() override final;
+		CHAOS_API virtual void ResetImpl() override final;
+		CHAOS_API virtual void BuildDetectorsImpl() override final;
 
 	protected:
-		virtual int32 GenerateCollisionsImpl(
+		CHAOS_API virtual int32 GenerateCollisionsImpl(
 			const FReal CullDistance,
 			const FReal Dt,
 			const FCollisionContext& Context) override final;
 
-		virtual void WakeCollisionsImpl(const int32 CurrentEpoch) override final;
+		CHAOS_API virtual void WakeCollisionsImpl(const int32 CurrentEpoch) override final;
 
-		virtual void InjectCollisionImpl(const FPBDCollisionConstraint& Constraint, const FCollisionContext& Context) override final;
+		CHAOS_API virtual void InjectCollisionImpl(const FPBDCollisionConstraint& Constraint, const FCollisionContext& Context) override final;
 
 	private:
-		void TryAddShapePair(
+		CHAOS_API void TryAddShapePair(
 			const FPerShapeData* Shape0,
 			const FPerShapeData* Shape1);
 
@@ -489,23 +489,23 @@ namespace Chaos
 	 * if there is one, so it can be much faster the the FShapePairParticlePairMidPhase
 	 * when the set of potentially colliding pairs is very large.
 	 */
-	class CHAOS_API FGenericParticlePairMidPhase : public FParticlePairMidPhase
+	class FGenericParticlePairMidPhase : public FParticlePairMidPhase
 	{
 	public:
 		friend class FParticlePairMidPhase;
 
-		FGenericParticlePairMidPhase();
-		~FGenericParticlePairMidPhase();
+		CHAOS_API FGenericParticlePairMidPhase();
+		CHAOS_API ~FGenericParticlePairMidPhase();
 
-		virtual void ResetImpl() override final;
-		virtual void BuildDetectorsImpl() override final;
+		CHAOS_API virtual void ResetImpl() override final;
+		CHAOS_API virtual void BuildDetectorsImpl() override final;
 
 		/**
 		 * @brief Callback from the narrow phase to create a collision constraint for this particle pair.
 		 * We should never be asked for a collision for a different particle pair, but the
 		 * implicit objects may be children of the root shape.
 		*/
-		virtual FPBDCollisionConstraint* FindOrCreateConstraint(
+		CHAOS_API virtual FPBDCollisionConstraint* FindOrCreateConstraint(
 			FGeometryParticleHandle* InParticle0,
 			const FImplicitObject* InImplicit0,
 			const int32 InImplicitId0,
@@ -529,21 +529,21 @@ namespace Chaos
 		 * @parame SleepEpoch The tick on which the particle went to sleep.
 		 * Only constraints that were active when the particle went to sleep should be reactivated.
 		*/
-		void WakeCollisionsImpl(const int32 SleepEpoch, const int32 CurrentEpoch);
+		CHAOS_API void WakeCollisionsImpl(const int32 SleepEpoch, const int32 CurrentEpoch);
 
 	protected:
-		virtual int32 GenerateCollisionsImpl(
+		CHAOS_API virtual int32 GenerateCollisionsImpl(
 			const FReal CullDistance,
 			const FReal Dt,
 			const FCollisionContext& Context) override final;
 
-		virtual void WakeCollisionsImpl(const int32 CurrentEpoch) override final;
+		CHAOS_API virtual void WakeCollisionsImpl(const int32 CurrentEpoch) override final;
 
-		virtual void InjectCollisionImpl(const FPBDCollisionConstraint& Constraint, const FCollisionContext& Context) override final;
+		CHAOS_API virtual void InjectCollisionImpl(const FPBDCollisionConstraint& Constraint, const FCollisionContext& Context) override final;
 
 	private:
 		// BVH on ParticleA versus implicit hierarchy of ParticleB
-		void GenerateCollisionsBVH(
+		CHAOS_API void GenerateCollisionsBVH(
 			FGeometryParticleHandle* ParticleA, const Private::FImplicitBVH* BVHA,
 			FGeometryParticleHandle* ParticleB, const FImplicitObject* RootImplicitB,
 			const FReal CullDistance,
@@ -551,7 +551,7 @@ namespace Chaos
 			const FCollisionContext& Context);
 
 		// Implicit hierarchy of particle A versus implicit hierarchy of ParticleB (used when no BVHs present)
-		void GenerateCollisionsImplicit(
+		CHAOS_API void GenerateCollisionsImplicit(
 			FGeometryParticleHandle* ParticleA, const FImplicitObject* RootImplicitA,
 			FGeometryParticleHandle* ParticleB, const FImplicitObject* RootImplicitB,
 			const FReal CullDistance,
@@ -559,7 +559,7 @@ namespace Chaos
 			const FCollisionContext& Context);
 
 		// Leaf Implicit on ParticleA versus implicit hierarchy of ParticleB
-		void GenerateCollisionsShapeHierarchy(
+		CHAOS_API void GenerateCollisionsShapeHierarchy(
 			FGeometryParticleHandle* ParticleA, const FImplicitObject* ImplicitA, const FRigidTransform3 ParticleWorldTransformA, const FRigidTransform3& RelativeTransformA, const FAABB3& RelativeBoundsA, const int32 RootObjectIndexA, const int32 LeafObjectIndexA,
 			FGeometryParticleHandle* ParticleB, const FImplicitObject* RootImplicitB, const FRigidTransform3 ParticleWorldTransformB,
 			const FRigidTransform3 ParticleTransformAToB,
@@ -568,7 +568,7 @@ namespace Chaos
 			const FCollisionContext& Context);
 
 		// Leaf Implicit on ParticleA versus Leaf Implicit of ParticleB
-		void GenerateCollisionsShapeShape(
+		CHAOS_API void GenerateCollisionsShapeShape(
 			FGeometryParticleHandle* ParticleA, const FImplicitObject* ImplicitA, const FShapeInstance* ShapeInstanceA, const FRigidTransform3 ParticleWorldTransformA, const FRigidTransform3& RelativeTransformA, const int32 LeafObjectIndexA,
 			FGeometryParticleHandle* ParticleB, const FImplicitObject* ImplicitB, const FShapeInstance* ShapeInstanceB, const FRigidTransform3 ParticleWorldTransformB, const FRigidTransform3& RelativeTransformB, const int32 LeafObjectIndexB,
 			const FReal CullDistance,
@@ -576,7 +576,7 @@ namespace Chaos
 			const FCollisionContext& Context);
 
 		// A bounds check between two implicits
-		bool DoBoundsOverlap(
+		CHAOS_API bool DoBoundsOverlap(
 			const FImplicitObject* ImplicitA,
 			const FRigidTransform3& ParticleWorldTransformA,
 			const FRigidTransform3& ShapeRelativeTransformA,
@@ -585,9 +585,9 @@ namespace Chaos
 			const FRigidTransform3& ShapeRelativeTransformB,
 			const FReal CullDistance);
 
-		FPBDCollisionConstraint* FindConstraint(const FCollisionParticlePairConstraintKey& Key);
+		CHAOS_API FPBDCollisionConstraint* FindConstraint(const FCollisionParticlePairConstraintKey& Key);
 
-		FPBDCollisionConstraint* CreateConstraint(
+		CHAOS_API FPBDCollisionConstraint* CreateConstraint(
 			FGeometryParticleHandle* Particle0,
 			const FImplicitObject* Implicit0,
 			const FPerShapeData* Shape0,
@@ -604,20 +604,20 @@ namespace Chaos
 			const FCollisionParticlePairConstraintKey& Key,
 			const FCollisionContext& Context);
 
-		int32 ProcessNewConstraints(
+		CHAOS_API int32 ProcessNewConstraints(
 			const FReal CullDistance,
 			const FReal Dt,
 			const FCollisionContext& Context);
 
-		void PruneConstraints(const int32 CurrentEpoch);
+		CHAOS_API void PruneConstraints(const int32 CurrentEpoch);
 
-		void UpdateCollision(
+		CHAOS_API void UpdateCollision(
 			FPBDCollisionConstraint* Constraint,
 			const FReal CullDistance,
 			const FReal Dt,
 			const FCollisionContext& Context);
 
-		void UpdateCollisionCCD(
+		CHAOS_API void UpdateCollisionCCD(
 			FPBDCollisionConstraint* Constraint,
 			const FReal CullDistance,
 			const FReal Dt,
diff --git a/Engine/Source/Runtime/Experimental/Chaos/Public/Chaos/Collision/SpatialAccelerationCollisionDetector.h b/Engine/Source/Runtime/Experimental/Chaos/Public/Chaos/Collision/SpatialAccelerationCollisionDetector.h
--- a/Engine/Source/Runtime/Experimental/Chaos/Public/Chaos/Collision/SpatialAccelerationCollisionDetector.h
+++ b/Engine/Source/Runtime/Experimental/Chaos/Public/Chaos/Collision/SpatialAccelerationCollisionDetector.h
@@ -10,7 +10,7 @@
 
 namespace Chaos
 {
-	class CHAOS_API FSpatialAccelerationCollisionDetector : public FCollisionDetector
+	class FSpatialAccelerationCollisionDetector : public FCollisionDetector
 	{
 	public:
 		FSpatialAccelerationCollisionDetector(FSpatialAccelerationBroadPhase& InBroadPhase, FPBDCollisionConstraints& InCollisionContainer)
diff --git a/Engine/Source/Runtime/Experimental/Chaos/Public/Chaos/CollisionConvexMesh.h b/Engine/Source/Runtime/Experimental/Chaos/Public/Chaos/CollisionConvexMesh.h
--- a/Engine/Source/Runtime/Experimental/Chaos/Public/Chaos/CollisionConvexMesh.h
+++ b/Engine/Source/Runtime/Experimental/Chaos/Public/Chaos/CollisionConvexMesh.h
@@ -57,7 +57,7 @@ namespace Chaos
 			FRealType HorizonEpsilon;
 		};
 
-		static CHAOS_API FRealType SuggestEpsilon(const TArray<FVec3Type>& InVertices)
+		static FRealType SuggestEpsilon(const TArray<FVec3Type>& InVertices)
 		{
 			if (ComputeHorizonEpsilonFromMeshExtends == 0)
 			{
@@ -345,7 +345,7 @@ namespace Chaos
 			return FTriangleMesh(MoveTemp(Indices));
 		}
 
-		static CHAOS_API bool IsPerformanceWarning(int32 NumPlanes, int32 NumVertices)
+		static bool IsPerformanceWarning(int32 NumPlanes, int32 NumVertices)
 		{
 			if (!PerformGeometryCheck)
 			{
@@ -355,7 +355,7 @@ namespace Chaos
 			return (NumVertices > VerticesThreshold);
 		}
 
-		static CHAOS_API bool IsGeometryReductionEnabled()
+		static bool IsGeometryReductionEnabled()
 		{
 			return (PerformGeometryReduction>0)?true:false;
 		}
@@ -365,7 +365,7 @@ namespace Chaos
 			return FString::Printf(TEXT("Planes %d, Vertices %d"), NumPlanes, NumVertices);
 		}
 
-		static CHAOS_API void Simplify(TArray<FPlaneType>& InOutPlanes, TArray<TArray<int32>>& InOutFaces, TArray<FVec3Type>& InOutVertices, FAABB3Type& InOutLocalBounds)
+		static void Simplify(TArray<FPlaneType>& InOutPlanes, TArray<TArray<int32>>& InOutFaces, TArray<FVec3Type>& InOutVertices, FAABB3Type& InOutLocalBounds)
 		{
 			struct TPair
 			{
diff --git a/Engine/Source/Runtime/Experimental/Chaos/Public/Chaos/CollisionFilterData.h b/Engine/Source/Runtime/Experimental/Chaos/Public/Chaos/CollisionFilterData.h
--- a/Engine/Source/Runtime/Experimental/Chaos/Public/Chaos/CollisionFilterData.h
+++ b/Engine/Source/Runtime/Experimental/Chaos/Public/Chaos/CollisionFilterData.h
@@ -22,7 +22,7 @@ namespace Chaos
 	};
 }
 
-struct CHAOS_API FCollisionFilterData
+struct FCollisionFilterData
 {
 	uint32 Word0;
 	uint32 Word1;
diff --git a/Engine/Source/Runtime/Experimental/Chaos/Public/Chaos/ConstraintHandle.h b/Engine/Source/Runtime/Experimental/Chaos/Public/Chaos/ConstraintHandle.h
--- a/Engine/Source/Runtime/Experimental/Chaos/Public/Chaos/ConstraintHandle.h
+++ b/Engine/Source/Runtime/Experimental/Chaos/Public/Chaos/ConstraintHandle.h
@@ -93,7 +93,7 @@ namespace Chaos
 	 * 
 	 * @see FIndexedConstraintHandle, FIntrusiveConstraintHandle
 	*/
-	class CHAOS_API FConstraintHandle
+	class FConstraintHandle
 	{
 	public:
 		using FGeometryParticleHandle = TGeometryParticleHandle<FReal, 3>;
@@ -168,7 +168,7 @@ namespace Chaos
 		virtual void SetWasAwakened(const bool bInWasAwakened) {}
 
 		// Implemented in ConstraintContainer.h
-		int32 GetContainerId() const;
+		CHAOS_API int32 GetContainerId() const;
 
 		// Implemented in ConstraintContainer.h
 		template<typename T>  T* As();
@@ -178,7 +178,7 @@ namespace Chaos
 		template<typename T>  T* AsUnsafe() { check(As<T>() != nullptr); return static_cast<T*>(this); }
 		template<typename T>  const T* AsUnsafe() const { check(As<T>() != nullptr); return static_cast<const T*>(this); }
 
-		const FConstraintHandleTypeID& GetType() const;
+		CHAOS_API const FConstraintHandleTypeID& GetType() const;
 
 		static const FConstraintHandleTypeID& StaticType()
 		{
@@ -212,7 +212,7 @@ namespace Chaos
 	 * address (as opposed to array-based containers where the array could relocate). The constraint class should inherit
 	 * this handle class. This effectively eliminates the handle, reducing cache misses and allocations.
 	*/
-	class CHAOS_API FIntrusiveConstraintHandle : public FConstraintHandle
+	class FIntrusiveConstraintHandle : public FConstraintHandle
 	{
 	public:
 		FIntrusiveConstraintHandle()
@@ -241,7 +241,7 @@ namespace Chaos
 	 * @tparam T_CONSTRAINT The constraint type
 	*/
 	template<typename T_CONSTRAINT>
-	class CHAOS_API TIntrusiveConstraintHandle : public FIntrusiveConstraintHandle
+	class TIntrusiveConstraintHandle : public FIntrusiveConstraintHandle
 	{
 	public:
 		using FConstraint = T_CONSTRAINT;
@@ -275,7 +275,7 @@ namespace Chaos
 	 * @todo(ccaulfield): block allocator for handles, or support custom allocators in constraint containers.
 	 */
 	template<class T_CONTAINER>
-	class CHAOS_API TConstraintHandleAllocator
+	class TConstraintHandleAllocator
 	{
 	public:
 		using FConstraintContainer = T_CONTAINER;
@@ -292,7 +292,7 @@ namespace Chaos
 	 * This acts as a FConstraintHandle*, but caches some extra debug data useful in tracking
 	 * down dangling pointer issues when they arise.
 	*/
-	class CHAOS_API FConstraintHandleHolder
+	class FConstraintHandleHolder
 	{
 	public:
 		FConstraintHandleHolder()
@@ -329,7 +329,7 @@ namespace Chaos
 		const FGeometryParticleHandle* GetParticle1() const { return Particles[1]; }
 
 	private:
-		void InitDebugData();
+		CHAOS_API void InitDebugData();
 
 		const FConstraintHandleTypeID* ConstraintType;
 		const FGeometryParticleHandle* Particles[2];
