[Chaos]

diff --git a/Engine/Source/Runtime/Experimental/Chaos/Public/Chaos/ContactModification.h b/Engine/Source/Runtime/Experimental/Chaos/Public/Chaos/ContactModification.h
--- a/Engine/Source/Runtime/Experimental/Chaos/Public/Chaos/ContactModification.h
+++ b/Engine/Source/Runtime/Experimental/Chaos/Public/Chaos/ContactModification.h
@@ -13,7 +13,7 @@ namespace Chaos
 	class FContactPairModifier;
 	class FPerShapeData;
 
-	class CHAOS_API FContactPairModifier
+	class FContactPairModifier
 	{
 	public:
 		FContactPairModifier()
@@ -29,54 +29,54 @@ namespace Chaos
 		/**
 		* Disable all contact points for this pair of bodies.
 		*/
-		void Disable();
+		CHAOS_API void Disable();
 
 		/**
 		* Enables all contact points for this pair of bodies. This does nothing unless Disable() was previously called.
 		*/
-		void Enable();
+		CHAOS_API void Enable();
 
 		/**
 		* Convert the constraint between this pair of bodies into a probe. Collision callbacks will still occur if the
 		* bodies collide (a contact occurring in contact modification does not mean a contact actually has occurred
 		* yet, but that it may)
 		*/
-		void ConvertToProbe();
+		CHAOS_API void ConvertToProbe();
 
 		/**
 		* @return Number of contact points in constraint pair. ContactPointIdx must be below number of contacts.
 		*/
-		int32 GetNumContacts() const;
+		CHAOS_API int32 GetNumContacts() const;
 
 		/**
 		* @return Contact point index of point in manifold penetrating the deepest.
 		*/
-		int32 GetDeepestContactIndex() const;
+		CHAOS_API int32 GetDeepestContactIndex() const;
 
 		/**
 		* @return Geometry of specified particle in contact pair.
 		*/
-		const FImplicitObject* GetContactGeometry(int32 ParticleIdx);
+		CHAOS_API const FImplicitObject* GetContactGeometry(int32 ParticleIdx);
 
 		/**
 		* @return Transformation from geometry to world.
 		*/
-		FRigidTransform3 GetShapeToWorld(int32 ParticleIdx) const;
+		CHAOS_API FRigidTransform3 GetShapeToWorld(int32 ParticleIdx) const;
 
 		/**
 		* @return separation of contact locations in direction of normal.
 		*/
-		FReal GetSeparation(int32 ContactPointIdx) const;
+		CHAOS_API FReal GetSeparation(int32 ContactPointIdx) const;
 
 		/**
 		* @return desired separation of contact locations in direction of normal.
 		*/
-		FReal GetTargetSeparation(int32 ContactPointIdx) const;
+		CHAOS_API FReal GetTargetSeparation(int32 ContactPointIdx) const;
 
 		/**
 		* Set the desired separation at the contact point (negative for penetration).
 		*/
-		void ModifyTargetSeparation(FReal TargetSeparation, int32 ContactPointIdx);
+		CHAOS_API void ModifyTargetSeparation(FReal TargetSeparation, int32 ContactPointIdx);
 
 		/**
 		* @brief Get the world-space contact normal.
@@ -85,167 +85,167 @@ namespace Chaos
 		* sphere is the first or second body in the constraint.
 		* @return Normal of contact in world space.
 		*/
-		FVec3 GetWorldNormal(int32 ContactPointIdx) const;
+		CHAOS_API FVec3 GetWorldNormal(int32 ContactPointIdx) const;
 
 		/**
 		* Modify contact normal in world space. If modifying separation and normal, order of operations should be considered.
 		* @note The contact normal should always point away from the second body.
 		*/
-		void ModifyWorldNormal(const FVec3& Normal, int32 ContactPointIdx);
+		CHAOS_API void ModifyWorldNormal(const FVec3& Normal, int32 ContactPointIdx);
 
 		/**
 		* Get contact location on each body in world space.
 		*/
-		void GetWorldContactLocations(int32 ContactPointIdx, FVec3& OutLocation0, FVec3& OutLocation1)  const;
+		CHAOS_API void GetWorldContactLocations(int32 ContactPointIdx, FVec3& OutLocation0, FVec3& OutLocation1)  const;
 
 		/**
 		* @return Contact location. This is midpoint of each body's contact location.
 		*/
-		FVec3 GetWorldContactLocation(int32 ContactPointIdx)  const;
+		CHAOS_API FVec3 GetWorldContactLocation(int32 ContactPointIdx)  const;
 
 		/**
 		* Modify contact location of each body in world space. If modifying locations and separation, order of operations should be considered.
 		*/
-		void ModifyWorldContactLocations(const FVec3& Location0, const FVec3& Location1, int32 ContactPointIdx);
+		CHAOS_API void ModifyWorldContactLocations(const FVec3& Location0, const FVec3& Location1, int32 ContactPointIdx);
 
 		/**
 		* @return Restitution of contact.
 		*/
-		FReal GetRestitution() const;
+		CHAOS_API FReal GetRestitution() const;
 
 		/**
 		* Modify restitution of contact. If object is moving slowly, RestitutionThreshold may need to be lowered if restitution is desired.
 		*/
-		void ModifyRestitution(FReal Restitution);
+		CHAOS_API void ModifyRestitution(FReal Restitution);
 
 		/**
 		* @return Restitution threshold of contact. If velocity of object is below this threshold, no restitution is applied. 
 		*/
-		FReal GetRestitutionThreshold() const;
+		CHAOS_API FReal GetRestitutionThreshold() const;
 
 		/**
 		* Modify restitution threshold of contact. If velocity of object is below this threshold, no restitution is applied.
 		*/
-		void ModifyRestitutionThreshold(FReal Restitution);
+		CHAOS_API void ModifyRestitutionThreshold(FReal Restitution);
 
 		/**
 		* @return Dynamic friction coefficient of contact.
 		*/
-		FReal GetDynamicFriction() const;
+		CHAOS_API FReal GetDynamicFriction() const;
 
 		/**
 		* Modify dynamic friction coefficient of contact.
 		*/
-		void ModifyDynamicFriction(FReal DynamicFriction);
+		CHAOS_API void ModifyDynamicFriction(FReal DynamicFriction);
 
 		/**
 		* @return Static friction coefficient of contact.
 		*/
-		FReal GetStaticFriction() const;
+		CHAOS_API FReal GetStaticFriction() const;
 
 		/**
 		* Modify static friction coefficient of contact.
 		*/
-		void ModifyStaticFriction(FReal StaticFriction);
+		CHAOS_API void ModifyStaticFriction(FReal StaticFriction);
 
 		/**
 		* @return Linear velocity of particle.
 		*/
-		FVec3 GetParticleVelocity(int32 ParticleIdx) const;
+		CHAOS_API FVec3 GetParticleVelocity(int32 ParticleIdx) const;
 
 		/**
 		* Modify linear velocity. For simulated objects: Previous frame's position is modified to maintain PBD implicit velocity.
 		*/
-		void ModifyParticleVelocity(FVec3 Velocity, int32 ParticleIdx);
+		CHAOS_API void ModifyParticleVelocity(FVec3 Velocity, int32 ParticleIdx);
 
 		/**
 		* @return Angular velocity of particle.
 		*/
-		FVec3 GetParticleAngularVelocity(int32 ParticleIdx) const;
+		CHAOS_API FVec3 GetParticleAngularVelocity(int32 ParticleIdx) const;
 
 		/**
 		* Modify angular velocity. For simulated objects: Previous frame's rotation is modified to maintain PBD implicit velocity.
 		*/
-		void ModifyParticleAngularVelocity(FVec3 AngularVelocity, int32 ParticleIdx);
+		CHAOS_API void ModifyParticleAngularVelocity(FVec3 AngularVelocity, int32 ParticleIdx);
 
 		/**
 		* @return Position of particle. For simulated objects: this is PBD predicted position and not previous position before integrating movement.
 		*/
-		FVec3 GetParticlePosition(int32 ParticleIdx) const;
+		CHAOS_API FVec3 GetParticlePosition(int32 ParticleIdx) const;
 
 		/**
 		* For simulated objects: this sets PBD predicted position and not position before integrating movement.
 		* @param bMaintainVelocity - if true, simulated object will update previous position to maintain implicit velocity. Otherwise changing position will change velocity.
 		*/
-		void ModifyParticlePosition(FVec3 Position, bool bMaintainVelocity, int32 ParticleIdx);
+		CHAOS_API void ModifyParticlePosition(FVec3 Position, bool bMaintainVelocity, int32 ParticleIdx);
 
 		/**
 		* @return Rotation of particle. For simulated objects: this is PBD predicted rotation and not previous rotation before integrating movement.
 		*/
-		FRotation3 GetParticleRotation(int32 ParticleIdx) const;
+		CHAOS_API FRotation3 GetParticleRotation(int32 ParticleIdx) const;
 
 		/**
 		* For simulated objects: this sets PBD predicted rotation and not rotation before integrating movement.
 		* @param bMaintainVelocity - if true, simulated object will update previous rotation to maintain implicit velocity. Otherwise changing rotation will change angular velocity.
 		*/
-		void ModifyParticleRotation(FRotation3 Rotation, bool bMaintainVelocity, int32 ParticleIdx);
+		CHAOS_API void ModifyParticleRotation(FRotation3 Rotation, bool bMaintainVelocity, int32 ParticleIdx);
 
 		/**
 		* @return InvInertiaScale of particle.
 		*/
-		FReal GetInvInertiaScale(int32 ParticleIdx) const;
+		CHAOS_API FReal GetInvInertiaScale(int32 ParticleIdx) const;
 
 		/**
 		* Modify InvInertiaScale of particle. 0 prevents all rotation of particle due to collision, values > 1 decrease rotational inertia.
 		*/
-		void ModifyInvInertiaScale(FReal InvInertiaScale, int32 ParticleIdx);
+		CHAOS_API void ModifyInvInertiaScale(FReal InvInertiaScale, int32 ParticleIdx);
 
 		/**
 		* @return InvMassScale of particle.
 		*/
-		FReal GetInvMassScale(int32 ParticleIdx) const;
+		CHAOS_API FReal GetInvMassScale(int32 ParticleIdx) const;
 
 		/**
 		* Modify InvMassScale of particle. 0 gives particle infinite mass during collision, values > 1 decrease mass during collision.
 		*/
-		void ModifyInvMassScale(FReal InvMassScale, int32 ParticleIdx);
+		CHAOS_API void ModifyInvMassScale(FReal InvMassScale, int32 ParticleIdx);
 
 		/*
 		* Retrieve physics handles for particles in contact pair.
 		*/
-		TVec2<FGeometryParticleHandle*> GetParticlePair() const;
+		CHAOS_API TVec2<FGeometryParticleHandle*> GetParticlePair() const;
 
 		/*
 		* Get shape pair from constraint.
 		*/
-		TVec2<const FPerShapeData*> GetShapePair() const;
+		CHAOS_API TVec2<const FPerShapeData*> GetShapePair() const;
 
 		/*
 		* Check to see if a contact point index is an edge contact.
 		*/
-		bool IsEdgeContactPoint(int32 ContactPointIdx) const;
+		CHAOS_API bool IsEdgeContactPoint(int32 ContactPointIdx) const;
 
 		/*
 		* Check to see if a contact point index is disabled.
 		*/
-		bool IsContactPointDisabled(int32 ContactPointIdx) const;
+		CHAOS_API bool IsContactPointDisabled(int32 ContactPointIdx) const;
 
 		/*
 		* Set a contact point disabled.
 		*/
-		void SetContactPointDisabled(int32 ContactPointIdx) const;
+		CHAOS_API void SetContactPointDisabled(int32 ContactPointIdx) const;
 
 	private:
 		/**
 		 * @brief Update cached shape transforms in the constraint after modifying particle positions
 		*/
-		void UpdateConstraintShapeTransforms();
+		CHAOS_API void UpdateConstraintShapeTransforms();
 
 		FPBDCollisionConstraint* Constraint;
 		FCollisionContactModifier* Modifier;
 	};
 
-	class CHAOS_API FContactPairModifierIterator
+	class FContactPairModifierIterator
 	{
 	public:
 		FContactPairModifierIterator() { SetToEnd(); }
@@ -317,7 +317,7 @@ namespace Chaos
 		}
 
 		// Moves to constraint at ConstraintIdx, or finds next valid constraint. Moves to end if no remaining valid constraints.
-		void SeekValidContact();
+		CHAOS_API void SeekValidContact();
 
 		int32 ConstraintIdx;
 		FCollisionContactModifier* Modifier;
@@ -327,7 +327,7 @@ namespace Chaos
 	/*
 	*  Provides interface for iterating over modifiable contact pairs
 	*/
-	class CHAOS_API FCollisionContactModifier
+	class FCollisionContactModifier
 	{
 	public:
 		friend FContactPairModifier;
@@ -347,18 +347,18 @@ namespace Chaos
 		FContactPairModifierIterator end() const { return FContactPairModifierIterator(); }
 
 	private:
-		TArrayView<FPBDCollisionConstraint* const>& GetConstraints();
-		void DisableConstraint(FPBDCollisionConstraint& Constraint);
-		void EnableConstraint(FPBDCollisionConstraint& Constraint);
+		CHAOS_API TArrayView<FPBDCollisionConstraint* const>& GetConstraints();
+		CHAOS_API void DisableConstraint(FPBDCollisionConstraint& Constraint);
+		CHAOS_API void EnableConstraint(FPBDCollisionConstraint& Constraint);
 
 		// Turn this constraint into a probe. It will still generate hit events,
 		// but will not produce impulses
-		void ConvertToProbeConstraint(FPBDCollisionConstraint& Constraint);
+		CHAOS_API void ConvertToProbeConstraint(FPBDCollisionConstraint& Constraint);
 
-		void MarkConstraintForManifoldUpdate(FPBDCollisionConstraint& Constraint);
+		CHAOS_API void MarkConstraintForManifoldUpdate(FPBDCollisionConstraint& Constraint);
 
 		// Update manifolds of modified constraints.
-		void UpdateConstraintManifolds();
+		CHAOS_API void UpdateConstraintManifolds();
 
 		TArrayView<FPBDCollisionConstraint* const>& Constraints;
 
diff --git a/Engine/Source/Runtime/Experimental/Chaos/Public/Chaos/Convex.h b/Engine/Source/Runtime/Experimental/Chaos/Public/Chaos/Convex.h
--- a/Engine/Source/Runtime/Experimental/Chaos/Public/Chaos/Convex.h
+++ b/Engine/Source/Runtime/Experimental/Chaos/Public/Chaos/Convex.h
@@ -25,7 +25,7 @@ namespace Chaos
 	// instance wrapper (ImplicitScaled, ImplicitTransformed, or ImplicitInstanced). Usually the
 	// margin on the convex itself is zero.
 	//
-	class CHAOS_API FConvex final : public FImplicitObject
+	class FConvex final : public FImplicitObject
 	{
 	public:
 		using FImplicitObject::GetTypeName;
@@ -169,12 +169,12 @@ namespace Chaos
 			return TUniquePtr<FImplicitObject>(new FConvex(*this));
 		}
 
-		virtual TUniquePtr<FImplicitObject> CopyWithScale(const FVec3& Scale) const override;
+		CHAOS_API virtual TUniquePtr<FImplicitObject> CopyWithScale(const FVec3& Scale) const override;
 
-		void MovePlanesAndRebuild(FRealType InDelta);
+		CHAOS_API void MovePlanesAndRebuild(FRealType InDelta);
 
 	private:
-		void CreateStructureData(TArray<TArray<int32>>&& FaceIndices);
+		CHAOS_API void CreateStructureData(TArray<TArray<int32>>&& FaceIndices);
 
 	public:
 		static constexpr EImplicitObjectType StaticType()
@@ -355,8 +355,8 @@ namespace Chaos
 		 * and \c OutNormal should be.  The burden for detecting this case is deferred to the
 		 * caller. 
 		 */
-		virtual bool Raycast(const FVec3& StartPoint, const FVec3& Dir, const FReal Length, const FReal Thickness, FReal& OutTime, FVec3& OutPosition, FVec3& OutNormal, int32& OutFaceIndex) const override;
-		bool RaycastFast(const FVec3& StartPoint, const FVec3& Dir, const FReal Length, const FReal Thickness, FReal& OutTime, FVec3& OutPosition, FVec3& OutNormal, int32& OutFaceIndex) const;
+		CHAOS_API virtual bool Raycast(const FVec3& StartPoint, const FVec3& Dir, const FReal Length, const FReal Thickness, FReal& OutTime, FVec3& OutPosition, FVec3& OutNormal, int32& OutFaceIndex) const override;
+		CHAOS_API bool RaycastFast(const FVec3& StartPoint, const FVec3& Dir, const FReal Length, const FReal Thickness, FReal& OutTime, FVec3& OutPosition, FVec3& OutNormal, int32& OutFaceIndex) const;
 
 		virtual Pair<FVec3, bool> FindClosestIntersectionImp(const FVec3& StartPoint, const FVec3& EndPoint, const FReal Thickness) const override
 		{
@@ -389,14 +389,14 @@ namespace Chaos
 		const FConvexStructureData& GetStructureData() const { return StructureData; }
 
 		// Get the index of the plane that most opposes the normal
-		int32 GetMostOpposingPlane(const FVec3& Normal) const;
+		CHAOS_API int32 GetMostOpposingPlane(const FVec3& Normal) const;
 
 		// Get the index of the plane that most opposes the normal
-		int32 GetMostOpposingPlaneScaled(const FVec3& Normal, const FVec3& Scale) const;
+		CHAOS_API int32 GetMostOpposingPlaneScaled(const FVec3& Normal, const FVec3& Scale) const;
 
 		// Get the nearest point on an edge and the edge vertices
 		// Used for manifold generation
-		FVec3 GetClosestEdge(int32 PlaneIndexHint, const FVec3& Position, FVec3& OutEdgePos0, FVec3& OutEdgePos1) const;
+		CHAOS_API FVec3 GetClosestEdge(int32 PlaneIndexHint, const FVec3& Position, FVec3& OutEdgePos0, FVec3& OutEdgePos1) const;
 
 		// Get the nearest point on an edge of the specified face
 		FVec3 GetClosestEdgePosition(int32 PlaneIndex, const FVec3& Position) const
@@ -405,13 +405,13 @@ namespace Chaos
 			return GetClosestEdge(PlaneIndex, Position, Unused0, Unused1);
 		}
 
-		bool GetClosestEdgeVertices(int32 PlaneIndex, const FVec3& Position, int32& OutVertexIndex0, int32& OutVertexIndex1) const;
+		CHAOS_API bool GetClosestEdgeVertices(int32 PlaneIndex, const FVec3& Position, int32& OutVertexIndex0, int32& OutVertexIndex1) const;
 
 		// Get an array of all the plane indices that belong to a vertex (up to MaxVertexPlanes).
 		// Returns the number of planes found.
-		int32 FindVertexPlanes(int32 VertexIndex, int32* OutVertexPlanes, int32 MaxVertexPlanes) const;
+		CHAOS_API int32 FindVertexPlanes(int32 VertexIndex, int32* OutVertexPlanes, int32 MaxVertexPlanes) const;
 
-		int32 GetVertexPlanes3(int32 VertexIndex, int32& PlaneIndex0, int32& PlaneIndex1, int32& PlaneIndex2) const;
+		CHAOS_API int32 GetVertexPlanes3(int32 VertexIndex, int32& PlaneIndex0, int32& PlaneIndex1, int32& PlaneIndex2) const;
 
 		// The number of vertices that make up the corners of the specified face
 		inline int32 NumPlaneVertices(int32 PlaneIndex) const
@@ -497,9 +497,9 @@ namespace Chaos
 		}
 
 
-		virtual int32 FindMostOpposingFace(const FVec3& Position, const FVec3& UnitDir, int32 HintFaceIndex, FReal SearchDist) const override;
+		CHAOS_API virtual int32 FindMostOpposingFace(const FVec3& Position, const FVec3& UnitDir, int32 HintFaceIndex, FReal SearchDist) const override;
 
-		virtual int32 FindMostOpposingFaceScaled(const FVec3& Position, const FVec3& UnitDir, int32 HintFaceIndex, FReal SearchDist, const FVec3& Scale) const override;
+		CHAOS_API virtual int32 FindMostOpposingFaceScaled(const FVec3& Position, const FVec3& UnitDir, int32 HintFaceIndex, FReal SearchDist, const FVec3& Scale) const override;
 
 		FVec3 FindGeometryOpposingNormal(const FVec3& DenormDir, int32 FaceIndex, const FVec3& OriginalNormal) const
 		{
@@ -514,7 +514,7 @@ namespace Chaos
 			return FVec3(0.f, 0.f, 1.f);
 		}
 
-		virtual int32 FindClosestFaceAndVertices(const FVec3& Position, TArray<FVec3>& FaceVertices, FReal SearchDist = 0.01f) const override;
+		CHAOS_API virtual int32 FindClosestFaceAndVertices(const FVec3& Position, TArray<FVec3>& FaceVertices, FReal SearchDist = 0.01f) const override;
 
 		// Returns a winding order multiplier used in the manifold clipping and required when we have negative scales (See ImplicitObjectScaled)
 		FReal GetWindingOrder() const
@@ -1045,7 +1045,7 @@ namespace Chaos
 #endif // #if INTEL_ISPC
 
 	private:
-		void ComputeUnitMassInertiaTensorAndRotationOfMass(const FReal InVolume);
+		CHAOS_API void ComputeUnitMassInertiaTensorAndRotationOfMass(const FReal InVolume);
 		
 		// IMPORTANT to keep this copy constructor private to avoid unintented expensive copies 
 		FConvex(const FConvex& Other)
diff --git a/Engine/Source/Runtime/Experimental/Chaos/Public/Chaos/ConvexFlattenedArrayStructureData.h b/Engine/Source/Runtime/Experimental/Chaos/Public/Chaos/ConvexFlattenedArrayStructureData.h
--- a/Engine/Source/Runtime/Experimental/Chaos/Public/Chaos/ConvexFlattenedArrayStructureData.h
+++ b/Engine/Source/Runtime/Experimental/Chaos/Public/Chaos/ConvexFlattenedArrayStructureData.h
@@ -27,7 +27,7 @@ namespace Chaos
 		// Base class for TConvexFlattenedArrayStructureData.
 		// NOTE: Deliberately no-virtual destructor so it should never be deleted from a 
 		// base class pointer.
-		class CHAOS_API FConvexFlattenedArrayStructureData
+		class FConvexFlattenedArrayStructureData
 		{
 		public:
 			FConvexFlattenedArrayStructureData() {}
@@ -48,7 +48,7 @@ namespace Chaos
 		//					Max(NumPlanes*AverageVertsPerPlane)
 		//
 		template<typename T_INDEX, typename T_OFFSETINDEX>
-		class CHAOS_API TConvexFlattenedArrayStructureData : public FConvexFlattenedArrayStructureData
+		class TConvexFlattenedArrayStructureData : public FConvexFlattenedArrayStructureData
 		{
 		public:
 			// Everything public - not used outside of legacy loader and unit tests
@@ -282,4 +282,4 @@ namespace Chaos
 		};
 	}
 
-}
\ No newline at end of file
+}
diff --git a/Engine/Source/Runtime/Experimental/Chaos/Public/Chaos/ConvexStructureData.h b/Engine/Source/Runtime/Experimental/Chaos/Public/Chaos/ConvexStructureData.h
--- a/Engine/Source/Runtime/Experimental/Chaos/Public/Chaos/ConvexStructureData.h
+++ b/Engine/Source/Runtime/Experimental/Chaos/Public/Chaos/ConvexStructureData.h
@@ -15,7 +15,7 @@ namespace Chaos
 
 	// Metadata for a convex shape used by the manifold generation system and anything
 	// else that can benefit from knowing which vertices are associated with the faces.
-	class CHAOS_API FConvexStructureData
+	class FConvexStructureData
 	{
 	public:
 		using FConvexStructureDataLarge = FConvexHalfEdgeStructureDataS32;
diff --git a/Engine/Source/Runtime/Experimental/Chaos/Public/Chaos/DebugDrawQueue.h b/Engine/Source/Runtime/Experimental/Chaos/Public/Chaos/DebugDrawQueue.h
--- a/Engine/Source/Runtime/Experimental/Chaos/Public/Chaos/DebugDrawQueue.h
+++ b/Engine/Source/Runtime/Experimental/Chaos/Public/Chaos/DebugDrawQueue.h
@@ -82,7 +82,7 @@ private:
 	TList<ElementType>* FirstElement;
 };
 
-struct CHAOS_API FLatentDrawCommand
+struct FLatentDrawCommand
 {
 	FVector LineStart;
 	FVector LineEnd;
@@ -251,7 +251,7 @@ struct CHAOS_API FLatentDrawCommand
 };
 
 /** A thread safe way to generate latent debug drawing. (This is picked up later by the geometry collection component which is a total hack for now, but needed to get into an engine world ) */
-class CHAOS_API FDebugDrawQueue
+class FDebugDrawQueue
 {
 public:
 	enum EBuffer
@@ -434,9 +434,9 @@ public:
 		return Bounds.ThickenSymmetrically(FVec3(RadiusOfInterest)).Contains(CenterOfInterest);
 	}
 
-	void SetConsumerActive(void* Consumer, bool bConsumerActive);
+	CHAOS_API void SetConsumerActive(void* Consumer, bool bConsumerActive);
 
-	static FDebugDrawQueue& GetInstance();
+	static CHAOS_API FDebugDrawQueue& GetInstance();
 
 	static bool IsDebugDrawingEnabled()
 	{
@@ -483,4 +483,4 @@ private:
 	bool bEnableDebugDrawing;
 };
 }
-#endif
\ No newline at end of file
+#endif
diff --git a/Engine/Source/Runtime/Experimental/Chaos/Public/Chaos/Deformable/ChaosDeformableCollisionsProxy.h b/Engine/Source/Runtime/Experimental/Chaos/Public/Chaos/Deformable/ChaosDeformableCollisionsProxy.h
--- a/Engine/Source/Runtime/Experimental/Chaos/Public/Chaos/Deformable/ChaosDeformableCollisionsProxy.h
+++ b/Engine/Source/Runtime/Experimental/Chaos/Public/Chaos/Deformable/ChaosDeformableCollisionsProxy.h
@@ -21,7 +21,7 @@ namespace Chaos::Softs
 
 	typedef TTuple< const UObject*, ERigidCollisionShapeType, int8> FCollisionObjectKey;
 
-	struct CHAOS_API FCollisionObjectAddedBodies
+	struct FCollisionObjectAddedBodies
 	{
 		FCollisionObjectAddedBodies( FCollisionObjectKey InKey = FCollisionObjectKey(),
 			FTransform InTransform = FTransform::Identity,
@@ -38,18 +38,18 @@ namespace Chaos::Softs
 		FImplicitObject* Shapes = nullptr;
 	};
 
-	struct CHAOS_API FCollisionObjectRemovedBodies
+	struct FCollisionObjectRemovedBodies
 	{
 		FCollisionObjectKey Key = FCollisionObjectKey();
 	};
 
-	struct CHAOS_API FCollisionObjectUpdatedBodies
+	struct FCollisionObjectUpdatedBodies
 	{
 		FCollisionObjectKey Key = FCollisionObjectKey();
 		FTransform Transform = FTransform::Identity;
 	};
 
-	struct CHAOS_API FCollisionObjectParticleHandel
+	struct FCollisionObjectParticleHandel
 	{
 		FCollisionObjectParticleHandel(int32 InParticleIndex = INDEX_NONE,
 									   int32 InActiveViewIndex = INDEX_NONE,
@@ -64,7 +64,7 @@ namespace Chaos::Softs
 	};
 
 
-	class CHAOS_API FCollisionManagerProxy : public FThreadingProxy
+	class FCollisionManagerProxy : public FThreadingProxy
 	{
 	public:
 		typedef FThreadingProxy Super;
diff --git a/Engine/Source/Runtime/Experimental/Chaos/Public/Chaos/Deformable/ChaosDeformableSolver.h b/Engine/Source/Runtime/Experimental/Chaos/Public/Chaos/Deformable/ChaosDeformableSolver.h
--- a/Engine/Source/Runtime/Experimental/Chaos/Public/Chaos/Deformable/ChaosDeformableSolver.h
+++ b/Engine/Source/Runtime/Experimental/Chaos/Public/Chaos/Deformable/ChaosDeformableSolver.h
@@ -26,40 +26,40 @@
 
 namespace Chaos::Softs
 {
-	class CHAOS_API FDeformableSolver : public FPhysicsSolverEvents
+	class FDeformableSolver : public FPhysicsSolverEvents
 	{
 		friend class FGameThreadAccess;
 		friend class FPhysicsThreadAccess;
 
 	public:
 
-		FDeformableSolver(FDeformableSolverProperties InProp = FDeformableSolverProperties());
-		virtual ~FDeformableSolver();
+		CHAOS_API FDeformableSolver(FDeformableSolverProperties InProp = FDeformableSolverProperties());
+		CHAOS_API virtual ~FDeformableSolver();
 
 		/* Physics Thread Access API */
-		class CHAOS_API FPhysicsThreadAccess
+		class FPhysicsThreadAccess
 		{
 		public:
 			FPhysicsThreadAccess(FDeformableSolver* InSolver, const FPhysicsThreadAccessor&) : Solver(InSolver) {}
 			bool operator()() { return Solver != nullptr; }
 
 			/* Simulation Advance */
-			void UpdateProxyInputPackages();
-			void Simulate(FSolverReal DeltaTime);
-			void AdvanceDt(FSolverReal DeltaTime);
-			void Reset(const FDeformableSolverProperties&);
-			void Update(FSolverReal DeltaTime);
-			void UpdateOutputState(FThreadingProxy&);
-			TUniquePtr<FDeformablePackage> PullInputPackage();
-			void PushOutputPackage(int32 Frame, FDeformableDataMap&& Package);
+			CHAOS_API void UpdateProxyInputPackages();
+			CHAOS_API void Simulate(FSolverReal DeltaTime);
+			CHAOS_API void AdvanceDt(FSolverReal DeltaTime);
+			CHAOS_API void Reset(const FDeformableSolverProperties&);
+			CHAOS_API void Update(FSolverReal DeltaTime);
+			CHAOS_API void UpdateOutputState(FThreadingProxy&);
+			CHAOS_API TUniquePtr<FDeformablePackage> PullInputPackage();
+			CHAOS_API void PushOutputPackage(int32 Frame, FDeformableDataMap&& Package);
 
 			/* Iteration Advance */
-			void InitializeSimulationObjects();
-			void InitializeSimulationObject(FThreadingProxy&);
-			void InitializeCollisionBodies(FCollisionManagerProxy&);
-			void InitializeKinematicConstraint();
-			void InitializeSelfCollisionVariables();
-			void RemoveSimulationObjects();
+			CHAOS_API void InitializeSimulationObjects();
+			CHAOS_API void InitializeSimulationObject(FThreadingProxy&);
+			CHAOS_API void InitializeCollisionBodies(FCollisionManagerProxy&);
+			CHAOS_API void InitializeKinematicConstraint();
+			CHAOS_API void InitializeSelfCollisionVariables();
+			CHAOS_API void RemoveSimulationObjects();
 
 			const FDeformableSolverProperties& GetProperties() const { return Solver->GetProperties(); }
 
@@ -76,22 +76,22 @@ namespace Chaos::Softs
 
 
 		/* Game Thread Access API */
-		class CHAOS_API FGameThreadAccess
+		class FGameThreadAccess
 		{
 		public:
 			FGameThreadAccess(FDeformableSolver* InSolver, const FGameThreadAccessor&) : Solver(InSolver) {}
 			bool operator()() { return Solver != nullptr; }
 
 			int32 GetFrame() const { return Solver->GetFrame(); }
-			bool HasObject(UObject* InObject) const;
-			void AddProxy(FThreadingProxy* InObject);
-			void RemoveProxy(FThreadingProxy* InObject);
-			void PushInputPackage(int32 Frame, FDeformableDataMap&& InPackage);
+			CHAOS_API bool HasObject(UObject* InObject) const;
+			CHAOS_API void AddProxy(FThreadingProxy* InObject);
+			CHAOS_API void RemoveProxy(FThreadingProxy* InObject);
+			CHAOS_API void PushInputPackage(int32 Frame, FDeformableDataMap&& InPackage);
 
-			void SetEnableSolver(bool InbEnableSolver);
-			bool GetEnableSolver();
+			CHAOS_API void SetEnableSolver(bool InbEnableSolver);
+			CHAOS_API bool GetEnableSolver();
 
-			TUniquePtr<FDeformablePackage> PullOutputPackage();
+			CHAOS_API TUniquePtr<FDeformablePackage> PullOutputPackage();
 
 		private:
 			FDeformableSolver* Solver;
@@ -104,60 +104,60 @@ namespace Chaos::Softs
 
 		/* Simulation Advance */
 		int32 GetFrame() const { return Frame; }
-		void UpdateProxyInputPackages();
-		void Simulate(FSolverReal DeltaTime);
-		void AdvanceDt(FSolverReal DeltaTime);
-		void Reset(const FDeformableSolverProperties&);
-		void Update(FSolverReal DeltaTime);
-		void UpdateSimulationObjects(FSolverReal DeltaTime);
-		void UpdateOutputState(FThreadingProxy&);
-		void PushOutputPackage(int32 Frame, FDeformableDataMap&& Package);
-		TUniquePtr<FDeformablePackage> PullInputPackage( );
+		CHAOS_API void UpdateProxyInputPackages();
+		CHAOS_API void Simulate(FSolverReal DeltaTime);
+		CHAOS_API void AdvanceDt(FSolverReal DeltaTime);
+		CHAOS_API void Reset(const FDeformableSolverProperties&);
+		CHAOS_API void Update(FSolverReal DeltaTime);
+		CHAOS_API void UpdateSimulationObjects(FSolverReal DeltaTime);
+		CHAOS_API void UpdateOutputState(FThreadingProxy&);
+		CHAOS_API void PushOutputPackage(int32 Frame, FDeformableDataMap&& Package);
+		CHAOS_API TUniquePtr<FDeformablePackage> PullInputPackage( );
 
 		/* Iteration Advance */
-		void InitializeSimulationSpace();
-		void InitializeSimulationObjects();
-		void InitializeSimulationObject(FThreadingProxy&);
-		void InitializeDeformableParticles(FFleshThreadingProxy&);
-		void UpdateTransientConstraints();
-		void PostProcessTransientConstraints();
-		void InitializeKinematicParticles(FFleshThreadingProxy&);
-		void InitializeTetrahedralConstraint(FFleshThreadingProxy&);
-		void InitializeGidBasedConstraints(FFleshThreadingProxy&);
-		void InitializeGaussSeidelConstraints(FFleshThreadingProxy& Proxy);
-		void InitializeWeakConstraint(FFleshThreadingProxy&);
-		void InitializeKinematicConstraint();
-		void InitializeCollisionBodies(FCollisionManagerProxy&);
-		void InitializeSelfCollisionVariables();
-		void InitializeGridBasedConstraintVariables();
-		void InitializeGaussSeidelConstraintVariables();
-		void UpdateCollisionBodies(FCollisionManagerProxy&, FThreadingProxy::FKey, FSolverReal DeltaTime);
-		void RemoveSimulationObjects();
-		TArray<Chaos::TVec3<FSolverReal>> ComputeParticleTargets(const TArray<TArray<int32>>& ParticleIndices);
+		CHAOS_API void InitializeSimulationSpace();
+		CHAOS_API void InitializeSimulationObjects();
+		CHAOS_API void InitializeSimulationObject(FThreadingProxy&);
+		CHAOS_API void InitializeDeformableParticles(FFleshThreadingProxy&);
+		CHAOS_API void UpdateTransientConstraints();
+		CHAOS_API void PostProcessTransientConstraints();
+		CHAOS_API void InitializeKinematicParticles(FFleshThreadingProxy&);
+		CHAOS_API void InitializeTetrahedralConstraint(FFleshThreadingProxy&);
+		CHAOS_API void InitializeGidBasedConstraints(FFleshThreadingProxy&);
+		CHAOS_API void InitializeGaussSeidelConstraints(FFleshThreadingProxy& Proxy);
+		CHAOS_API void InitializeWeakConstraint(FFleshThreadingProxy&);
+		CHAOS_API void InitializeKinematicConstraint();
+		CHAOS_API void InitializeCollisionBodies(FCollisionManagerProxy&);
+		CHAOS_API void InitializeSelfCollisionVariables();
+		CHAOS_API void InitializeGridBasedConstraintVariables();
+		CHAOS_API void InitializeGaussSeidelConstraintVariables();
+		CHAOS_API void UpdateCollisionBodies(FCollisionManagerProxy&, FThreadingProxy::FKey, FSolverReal DeltaTime);
+		CHAOS_API void RemoveSimulationObjects();
+		CHAOS_API TArray<Chaos::TVec3<FSolverReal>> ComputeParticleTargets(const TArray<TArray<int32>>& ParticleIndices);
 
 		/*Debug Output*/
-		void DebugDrawSimulationData();
-		void DebugDrawTetrahedralParticles(FFleshThreadingProxy& Proxy);
-		void WriteFrame(FThreadingProxy& InProxy, const FSolverReal DeltaTime);
-		void WriteTrisGEO(const FSolverParticles& Particles, const TArray<TVec3<int32>>& Mesh);
+		CHAOS_API void DebugDrawSimulationData();
+		CHAOS_API void DebugDrawTetrahedralParticles(FFleshThreadingProxy& Proxy);
+		CHAOS_API void WriteFrame(FThreadingProxy& InProxy, const FSolverReal DeltaTime);
+		CHAOS_API void WriteTrisGEO(const FSolverParticles& Particles, const TArray<TVec3<int32>>& Mesh);
 
 		/*Game Thread API*/
 		bool HasObject(UObject* InObject) const { return InitializedObjects_External.Contains(InObject); }
-		void AddProxy(FThreadingProxy* InObject);
-		void RemoveProxy(FThreadingProxy* InObject);
-		TUniquePtr<FDeformablePackage> PullOutputPackage();
-		void PushInputPackage(int32 Frame, FDeformableDataMap&& InPackage);
+		CHAOS_API void AddProxy(FThreadingProxy* InObject);
+		CHAOS_API void RemoveProxy(FThreadingProxy* InObject);
+		CHAOS_API TUniquePtr<FDeformablePackage> PullOutputPackage();
+		CHAOS_API void PushInputPackage(int32 Frame, FDeformableDataMap&& InPackage);
 
 		const FDeformableSolverProperties& GetProperties() const { return Property; }
 
 	private:
 
 		// connections outside the solver.
-		static FCriticalSection	InitializationMutex; // @todo(flesh) : change to threaded commands to prevent the lock. 
-		static FCriticalSection	RemovalMutex; // @todo(flesh) : change to threaded commands to prevent the lock. 
-		static FCriticalSection	PackageOutputMutex;
-		static FCriticalSection	PackageInputMutex;
-		static FCriticalSection	SolverEnabledMutex;
+		static CHAOS_API FCriticalSection	InitializationMutex; // @todo(flesh) : change to threaded commands to prevent the lock. 
+		static CHAOS_API FCriticalSection	RemovalMutex; // @todo(flesh) : change to threaded commands to prevent the lock. 
+		static CHAOS_API FCriticalSection	PackageOutputMutex;
+		static CHAOS_API FCriticalSection	PackageInputMutex;
+		static CHAOS_API FCriticalSection	SolverEnabledMutex;
 
 		TArray< FThreadingProxy* > RemovedProxys_Internal;
 		TArray< FThreadingProxy* > UninitializedProxys_Internal;
diff --git a/Engine/Source/Runtime/Experimental/Chaos/Public/Chaos/Deformable/ChaosDeformableSolverProxy.h b/Engine/Source/Runtime/Experimental/Chaos/Public/Chaos/Deformable/ChaosDeformableSolverProxy.h
--- a/Engine/Source/Runtime/Experimental/Chaos/Public/Chaos/Deformable/ChaosDeformableSolverProxy.h
+++ b/Engine/Source/Runtime/Experimental/Chaos/Public/Chaos/Deformable/ChaosDeformableSolverProxy.h
@@ -21,7 +21,7 @@ enum ChaosDeformableSimSpace : uint8
 
 namespace Chaos::Softs
 {
-	class CHAOS_API FThreadingProxy {
+	class FThreadingProxy {
 		const UObject* Owner = nullptr;
 		FName TypeName = FName("");
 
@@ -70,7 +70,7 @@ namespace Chaos::Softs
 		};
 	};
 
-	class CHAOS_API FFleshThreadingProxy : public FThreadingProxy
+	class FFleshThreadingProxy : public FThreadingProxy
 	{
 		typedef FThreadingProxy Super;
 
diff --git a/Engine/Source/Runtime/Experimental/Chaos/Public/Chaos/Deformable/ChaosDeformableSolverTypes.h b/Engine/Source/Runtime/Experimental/Chaos/Public/Chaos/Deformable/ChaosDeformableSolverTypes.h
--- a/Engine/Source/Runtime/Experimental/Chaos/Public/Chaos/Deformable/ChaosDeformableSolverTypes.h
+++ b/Engine/Source/Runtime/Experimental/Chaos/Public/Chaos/Deformable/ChaosDeformableSolverTypes.h
@@ -13,7 +13,7 @@ class FFleshCacheAdapter;
 
 namespace Chaos::Softs
 {
-	struct CHAOS_API FDeformableSolverProperties
+	struct FDeformableSolverProperties
 	{
 		FDeformableSolverProperties(
 			int32 InNumSolverSubSteps = 2,
@@ -91,7 +91,7 @@ namespace Chaos::Softs
 	typedef TSharedPtr<const FThreadingProxy::FBuffer> FDataMapValue; // Buffer Pointer
 	typedef TMap<FThreadingProxy::FKey, FDataMapValue > FDeformableDataMap; // <const UObject*,FBufferSharedPtr>
 
-	struct CHAOS_API FDeformablePackage {
+	struct FDeformablePackage {
 		FDeformablePackage()
 		{}
 
@@ -105,7 +105,7 @@ namespace Chaos::Softs
 	};
 
 	/* Accessor for the Game Thread*/
-	class CHAOS_API FGameThreadAccessor
+	class FGameThreadAccessor
 	{
 	public:
 //		friend class UDeformableSolverComponent;
@@ -118,7 +118,7 @@ namespace Chaos::Softs
 
 
 	/* Accessor for the Physics Thread*/
-	class CHAOS_API FPhysicsThreadAccessor
+	class FPhysicsThreadAccessor
 	{
 	public:
 //		friend class UDeformableSolverComponent;
@@ -130,7 +130,7 @@ namespace Chaos::Softs
 	};
 
 
-	struct CHAOS_API FDeformableDebugParams
+	struct FDeformableDebugParams
 	{				
 		bool bDoDrawTetrahedralParticles = false;
 		bool bDoDrawKinematicParticles = false;
@@ -149,7 +149,7 @@ namespace Chaos::Softs
 		}
 	};
 
-	struct CHAOS_API FDeformableXPBDCorotatedParams
+	struct FDeformableXPBDCorotatedParams
 	{
 		int32 XPBDCorotatedBatchSize = 5;
 		int32 XPBDCorotatedBatchThreshold = 5;
@@ -157,4 +157,4 @@ namespace Chaos::Softs
 	};
 
 
-}; // namesapce Chaos::Softs
\ No newline at end of file
+}; // namesapce Chaos::Softs
diff --git a/Engine/Source/Runtime/Experimental/Chaos/Public/Chaos/Evolution/ABTestingConstraintContainerSolver.h b/Engine/Source/Runtime/Experimental/Chaos/Public/Chaos/Evolution/ABTestingConstraintContainerSolver.h
--- a/Engine/Source/Runtime/Experimental/Chaos/Public/Chaos/Evolution/ABTestingConstraintContainerSolver.h
+++ b/Engine/Source/Runtime/Experimental/Chaos/Public/Chaos/Evolution/ABTestingConstraintContainerSolver.h
@@ -31,7 +31,7 @@ namespace Chaos
 		* as long as we only have one constraint type in the island group which will do for now.
 		*/
 		template<typename T1, typename T2>
-		class CHAOS_API TABTestingConstraintContainerSolver : public FConstraintContainerSolver
+		class TABTestingConstraintContainerSolver : public FConstraintContainerSolver
 		{
 		public:
 			enum class ESolverPhase
@@ -204,4 +204,4 @@ namespace Chaos
 #endif // CHAOS_ABTEST_CONSTRAINTSOLVER_ENABLED
 
 	}	// namespace Private
-}	// namespace Chaos
\ No newline at end of file
+}	// namespace Chaos
diff --git a/Engine/Source/Runtime/Experimental/Chaos/Public/Chaos/Evolution/ConstraintGroupSolver.h b/Engine/Source/Runtime/Experimental/Chaos/Public/Chaos/Evolution/ConstraintGroupSolver.h
--- a/Engine/Source/Runtime/Experimental/Chaos/Public/Chaos/Evolution/ConstraintGroupSolver.h
+++ b/Engine/Source/Runtime/Experimental/Chaos/Public/Chaos/Evolution/ConstraintGroupSolver.h
@@ -18,13 +18,13 @@ namespace Chaos
 		 * All constraints in the groups are solved in sequence in a single thread. We can create one of these for 
 		 * each Island (or group of Islands, or subset of a color of constraints) for parallelism.
 		*/
-		class CHAOS_API FPBDConstraintGroupSolver
+		class FPBDConstraintGroupSolver
 		{
 		public:
 			UE_NONCOPYABLE(FPBDConstraintGroupSolver);
 
-			FPBDConstraintGroupSolver();
-			virtual ~FPBDConstraintGroupSolver();
+			CHAOS_API FPBDConstraintGroupSolver();
+			CHAOS_API virtual ~FPBDConstraintGroupSolver();
 
 			/**
 			 * Get the iterations settings for this group of islands
@@ -63,89 +63,89 @@ namespace Chaos
 			/**
 			 * Attach a constraint solver to the specified ContainerId. This must be for the same constraint type as the container with that Id.
 			*/
-			void SetConstraintSolver(const int32 ContainerId, TUniquePtr<FConstraintContainerSolver>&& Solver);
+			CHAOS_API void SetConstraintSolver(const int32 ContainerId, TUniquePtr<FConstraintContainerSolver>&& Solver);
 
 			/**
 			 * Set the solver priority of the specified constraint type
 			*/
-			void SetConstraintSolverPriority(const int32 ContainerId, const int32 Priority);
+			CHAOS_API void SetConstraintSolverPriority(const int32 ContainerId, const int32 Priority);
 
 			/**
 			 * Reset all state - called once per tick
 			*/
-			void Reset();
+			CHAOS_API void Reset();
 
 			/**
 			 * Set up the constraints solvers and body containers with pointers to their constraint and particles, but do not collect any data.
 			*/
-			void AddConstraintsAndBodies();
+			CHAOS_API void AddConstraintsAndBodies();
 
 			/**
 			 * Collect all the data for all solver bodies from their respective particles.
 			*/
-			void GatherBodies(const FReal Dt);
+			CHAOS_API void GatherBodies(const FReal Dt);
 
 			/**
 			 * Collect all the data for the specified range of solver bodies from their respective particles. Will be called from multiple
 			 * threads with different non-overlapping indices.
 			*/
-			void GatherBodies(const FReal Dt, const int32 BeginBodyIndex, const int32 EndBodyIndex);
+			CHAOS_API void GatherBodies(const FReal Dt, const int32 BeginBodyIndex, const int32 EndBodyIndex);
 
 			/**
 			 * Collect all the data for all constraint solvers from their respective constraints.
 			*/
-			void GatherConstraints(const FReal Dt);
+			CHAOS_API void GatherConstraints(const FReal Dt);
 
 			/**
 			 * Collect all the data for the specified range of constraint solvers from their respective constraints. Will be called from multiple
 			 * threads with different non-overlapping indices.
 			*/
-			void GatherConstraints(const FReal Dt, const int32 BeginConstraintIndex, const int32 EndConstraintIndex);
+			CHAOS_API void GatherConstraints(const FReal Dt, const int32 BeginConstraintIndex, const int32 EndConstraintIndex);
 
 			/**
 			 * For additional processing after gathering all the data.
 			 * Calls PreApplyConstraints on each container solver.
 			 */
-			void PreApplyPositionConstraints(const FReal Dt);
-			void PreApplyVelocityConstraints(const FReal Dt);
-			void PreApplyProjectionConstraints(const FReal Dt);
+			CHAOS_API void PreApplyPositionConstraints(const FReal Dt);
+			CHAOS_API void PreApplyVelocityConstraints(const FReal Dt);
+			CHAOS_API void PreApplyProjectionConstraints(const FReal Dt);
 
 			/**
 			 * Apply positional constraints, and set the body velocities
 			*/
-			void ApplyPositionConstraints(const FReal Dt);
+			CHAOS_API void ApplyPositionConstraints(const FReal Dt);
 
 			/**
 			 * Apply any velocity constraints, and update the body velocities
 			*/
-			void ApplyVelocityConstraints(const FReal Dt);
+			CHAOS_API void ApplyVelocityConstraints(const FReal Dt);
 
 			/**
 			 * Apply projection to attempt to fix up any errors left over from the Position and Velocity saolver phases
 			*/
-			void ApplyProjectionConstraints(const FReal Dt);
+			CHAOS_API void ApplyProjectionConstraints(const FReal Dt);
 
 			/**
 			 * Push results from all solver bodies back to their respective particles.
 			*/
-			void ScatterBodies(const FReal Dt);
+			CHAOS_API void ScatterBodies(const FReal Dt);
 
 			/**
 			 * Push results from the specified range of solver bodies back to their respective particles. Will be called from multiple
 			 * threads with different non-overlapping indices.
 			*/
-			void ScatterBodies(const FReal Dt, const int32 BeginBodyIndex, const int32 EndBodyIndex);
+			CHAOS_API void ScatterBodies(const FReal Dt, const int32 BeginBodyIndex, const int32 EndBodyIndex);
 
 			/**
 			 * Push results from all constraint solvers back to their respective constraints.
 			*/
-			void ScatterConstraints(const FReal Dt);
+			CHAOS_API void ScatterConstraints(const FReal Dt);
 
 			/**
 			 * Push results from the specified range of constraint solvers back to their respective constraints.Will be called from multiple
 			 * threads with different non-overlapping indices. The range covers all constraint types.
 			*/
-			void ScatterConstraints(const FReal Dt, const int32 BeginConstraintIndex, const int32 EndConstraintIndex);
+			CHAOS_API void ScatterConstraints(const FReal Dt, const int32 BeginConstraintIndex, const int32 EndConstraintIndex);
 
 		protected:
 
@@ -155,7 +155,7 @@ namespace Chaos
 			void ApplyToConstraintRange(const int32 BeginConstraintIndex, const int32 EndConstraintIndex, const LambdaType& Lambda);
 
 			// Sort the solvers based on Level and other criteria for more stable solving
-			void SortSolverContainers();
+			CHAOS_API void SortSolverContainers();
 
 			// Allow a derived class to perform per-tick reset
 			virtual void ResetImpl() {}
@@ -196,7 +196,7 @@ namespace Chaos
 		 * 
 		 * There will only be one (or zero) FPBDSceneConstraintGroupSolver per simulation world.
 		*/
-		class CHAOS_API FPBDSceneConstraintGroupSolver : public FPBDConstraintGroupSolver
+		class FPBDSceneConstraintGroupSolver : public FPBDConstraintGroupSolver
 		{
 		public:
 			FPBDSceneConstraintGroupSolver(const FIterationSettings& InIterations)
@@ -206,7 +206,7 @@ namespace Chaos
 
 		protected:
 			// The Scene Group Solver adds all constraints from all containers
-			virtual void AddConstraintsImpl() override final;
+			CHAOS_API virtual void AddConstraintsImpl() override final;
 		};
 
 	} // namespace Private
@@ -219,4 +219,4 @@ namespace Chaos
 	//UE_DEPRECATED(5.2, "FPBDSceneConstraintGroupSolver is for internal use only")
 	//using FPBDSceneConstraintGroupSolver = Private::FPBDSceneConstraintGroupSolver;
 
-}	// namespace Chaos
\ No newline at end of file
+}	// namespace Chaos
diff --git a/Engine/Source/Runtime/Experimental/Chaos/Public/Chaos/Evolution/IndexedConstraintContainer.h b/Engine/Source/Runtime/Experimental/Chaos/Public/Chaos/Evolution/IndexedConstraintContainer.h
--- a/Engine/Source/Runtime/Experimental/Chaos/Public/Chaos/Evolution/IndexedConstraintContainer.h
+++ b/Engine/Source/Runtime/Experimental/Chaos/Public/Chaos/Evolution/IndexedConstraintContainer.h
@@ -11,7 +11,7 @@ namespace Chaos
 	/**
 	 * Base class for handles to constraints in an index-based container (e.g., FJointConstraints)
 	 */
-	class CHAOS_API FIndexedConstraintHandle : public FConstraintHandle
+	class FIndexedConstraintHandle : public FConstraintHandle
 	{
 	public:
 		using FGeometryParticleHandle = TGeometryParticleHandle<FReal, 3>;
@@ -64,7 +64,7 @@ namespace Chaos
 	 * Utility base class for FIndexedConstraintHandles
 	 */
 	template<typename T_CONTAINER>
-	class CHAOS_API TIndexedContainerConstraintHandle : public FIndexedConstraintHandle
+	class TIndexedContainerConstraintHandle : public FIndexedConstraintHandle
 	{
 	public:
 		using Base = FIndexedConstraintHandle;
@@ -125,7 +125,7 @@ namespace Chaos
 	 * as determined by the constraint graph.
 	*/
 	template<typename ConstraintContainerType>
-	class CHAOS_API TIndexedConstraintContainerSolver : public FConstraintContainerSolver
+	class TIndexedConstraintContainerSolver : public FConstraintContainerSolver
 	{
 	public:
 		using FConstraintContainerType = ConstraintContainerType;
@@ -228,7 +228,7 @@ namespace Chaos
 	/**
 	* Utility base class for constraint containers usingindexed-based constraint handles (e.g., FJointConstraints)
 	*/
-	class CHAOS_API FPBDIndexedConstraintContainer : public FPBDConstraintContainer
+	class FPBDIndexedConstraintContainer : public FPBDConstraintContainer
 	{
 	public:
 		FPBDIndexedConstraintContainer(FConstraintHandleTypeID InType)
@@ -253,7 +253,7 @@ namespace Chaos
 	};
 
 	template<typename ConstaintContainerType>
-	class CHAOS_API TPBDIndexedConstraintContainer : public FPBDIndexedConstraintContainer
+	class TPBDIndexedConstraintContainer : public FPBDIndexedConstraintContainer
 	{
 	public:
 		using FConstaintContainerType = ConstaintContainerType;
diff --git a/Engine/Source/Runtime/Experimental/Chaos/Public/Chaos/Evolution/IterationSettings.h b/Engine/Source/Runtime/Experimental/Chaos/Public/Chaos/Evolution/IterationSettings.h
--- a/Engine/Source/Runtime/Experimental/Chaos/Public/Chaos/Evolution/IterationSettings.h
+++ b/Engine/Source/Runtime/Experimental/Chaos/Public/Chaos/Evolution/IterationSettings.h
@@ -13,7 +13,7 @@ namespace Chaos
 		 * Iteration counts for use by bodies and constraints. 
 		 * A simulation island will use the maximum number of iterations of any body or constraint in the island.
 		*/
-		class CHAOS_API FIterationSettings
+		class FIterationSettings
 		{
 		private:
 			using FIndexType = int8;
@@ -64,4 +64,4 @@ namespace Chaos
 			}
 		};
 	}
-}
\ No newline at end of file
+}
diff --git a/Engine/Source/Runtime/Experimental/Chaos/Public/Chaos/Evolution/PBDMinEvolution.h b/Engine/Source/Runtime/Experimental/Chaos/Public/Chaos/Evolution/PBDMinEvolution.h
--- a/Engine/Source/Runtime/Experimental/Chaos/Public/Chaos/Evolution/PBDMinEvolution.h
+++ b/Engine/Source/Runtime/Experimental/Chaos/Public/Chaos/Evolution/PBDMinEvolution.h
@@ -27,21 +27,21 @@ namespace Chaos
 	 *
 	 * It is single-threaded and does not use a constraint graph or partition the particles into islands.
 	 */
-	class CHAOS_API FPBDMinEvolution
+	class FPBDMinEvolution
 	{
 	public:
 		using FCollisionDetector = FBasicCollisionDetector;
 		using FEvolutionCallback = TFunction<void()>;
 		using FRigidParticleSOAs = FPBDRigidsSOAs;
 
-		FPBDMinEvolution(FRigidParticleSOAs& InParticles, TArrayCollectionArray<FVec3>& InPrevX, TArrayCollectionArray<FRotation3>& InPrevR, FCollisionDetector& InCollisionDetector);
-		~FPBDMinEvolution();
+		CHAOS_API FPBDMinEvolution(FRigidParticleSOAs& InParticles, TArrayCollectionArray<FVec3>& InPrevX, TArrayCollectionArray<FRotation3>& InPrevR, FCollisionDetector& InCollisionDetector);
+		CHAOS_API ~FPBDMinEvolution();
 
-		void AddConstraintContainer(FPBDConstraintContainer& InContainer, const int32 Priority = 0);
-		void SetConstraintContainerPriority(const int32 ContainerId, const int32 Priority);
+		CHAOS_API void AddConstraintContainer(FPBDConstraintContainer& InContainer, const int32 Priority = 0);
+		CHAOS_API void SetConstraintContainerPriority(const int32 ContainerId, const int32 Priority);
 
-		void Advance(const FReal StepDt, const int32 NumSteps, const FReal RewindDt);
-		void AdvanceOneTimeStep(const FReal Dt, const FReal StepFraction);
+		CHAOS_API void Advance(const FReal StepDt, const int32 NumSteps, const FReal RewindDt);
+		CHAOS_API void AdvanceOneTimeStep(const FReal Dt, const FReal StepFraction);
 
 		void SetNumPositionIterations(const int32 NumIts)
 		{
@@ -103,17 +103,17 @@ namespace Chaos
 		}
 
 	private:
-		void PrepareTick();
-		void UnprepareTick();
-		void Rewind(FReal Dt, FReal RewindDt);
-		void Integrate(FReal Dt);
-		void ApplyKinematicTargets(FReal Dt, FReal StepFraction);
-		void DetectCollisions(FReal Dt);
-		void GatherInput(FReal Dt);
-		void ScatterOutput(FReal Dt);
-		void ApplyConstraintsPhase1(FReal Dt);
-		void ApplyConstraintsPhase2(FReal Dt);
-		void ApplyConstraintsPhase3(FReal Dt);
+		CHAOS_API void PrepareTick();
+		CHAOS_API void UnprepareTick();
+		CHAOS_API void Rewind(FReal Dt, FReal RewindDt);
+		CHAOS_API void Integrate(FReal Dt);
+		CHAOS_API void ApplyKinematicTargets(FReal Dt, FReal StepFraction);
+		CHAOS_API void DetectCollisions(FReal Dt);
+		CHAOS_API void GatherInput(FReal Dt);
+		CHAOS_API void ScatterOutput(FReal Dt);
+		CHAOS_API void ApplyConstraintsPhase1(FReal Dt);
+		CHAOS_API void ApplyConstraintsPhase2(FReal Dt);
+		CHAOS_API void ApplyConstraintsPhase3(FReal Dt);
 
 		FRigidParticleSOAs& Particles;
 		FCollisionDetector& CollisionDetector;
diff --git a/Engine/Source/Runtime/Experimental/Chaos/Public/Chaos/Evolution/SimulationSpace.h b/Engine/Source/Runtime/Experimental/Chaos/Public/Chaos/Evolution/SimulationSpace.h
--- a/Engine/Source/Runtime/Experimental/Chaos/Public/Chaos/Evolution/SimulationSpace.h
+++ b/Engine/Source/Runtime/Experimental/Chaos/Public/Chaos/Evolution/SimulationSpace.h
@@ -11,7 +11,7 @@ namespace Chaos
 	 * be used to pass a tunable fraction of the world-space movement
 	 * on to the bodies.
 	 */
-	class CHAOS_API FSimulationSpace
+	class FSimulationSpace
 	{
 	public:
 		FSimulationSpace()
@@ -42,7 +42,7 @@ namespace Chaos
 	 *
 	 * https://en.wikipedia.org/wiki/Rotating_reference_frame
 	 */
-	class CHAOS_API FSimulationSpaceSettings
+	class FSimulationSpaceSettings
 	{
 	public:
 		FSimulationSpaceSettings()
@@ -95,4 +95,4 @@ PRAGMA_ENABLE_DEPRECATION_WARNINGS
 		FVec3 ExternalLinearEtherDrag;
 	};
 
-}
\ No newline at end of file
+}
diff --git a/Engine/Source/Runtime/Experimental/Chaos/Public/Chaos/Evolution/SolverConstraintContainer.h b/Engine/Source/Runtime/Experimental/Chaos/Public/Chaos/Evolution/SolverConstraintContainer.h
--- a/Engine/Source/Runtime/Experimental/Chaos/Public/Chaos/Evolution/SolverConstraintContainer.h
+++ b/Engine/Source/Runtime/Experimental/Chaos/Public/Chaos/Evolution/SolverConstraintContainer.h
@@ -25,7 +25,7 @@ namespace Chaos
 	 * RBAN does not attempt to partition its scene into islands and just solves all constraints on its main
 	 * thread, so it only has one FConstraintContainerSolver (per constraint type).
 	*/
-	class CHAOS_API FConstraintContainerSolver
+	class FConstraintContainerSolver
 	{
 	public:
 
@@ -118,7 +118,7 @@ namespace Chaos
 	 * @todo(chaos): really we should split the base class into group and scene versions. See CreateSceneSolver CreateGroupSolver
 	*/
 	template<typename ConstraintContainerType>
-	class CHAOS_API TSimpleConstraintContainerSolver : public FConstraintContainerSolver
+	class TSimpleConstraintContainerSolver : public FConstraintContainerSolver
 	{
 	public:
 		using FConstraintContainerType = ConstraintContainerType;
diff --git a/Engine/Source/Runtime/Experimental/Chaos/Public/Chaos/Framework/ChaosResultsManager.h b/Engine/Source/Runtime/Experimental/Chaos/Public/Chaos/Framework/ChaosResultsManager.h
--- a/Engine/Source/Runtime/Experimental/Chaos/Public/Chaos/Framework/ChaosResultsManager.h
+++ b/Engine/Source/Runtime/Experimental/Chaos/Public/Chaos/Framework/ChaosResultsManager.h
@@ -10,25 +10,25 @@ namespace Chaos
 	extern FRealSingle SecondChannelDelay;
 	extern int32 DefaultNumActiveChannels;
 
-	struct CHAOS_API FChaosRigidInterpolationData
+	struct FChaosRigidInterpolationData
 	{
 		FDirtyRigidParticleData Prev;
 		FDirtyRigidParticleData Next;
 	};
 
-	struct CHAOS_API FChaosGeometryCollectionInterpolationData
+	struct FChaosGeometryCollectionInterpolationData
 	{
 		FDirtyGeometryCollectionData Prev;
 		FDirtyGeometryCollectionData Next;
 	};
 
-	struct CHAOS_API FChaosClusterUnionInterpolationData
+	struct FChaosClusterUnionInterpolationData
 	{
 		FDirtyClusterUnionData Prev;
 		FDirtyClusterUnionData Next;
 	};
 	
-	struct CHAOS_API FChaosInterpolationResults
+	struct FChaosInterpolationResults
 	{
 		FChaosInterpolationResults()
 			: Prev(nullptr)
@@ -49,7 +49,7 @@ namespace Chaos
 			Other.Next = nullptr;
 		}
 
-		void Reset();
+		CHAOS_API void Reset();
 		
 		TArray<FChaosRigidInterpolationData> RigidInterpolations;
 		TArray<FChaosGeometryCollectionInterpolationData> GeometryCollectionInterpolations;
@@ -61,26 +61,26 @@ namespace Chaos
 
 	struct FChaosResultsChannel;
 
-	class CHAOS_API FChaosResultsManager
+	class FChaosResultsManager
 	{
 	public:
-		FChaosResultsManager(FChaosMarshallingManager& InMarshallingManager);
+		CHAOS_API FChaosResultsManager(FChaosMarshallingManager& InMarshallingManager);
 		FChaosResultsManager(const FChaosResultsManager& Other) = delete;
 		FChaosResultsManager(FChaosResultsManager&& Other) = default;
 		
-		~FChaosResultsManager();
+		CHAOS_API ~FChaosResultsManager();
 
-		const FChaosInterpolationResults& PullSyncPhysicsResults_External();
-		TArray<const FChaosInterpolationResults*> PullAsyncPhysicsResults_External(const FReal ResultsTime);
+		CHAOS_API const FChaosInterpolationResults& PullSyncPhysicsResults_External();
+		CHAOS_API TArray<const FChaosInterpolationResults*> PullAsyncPhysicsResults_External(const FReal ResultsTime);
 
-		void RemoveProxy_External(FSingleParticlePhysicsProxy* Proxy);
+		CHAOS_API void RemoveProxy_External(FSingleParticlePhysicsProxy* Proxy);
 
 	private:
 
 		friend FChaosResultsChannel;
 
-		FPullPhysicsData* PopPullData_External(int32 ChannelIdx);
-		void FreePullData_External(FPullPhysicsData* PullData, int32 ChannelIdx);
+		CHAOS_API FPullPhysicsData* PopPullData_External(int32 ChannelIdx);
+		CHAOS_API void FreePullData_External(FPullPhysicsData* PullData, int32 ChannelIdx);
 
 		TArray<FChaosResultsChannel*> Channels;
 		FChaosMarshallingManager& MarshallingManager;
diff --git a/Engine/Source/Runtime/Experimental/Chaos/Public/Chaos/Framework/DebugSubstep.h b/Engine/Source/Runtime/Experimental/Chaos/Public/Chaos/Framework/DebugSubstep.h
--- a/Engine/Source/Runtime/Experimental/Chaos/Public/Chaos/Framework/DebugSubstep.h
+++ b/Engine/Source/Runtime/Experimental/Chaos/Public/Chaos/Framework/DebugSubstep.h
@@ -19,14 +19,14 @@ class FEvent;
 namespace Chaos
 {
 	/** Debug only class used to control the pausing/stepping/substepping of a debug solver thread. */
-	class CHAOS_API FDebugSubstep final
+	class FDebugSubstep final
 	{
 		friend class FDebugSolverTask;
 		friend class FDebugSolverTasks;
 
 	public:
 		FDebugSubstep();
-		~FDebugSubstep();
+		CHAOS_API ~FDebugSubstep();
 
 		/** Return whether debugging mode/pausing to substeps is enabled. */
 		bool IsEnabled() const { return bIsEnabled; }
@@ -40,23 +40,23 @@ namespace Chaos
 		FORCEINLINE void Add(const TCHAR* Label = nullptr) const { Add(false, Label); }
 
 		/** Enable/disable substep pause points. */
-		void Enable(bool bEnable);
+		CHAOS_API void Enable(bool bEnable);
 
 		/** Allow progress to the next substep (only works once this object is enabled). */
-		void ProgressToSubstep();
+		CHAOS_API void ProgressToSubstep();
 
 		/** Allow progress to the next step (only works once this object is enabled). */
-		void ProgressToStep();
+		CHAOS_API void ProgressToStep();
 
 	private:
 		FDebugSubstep(const FDebugSubstep&) = delete;
 		FDebugSubstep& operator=(const FDebugSubstep&) = delete;
 
 		/** Initialize sync events. */
-		void Initialize();
+		CHAOS_API void Initialize();
 
 		/** Shutdown and release sync events. Once released the debug thread cannot be restarted unless Initialize is called first. */
-		void Release();
+		CHAOS_API void Release();
 
 		/** Return true when initialized. */
 		FORCEINLINE bool IsInitialized() const { return !!ProgressEvent && !!SubstepEvent; }
@@ -66,10 +66,10 @@ namespace Chaos
 		 * Not thread safe. Must be called from within the physics thread or with the physics thread locked.
 		 * Once shutdown, the debug thread can still be restarted when SyncAdvance is called.
 		 */
-		void Shutdown();
+		CHAOS_API void Shutdown();
 
 		/** Set the id of the thread the debug substepping will be running in. */
-		void AssumeThisThread();
+		CHAOS_API void AssumeThisThread();
 
 		/**
 		 * Control substepping progress.
@@ -77,23 +77,23 @@ namespace Chaos
 		 * @param bIsSolverEnabled The state of the solver. The substepping will stop when the solver is disabled and delay starting again until re-enabled.
 		 * @return Whether the debug thread needs running.
 		 */
-		bool SyncAdvance(bool bIsSolverEnabled);
+		CHAOS_API bool SyncAdvance(bool bIsSolverEnabled);
 
 		/**
 		 * Add a new step or substep.
 		 * @param bInStep Add a step instead of a substep when this is true (for internal use only in the solver debug thread loop).
 		 * @param Label The reference (if any) that will be used in verbose logs when this point is reached, or nullptr otherwise.
 		 */
-		void Add(bool bInStep, const TCHAR* Label) const;
+		CHAOS_API void Add(bool bInStep, const TCHAR* Label) const;
 
 		/** Set sync events to start the substepping process. */
-		void Start();
+		CHAOS_API void Start();
 
 		/** Trigger/wait for sync events to stop the substepping process. */
-		void Stop();
+		CHAOS_API void Stop();
 
 		/** Trigger/wait for sync events to step/substep. */
-		void Substep(bool bShouldStep);
+		CHAOS_API void Substep(bool bShouldStep);
 
 	private:
 		enum class ECommand { Enable, Disable, ProgressToSubstep, ProgressToStep };
diff --git a/Engine/Source/Runtime/Experimental/Chaos/Public/Chaos/Framework/PhysicsProxyBase.h b/Engine/Source/Runtime/Experimental/Chaos/Public/Chaos/Framework/PhysicsProxyBase.h
--- a/Engine/Source/Runtime/Experimental/Chaos/Public/Chaos/Framework/PhysicsProxyBase.h
+++ b/Engine/Source/Runtime/Experimental/Chaos/Public/Chaos/Framework/PhysicsProxyBase.h
@@ -26,7 +26,7 @@ namespace Chaos
 	class FPhysicsSolverBase;
 }
 
-struct CHAOS_API FProxyTimestampBase
+struct FProxyTimestampBase
 {
 	bool bDeleted = false;
 };
@@ -44,7 +44,7 @@ struct TTimestampProperty
 	int32 Timestamp = INDEX_NONE;
 };
 
-struct CHAOS_API FSingleParticleProxyTimestamp: public FProxyTimestampBase
+struct FSingleParticleProxyTimestamp: public FProxyTimestampBase
 {
 	int32 ObjectStateTimestamp = INDEX_NONE;
 	TTimestampProperty<Chaos::FVec3> OverWriteX;
@@ -53,13 +53,13 @@ struct CHAOS_API FSingleParticleProxyTimestamp: public FProxyTimestampBase
 	TTimestampProperty<Chaos::FVec3> OverWriteW;
 };
 
-struct CHAOS_API FGeometryCollectionProxyTimestamp: public FProxyTimestampBase
+struct FGeometryCollectionProxyTimestamp: public FProxyTimestampBase
 {
 	// nothing to add as Geometry collections are driven from the Physics thread only
 	// ( including kinematic targeting )
 };
 
-struct CHAOS_API FClusterUnionProxyTimestamp : public FProxyTimestampBase
+struct FClusterUnionProxyTimestamp : public FProxyTimestampBase
 {
 	TTimestampProperty<Chaos::FVec3> OverWriteX;
 	TTimestampProperty<Chaos::FRotation3> OverWriteR;
@@ -67,7 +67,7 @@ struct CHAOS_API FClusterUnionProxyTimestamp : public FProxyTimestampBase
 	TTimestampProperty<Chaos::FVec3> OverWriteW;
 };
 
-class CHAOS_API IPhysicsProxyBase
+class IPhysicsProxyBase
 {
 public:
 	IPhysicsProxyBase(EPhysicsProxyType InType, UObject* InOwner, TSharedPtr<FProxyTimestampBase,ESPMode::ThreadSafe> InProxyTimeStamp)
@@ -124,7 +124,7 @@ public:
 protected:
 	// Ensures that derived classes can successfully call this destructor
 	// but no one can delete using a IPhysicsProxyBase*
-	virtual ~IPhysicsProxyBase();
+	CHAOS_API virtual ~IPhysicsProxyBase();
 
 	template<typename TProxyTimeStamp>
 	FORCEINLINE_DEBUGGABLE TProxyTimeStamp& GetSyncTimestampAs()
@@ -146,7 +146,7 @@ protected:
 	int32 InitializedOnStep = INDEX_NONE;
 	int32 IgnoreDataOnStep_Internal = INDEX_NONE;
 
-	int32 GetSolverSyncTimestamp_External() const;
+	CHAOS_API int32 GetSolverSyncTimestamp_External() const;
 };
 
 struct PhysicsProxyWrapper
@@ -304,4 +304,4 @@ protected:
 	Chaos::FVec3 ErrorVelocitySmoothingX = { 0,0,0 };
 	Chaos::FVec3 ErrorVelocitySmoothingXPrev = { 0,0,0 };
 	int32 ErrorVelocitySmoothingCount = 0;
-};
\ No newline at end of file
+};
diff --git a/Engine/Source/Runtime/Experimental/Chaos/Public/Chaos/Framework/PhysicsSolverBase.h b/Engine/Source/Runtime/Experimental/Chaos/Public/Chaos/Framework/PhysicsSolverBase.h
--- a/Engine/Source/Runtime/Experimental/Chaos/Public/Chaos/Framework/PhysicsSolverBase.h
+++ b/Engine/Source/Runtime/Experimental/Chaos/Public/Chaos/Framework/PhysicsSolverBase.h
@@ -51,7 +51,7 @@ namespace Chaos
 	extern CHAOS_API int32 ForceDisableAsyncPhysics;
 	extern CHAOS_API FRealSingle AsyncInterpolationMultiplier;
 
-	struct CHAOS_API FSubStepInfo
+	struct FSubStepInfo
 	{
 		FSubStepInfo()
 			: PseudoFraction(1.0)
@@ -89,7 +89,7 @@ namespace Chaos
 	/**
 	 * Task responsible for processing the command buffer of a single solver and preparing data before solver task and callbacks are run
 	 */
-	class CHAOS_API FPhysicsSolverProcessPushDataTask
+	class FPhysicsSolverProcessPushDataTask
 	{
 	public:
 
@@ -98,10 +98,10 @@ namespace Chaos
 			, PushData(InPushData){}
 
 		TStatId GetStatId() const { RETURN_QUICK_DECLARE_CYCLE_STAT(FPhysicsSolverProcessPushDataTask, STATGROUP_TaskGraphTasks); }
-		static ENamedThreads::Type GetDesiredThread();
+		static CHAOS_API ENamedThreads::Type GetDesiredThread();
 		static ESubsequentsMode::Type GetSubsequentsMode() { return ESubsequentsMode::TrackSubsequents; }
 		void DoTask(ENamedThreads::Type CurrentThread, const FGraphEventRef& MyCompletionGraphEvent) { ProcessPushData(); }
-		void ProcessPushData();
+		CHAOS_API void ProcessPushData();
 
 	private:
 
@@ -112,7 +112,7 @@ namespace Chaos
 	/**
 	 * Task responsible for triggering any game thread callbacks before solver can advance (not scheduled if none are registered)
 	 */
-	class CHAOS_API FPhysicsSolverFrozenGTPreSimCallbacks
+	class FPhysicsSolverFrozenGTPreSimCallbacks
 	{
 	public:
 
@@ -122,7 +122,7 @@ namespace Chaos
 		static ENamedThreads::Type GetDesiredThread() { return ENamedThreads::SetTaskPriority(ENamedThreads::GameThread, ENamedThreads::HighTaskPriority); }
 		static ESubsequentsMode::Type GetSubsequentsMode() { return ESubsequentsMode::TrackSubsequents; }
 		void DoTask(ENamedThreads::Type CurrentThread, const FGraphEventRef& MyCompletionGraphEvent) { GTPreSimCallbacks(); }
-		void GTPreSimCallbacks();
+		CHAOS_API void GTPreSimCallbacks();
 
 	private:
 
@@ -132,17 +132,17 @@ namespace Chaos
 	/**
 	 * Task responsible for advancing the solver once data has been prepared and GT callbacks have fired
 	 */
-	class CHAOS_API FPhysicsSolverAdvanceTask
+	class FPhysicsSolverAdvanceTask
 	{
 	public:
 
-		FPhysicsSolverAdvanceTask(FPhysicsSolverBase& InSolver, FPushPhysicsData* InPushData);
+		CHAOS_API FPhysicsSolverAdvanceTask(FPhysicsSolverBase& InSolver, FPushPhysicsData* InPushData);
 		
 		TStatId GetStatId() const { RETURN_QUICK_DECLARE_CYCLE_STAT(FPhysicsSolverAdvanceTask, STATGROUP_TaskGraphTasks); }
-		static ENamedThreads::Type GetDesiredThread();
+		static CHAOS_API ENamedThreads::Type GetDesiredThread();
 		static ESubsequentsMode::Type GetSubsequentsMode() { return ESubsequentsMode::TrackSubsequents; }
-		void DoTask(ENamedThreads::Type CurrentThread, const FGraphEventRef& MyCompletionGraphEvent);
-		void AdvanceSolver();
+		CHAOS_API void DoTask(ENamedThreads::Type CurrentThread, const FGraphEventRef& MyCompletionGraphEvent);
+		CHAOS_API void AdvanceSolver();
 
 	private:
 
@@ -159,7 +159,7 @@ namespace Chaos
 #endif
 	};
 
-	struct CHAOS_API FAllSolverTasks
+	struct FAllSolverTasks
 	{
 		FAllSolverTasks(FPhysicsSolverBase& InSolver, FPushPhysicsData* PushData)
 			: ProcessPushData(InSolver, PushData)
@@ -173,7 +173,7 @@ namespace Chaos
 		FPhysicsSolverFrozenGTPreSimCallbacks GTPreSimCallbacks;
 		FPhysicsSolverAdvanceTask AdvanceTask;
 
-		void AdvanceSolver();
+		CHAOS_API void AdvanceSolver();
 
 		FPhysicsSolverBase& Solver;
 	};
@@ -192,7 +192,7 @@ namespace Chaos
 	
 	/** Base solver class storing events that will be used by the derived solver during the solve.
 	 * Currently used by caching but could be used for any other pre/post solve works that need to be done.  */
-	class CHAOS_API FPhysicsSolverEvents
+	class FPhysicsSolverEvents
 	{
 
 	public:
@@ -207,20 +207,20 @@ namespace Chaos
 		/** Events */
 		/** WARNING: Events are not threadsafe!*/
 		/** Pre advance is called before any physics processing or simulation happens in a given physics update */
-		FDelegateHandle AddPreAdvanceCallback(FSolverPreAdvance::FDelegate InDelegate);
-		bool            RemovePreAdvanceCallback(FDelegateHandle InHandle);
+		CHAOS_API FDelegateHandle AddPreAdvanceCallback(FSolverPreAdvance::FDelegate InDelegate);
+		CHAOS_API bool            RemovePreAdvanceCallback(FDelegateHandle InHandle);
 
 		/** Pre buffer happens after the simulation has been advanced (particle positions etc. will have been updated) but GT results haven't been prepared yet */
-		FDelegateHandle AddPreBufferCallback(FSolverPreAdvance::FDelegate InDelegate);
-		bool            RemovePreBufferCallback(FDelegateHandle InHandle);
+		CHAOS_API FDelegateHandle AddPreBufferCallback(FSolverPreAdvance::FDelegate InDelegate);
+		CHAOS_API bool            RemovePreBufferCallback(FDelegateHandle InHandle);
 
 		/** Post advance happens after all processing and results generation has been completed */
-		FDelegateHandle AddPostAdvanceCallback(FSolverPostAdvance::FDelegate InDelegate);
-		bool            RemovePostAdvanceCallback(FDelegateHandle InHandle);
+		CHAOS_API FDelegateHandle AddPostAdvanceCallback(FSolverPostAdvance::FDelegate InDelegate);
+		CHAOS_API bool            RemovePostAdvanceCallback(FDelegateHandle InHandle);
 
 		/** Teardown happens as the solver is destroyed or streamed out */
-		FDelegateHandle AddTeardownCallback(FSolverTeardown::FDelegate InDelegate);
-		bool            RemoveTeardownCallback(FDelegateHandle InHandle);
+		CHAOS_API FDelegateHandle AddTeardownCallback(FSolverTeardown::FDelegate InDelegate);
+		CHAOS_API bool            RemoveTeardownCallback(FDelegateHandle InHandle);
 
 		/** Clear all the callbacks*/
 		void            ClearCallbacks() 
@@ -269,7 +269,7 @@ namespace Chaos
 	};
 #endif
 
-	class CHAOS_API FPhysicsSolverBase : public FPhysicsSolverEvents
+	class FPhysicsSolverBase : public FPhysicsSolverEvents
 	{
 	public:
 
@@ -284,7 +284,7 @@ namespace Chaos
 			return (FPBDRigidsSolver&)(*this);
 		}
 
-		void ChangeBufferMode(EMultiBufferMode InBufferMode);
+		CHAOS_API void ChangeBufferMode(EMultiBufferMode InBufferMode);
 
 		void AddDirtyProxy(IPhysicsProxyBase * ProxyBaseIn)
 		{
@@ -363,7 +363,7 @@ namespace Chaos
 			return CallbackObject;
 		}
 
-		void EnqueueSimcallbackRewindRegisteration(ISimCallbackObject* Callback);
+		CHAOS_API void EnqueueSimcallbackRewindRegisteration(ISimCallbackObject* Callback);
 
 		void UnregisterAndFreeSimCallbackObject_External(ISimCallbackObject* SimCallbackObject)
 		{
@@ -386,7 +386,7 @@ namespace Chaos
 			RegisterSimOneShotCallback(MoveTemp(Func));
 		}
 
-		void SetRewindCallback(TUniquePtr<IRewindCallback>&& RewindCallback);
+		CHAOS_API void SetRewindCallback(TUniquePtr<IRewindCallback>&& RewindCallback);
 
 		FRewindData* GetRewindData()
 		{
@@ -471,11 +471,11 @@ namespace Chaos
 
 		bool IsShuttingDown() const { return bIsShuttingDown; }
 
-		void EnableAsyncMode(FReal FixedDt);
+		CHAOS_API void EnableAsyncMode(FReal FixedDt);
 		
 		FReal GetAsyncDeltaTime() const { return AsyncDt; } 
 		
-		void DisableAsyncMode();
+		CHAOS_API void DisableAsyncMode();
 		
 		virtual void ConditionalApplyRewind_Internal(){}
 		virtual bool IsResimming() const {return false;}
@@ -488,7 +488,7 @@ namespace Chaos
 			return ThreadingMode;
 		}
 
-		FGraphEventRef AdvanceAndDispatch_External(FReal InDt);
+		CHAOS_API FGraphEventRef AdvanceAndDispatch_External(FReal InDt);
 
 #if CHAOS_DEBUG_NAME
 		void SetDebugName(const FName& Name)
@@ -541,7 +541,7 @@ namespace Chaos
 			}
 		}
 
-		void UpdateParticleInAccelerationStructure_External(FGeometryParticle* Particle, EPendingSpatialDataOperation InOperation);
+		CHAOS_API void UpdateParticleInAccelerationStructure_External(FGeometryParticle* Particle, EPendingSpatialDataOperation InOperation);
 
 		bool IsPaused_External() const
 		{
@@ -691,12 +691,12 @@ namespace Chaos
 		EThreadingModeTemp ThreadingMode;
 
 		/** Protected construction so callers still have to go through the module to create new instances */
-		FPhysicsSolverBase(const EMultiBufferMode BufferingModeIn,const EThreadingModeTemp InThreadingMode,UObject* InOwner, FReal AsyncDt);
+		CHAOS_API FPhysicsSolverBase(const EMultiBufferMode BufferingModeIn,const EThreadingModeTemp InThreadingMode,UObject* InOwner, FReal AsyncDt);
 
 		/** Only allow construction with valid parameters as well as restricting to module construction */
-		virtual ~FPhysicsSolverBase();
+		CHAOS_API virtual ~FPhysicsSolverBase();
 
-		static void DestroySolver(FPhysicsSolverBase& InSolver);
+		static CHAOS_API void DestroySolver(FPhysicsSolverBase& InSolver);
 
 		FPhysicsSolverBase() = delete;
 		FPhysicsSolverBase(const FPhysicsSolverBase& InCopy) = delete;
@@ -796,8 +796,8 @@ namespace Chaos
 
 	protected:
 
-		void TrackGTParticle_External(FGeometryParticle& Particle);
-		void ClearGTParticle_External(FGeometryParticle& Particle);
+		CHAOS_API void TrackGTParticle_External(FGeometryParticle& Particle);
+		CHAOS_API void ClearGTParticle_External(FGeometryParticle& Particle);
 		
 
 #if !UE_BUILD_SHIPPING
@@ -807,8 +807,8 @@ namespace Chaos
 		bool bStealAdvanceTasksForTesting;
 		TArray<FAllSolverTasks> StolenSolverAdvanceTasks;
 	public:
-		void SetStealAdvanceTasks_ForTesting(bool bInStealAdvanceTasksForTesting);
-		void PopAndExecuteStolenAdvanceTask_ForTesting();
+		CHAOS_API void SetStealAdvanceTasks_ForTesting(bool bInStealAdvanceTasksForTesting);
+		CHAOS_API void PopAndExecuteStolenAdvanceTask_ForTesting();
 #endif
 
 #if WITH_CHAOS_VISUAL_DEBUGGER
diff --git a/Engine/Source/Runtime/Experimental/Chaos/Public/Chaos/GeometryParticles.h b/Engine/Source/Runtime/Experimental/Chaos/Public/Chaos/GeometryParticles.h
--- a/Engine/Source/Runtime/Experimental/Chaos/Public/Chaos/GeometryParticles.h
+++ b/Engine/Source/Runtime/Experimental/Chaos/Public/Chaos/GeometryParticles.h
@@ -158,7 +158,7 @@ namespace Chaos
 
 		CHAOS_API static TGeometryParticlesImp<T, d, SimType>* SerializationFactory(FChaosArchive& Ar, TGeometryParticlesImp < T, d, SimType>* Particles);
 		
-		CHAOS_API TGeometryParticlesImp()
+		TGeometryParticlesImp()
 		    : TParticles<T, d>()
 		{
 			MParticleType = EParticleType::Static;
@@ -200,7 +200,7 @@ namespace Chaos
 
 		}
 		TGeometryParticlesImp(const TGeometryParticlesImp<T, d, SimType>& Other) = delete;
-		CHAOS_API TGeometryParticlesImp(TGeometryParticlesImp<T, d, SimType>&& Other)
+		TGeometryParticlesImp(TGeometryParticlesImp<T, d, SimType>&& Other)
 		    : TParticles<T, d>(MoveTemp(Other))
 			, MUniqueIdx(MoveTemp(Other.MUniqueIdx))
 			, MR(MoveTemp(Other.MR))
@@ -272,7 +272,7 @@ namespace Chaos
 
 		static constexpr bool IsRigidBodySim() { return SimType == EGeometryParticlesSimType::RigidBodySim; }
 
-		CHAOS_API TGeometryParticlesImp(TParticles<T, d>&& Other)
+		TGeometryParticlesImp(TParticles<T, d>&& Other)
 		    : TParticles<T, d>(MoveTemp(Other))
 		{
 			MParticleType = EParticleType::Static;
@@ -313,37 +313,37 @@ namespace Chaos
 			}
 		}
 
-		CHAOS_API virtual ~TGeometryParticlesImp()
+		virtual ~TGeometryParticlesImp()
 		{}
 
 		FORCEINLINE const TRotation<T, d>& R(const int32 Index) const { return MR[Index]; }
 		FORCEINLINE TRotation<T, d>& R(const int32 Index) { return MR[Index]; }
 
-		CHAOS_API FUniqueIdx UniqueIdx(const int32 Index) const { return MUniqueIdx[Index]; }
-		CHAOS_API FUniqueIdx& UniqueIdx(const int32 Index) { return MUniqueIdx[Index]; }
+		FUniqueIdx UniqueIdx(const int32 Index) const { return MUniqueIdx[Index]; }
+		FUniqueIdx& UniqueIdx(const int32 Index) { return MUniqueIdx[Index]; }
 
-		CHAOS_API ESyncState& SyncState(const int32 Index) { return MSyncState[Index].State; }
-		CHAOS_API ESyncState SyncState(const int32 Index) const { return MSyncState[Index].State; }
+		ESyncState& SyncState(const int32 Index) { return MSyncState[Index].State; }
+		ESyncState SyncState(const int32 Index) const { return MSyncState[Index].State; }
 
-		CHAOS_API TSerializablePtr<FImplicitObject> Geometry(const int32 Index) const { return MGeometry[Index]; }
+		TSerializablePtr<FImplicitObject> Geometry(const int32 Index) const { return MGeometry[Index]; }
 
-		CHAOS_API const TUniquePtr<FImplicitObject>& DynamicGeometry(const int32 Index) const { return MDynamicGeometry[Index]; }
+		const TUniquePtr<FImplicitObject>& DynamicGeometry(const int32 Index) const { return MDynamicGeometry[Index]; }
 
-		CHAOS_API const TSharedPtr<const FImplicitObject, ESPMode::ThreadSafe>& SharedGeometry(const int32 Index) const { return MSharedGeometry[Index]; }
+		const TSharedPtr<const FImplicitObject, ESPMode::ThreadSafe>& SharedGeometry(const int32 Index) const { return MSharedGeometry[Index]; }
 
-		CHAOS_API bool HasCollision(const int32 Index) const { return MHasCollision[Index]; }
-		CHAOS_API bool& HasCollision(const int32 Index) { return MHasCollision[Index]; }
+		bool HasCollision(const int32 Index) const { return MHasCollision[Index]; }
+		bool& HasCollision(const int32 Index) { return MHasCollision[Index]; }
 
-		CHAOS_API const FShapesArray& ShapesArray(const int32 Index) const { return reinterpret_cast<const FShapesArray&>(MShapesArray[Index]); }
+		const FShapesArray& ShapesArray(const int32 Index) const { return reinterpret_cast<const FShapesArray&>(MShapesArray[Index]); }
 
-		CHAOS_API const FShapeInstanceArray& ShapeInstances(const int32 Index) const { return MShapesArray[Index]; }
+		const FShapeInstanceArray& ShapeInstances(const int32 Index) const { return MShapesArray[Index]; }
 
 #if CHAOS_DETERMINISTIC
-		CHAOS_API FParticleID ParticleID(const int32 Idx) const { return MParticleIDs[Idx]; }
-		CHAOS_API FParticleID& ParticleID(const int32 Idx) { return MParticleIDs[Idx]; }
+		FParticleID ParticleID(const int32 Idx) const { return MParticleIDs[Idx]; }
+		FParticleID& ParticleID(const int32 Idx) { return MParticleIDs[Idx]; }
 #endif
 		// Set a dynamic geometry. Note that X and R must be initialized before calling this function.
-		CHAOS_API void SetDynamicGeometry(const int32 Index, TUniquePtr<FImplicitObject>&& InUnique)
+		void SetDynamicGeometry(const int32 Index, TUniquePtr<FImplicitObject>&& InUnique)
 		{
 			check(!SharedGeometry(Index));	// If shared geometry exists we should not be setting dynamic geometry on top
 			SetGeometryImpl(Index, MakeSerializable(InUnique));
@@ -351,14 +351,14 @@ namespace Chaos
 		}
 
 		// Set a shared geometry. Note that X and R must be initialized before calling this function.
-		CHAOS_API void SetSharedGeometry(const int32 Index, TSharedPtr<const FImplicitObject, ESPMode::ThreadSafe> InShared)
+		void SetSharedGeometry(const int32 Index, TSharedPtr<const FImplicitObject, ESPMode::ThreadSafe> InShared)
 		{
 			check(!DynamicGeometry(Index));	// If dynamic geometry exists we should not be setting shared geometry on top
 			SetGeometryImpl(Index, MakeSerializable(InShared));
 			MSharedGeometry[Index] = InShared;
 		}
 		
-		CHAOS_API void SetGeometry(const int32 Index, TSerializablePtr<FImplicitObject> InGeometry)
+		void SetGeometry(const int32 Index, TSerializablePtr<FImplicitObject> InGeometry)
 		{
 			check(!DynamicGeometry(Index));
 			check(!SharedGeometry(Index));
@@ -366,7 +366,7 @@ namespace Chaos
 		}
 
 	private:
-		CHAOS_API void SetGeometryImpl(const int32 Index, TSerializablePtr<FImplicitObject> InGeometry)
+		void SetGeometryImpl(const int32 Index, TSerializablePtr<FImplicitObject> InGeometry)
 		{
 			MGeometry[Index] = InGeometry;
 
@@ -444,37 +444,37 @@ namespace Chaos
 		}
 	public:
 
-		CHAOS_API const TAABB<T,d>& LocalBounds(const int32 Index) const
+		const TAABB<T,d>& LocalBounds(const int32 Index) const
 		{
 			return MLocalBounds[Index];
 		}
 
-		CHAOS_API TAABB<T, d>& LocalBounds(const int32 Index)
+		TAABB<T, d>& LocalBounds(const int32 Index)
 		{
 			return MLocalBounds[Index];
 		}
 
-		CHAOS_API const TVector<T,d>& CCDAxisThreshold(const int32 Index) const
+		const TVector<T,d>& CCDAxisThreshold(const int32 Index) const
 		{
 			return MCCDAxisThreshold[Index];
 		}
 
-		CHAOS_API bool HasBounds(const int32 Index) const
+		bool HasBounds(const int32 Index) const
 		{
 			return MHasBounds[Index];
 		}
 
-		CHAOS_API bool& HasBounds(const int32 Index)
+		bool& HasBounds(const int32 Index)
 		{
 			return MHasBounds[Index];
 		}
 
-		CHAOS_API FSpatialAccelerationIdx SpatialIdx(const int32 Index) const
+		FSpatialAccelerationIdx SpatialIdx(const int32 Index) const
 		{
 			return MSpatialIdx[Index];
 		}
 
-		CHAOS_API FSpatialAccelerationIdx& SpatialIdx(const int32 Index)
+		FSpatialAccelerationIdx& SpatialIdx(const int32 Index)
 		{
 			return MSpatialIdx[Index];
 		}
@@ -491,12 +491,12 @@ namespace Chaos
 		}
 #endif
 
-		CHAOS_API const TAABB<T, d>& WorldSpaceInflatedBounds(const int32 Index) const
+		const TAABB<T, d>& WorldSpaceInflatedBounds(const int32 Index) const
 		{
 			return MWorldSpaceInflatedBounds[Index];
 		}
 
-		CHAOS_API void UpdateWorldSpaceState(const int32 Index, const FRigidTransform3& WorldTransform, const FVec3& BoundsExpansion)
+		void UpdateWorldSpaceState(const int32 Index, const FRigidTransform3& WorldTransform, const FVec3& BoundsExpansion)
 		{
 			const FShapesArray& Shapes = ShapesArray(Index);
 
@@ -526,17 +526,17 @@ namespace Chaos
 
 		CHAOS_API void SetHandle(int32 Index, FGeometryParticleHandle* Handle);
 		
-		CHAOS_API FGeometryParticle* GTGeometryParticle(const int32 Index) const { return MGeometryParticle[Index]; }
-		CHAOS_API FGeometryParticle*& GTGeometryParticle(const int32 Index) { return MGeometryParticle[Index]; }
+		FGeometryParticle* GTGeometryParticle(const int32 Index) const { return MGeometryParticle[Index]; }
+		FGeometryParticle*& GTGeometryParticle(const int32 Index) { return MGeometryParticle[Index]; }
 
-		CHAOS_API const IPhysicsProxyBase* PhysicsProxy(const int32 Index) const { return MPhysicsProxy[Index];  }
-		CHAOS_API IPhysicsProxyBase* PhysicsProxy(const int32 Index) { return MPhysicsProxy[Index]; }
-		CHAOS_API void SetPhysicsProxy(const int32 Index, IPhysicsProxyBase* InPhysicsProxy)
+		const IPhysicsProxyBase* PhysicsProxy(const int32 Index) const { return MPhysicsProxy[Index];  }
+		IPhysicsProxyBase* PhysicsProxy(const int32 Index) { return MPhysicsProxy[Index]; }
+		void SetPhysicsProxy(const int32 Index, IPhysicsProxyBase* InPhysicsProxy)
 		{
 			MPhysicsProxy[Index] = InPhysicsProxy;
 		}
 
-		CHAOS_API FWeakParticleHandle& WeakParticleHandle(const int32 Index)
+		FWeakParticleHandle& WeakParticleHandle(const int32 Index)
 		{
 			FWeakParticleHandle& WeakHandle = MWeakParticleHandle[Index];
 			if(WeakHandle.IsInitialized())
@@ -551,19 +551,19 @@ namespace Chaos
 		/**
 		 * @brief All of the persistent (non-collision) constraints affecting the particle
 		*/
-		CHAOS_API FConstraintHandleArray& ParticleConstraints(const int32 Index)
+		FConstraintHandleArray& ParticleConstraints(const int32 Index)
 		{
 			return MParticleConstraints[Index];
 		}
 
-		CHAOS_API void AddConstraintHandle(const int32& Index, FConstraintHandle* InConstraintHandle)
+		void AddConstraintHandle(const int32& Index, FConstraintHandle* InConstraintHandle)
 		{
 			CHAOS_ENSURE(!MParticleConstraints[Index].Contains(InConstraintHandle));
 			MParticleConstraints[Index].Add(InConstraintHandle);
 		}
 
 
-		CHAOS_API void RemoveConstraintHandle(const int32& Index, FConstraintHandle* InConstraintHandle)
+		void RemoveConstraintHandle(const int32& Index, FConstraintHandle* InConstraintHandle)
 		{
 			MParticleConstraints[Index].RemoveSingleSwap(InConstraintHandle);
 			CHAOS_ENSURE(!MParticleConstraints[Index].Contains(InConstraintHandle));
@@ -572,7 +572,7 @@ namespace Chaos
 		/**
 		 * @brief All of the collision constraints affecting the particle
 		*/
-		CHAOS_API FParticleCollisions& ParticleCollisions(const int32 Index)
+		FParticleCollisions& ParticleCollisions(const int32 Index)
 		{
 			return MParticleCollisions[Index];
 		}
@@ -596,7 +596,7 @@ namespace Chaos
 
 private:
 		friend THandleType;
-		CHAOS_API void ResetWeakParticleHandle(const int32 Index)
+		void ResetWeakParticleHandle(const int32 Index)
 		{
 			FWeakParticleHandle& WeakHandle = MWeakParticleHandle[Index];
 			if(WeakHandle.IsInitialized())
@@ -612,7 +612,7 @@ public:
 			return FString::Printf(TEXT("%s, MUniqueIdx:%d MR:%s, MGeometry:%s, IsDynamic:%d"), *BaseString, UniqueIdx(index).Idx, *R(index).ToString(), (Geometry(index) ? *(Geometry(index)->ToString()) : TEXT("none")), (DynamicGeometry(index) != nullptr));
 		}
 
-		CHAOS_API virtual void Serialize(FChaosArchive& Ar)
+		virtual void Serialize(FChaosArchive& Ar)
 		{
 			LLM_SCOPE(ELLMTag::ChaosParticles);
 			TParticles<T, d>::Serialize(Ar);
diff --git a/Engine/Source/Runtime/Experimental/Chaos/Public/Chaos/GraphColoring.h b/Engine/Source/Runtime/Experimental/Chaos/Public/Chaos/GraphColoring.h
--- a/Engine/Source/Runtime/Experimental/Chaos/Public/Chaos/GraphColoring.h
+++ b/Engine/Source/Runtime/Experimental/Chaos/Public/Chaos/GraphColoring.h
@@ -62,28 +62,28 @@ class FGraphColoring
 
   public:
 	template<typename T>
-	CHAOS_API static TArray<TArray<int32>> ComputeGraphColoring(const TArray<TVector<int32, 2>>& Graph, const TDynamicParticles<T, 3>& InParticles, const int32 GraphParticlesStart, const int32 GraphParticlesEnd);
+	static TArray<TArray<int32>> ComputeGraphColoring(const TArray<TVector<int32, 2>>& Graph, const TDynamicParticles<T, 3>& InParticles, const int32 GraphParticlesStart, const int32 GraphParticlesEnd);
 	template<typename T>
 	inline static TArray<TArray<int32>> ComputeGraphColoring(const TArray<TVector<int32, 2>>& Graph, const TDynamicParticles<T, 3>& InParticles)
 	{
 		return ComputeGraphColoring(Graph, InParticles, 0, InParticles.Size());
 	}
 	template<typename T>
-	CHAOS_API static TArray<TArray<int32>> ComputeGraphColoring(const TArray<TVector<int32, 3>>& Graph, const TDynamicParticles<T, 3>& InParticles, const int32 GraphParticlesStart, const int32 GraphParticlesEnd);
+	static TArray<TArray<int32>> ComputeGraphColoring(const TArray<TVector<int32, 3>>& Graph, const TDynamicParticles<T, 3>& InParticles, const int32 GraphParticlesStart, const int32 GraphParticlesEnd);
 	template<typename T>
 	inline static TArray<TArray<int32>> ComputeGraphColoring(const TArray<TVector<int32, 3>>& Graph, const TDynamicParticles<T, 3>& InParticles)
 	{
 		return ComputeGraphColoring(Graph, InParticles, 0, InParticles.Size());
 	}
 	template<typename T>
-	CHAOS_API static TArray<TArray<int32>> ComputeGraphColoring(const TArray<TVector<int32, 4>>& Graph, const TDynamicParticles<T, 3>& InParticles, const int32 GraphParticlesStart, const int32 GraphParticlesEnd);
+	static TArray<TArray<int32>> ComputeGraphColoring(const TArray<TVector<int32, 4>>& Graph, const TDynamicParticles<T, 3>& InParticles, const int32 GraphParticlesStart, const int32 GraphParticlesEnd);
 	template<typename T>
 	inline static TArray<TArray<int32>> ComputeGraphColoring(const TArray<TVector<int32, 4>>& Graph, const TDynamicParticles<T, 3>& InParticles)
 	{
 		return ComputeGraphColoring(Graph, InParticles, 0, InParticles.Size());
 	}
 	template<typename T>
-	CHAOS_API static TArray<TArray<int32>> ComputeGraphColoringAllDynamic(const TArray<TVec4<int32>>& Graph, const Chaos::TDynamicParticles<T, 3>& InParticles, const int32 GraphParticlesStart, const int32 GraphParticlesEnd);
+	static TArray<TArray<int32>> ComputeGraphColoringAllDynamic(const TArray<TVec4<int32>>& Graph, const Chaos::TDynamicParticles<T, 3>& InParticles, const int32 GraphParticlesStart, const int32 GraphParticlesEnd);
 	template<typename T>
 	inline static TArray<TArray<int32>> ComputeGraphColoringAllDynamic(const TArray<TVec4<int32>>& Graph, const Chaos::TDynamicParticles<T, 3>& InParticles)
 	{
