[Chaos]

diff --git a/Engine/Source/Runtime/Experimental/Chaos/Public/Chaos/HeightField.h b/Engine/Source/Runtime/Experimental/Chaos/Public/Chaos/HeightField.h
--- a/Engine/Source/Runtime/Experimental/Chaos/Public/Chaos/HeightField.h
+++ b/Engine/Source/Runtime/Experimental/Chaos/Public/Chaos/HeightField.h
@@ -23,13 +23,13 @@ namespace Chaos
 
 namespace Chaos
 {
-	class CHAOS_API FHeightField final : public FImplicitObject
+	class FHeightField final : public FImplicitObject
 	{
 	public:
 		using FImplicitObject::GetTypeName;
 
-		FHeightField(TArray<FReal>&& Height, TArray<uint8>&& InMaterialIndices, int32 InNumRows, int32 InNumCols, const FVec3& InScale);
-		FHeightField(TArrayView<const uint16> InHeights, TArrayView<const uint8> InMaterialIndices, int32 InNumRows, int32 InNumCols, const FVec3& InScale);
+		CHAOS_API FHeightField(TArray<FReal>&& Height, TArray<uint8>&& InMaterialIndices, int32 InNumRows, int32 InNumCols, const FVec3& InScale);
+		CHAOS_API FHeightField(TArrayView<const uint16> InHeights, TArrayView<const uint8> InMaterialIndices, int32 InNumRows, int32 InNumCols, const FVec3& InScale);
 		FHeightField(const FHeightField& Other) = delete;
 		
 		// Not required as long as FImplicitObject also has deleted move constructor (adding this causes an error on Linux build)
@@ -38,66 +38,66 @@ namespace Chaos
 		virtual ~FHeightField() {}
 
 		/** Support for editing a subsection of the heightfield */
-		void EditHeights(TArrayView<FReal> InHeights, int32 InBeginRow, int32 InBeginCol, int32 InNumRows, int32 InNumCols);
-		void EditHeights(TArrayView<const uint16> InHeights, int32 InBeginRow, int32 InBeginCol, int32 InNumRows, int32 InNumCols);
-		FReal GetHeight(int32 InIndex) const;
-		FVec3 GetPointScaled(int32 InIndex) const;
-		FReal GetHeight(int32 InX, int32 InY) const;
-		uint8 GetMaterialIndex(int32 InIndex) const;
-		uint8 GetMaterialIndex(int32 InX, int32 InY) const;
-		bool IsHole(int32 InIndex) const;
-		bool IsHole(int32 InCellX, int32 InCellY) const;
-		FVec3 GetNormalAt(const FVec2& InGridLocationLocal) const;
-		FReal GetHeightAt(const FVec2& InGridLocationLocal) const;
+		CHAOS_API void EditHeights(TArrayView<FReal> InHeights, int32 InBeginRow, int32 InBeginCol, int32 InNumRows, int32 InNumCols);
+		CHAOS_API void EditHeights(TArrayView<const uint16> InHeights, int32 InBeginRow, int32 InBeginCol, int32 InNumRows, int32 InNumCols);
+		CHAOS_API FReal GetHeight(int32 InIndex) const;
+		CHAOS_API FVec3 GetPointScaled(int32 InIndex) const;
+		CHAOS_API FReal GetHeight(int32 InX, int32 InY) const;
+		CHAOS_API uint8 GetMaterialIndex(int32 InIndex) const;
+		CHAOS_API uint8 GetMaterialIndex(int32 InX, int32 InY) const;
+		CHAOS_API bool IsHole(int32 InIndex) const;
+		CHAOS_API bool IsHole(int32 InCellX, int32 InCellY) const;
+		CHAOS_API FVec3 GetNormalAt(const FVec2& InGridLocationLocal) const;
+		CHAOS_API FReal GetHeightAt(const FVec2& InGridLocationLocal) const;
 
 		int32 GetNumRows() const { return GeomData.NumRows; }
 		int32 GetNumCols() const { return GeomData.NumCols; }
 
-		virtual FReal PhiWithNormal(const FVec3& x, FVec3& Normal) const;
+		CHAOS_API virtual FReal PhiWithNormal(const FVec3& x, FVec3& Normal) const;
 
-		virtual bool Raycast(const FVec3& StartPoint, const FVec3& Dir, const FReal Length, const FReal Thickness, FReal& OutTime, FVec3& OutPosition, FVec3& OutNormal, int32& OutFaceIndex) const override;
-		virtual bool Overlap(const FVec3& Point, const FReal Thickness) const override;
+		CHAOS_API virtual bool Raycast(const FVec3& StartPoint, const FVec3& Dir, const FReal Length, const FReal Thickness, FReal& OutTime, FVec3& OutPosition, FVec3& OutNormal, int32& OutFaceIndex) const override;
+		CHAOS_API virtual bool Overlap(const FVec3& Point, const FReal Thickness) const override;
 		
-		bool OverlapGeom(const TSphere<FReal, 3>& QueryGeom, const FRigidTransform3& QueryTM, const FReal Thickness, FMTDInfo* OutMTD = nullptr) const;
-		bool OverlapGeom(const TBox<FReal, 3>& QueryGeom, const FRigidTransform3& QueryTM, const FReal Thickness, FMTDInfo* OutMTD = nullptr) const;
-		bool OverlapGeom(const FCapsule& QueryGeom, const FRigidTransform3& QueryTM, const FReal Thickness, FMTDInfo* OutMTD = nullptr) const;
-		bool OverlapGeom(const FConvex& QueryGeom, const FRigidTransform3& QueryTM, const FReal Thickness, FMTDInfo* OutMTD = nullptr) const;
-		bool OverlapGeom(const TImplicitObjectScaled<TSphere<FReal, 3>>& QueryGeom, const FRigidTransform3& QueryTM, const FReal Thickness, FMTDInfo* OutMTD = nullptr) const;
-		bool OverlapGeom(const TImplicitObjectScaled<TBox<FReal, 3>>& QueryGeom, const FRigidTransform3& QueryTM, const FReal Thickness, FMTDInfo* OutMTD = nullptr) const;
-		bool OverlapGeom(const TImplicitObjectScaled<FCapsule>& QueryGeom, const FRigidTransform3& QueryTM, const FReal Thickness, FMTDInfo* OutMTD = nullptr) const;
-		bool OverlapGeom(const TImplicitObjectScaled<FConvex>& QueryGeom, const FRigidTransform3& QueryTM, const FReal Thickness, FMTDInfo* OutMTD = nullptr) const;
-
-		bool SweepGeom(const TSphere<FReal, 3>& QueryGeom, const FRigidTransform3& StartTM, const FVec3& Dir, const FReal Length, FReal& OutTime, FVec3& OutPosition, FVec3& OutNormal, int32& OutFaceIndex, FVec3& OutFaceNormal, const FReal Thickness = 0, bool bComputeMTD = false) const;
-		bool SweepGeom(const TBox<FReal, 3>& QueryGeom, const FRigidTransform3& StartTM, const FVec3& Dir, const FReal Length, FReal& OutTime, FVec3& OutPosition, FVec3& OutNormal, int32& OutFaceIndex, FVec3& OutFaceNormal, const FReal Thickness = 0, bool bComputeMTD = false) const;
-		bool SweepGeom(const FCapsule& QueryGeom, const FRigidTransform3& StartTM, const FVec3& Dir, const FReal Length, FReal& OutTime, FVec3& OutPosition, FVec3& OutNormal, int32& OutFaceIndex, FVec3& OutFaceNormal, const FReal Thickness = 0, bool bComputeMTD = false) const;
-		bool SweepGeom(const FConvex& QueryGeom, const FRigidTransform3& StartTM, const FVec3& Dir, const FReal Length, FReal& OutTime, FVec3& OutPosition, FVec3& OutNormal, int32& OutFaceIndex, FVec3& OutFaceNormal, const FReal Thickness = 0, bool bComputeMTD = false) const;
-		bool SweepGeom(const TImplicitObjectScaled<TSphere<FReal, 3>>& QueryGeom, const FRigidTransform3& StartTM, const FVec3& Dir, const FReal Length, FReal& OutTime, FVec3& OutPosition, FVec3& OutNormal, int32& OutFaceIndex, FVec3& OutFaceNormal, const FReal Thickness = 0, bool bComputeMTD = false) const;
-		bool SweepGeom(const TImplicitObjectScaled<TBox<FReal, 3>>& QueryGeom, const FRigidTransform3& StartTM, const FVec3& Dir, const FReal Length, FReal& OutTime, FVec3& OutPosition, FVec3& OutNormal, int32& OutFaceIndex, FVec3& OutFaceNormal, const FReal Thickness = 0, bool bComputeMTD = false) const;
-		bool SweepGeom(const TImplicitObjectScaled<FCapsule>& QueryGeom, const FRigidTransform3& StartTM, const FVec3& Dir, const FReal Length, FReal& OutTime, FVec3& OutPosition, FVec3& OutNormal, int32& OutFaceIndex, FVec3& OutFaceNormal, const FReal Thickness = 0, bool bComputeMTD = false) const;
-		bool SweepGeom(const TImplicitObjectScaled<FConvex>& QueryGeom, const FRigidTransform3& StartTM, const FVec3& Dir, const FReal Length, FReal& OutTime, FVec3& OutPosition, FVec3& OutNormal, int32& OutFaceIndex, FVec3& OutFaceNormal, const FReal Thickness = 0, bool bComputeMTD = false) const;
+		CHAOS_API bool OverlapGeom(const TSphere<FReal, 3>& QueryGeom, const FRigidTransform3& QueryTM, const FReal Thickness, FMTDInfo* OutMTD = nullptr) const;
+		CHAOS_API bool OverlapGeom(const TBox<FReal, 3>& QueryGeom, const FRigidTransform3& QueryTM, const FReal Thickness, FMTDInfo* OutMTD = nullptr) const;
+		CHAOS_API bool OverlapGeom(const FCapsule& QueryGeom, const FRigidTransform3& QueryTM, const FReal Thickness, FMTDInfo* OutMTD = nullptr) const;
+		CHAOS_API bool OverlapGeom(const FConvex& QueryGeom, const FRigidTransform3& QueryTM, const FReal Thickness, FMTDInfo* OutMTD = nullptr) const;
+		CHAOS_API bool OverlapGeom(const TImplicitObjectScaled<TSphere<FReal, 3>>& QueryGeom, const FRigidTransform3& QueryTM, const FReal Thickness, FMTDInfo* OutMTD = nullptr) const;
+		CHAOS_API bool OverlapGeom(const TImplicitObjectScaled<TBox<FReal, 3>>& QueryGeom, const FRigidTransform3& QueryTM, const FReal Thickness, FMTDInfo* OutMTD = nullptr) const;
+		CHAOS_API bool OverlapGeom(const TImplicitObjectScaled<FCapsule>& QueryGeom, const FRigidTransform3& QueryTM, const FReal Thickness, FMTDInfo* OutMTD = nullptr) const;
+		CHAOS_API bool OverlapGeom(const TImplicitObjectScaled<FConvex>& QueryGeom, const FRigidTransform3& QueryTM, const FReal Thickness, FMTDInfo* OutMTD = nullptr) const;
+
+		CHAOS_API bool SweepGeom(const TSphere<FReal, 3>& QueryGeom, const FRigidTransform3& StartTM, const FVec3& Dir, const FReal Length, FReal& OutTime, FVec3& OutPosition, FVec3& OutNormal, int32& OutFaceIndex, FVec3& OutFaceNormal, const FReal Thickness = 0, bool bComputeMTD = false) const;
+		CHAOS_API bool SweepGeom(const TBox<FReal, 3>& QueryGeom, const FRigidTransform3& StartTM, const FVec3& Dir, const FReal Length, FReal& OutTime, FVec3& OutPosition, FVec3& OutNormal, int32& OutFaceIndex, FVec3& OutFaceNormal, const FReal Thickness = 0, bool bComputeMTD = false) const;
+		CHAOS_API bool SweepGeom(const FCapsule& QueryGeom, const FRigidTransform3& StartTM, const FVec3& Dir, const FReal Length, FReal& OutTime, FVec3& OutPosition, FVec3& OutNormal, int32& OutFaceIndex, FVec3& OutFaceNormal, const FReal Thickness = 0, bool bComputeMTD = false) const;
+		CHAOS_API bool SweepGeom(const FConvex& QueryGeom, const FRigidTransform3& StartTM, const FVec3& Dir, const FReal Length, FReal& OutTime, FVec3& OutPosition, FVec3& OutNormal, int32& OutFaceIndex, FVec3& OutFaceNormal, const FReal Thickness = 0, bool bComputeMTD = false) const;
+		CHAOS_API bool SweepGeom(const TImplicitObjectScaled<TSphere<FReal, 3>>& QueryGeom, const FRigidTransform3& StartTM, const FVec3& Dir, const FReal Length, FReal& OutTime, FVec3& OutPosition, FVec3& OutNormal, int32& OutFaceIndex, FVec3& OutFaceNormal, const FReal Thickness = 0, bool bComputeMTD = false) const;
+		CHAOS_API bool SweepGeom(const TImplicitObjectScaled<TBox<FReal, 3>>& QueryGeom, const FRigidTransform3& StartTM, const FVec3& Dir, const FReal Length, FReal& OutTime, FVec3& OutPosition, FVec3& OutNormal, int32& OutFaceIndex, FVec3& OutFaceNormal, const FReal Thickness = 0, bool bComputeMTD = false) const;
+		CHAOS_API bool SweepGeom(const TImplicitObjectScaled<FCapsule>& QueryGeom, const FRigidTransform3& StartTM, const FVec3& Dir, const FReal Length, FReal& OutTime, FVec3& OutPosition, FVec3& OutNormal, int32& OutFaceIndex, FVec3& OutFaceNormal, const FReal Thickness = 0, bool bComputeMTD = false) const;
+		CHAOS_API bool SweepGeom(const TImplicitObjectScaled<FConvex>& QueryGeom, const FRigidTransform3& StartTM, const FVec3& Dir, const FReal Length, FReal& OutTime, FVec3& OutPosition, FVec3& OutNormal, int32& OutFaceIndex, FVec3& OutFaceNormal, const FReal Thickness = 0, bool bComputeMTD = false) const;
 
 		// Sweep used for CCD. Ignores triangles we penetrate by less than IgnorePenetration, and calculate the TOI for a depth of TargetPenetration. If both are zero, this is equivalent to SweepGeom
-		bool SweepGeomCCD(const TSphere<FReal, 3>& QueryGeom, const FRigidTransform3& StartTM, const FVec3& Dir, const FReal Length, const FReal IgnorePenetration, const FReal TargetPenetration, FReal& OutTOI, FReal& OutPhi, FVec3& OutPosition, FVec3& OutNormal, int32& OutFaceIndex, FVec3& OutFaceNormal) const;
-		bool SweepGeomCCD(const TBox<FReal, 3>& QueryGeom, const FRigidTransform3& StartTM, const FVec3& Dir, const FReal Length, const FReal IgnorePenetration, const FReal TargetPenetration, FReal& OutTOI, FReal& OutPhi, FVec3& OutPosition, FVec3& OutNormal, int32& OutFaceIndex, FVec3& OutFaceNormal) const;
-		bool SweepGeomCCD(const FCapsule& QueryGeom, const FRigidTransform3& StartTM, const FVec3& Dir, const FReal Length, const FReal IgnorePenetration, const FReal TargetPenetration, FReal& OutTOI, FReal& OutPhi, FVec3& OutPosition, FVec3& OutNormal, int32& OutFaceIndex, FVec3& OutFaceNormal) const;
-		bool SweepGeomCCD(const FConvex& QueryGeom, const FRigidTransform3& StartTM, const FVec3& Dir, const FReal Length, const FReal IgnorePenetration, const FReal TargetPenetration, FReal& OutTOI, FReal& OutPhi, FVec3& OutPosition, FVec3& OutNormal, int32& OutFaceIndex, FVec3& OutFaceNormal) const;
-		bool SweepGeomCCD(const TImplicitObjectScaled<TSphere<FReal, 3>>& QueryGeom, const FRigidTransform3& StartTM, const FVec3& Dir, const FReal Length, const FReal IgnorePenetration, const FReal TargetPenetration, FReal& OutTOI, FReal& OutPhi, FVec3& OutPosition, FVec3& OutNormal, int32& OutFaceIndex, FVec3& OutFaceNormal) const;
-		bool SweepGeomCCD(const TImplicitObjectScaled<TBox<FReal, 3>>& QueryGeom, const FRigidTransform3& StartTM, const FVec3& Dir, const FReal Length, const FReal IgnorePenetration, const FReal TargetPenetration, FReal& OutTOI, FReal& OutPhi, FVec3& OutPosition, FVec3& OutNormal, int32& OutFaceIndex, FVec3& OutFaceNormal) const;
-		bool SweepGeomCCD(const TImplicitObjectScaled<FCapsule>& QueryGeom, const FRigidTransform3& StartTM, const FVec3& Dir, const FReal Length, const FReal IgnorePenetration, const FReal TargetPenetration, FReal& OutTOI, FReal& OutPhi, FVec3& OutPosition, FVec3& OutNormal, int32& OutFaceIndex, FVec3& OutFaceNormal) const;
-		bool SweepGeomCCD(const TImplicitObjectScaled<FConvex>& QueryGeom, const FRigidTransform3& StartTM, const FVec3& Dir, const FReal Length, const FReal IgnorePenetration, const FReal TargetPenetration, FReal& OutTOI, FReal& OutPhi, FVec3& OutPosition, FVec3& OutNormal, int32& OutFaceIndex, FVec3& OutFaceNormal) const;
-
-		bool GJKContactPoint(const TBox<FReal, 3>& QueryGeom, const FRigidTransform3& QueryTM, const FReal Thickness, FVec3& ContactLocation, FVec3& ContactNormal, FReal& ContactPhi, int32& ContactFaceIndex) const;
-		bool GJKContactPoint(const TSphere<FReal, 3>& QueryGeom, const FRigidTransform3& QueryTM, const FReal Thickness, FVec3& ContactLocation, FVec3& ContactNormal, FReal& ContactPhi, int32& ContactFaceIndex) const;
-		bool GJKContactPoint(const FCapsule& QueryGeom, const FRigidTransform3& QueryTM, const FReal Thickness, FVec3& ContactLocation, FVec3& ContactNormal, FReal& ContactPhi, int32& ContactFaceIndex) const;
-		bool GJKContactPoint(const FConvex& QueryGeom, const FRigidTransform3& QueryTM, const FReal Thickness, FVec3& ContactLocation, FVec3& ContactNormal, FReal& ContactPhi, int32& ContactFaceIndex) const;
-		bool GJKContactPoint(const TImplicitObjectScaled<TBox<FReal, 3>>& QueryGeom, const FRigidTransform3& QueryTM, const FReal Thickness, FVec3& ContactLocation, FVec3& ContactNormal, FReal& ContactPhi, int32& ContactFaceIndex) const;
-		bool GJKContactPoint(const TImplicitObjectScaled<TSphere<FReal, 3>>& QueryGeom, const FRigidTransform3& QueryTM, const FReal Thickness, FVec3& ContactLocation, FVec3& ContactNormal, FReal& ContactPhi, int32& ContactFaceIndex) const;
-		bool GJKContactPoint(const TImplicitObjectScaled<FCapsule>& QueryGeom, const FRigidTransform3& QueryTM, const FReal Thickness, FVec3& ContactLocation, FVec3& ContactNormal, FReal& ContactPhi, int32& ContactFaceIndex) const;
-		bool GJKContactPoint(const TImplicitObjectScaled<FConvex>& QueryGeom, const FRigidTransform3& QueryTM, const FReal Thickness, FVec3& ContactLocation, FVec3& ContactNormal, FReal& ContactPhi, int32& ContactFaceIndex) const;
-
-
-		virtual int32 FindMostOpposingFace(const FVec3& Position, const FVec3& UnitDir, int32 HintFaceIndex, FReal SearchDist) const override;
-		virtual FVec3 FindGeometryOpposingNormal(const FVec3& DenormDir, int32 FaceIndex, const FVec3& OriginalNormal) const override;
+		CHAOS_API bool SweepGeomCCD(const TSphere<FReal, 3>& QueryGeom, const FRigidTransform3& StartTM, const FVec3& Dir, const FReal Length, const FReal IgnorePenetration, const FReal TargetPenetration, FReal& OutTOI, FReal& OutPhi, FVec3& OutPosition, FVec3& OutNormal, int32& OutFaceIndex, FVec3& OutFaceNormal) const;
+		CHAOS_API bool SweepGeomCCD(const TBox<FReal, 3>& QueryGeom, const FRigidTransform3& StartTM, const FVec3& Dir, const FReal Length, const FReal IgnorePenetration, const FReal TargetPenetration, FReal& OutTOI, FReal& OutPhi, FVec3& OutPosition, FVec3& OutNormal, int32& OutFaceIndex, FVec3& OutFaceNormal) const;
+		CHAOS_API bool SweepGeomCCD(const FCapsule& QueryGeom, const FRigidTransform3& StartTM, const FVec3& Dir, const FReal Length, const FReal IgnorePenetration, const FReal TargetPenetration, FReal& OutTOI, FReal& OutPhi, FVec3& OutPosition, FVec3& OutNormal, int32& OutFaceIndex, FVec3& OutFaceNormal) const;
+		CHAOS_API bool SweepGeomCCD(const FConvex& QueryGeom, const FRigidTransform3& StartTM, const FVec3& Dir, const FReal Length, const FReal IgnorePenetration, const FReal TargetPenetration, FReal& OutTOI, FReal& OutPhi, FVec3& OutPosition, FVec3& OutNormal, int32& OutFaceIndex, FVec3& OutFaceNormal) const;
+		CHAOS_API bool SweepGeomCCD(const TImplicitObjectScaled<TSphere<FReal, 3>>& QueryGeom, const FRigidTransform3& StartTM, const FVec3& Dir, const FReal Length, const FReal IgnorePenetration, const FReal TargetPenetration, FReal& OutTOI, FReal& OutPhi, FVec3& OutPosition, FVec3& OutNormal, int32& OutFaceIndex, FVec3& OutFaceNormal) const;
+		CHAOS_API bool SweepGeomCCD(const TImplicitObjectScaled<TBox<FReal, 3>>& QueryGeom, const FRigidTransform3& StartTM, const FVec3& Dir, const FReal Length, const FReal IgnorePenetration, const FReal TargetPenetration, FReal& OutTOI, FReal& OutPhi, FVec3& OutPosition, FVec3& OutNormal, int32& OutFaceIndex, FVec3& OutFaceNormal) const;
+		CHAOS_API bool SweepGeomCCD(const TImplicitObjectScaled<FCapsule>& QueryGeom, const FRigidTransform3& StartTM, const FVec3& Dir, const FReal Length, const FReal IgnorePenetration, const FReal TargetPenetration, FReal& OutTOI, FReal& OutPhi, FVec3& OutPosition, FVec3& OutNormal, int32& OutFaceIndex, FVec3& OutFaceNormal) const;
+		CHAOS_API bool SweepGeomCCD(const TImplicitObjectScaled<FConvex>& QueryGeom, const FRigidTransform3& StartTM, const FVec3& Dir, const FReal Length, const FReal IgnorePenetration, const FReal TargetPenetration, FReal& OutTOI, FReal& OutPhi, FVec3& OutPosition, FVec3& OutNormal, int32& OutFaceIndex, FVec3& OutFaceNormal) const;
+
+		CHAOS_API bool GJKContactPoint(const TBox<FReal, 3>& QueryGeom, const FRigidTransform3& QueryTM, const FReal Thickness, FVec3& ContactLocation, FVec3& ContactNormal, FReal& ContactPhi, int32& ContactFaceIndex) const;
+		CHAOS_API bool GJKContactPoint(const TSphere<FReal, 3>& QueryGeom, const FRigidTransform3& QueryTM, const FReal Thickness, FVec3& ContactLocation, FVec3& ContactNormal, FReal& ContactPhi, int32& ContactFaceIndex) const;
+		CHAOS_API bool GJKContactPoint(const FCapsule& QueryGeom, const FRigidTransform3& QueryTM, const FReal Thickness, FVec3& ContactLocation, FVec3& ContactNormal, FReal& ContactPhi, int32& ContactFaceIndex) const;
+		CHAOS_API bool GJKContactPoint(const FConvex& QueryGeom, const FRigidTransform3& QueryTM, const FReal Thickness, FVec3& ContactLocation, FVec3& ContactNormal, FReal& ContactPhi, int32& ContactFaceIndex) const;
+		CHAOS_API bool GJKContactPoint(const TImplicitObjectScaled<TBox<FReal, 3>>& QueryGeom, const FRigidTransform3& QueryTM, const FReal Thickness, FVec3& ContactLocation, FVec3& ContactNormal, FReal& ContactPhi, int32& ContactFaceIndex) const;
+		CHAOS_API bool GJKContactPoint(const TImplicitObjectScaled<TSphere<FReal, 3>>& QueryGeom, const FRigidTransform3& QueryTM, const FReal Thickness, FVec3& ContactLocation, FVec3& ContactNormal, FReal& ContactPhi, int32& ContactFaceIndex) const;
+		CHAOS_API bool GJKContactPoint(const TImplicitObjectScaled<FCapsule>& QueryGeom, const FRigidTransform3& QueryTM, const FReal Thickness, FVec3& ContactLocation, FVec3& ContactNormal, FReal& ContactPhi, int32& ContactFaceIndex) const;
+		CHAOS_API bool GJKContactPoint(const TImplicitObjectScaled<FConvex>& QueryGeom, const FRigidTransform3& QueryTM, const FReal Thickness, FVec3& ContactLocation, FVec3& ContactNormal, FReal& ContactPhi, int32& ContactFaceIndex) const;
+
+
+		CHAOS_API virtual int32 FindMostOpposingFace(const FVec3& Position, const FVec3& UnitDir, int32 HintFaceIndex, FReal SearchDist) const override;
+		CHAOS_API virtual FVec3 FindGeometryOpposingNormal(const FVec3& DenormDir, int32 FaceIndex, const FVec3& OriginalNormal) const override;
 
 		/**
 		 * Visit all triangles transformed into the specified space. Triangles will have standard winding, regardless of heightfield scale.
@@ -234,7 +234,7 @@ namespace Chaos
 			bool bWasSampleBehind = false;
 		};
 
-		FClosestFaceData FindClosestFace(const FVec3& Position, FReal SearchDist) const;
+		CHAOS_API FClosestFaceData FindClosestFace(const FVec3& Position, FReal SearchDist) const;
 		
 		virtual uint16 GetMaterialIndex(uint32 HintIndex) const override
 		{
@@ -987,33 +987,33 @@ namespace Chaos
 		};
 
 		// Helpers for accessing bounds
-		bool GetCellBounds2D(const TVec2<int32> InCoord, FBounds2D& OutBounds, const FVec2& InInflate = {0}) const;
-		bool GetCellBounds3D(const TVec2<int32> InCoord, FVec3& OutMin, FVec3& OutMax, const FVec3& InInflate = FVec3(0)) const;
-		bool GetCellBounds2DScaled(const TVec2<int32> InCoord, FBounds2D& OutBounds, const FVec2& InInflate = {0}) const;
-		bool GetCellBounds3DScaled(const TVec2<int32> InCoord, FVec3& OutMin, FVec3& OutMax, const FVec3& InInflate = FVec3(0)) const;
-		bool CalcCellBounds3D(const TVec2<int32> InCoord, FVec3& OutMin, FVec3& OutMax, const FVec3& InInflate = FVec3(0)) const;
+		CHAOS_API bool GetCellBounds2D(const TVec2<int32> InCoord, FBounds2D& OutBounds, const FVec2& InInflate = {0}) const;
+		CHAOS_API bool GetCellBounds3D(const TVec2<int32> InCoord, FVec3& OutMin, FVec3& OutMax, const FVec3& InInflate = FVec3(0)) const;
+		CHAOS_API bool GetCellBounds2DScaled(const TVec2<int32> InCoord, FBounds2D& OutBounds, const FVec2& InInflate = {0}) const;
+		CHAOS_API bool GetCellBounds3DScaled(const TVec2<int32> InCoord, FVec3& OutMin, FVec3& OutMax, const FVec3& InInflate = FVec3(0)) const;
+		CHAOS_API bool CalcCellBounds3D(const TVec2<int32> InCoord, FVec3& OutMin, FVec3& OutMax, const FVec3& InInflate = FVec3(0)) const;
 
 		// Query functions - sweep, ray, overlap
 		template<typename SQVisitor>
 		bool GridSweep(const FVec3& StartPoint, const FVec3& Dir, const FReal Length, const FVec3 InHalfExtents, SQVisitor& Visitor) const;
 		
-		bool WalkSlow(TVec2<int32>& CellIdx, FHeightfieldRaycastVisitor& Visitor, FReal CurrentLength, const VectorRegister4Float& CurrentLengthSimd,
+		CHAOS_API bool WalkSlow(TVec2<int32>& CellIdx, FHeightfieldRaycastVisitor& Visitor, FReal CurrentLength, const VectorRegister4Float& CurrentLengthSimd,
 			const FVec2& ScaledMin, FReal ZMidPoint, const FVec3& Dir, const FVec3& InvDir, bool bParallel[3], const FVec2& ScaledDx2D, FVec3& NextStart,
 			const FVec3& ScaleSign, const FVec3& ScaledDx, int32 IndexLowResX, int32 IndexLowResY) const;
 		
-		bool WalkFast(TVec2<int32>& CellIdx, FHeightfieldRaycastVisitor& Visitor, FReal CurrentLength, const VectorRegister4Float& CurrentLengthSimd,
+		CHAOS_API bool WalkFast(TVec2<int32>& CellIdx, FHeightfieldRaycastVisitor& Visitor, FReal CurrentLength, const VectorRegister4Float& CurrentLengthSimd,
 			const FVec2& ScaledMin, FReal ZMidPoint, const FVec3& Dir, const FVec3& InvDir, bool bParallel[3], const FVec2& ScaledDx2D, FVec3& NextStart,
 			const FVec3& ScaleSign, const FVec3& ScaledDx, const FVec2& Scale2D, const FVec3& DirScaled) const;
 
-		bool WalkOnLowRes(TVec2<int32>& CellIdx, FHeightfieldRaycastVisitor& Visitor, FReal CurrentLength, const VectorRegister4Float& CurrentLengthSimd,
+		CHAOS_API bool WalkOnLowRes(TVec2<int32>& CellIdx, FHeightfieldRaycastVisitor& Visitor, FReal CurrentLength, const VectorRegister4Float& CurrentLengthSimd,
 			const FVec2& ScaledMin, FReal ZMidPoint, const FVec3& Dir, const FVec3& InvDir, bool bParallel[3], const FVec2& ScaledDx2D, FVec3& NextStart,
 			const FVec3& ScaleSign, const FVec3& ScaledDx, const FVec2& Scale2D, const FVec3& DirScaled) const;
 
-		bool GridCast(const FVec3& StartPoint, const FVec3& Dir, const FReal Length, FHeightfieldRaycastVisitor& Visitor) const;
-		bool GetGridIntersections(FBounds2D InFlatBounds, TArray<TVec2<int32>>& OutInterssctions) const;
-		bool GetGridIntersectionsBatch(FBounds2D InFlatBounds, TArray<TVec2<int32>>& OutIntersections, const FAABBVectorized& Bounds) const;
+		CHAOS_API bool GridCast(const FVec3& StartPoint, const FVec3& Dir, const FReal Length, FHeightfieldRaycastVisitor& Visitor) const;
+		CHAOS_API bool GetGridIntersections(FBounds2D InFlatBounds, TArray<TVec2<int32>>& OutInterssctions) const;
+		CHAOS_API bool GetGridIntersectionsBatch(FBounds2D InFlatBounds, TArray<TVec2<int32>>& OutIntersections, const FAABBVectorized& Bounds) const;
 		
-		FBounds2D GetFlatBounds() const;
+		CHAOS_API FBounds2D GetFlatBounds() const;
 
 		// Grid for queries, faster than bounding volumes for heightfields
 		TUniformGrid<FReal, 2> FlatGrid;
@@ -1024,8 +1024,8 @@ namespace Chaos
 		// Cached when bounds are requested. Mutable to allow GetBounds to be logical const
 		mutable FAABB3 CachedBounds;
 
-		void CalcBounds();
-		void BuildQueryData();
+		CHAOS_API void CalcBounds();
+		CHAOS_API void BuildQueryData();
 
 		// Needed for serialization
 		FHeightField() : FImplicitObject(EImplicitObject::HasBoundingBox, ImplicitObjectType::HeightField) {}
@@ -1049,4 +1049,4 @@ namespace Chaos
 		Value.Serialize(Ar);
 		return Ar;
 	}
-}
\ No newline at end of file
+}
diff --git a/Engine/Source/Runtime/Experimental/Chaos/Public/Chaos/ISpatialAcceleration.h b/Engine/Source/Runtime/Experimental/Chaos/Public/Chaos/ISpatialAcceleration.h
--- a/Engine/Source/Runtime/Experimental/Chaos/Public/Chaos/ISpatialAcceleration.h
+++ b/Engine/Source/Runtime/Experimental/Chaos/Public/Chaos/ISpatialAcceleration.h
@@ -9,7 +9,7 @@
 namespace Chaos
 {
 
-struct CHAOS_API FQueryFastData
+struct FQueryFastData
 {
 	FQueryFastData(const FVec3& InDir, const FReal InLength)
 		: Dir(InDir)
@@ -59,7 +59,7 @@ protected:
 };
 
 //dummy struct for templatized paths
-struct CHAOS_API FQueryFastDataVoid : public FQueryFastData
+struct FQueryFastDataVoid : public FQueryFastData
 {
 	FQueryFastDataVoid() : FQueryFastData(DummyDir) {}
 	
diff --git a/Engine/Source/Runtime/Experimental/Chaos/Public/Chaos/ISpatialAccelerationCollection.h b/Engine/Source/Runtime/Experimental/Chaos/Public/Chaos/ISpatialAccelerationCollection.h
--- a/Engine/Source/Runtime/Experimental/Chaos/Public/Chaos/ISpatialAccelerationCollection.h
+++ b/Engine/Source/Runtime/Experimental/Chaos/Public/Chaos/ISpatialAccelerationCollection.h
@@ -19,7 +19,7 @@ namespace Chaos
 	class IResimCacheBase;
 
 template <typename TPayloadType, typename T, int d>
-class CHAOS_API ISpatialAccelerationCollection : public ISpatialAcceleration<TPayloadType, T, d>
+class ISpatialAccelerationCollection : public ISpatialAcceleration<TPayloadType, T, d>
 {
 public:
 	UE_NONCOPYABLE(ISpatialAccelerationCollection)
diff --git a/Engine/Source/Runtime/Experimental/Chaos/Public/Chaos/ImplicitObject.h b/Engine/Source/Runtime/Experimental/Chaos/Public/Chaos/ImplicitObject.h
--- a/Engine/Source/Runtime/Experimental/Chaos/Public/Chaos/ImplicitObject.h
+++ b/Engine/Source/Runtime/Experimental/Chaos/Public/Chaos/ImplicitObject.h
@@ -99,17 +99,17 @@ using FImplicitHierarchyVisitorBool = TFunctionRef<bool(const FImplicitObject* I
 #pragma warning(error : 4820)
 #endif
 #endif // #if DISALLOW_FIMPLICIT_OBJECT_TAIL_PADDING
-class CHAOS_API FImplicitObject
+class FImplicitObject
 {
 public:
 	using TType = FReal;
 	static constexpr int D = 3;
-	static FImplicitObject* SerializationFactory(FChaosArchive& Ar, FImplicitObject* Obj);
+	static CHAOS_API FImplicitObject* SerializationFactory(FChaosArchive& Ar, FImplicitObject* Obj);
 
-	FImplicitObject(int32 Flags, EImplicitObjectType InType = ImplicitObjectType::Unknown);
+	CHAOS_API FImplicitObject(int32 Flags, EImplicitObjectType InType = ImplicitObjectType::Unknown);
 	FImplicitObject(const FImplicitObject&) = delete;
 	FImplicitObject(FImplicitObject&&) = delete;
-	virtual ~FImplicitObject();
+	CHAOS_API virtual ~FImplicitObject();
 
 	template<class T_DERIVED>
 	T_DERIVED* GetObject()
@@ -149,30 +149,30 @@ public:
 	virtual FImplicitObject* Duplicate() const { check(false); return nullptr; }
 
 	virtual EImplicitObjectType GetNestedType() const { return GetType(); }
-	EImplicitObjectType GetType() const;
+	CHAOS_API EImplicitObjectType GetType() const;
 	static int32 GetOffsetOfType() { return offsetof(FImplicitObject, Type); }
 
-	EImplicitObjectType GetCollisionType() const;
+	CHAOS_API EImplicitObjectType GetCollisionType() const;
 	void SetCollisionType(EImplicitObjectType InCollisionType) { CollisionType = InCollisionType; }
 
 	FReal GetMargin() const { return Margin; }
 	static int32 GetOffsetOfMargin() { return offsetof(FImplicitObject, Margin); }
 
-	virtual bool IsValidGeometry() const;
+	CHAOS_API virtual bool IsValidGeometry() const;
 
-	virtual TUniquePtr<FImplicitObject> Copy() const;
-	virtual TUniquePtr<FImplicitObject> CopyWithScale(const FVec3& Scale) const;
+	CHAOS_API virtual TUniquePtr<FImplicitObject> Copy() const;
+	CHAOS_API virtual TUniquePtr<FImplicitObject> CopyWithScale(const FVec3& Scale) const;
 	virtual TUniquePtr<FImplicitObject> DeepCopy() const { return Copy(); }
 	virtual TUniquePtr<FImplicitObject> DeepCopyWithScale(const FVec3& Scale) const { return CopyWithScale(Scale); }
 
 	//This is strictly used for optimization purposes
-	bool IsUnderlyingUnion() const;
+	CHAOS_API bool IsUnderlyingUnion() const;
 
 	// Explicitly non-virtual.  Must cast to derived types to target their implementation.
-	FReal SignedDistance(const FVec3& x) const;
+	CHAOS_API FReal SignedDistance(const FVec3& x) const;
 
 	// Explicitly non-virtual.  Must cast to derived types to target their implementation.
-	FVec3 Normal(const FVec3& x) const;
+	CHAOS_API FVec3 Normal(const FVec3& x) const;
 
 	// Find the closest point on the surface, and return the separating distance and axis
 	virtual FReal PhiWithNormal(const FVec3& x, FVec3& Normal) const = 0;
@@ -190,7 +190,7 @@ public:
 		return ScaledPhi;
 	}
 
-	virtual const FAABB3 BoundingBox() const;
+	CHAOS_API virtual const FAABB3 BoundingBox() const;
 
 	// Calculate the tight-fitting world-space bounding box
 	virtual FAABB3 CalculateTransformedBounds(const FRigidTransform3& Transform) const
@@ -209,7 +209,7 @@ public:
 	
 #if TRACK_CHAOS_GEOMETRY
 	//Turn on memory tracking. Must pass object itself as a serializable ptr so we can save it out
-	void Track(TSerializablePtr<FImplicitObject> This, const FString& DebugInfo);
+	CHAOS_API void Track(TSerializablePtr<FImplicitObject> This, const FString& DebugInfo);
 #endif
 
 	virtual bool IsPerformanceWarning() const { return false; }
@@ -218,9 +218,9 @@ public:
 		return FString::Printf(TEXT("ImplicitObject - No Performance String"));
 	};
 
-	Pair<FVec3, bool> FindDeepestIntersection(const FImplicitObject* Other, const FBVHParticles* Particles, const FMatrix33& OtherToLocalTransform, const FReal Thickness) const;
-	Pair<FVec3, bool> FindDeepestIntersection(const FImplicitObject* Other, const FParticles* Particles, const FMatrix33& OtherToLocalTransform, const FReal Thickness) const;
-	Pair<FVec3, bool> FindClosestIntersection(const FVec3& StartPoint, const FVec3& EndPoint, const FReal Thickness) const;
+	CHAOS_API Pair<FVec3, bool> FindDeepestIntersection(const FImplicitObject* Other, const FBVHParticles* Particles, const FMatrix33& OtherToLocalTransform, const FReal Thickness) const;
+	CHAOS_API Pair<FVec3, bool> FindDeepestIntersection(const FImplicitObject* Other, const FParticles* Particles, const FMatrix33& OtherToLocalTransform, const FReal Thickness) const;
+	CHAOS_API Pair<FVec3, bool> FindClosestIntersection(const FVec3& StartPoint, const FVec3& EndPoint, const FReal Thickness) const;
 
 	//This gives derived types a way to avoid calling PhiWithNormal todo: this api is confusing
 	virtual bool Raycast(const FVec3& StartPoint, const FVec3& Dir, const FReal Length, const FReal Thickness, FReal& OutTime, FVec3& OutPosition, FVec3& OutNormal, int32& OutFaceIndex) const
@@ -311,14 +311,14 @@ public:
 		Out.Add(MakePair(This, ParentTM));
 	}
 
-	virtual void FindAllIntersectingObjects(TArray < Pair<const FImplicitObject*, FRigidTransform3>>& Out, const FAABB3& LocalBounds) const;
+	CHAOS_API virtual void FindAllIntersectingObjects(TArray < Pair<const FImplicitObject*, FRigidTransform3>>& Out, const FAABB3& LocalBounds) const;
 
 	virtual FString ToString() const
 	{
 		return FString::Printf(TEXT("ImplicitObject bIsConvex:%d, bDoCollide:%d, bHasBoundingBox:%d"), bIsConvex, bDoCollide, bHasBoundingBox);
 	}
 
-	void SerializeImp(FArchive& Ar);
+	CHAOS_API void SerializeImp(FArchive& Ar);
 
 	constexpr static EImplicitObjectType StaticType()
 	{
@@ -330,15 +330,15 @@ public:
 		check(false);	//Aggregate implicits require FChaosArchive - check false by default
 	}
 
-	virtual void Serialize(FChaosArchive& Ar);
+	CHAOS_API virtual void Serialize(FChaosArchive& Ar);
 	
-	static FArchive& SerializeLegacyHelper(FArchive& Ar, TUniquePtr<FImplicitObject>& Value);
+	static CHAOS_API FArchive& SerializeLegacyHelper(FArchive& Ar, TUniquePtr<FImplicitObject>& Value);
 
 	virtual uint32 GetTypeHash() const = 0;
 
 	virtual FName GetTypeName() const { return GetTypeName(GetType()); }
 
-	static const FName GetTypeName(const EImplicitObjectType InType);
+	static CHAOS_API const FName GetTypeName(const EImplicitObjectType InType);
 
 	virtual uint16 GetMaterialIndex(uint32 HintIndex) const { return 0; }
 
@@ -501,7 +501,7 @@ protected:
 #endif // DISALLOW_FIMPLICIT_OBJECT_TAIL_PADDING
 
 private:
-	virtual Pair<FVec3, bool> FindClosestIntersectionImp(const FVec3& StartPoint, const FVec3& EndPoint, const FReal Thickness) const;
+	CHAOS_API virtual Pair<FVec3, bool> FindClosestIntersectionImp(const FVec3& StartPoint, const FVec3& EndPoint, const FReal Thickness) const;
 };
 #if DISALLOW_FIMPLICIT_OBJECT_TAIL_PADDING
 #if defined(__clang__) || defined(__GNUC__)
diff --git a/Engine/Source/Runtime/Experimental/Chaos/Public/Chaos/ImplicitObjectUnion.h b/Engine/Source/Runtime/Experimental/Chaos/Public/Chaos/ImplicitObjectUnion.h
--- a/Engine/Source/Runtime/Experimental/Chaos/Public/Chaos/ImplicitObjectUnion.h
+++ b/Engine/Source/Runtime/Experimental/Chaos/Public/Chaos/ImplicitObjectUnion.h
@@ -23,24 +23,24 @@ namespace CVars
 	extern bool bChaosUnionBVHEnabled;
 }
 
-class CHAOS_API FImplicitObjectUnion : public FImplicitObject
+class FImplicitObjectUnion : public FImplicitObject
 {
   public:
 
 	using FImplicitObject::GetTypeName;
 
-	FImplicitObjectUnion(TArray<TUniquePtr<FImplicitObject>>&& Objects);
+	CHAOS_API FImplicitObjectUnion(TArray<TUniquePtr<FImplicitObject>>&& Objects);
 	FImplicitObjectUnion(const FImplicitObjectUnion& Other) = delete;
-	FImplicitObjectUnion(FImplicitObjectUnion&& Other);
-	virtual ~FImplicitObjectUnion();
+	CHAOS_API FImplicitObjectUnion(FImplicitObjectUnion&& Other);
+	CHAOS_API virtual ~FImplicitObjectUnion();
 
 	FORCEINLINE static constexpr EImplicitObjectType StaticType()
 	{
 		return ImplicitObjectType::Union;
 	}
 
-	void Combine(TArray<TUniquePtr<FImplicitObject>>& Objects);
-	void RemoveAt(int32 RemoveIndex);
+	CHAOS_API void Combine(TArray<TUniquePtr<FImplicitObject>>& Objects);
+	CHAOS_API void RemoveAt(int32 RemoveIndex);
 
 	// The total number of root objects in the hierarchy (same as GetObjects().Num())
 	int32 GetNumRootObjects() const
@@ -64,10 +64,10 @@ class CHAOS_API FImplicitObjectUnion : public FImplicitObject
 		}
 	}
 
-	virtual TUniquePtr<FImplicitObject> Copy() const;
-	virtual TUniquePtr<FImplicitObject> CopyWithScale(const FVec3& Scale) const override;
-	virtual TUniquePtr<FImplicitObject> DeepCopy() const;
-	virtual TUniquePtr<FImplicitObject> DeepCopyWithScale(const FVec3& Scale) const override;
+	CHAOS_API virtual TUniquePtr<FImplicitObject> Copy() const;
+	CHAOS_API virtual TUniquePtr<FImplicitObject> CopyWithScale(const FVec3& Scale) const override;
+	CHAOS_API virtual TUniquePtr<FImplicitObject> DeepCopy() const;
+	CHAOS_API virtual TUniquePtr<FImplicitObject> DeepCopyWithScale(const FVec3& Scale) const override;
 	
 	virtual FReal PhiWithNormal(const FVec3& x, FVec3& Normal) const override
 	{
@@ -123,7 +123,7 @@ class CHAOS_API FImplicitObjectUnion : public FImplicitObject
 		}
 	}
 
-	virtual void FindAllIntersectingObjects(TArray < Pair<const FImplicitObject*, FRigidTransform3>>& Out, const FAABB3& LocalBounds) const;
+	CHAOS_API virtual void FindAllIntersectingObjects(TArray < Pair<const FImplicitObject*, FRigidTransform3>>& Out, const FAABB3& LocalBounds) const;
 	
 	virtual void CacheAllImplicitObjects() { RebuildBVH(); }
 
@@ -168,7 +168,7 @@ class CHAOS_API FImplicitObjectUnion : public FImplicitObject
 		return false;
 	}
 
-	virtual void Serialize(FChaosArchive& Ar) override;
+	CHAOS_API virtual void Serialize(FChaosArchive& Ar) override;
 
 	virtual bool IsValidGeometry() const
 	{
@@ -182,7 +182,7 @@ class CHAOS_API FImplicitObjectUnion : public FImplicitObject
 	TArray<TUniquePtr<FImplicitObject>>& GetObjects() { return MObjects; }
 
 	// The lambda returns TRUE if an object was found and iteration should stop.
-	void ForEachObject(TFunctionRef<bool(const FImplicitObject&, const FRigidTransform3&)> Lambda) const;
+	CHAOS_API void ForEachObject(TFunctionRef<bool(const FImplicitObject&, const FRigidTransform3&)> Lambda) const;
 
 	virtual uint32 GetTypeHash() const override
 	{
@@ -253,7 +253,7 @@ protected:
 		return ClosestIntersection;
 	}
 
-	virtual void VisitOverlappingLeafObjectsImpl(
+	CHAOS_API virtual void VisitOverlappingLeafObjectsImpl(
 		const FAABB3& LocalBounds,
 		const FRigidTransform3& ObjectTransform,
 		const int32 RootObjectIndex,
@@ -261,34 +261,34 @@ protected:
 		int32& LeafObjectIndex,
 		const FImplicitHierarchyVisitor& VisitorFunc) const override final;
 
-	virtual void VisitLeafObjectsImpl(
+	CHAOS_API virtual void VisitLeafObjectsImpl(
 		const FRigidTransform3& ObjectTransform,
 		const int32 RootObjectIndex,
 		int32& ObjectIndex,
 		int32& LeafObjectIndex,
 		const FImplicitHierarchyVisitor& VisitorFunc) const override final;
 
-	virtual bool VisitObjectsImpl(
+	CHAOS_API virtual bool VisitObjectsImpl(
 		const FRigidTransform3& ObjectTransform,
 		const int32 RootObjectIndex,
 		int32& ObjectIndex,
 		int32& LeafObjectIndex,
 		const FImplicitHierarchyVisitorBool& VisitorFunc) const override final;
 
-	virtual bool IsOverlappingBoundsImpl(
+	CHAOS_API virtual bool IsOverlappingBoundsImpl(
 		const FAABB3& LocalBounds) const override final;
 
   protected:
 	// Needed for serialization
-	FImplicitObjectUnion();
+	CHAOS_API FImplicitObjectUnion();
 	friend FImplicitObject;
 
-	void SetNumLeafObjects(const int32 InNumLeafObjects);
-	void CreateBVH();
-	void DestroyBVH();
-	void RebuildBVH();
+	CHAOS_API void SetNumLeafObjects(const int32 InNumLeafObjects);
+	CHAOS_API void CreateBVH();
+	CHAOS_API void DestroyBVH();
+	CHAOS_API void RebuildBVH();
 
-	void LegacySerializeBVH(FChaosArchive& Ar);
+	CHAOS_API void LegacySerializeBVH(FChaosArchive& Ar);
 
 	union FFLags
 	{
@@ -327,13 +327,13 @@ struct FLargeUnionClusteredImplicitInfo
 	const FBVHParticles* BVHParticles;
 };
 
-class CHAOS_API FImplicitObjectUnionClustered: public FImplicitObjectUnion
+class FImplicitObjectUnionClustered: public FImplicitObjectUnion
 {
 public:
-	FImplicitObjectUnionClustered();
-	FImplicitObjectUnionClustered(TArray<TUniquePtr<FImplicitObject>>&& Objects, const TArray<FPBDRigidParticleHandle*>& OriginalParticleLookupHack = TArray<FPBDRigidParticleHandle*>());
+	CHAOS_API FImplicitObjectUnionClustered();
+	CHAOS_API FImplicitObjectUnionClustered(TArray<TUniquePtr<FImplicitObject>>&& Objects, const TArray<FPBDRigidParticleHandle*>& OriginalParticleLookupHack = TArray<FPBDRigidParticleHandle*>());
 	FImplicitObjectUnionClustered(const FImplicitObjectUnionClustered& Other) = delete;
-	FImplicitObjectUnionClustered(FImplicitObjectUnionClustered&& Other);
+	CHAOS_API FImplicitObjectUnionClustered(FImplicitObjectUnionClustered&& Other);
 	virtual ~FImplicitObjectUnionClustered() = default;
 
 	FORCEINLINE static constexpr EImplicitObjectType StaticType()
@@ -344,15 +344,15 @@ public:
 	UE_DEPRECATED(5.3, "Not supported")
 	void FindAllIntersectingClusteredObjects(TArray<FLargeUnionClusteredImplicitInfo>& Out, const FAABB3& LocalBounds) const {}
 
-	TArray<FPBDRigidParticleHandle*> FindAllIntersectingChildren(const FAABB3& LocalBounds) const;
+	CHAOS_API TArray<FPBDRigidParticleHandle*> FindAllIntersectingChildren(const FAABB3& LocalBounds) const;
 
 	// DO NOT USE!!
 	// @todo(chaos): we should get rid of this. Instead we should hold the map/whatever on the GeometryParticle that currently owns the geom
-	const FPBDRigidParticleHandle* FindParticleForImplicitObject(const FImplicitObject* Object) const;
+	CHAOS_API const FPBDRigidParticleHandle* FindParticleForImplicitObject(const FImplicitObject* Object) const;
 
 	// DO NOT USE!!
 	// @todo(chaos): move this fucntionality to the geometry particle?
-	const FBVHParticles* GetChildSimplicial(const int32 ChildIndex) const;
+	CHAOS_API const FBVHParticles* GetChildSimplicial(const int32 ChildIndex) const;
 
 private:
 	// Temp hack for finding original particles
diff --git a/Engine/Source/Runtime/Experimental/Chaos/Public/Chaos/Interface/SQTypes.h b/Engine/Source/Runtime/Experimental/Chaos/Public/Chaos/Interface/SQTypes.h
--- a/Engine/Source/Runtime/Experimental/Chaos/Public/Chaos/Interface/SQTypes.h
+++ b/Engine/Source/Runtime/Experimental/Chaos/Public/Chaos/Interface/SQTypes.h
@@ -23,12 +23,12 @@ namespace Chaos
 
 namespace ChaosInterface
 {
-	struct CHAOS_API FActorShape
+	struct FActorShape
 	{
 		Chaos::FGeometryParticle* Actor;
 		const Chaos::FPerShapeData* Shape;
 
-		void Serialize(Chaos::FChaosArchive& Ar);
+		CHAOS_API void Serialize(Chaos::FChaosArchive& Ar);
 	};
 
 	inline Chaos::FChaosArchive& operator<<(Chaos::FChaosArchive& Ar, FActorShape& ActorShape)
@@ -37,7 +37,7 @@ namespace ChaosInterface
 		return Ar;
 	}
 
-	struct CHAOS_API FQueryHit : public FActorShape
+	struct FQueryHit : public FActorShape
 	{
 		FQueryHit() : FaceIndex(-1), ElementIndex(-1) {}
 
@@ -50,7 +50,7 @@ namespace ChaosInterface
 
 		int32 ElementIndex; // Currently used to indicate which shape was hit for a particle with multiple shapes.
 
-		void Serialize(Chaos::FChaosArchive& Ar);
+		CHAOS_API void Serialize(Chaos::FChaosArchive& Ar);
 	};
 
 	inline Chaos::FChaosArchive& operator<<(Chaos::FChaosArchive& Ar, FQueryHit& QueryHit)
@@ -59,14 +59,14 @@ namespace ChaosInterface
 		return Ar;
 	}
 
-	struct CHAOS_API FLocationHit : public FQueryHit
+	struct FLocationHit : public FQueryHit
 	{
 		FHitFlags Flags;
 		FVector WorldPosition;
 		FVector WorldNormal;
 		float Distance;	// LWC_TODO: Should be FVector::FReal, but that causes precision issues resulting in collision failures. Investigate!
 
-		void Serialize(Chaos::FChaosArchive& Ar);
+		CHAOS_API void Serialize(Chaos::FChaosArchive& Ar);
 
 		bool operator<(const FLocationHit& Other) const { return Distance < Other.Distance; }
 	};
@@ -77,12 +77,12 @@ namespace ChaosInterface
 		return Ar;
 	}
 
-	struct CHAOS_API FRaycastHit : public FLocationHit
+	struct FRaycastHit : public FLocationHit
 	{
 		float U;
 		float V;
 
-		void Serialize(Chaos::FChaosArchive& Ar);
+		CHAOS_API void Serialize(Chaos::FChaosArchive& Ar);
 
 	};
 
@@ -92,7 +92,7 @@ namespace ChaosInterface
 		return Ar;
 	}
 
-	struct CHAOS_API FOverlapHit : public FQueryHit
+	struct FOverlapHit : public FQueryHit
 	{
 	};
 
@@ -102,7 +102,7 @@ namespace ChaosInterface
 		return Ar;
 	}
 
-	struct CHAOS_API FSweepHit : public FLocationHit
+	struct FSweepHit : public FLocationHit
 	{
 	};
 
@@ -112,13 +112,13 @@ namespace ChaosInterface
 		return Ar;
 	}
 
-	struct CHAOS_API FPTActorShape
+	struct FPTActorShape
 	{
 		Chaos::FGeometryParticleHandle* Actor;
 		const Chaos::FPerShapeData* Shape;
 	};
 
-	struct CHAOS_API FPTQueryHit : public FPTActorShape
+	struct FPTQueryHit : public FPTActorShape
 	{
 		/**
 		Face index of touched triangle, for triangle meshes, convex meshes and height fields. Defaults to -1 if face index is not available
@@ -130,7 +130,7 @@ namespace ChaosInterface
 		int32 ElementIndex = -1; // Currently used to indicate which shape was hit for a particle with multiple shapes.
 	};
 
-	struct CHAOS_API FPTLocationHit : public FPTQueryHit
+	struct FPTLocationHit : public FPTQueryHit
 	{
 		FHitFlags Flags;
 		FVector WorldPosition;
@@ -140,17 +140,17 @@ namespace ChaosInterface
 		bool operator<(const FPTLocationHit& Other) const { return Distance < Other.Distance; }
 	};
 
-	struct CHAOS_API FPTRaycastHit : public FPTLocationHit
+	struct FPTRaycastHit : public FPTLocationHit
 	{
 		float U;
 		float V;
 	};
 
-	struct CHAOS_API FPTOverlapHit : public FPTQueryHit
+	struct FPTOverlapHit : public FPTQueryHit
 	{
 	};
 
-	struct CHAOS_API FPTSweepHit : public FPTLocationHit
+	struct FPTSweepHit : public FPTLocationHit
 	{
 	};
 
diff --git a/Engine/Source/Runtime/Experimental/Chaos/Public/Chaos/Island/IslandGroup.h b/Engine/Source/Runtime/Experimental/Chaos/Public/Chaos/Island/IslandGroup.h
--- a/Engine/Source/Runtime/Experimental/Chaos/Public/Chaos/Island/IslandGroup.h
+++ b/Engine/Source/Runtime/Experimental/Chaos/Public/Chaos/Island/IslandGroup.h
@@ -15,7 +15,7 @@ namespace Chaos
 		* from several islands, but may be a sub set of constraints of a single island when coloring is enabled for
 		* that island.
 		*/
-		class CHAOS_API FPBDIslandConstraintGroupSolver : public FPBDConstraintGroupSolver
+		class FPBDIslandConstraintGroupSolver : public FPBDConstraintGroupSolver
 		{
 		public:
 			UE_NONCOPYABLE(FPBDIslandConstraintGroupSolver);
@@ -23,13 +23,13 @@ namespace Chaos
 			/**
 			* Init the island group 
 			*/
-			FPBDIslandConstraintGroupSolver(FPBDIslandManager& InIslandManager);
+			CHAOS_API FPBDIslandConstraintGroupSolver(FPBDIslandManager& InIslandManager);
 
 			/**
 			* Add island to the group
 			* @param Island Island to be added
 			*/
-			void AddIsland(FPBDIsland* Island);
+			CHAOS_API void AddIsland(FPBDIsland* Island);
 
 			/**
 			* Return the islands within the group
@@ -54,14 +54,14 @@ namespace Chaos
 			/**
 			 * Update the iteration settings based on the islands in the group, using at least as many as specified in InIterations
 			*/
-			virtual void SetIterationSettings(const FIterationSettings& InIterations) override;
+			CHAOS_API virtual void SetIterationSettings(const FIterationSettings& InIterations) override;
 
 		protected:
 			// Base class overrides
-			virtual void SetConstraintSolverImpl(const int32 ContainerId) override final;
-			virtual void ResetImpl() override final;
-			virtual void AddConstraintsImpl() override final;
-			virtual void GatherBodiesImpl(const FReal Dt, const int32 BeginBodyIndex, const int32 EndBodyIndex) override final;
+			CHAOS_API virtual void SetConstraintSolverImpl(const int32 ContainerId) override final;
+			CHAOS_API virtual void ResetImpl() override final;
+			CHAOS_API virtual void AddConstraintsImpl() override final;
+			CHAOS_API virtual void GatherBodiesImpl(const FReal Dt, const int32 BeginBodyIndex, const int32 EndBodyIndex) override final;
 
 		private:
 			FPBDIslandManager& IslandManager;
@@ -80,4 +80,4 @@ namespace Chaos
 
 	using FPBDIslandConstraintGroupSolver UE_DEPRECATED(5.2, "Internal class moved to Private namespace") = Private::FPBDIslandConstraintGroupSolver;
 
-}	// namespace Chaos
\ No newline at end of file
+}	// namespace Chaos
diff --git a/Engine/Source/Runtime/Experimental/Chaos/Public/Chaos/Island/IslandGroupManager.h b/Engine/Source/Runtime/Experimental/Chaos/Public/Chaos/Island/IslandGroupManager.h
--- a/Engine/Source/Runtime/Experimental/Chaos/Public/Chaos/Island/IslandGroupManager.h
+++ b/Engine/Source/Runtime/Experimental/Chaos/Public/Chaos/Island/IslandGroupManager.h
@@ -73,29 +73,29 @@ namespace Chaos
 		 * Each Island Group may be solved in parallel with the others and contains its own set of ConstraintContainerSolvers
 		 * (one for each type of constraint) so that solver data access may be cache efficient during the solver phases.
 		*/
-		class CHAOS_API FPBDIslandGroupManager
+		class FPBDIslandGroupManager
 		{
 		public:
 			UE_NONCOPYABLE(FPBDIslandGroupManager)
 
-			FPBDIslandGroupManager(FPBDIslandManager& InIslandManager);
-			~FPBDIslandGroupManager();
+			CHAOS_API FPBDIslandGroupManager(FPBDIslandManager& InIslandManager);
+			CHAOS_API ~FPBDIslandGroupManager();
 
 			/**
 			 * Register a constraint type with the manager. This will create a ConstraintContainerSolver for the ConstraintsContainer.
 			 * Constraints are solved in Priority order with lower priorities first, so higher priorities "win".
 			*/
-			void AddConstraintContainer(FPBDConstraintContainer& ConstraintContainer, const int32 Priority = 0);
+			CHAOS_API void AddConstraintContainer(FPBDConstraintContainer& ConstraintContainer, const int32 Priority = 0);
 
 			/**
 			 * Remove a previously-added container. Generall this is only needed for debuggin as cleanup is automatic on destruction.
 			 */
-			void RemoveConstraintContainer(FPBDConstraintContainer& ConstraintContainer);
+			CHAOS_API void RemoveConstraintContainer(FPBDConstraintContainer& ConstraintContainer);
 
 			/**
 			 * Set the priority for the specified container (that must have been pre-registered with AddConstraintContainer).
 			*/
-			void SetConstraintContainerPriority(const int32 ContainerId, const int32 Priority);
+			CHAOS_API void SetConstraintContainerPriority(const int32 ContainerId, const int32 Priority);
 
 			/**
 			 * The number of groups with at least one constraint in them
@@ -111,12 +111,12 @@ namespace Chaos
 			 * Pull all the active islands from the IslandManager and assign to groups.
 			 * @return 
 			*/
-			int32 BuildGroups(const bool bIsResimming);
+			CHAOS_API int32 BuildGroups(const bool bIsResimming);
 
 			/**
 			 * Set the default number of iterations to perform in the Solve step (can be increased on a per-island basis by any dynamic body).
 			*/
-			void SetIterationSettings(const FIterationSettings& InIterations);
+			CHAOS_API void SetIterationSettings(const FIterationSettings& InIterations);
 			void SetNumPositionIterations(const int32 InNumIterations) { Iterations.SetNumPositionIterations(InNumIterations); }
 			void SetNumVelocityIterations(const int32 InNumIterations) { Iterations.SetNumVelocityIterations(InNumIterations); }
 			void SetNumProjectionIterations(const int32 InNumIterations) { Iterations.SetNumProjectionIterations(InNumIterations); }
@@ -130,7 +130,7 @@ namespace Chaos
 			/**
 			 * Solve all constraints.
 			*/
-			void Solve(const FReal Dt);
+			CHAOS_API void Solve(const FReal Dt);
 
 
 		private:
@@ -155,12 +155,12 @@ namespace Chaos
 				FPBDIslandConstraintGroupSolver* IslandGroup;
 			};
 
-			void SolveSerial(const FReal Dt);
-			void SolveParallelFor(const FReal Dt);
-			void SolveParallelTasks(const FReal Dt);
+			CHAOS_API void SolveSerial(const FReal Dt);
+			CHAOS_API void SolveParallelFor(const FReal Dt);
+			CHAOS_API void SolveParallelTasks(const FReal Dt);
 
-			void BuildGatherBatches(TArray<FIslandGroupRange>& BodyRanges, TArray<FIslandGroupRange>& ConstraintRanges);
-			void SolveGroupConstraints(const int32 GroupIndex, const FReal Dt);
+			CHAOS_API void BuildGatherBatches(TArray<FIslandGroupRange>& BodyRanges, TArray<FIslandGroupRange>& ConstraintRanges);
+			CHAOS_API void SolveGroupConstraints(const int32 GroupIndex, const FReal Dt);
 
 
 			FPBDIslandManager& IslandManager;
diff --git a/Engine/Source/Runtime/Experimental/Chaos/Public/Chaos/Island/IslandManager.h b/Engine/Source/Runtime/Experimental/Chaos/Public/Chaos/Island/IslandManager.h
--- a/Engine/Source/Runtime/Experimental/Chaos/Public/Chaos/Island/IslandManager.h
+++ b/Engine/Source/Runtime/Experimental/Chaos/Public/Chaos/Island/IslandManager.h
@@ -20,12 +20,12 @@ namespace Chaos::Private
 	/**
 	* A particle in the constraint graph.
 	*/
-	class CHAOS_API FPBDIslandParticle
+	class FPBDIslandParticle
 	{
 	public:
-		FPBDIslandParticle();
-		FPBDIslandParticle(FGeometryParticleHandle* InParticle);
-		~FPBDIslandParticle();
+		CHAOS_API FPBDIslandParticle();
+		CHAOS_API FPBDIslandParticle(FGeometryParticleHandle* InParticle);
+		CHAOS_API ~FPBDIslandParticle();
 
 		FGeometryParticleHandle* GetParticle()
 		{
@@ -40,11 +40,11 @@ namespace Chaos::Private
 		// TPoolBackedItemAdapter for FPBDIslandManager::Nodes
 		const int32 GetArrayIndex() const { return ArrayIndex; }
 		void SetArrayIndex(const int32 InIndex) { ArrayIndex = InIndex; }
-		void Reuse(FGeometryParticleHandle* InParticle);
-		void Trash();
+		CHAOS_API void Reuse(FGeometryParticleHandle* InParticle);
+		CHAOS_API void Trash();
 
 		// An integer ID for an island (for debug only). @see FPBDIsland::GetIslandId()
-		int32 GetIslandId() const;
+		CHAOS_API int32 GetIslandId() const;
 
 	private:
 		friend class FPBDIslandManager;
@@ -107,12 +107,12 @@ namespace Chaos::Private
 	/**
 	* A constraint in the constraint graph.
 	*/
-	class CHAOS_API FPBDIslandConstraint
+	class FPBDIslandConstraint
 	{
 	public:
-		FPBDIslandConstraint();
-		FPBDIslandConstraint(const int32 InContainerId, FConstraintHandle* InConstraint);
-		~FPBDIslandConstraint();
+		CHAOS_API FPBDIslandConstraint();
+		CHAOS_API FPBDIslandConstraint(const int32 InContainerId, FConstraintHandle* InConstraint);
+		CHAOS_API ~FPBDIslandConstraint();
 
 		FConstraintHandle* GetConstraint() const
 		{
@@ -128,11 +128,11 @@ namespace Chaos::Private
 		// TPoolBackedItemAdapter for FPBDIslandManager::Edges
 		const int32 GetArrayIndex() const { return ArrayIndex; }
 		void SetArrayIndex(const int32 InIndex) { ArrayIndex = InIndex; }
-		void Reuse(const int32 InContainerId, FConstraintHandle* InConstraint);
-		void Trash();
+		CHAOS_API void Reuse(const int32 InContainerId, FConstraintHandle* InConstraint);
+		CHAOS_API void Trash();
 
 		// An integer ID for an island (for debug only). @see FPBDIsland::GetIslandId()
-		int32 GetIslandId() const;
+		CHAOS_API int32 GetIslandId() const;
 
 	private:
 		friend class FPBDIslandManager;
@@ -185,11 +185,11 @@ namespace Chaos::Private
 	* the set of these. To find kinematic particles used by an island you would need to visit all
 	* the constraints.
 	*/
-	class CHAOS_API FPBDIsland
+	class FPBDIsland
 	{
 	public:
-		FPBDIsland();
-		~FPBDIsland();
+		CHAOS_API FPBDIsland();
+		CHAOS_API ~FPBDIsland();
 
 		int32 GetNumParticles() const
 		{
@@ -281,13 +281,13 @@ namespace Chaos::Private
 		// TPoolBackedItemAdapter for FPBDIslandManager::Islands
 		const int32 GetArrayIndex() const { return ArrayIndex; }
 		void SetArrayIndex(const int32 InIndex) { ArrayIndex = InIndex; }
-		void Reuse();
-		void Trash();
+		CHAOS_API void Reuse();
+		CHAOS_API void Trash();
 
 	private:
 		friend class FPBDIslandManager;
 
-		void UpdateSyncState();
+		CHAOS_API void UpdateSyncState();
 
 		union FFlags
 		{
@@ -344,17 +344,17 @@ namespace Chaos::Private
 	/**
 	* A set of islands to be merged together.
 	*/
-	class CHAOS_API FPBDIslandMergeSet
+	class FPBDIslandMergeSet
 	{
 	public:
-		FPBDIslandMergeSet();
-		~FPBDIslandMergeSet();
+		CHAOS_API FPBDIslandMergeSet();
+		CHAOS_API ~FPBDIslandMergeSet();
 
 		// TPoolBackedItemAdapter for FPBDIslandManager::MergeSets
 		const int32 GetArrayIndex() const { return ArrayIndex; }
 		void SetArrayIndex(const int32 InIndex) { ArrayIndex = InIndex; }
-		void Reuse();
-		void Trash();
+		CHAOS_API void Reuse();
+		CHAOS_API void Trash();
 
 	private:
 		friend class FPBDIslandManager;
@@ -380,58 +380,58 @@ namespace Chaos::Private
 	* The main reason it is here is so that we can cache the thresholds and don't have to check 
 	* the materials for every particle every tick.
 	*/
-	class CHAOS_API FPBDIslandManager
+	class FPBDIslandManager
 	{
 	public:
-		FPBDIslandManager(FPBDRigidsSOAs& InParticles);
-		~FPBDIslandManager();
+		CHAOS_API FPBDIslandManager(FPBDRigidsSOAs& InParticles);
+		CHAOS_API ~FPBDIslandManager();
 
 		// The number of registered constraint containers
-		int32 GetNumConstraintContainers() const;
+		CHAOS_API int32 GetNumConstraintContainers() const;
 
 		// Add/Remove particles
 		int32 GetNumParticles() const { return Nodes.Num(); }
-		void AddParticle(FGeometryParticleHandle* Particle);
-		void RemoveParticle(FGeometryParticleHandle* Particle);
-		int32 ReserveParticles(const int32 InNumParticles);
-		void UpdateParticleMaterial(FGeometryParticleHandle* Particle);
-		int32 GetParticleLevel(FGeometryParticleHandle* Particle) const;
+		CHAOS_API void AddParticle(FGeometryParticleHandle* Particle);
+		CHAOS_API void RemoveParticle(FGeometryParticleHandle* Particle);
+		CHAOS_API int32 ReserveParticles(const int32 InNumParticles);
+		CHAOS_API void UpdateParticleMaterial(FGeometryParticleHandle* Particle);
+		CHAOS_API int32 GetParticleLevel(FGeometryParticleHandle* Particle) const;
 
 		// Add/Remove constraints
 		int32 GetNumConstraints() const { return Edges.Num(); }
-		void AddConstraint(const int32 ContainerId, FConstraintHandle* Constraint, const TVec2<FGeometryParticleHandle*>& ConstrainedParticles);
+		CHAOS_API void AddConstraint(const int32 ContainerId, FConstraintHandle* Constraint, const TVec2<FGeometryParticleHandle*>& ConstrainedParticles);
 		template<typename ConstraintContainerType> void AddContainerConstraints(ConstraintContainerType& ConstraintContainer);
-		void RemoveConstraint(FConstraintHandle* ConstraintHandle);
-		void RemoveParticleConstraints(FGeometryParticleHandle* Particle);
-		void RemoveParticleContainerConstraints(FGeometryParticleHandle* Particle, const int32 ContainerId);
-		void RemoveContainerConstraints(const int32 ContainerId);
+		CHAOS_API void RemoveConstraint(FConstraintHandle* ConstraintHandle);
+		CHAOS_API void RemoveParticleConstraints(FGeometryParticleHandle* Particle);
+		CHAOS_API void RemoveParticleContainerConstraints(FGeometryParticleHandle* Particle, const int32 ContainerId);
+		CHAOS_API void RemoveContainerConstraints(const int32 ContainerId);
 
 		// Access to Islands
 		int32 GetNumIslands() const { return Islands.Num(); }
 		const FPBDIsland* GetIsland(const int32 IslandIndex) const { return Islands[IslandIndex]; }
-		const FPBDIsland* GetParticleIsland(const FGeometryParticleHandle* Particle) const;
-		const FPBDIsland* GetConstraintIsland(const FConstraintHandle* Constraint) const;
+		CHAOS_API const FPBDIsland* GetParticleIsland(const FGeometryParticleHandle* Particle) const;
+		CHAOS_API const FPBDIsland* GetConstraintIsland(const FConstraintHandle* Constraint) const;
 
 		// Setup and update
-		void Reset();
-		void RemoveAllConstraints();
-		void AddConstraintContainer(const FPBDConstraintContainer& Container);
-		void RemoveConstraintContainer(const FPBDConstraintContainer& Container);
-		void SetMaterialContainers(const TArrayCollectionArray<TSerializablePtr<FChaosPhysicsMaterial>>* InPhysicsMaterials, const TArrayCollectionArray<TUniquePtr<FChaosPhysicsMaterial>>* InPerParticlePhysicsMaterials, const THandleArray<FChaosPhysicsMaterial>* InSimMaterials);
-		void SetGravityForces(const FPerParticleGravity* InGravity);
-		void SetDisableCounterThreshold(const int32 InDisableCounterThreshold);
-		void SetIsDeterministic(const bool bInIsDeterministic);
-		void UpdateParticles();
-		void UpdateIslands();
-		void UpdateSleep();
-		void UpdateDisable(TFunctionRef<void(FPBDRigidParticleHandle*)> ParticleDisableFunctor);
-		void EndTick();
+		CHAOS_API void Reset();
+		CHAOS_API void RemoveAllConstraints();
+		CHAOS_API void AddConstraintContainer(const FPBDConstraintContainer& Container);
+		CHAOS_API void RemoveConstraintContainer(const FPBDConstraintContainer& Container);
+		CHAOS_API void SetMaterialContainers(const TArrayCollectionArray<TSerializablePtr<FChaosPhysicsMaterial>>* InPhysicsMaterials, const TArrayCollectionArray<TUniquePtr<FChaosPhysicsMaterial>>* InPerParticlePhysicsMaterials, const THandleArray<FChaosPhysicsMaterial>* InSimMaterials);
+		CHAOS_API void SetGravityForces(const FPerParticleGravity* InGravity);
+		CHAOS_API void SetDisableCounterThreshold(const int32 InDisableCounterThreshold);
+		CHAOS_API void SetIsDeterministic(const bool bInIsDeterministic);
+		CHAOS_API void UpdateParticles();
+		CHAOS_API void UpdateIslands();
+		CHAOS_API void UpdateSleep();
+		CHAOS_API void UpdateDisable(TFunctionRef<void(FPBDRigidParticleHandle*)> ParticleDisableFunctor);
+		CHAOS_API void EndTick();
 
 		// Levels and Colors for sorting an parallelization
-		int32 GetParticleLevel(const FPBDIslandParticle* Node) const;
-		int32 GetParticleColor(const FPBDIslandParticle* Node) const;
-		int32 GetConstraintLevel(const FPBDIslandConstraint* Edge) const;
-		int32 GetConstraintColor(const FPBDIslandConstraint* Edge) const;
+		CHAOS_API int32 GetParticleLevel(const FPBDIslandParticle* Node) const;
+		CHAOS_API int32 GetParticleColor(const FPBDIslandParticle* Node) const;
+		CHAOS_API int32 GetConstraintLevel(const FPBDIslandConstraint* Edge) const;
+		CHAOS_API int32 GetConstraintColor(const FPBDIslandConstraint* Edge) const;
 
 		/**
 		* Visit all the awake constraints from the specified container.
@@ -448,18 +448,18 @@ namespace Chaos::Private
 		template<typename VisitorType> void VisitConstConstraints(const int32 ContainerId, const VisitorType& Visitor) const;
 
 		// Resim functionality
-		int32 GetParticleResimFrame(const FGeometryParticleHandle* Particle) const;
-		void SetParticleResimFrame(FGeometryParticleHandle* Particle, const int32 ResimFrame);
-		void ResetParticleResimFrame(const int32 ResetFrame = INDEX_NONE);
+		CHAOS_API int32 GetParticleResimFrame(const FGeometryParticleHandle* Particle) const;
+		CHAOS_API void SetParticleResimFrame(FGeometryParticleHandle* Particle, const int32 ResimFrame);
+		CHAOS_API void ResetParticleResimFrame(const int32 ResetFrame = INDEX_NONE);
 
 		// The following is intended for testing, debugging and visualization only and is not considered part of the API
 		FPBDIsland* GetIsland(const int32 IslandIndex) { return Islands[IslandIndex]; }
-		int32 GetIslandIndex(const FPBDIsland* Island) const;
-		int32 GetIslandArrayIndex(const FPBDIslandConstraint* Edge) const;
-		TArray<const FPBDIsland*> FindParticleIslands(const FGeometryParticleHandle* Particle) const;
-		TArray<const FGeometryParticleHandle*> FindParticlesInIslands(const TArray<const FPBDIsland*> Islands) const;
-		TArray<const FConstraintHandle*> FindConstraintsInIslands(const TArray<const FPBDIsland*> Islands, int32 ContainerId) const;
-		void SetParticleIslandIsSleeping(FGeometryParticleHandle* Particle, const bool bInIsSleeping);
+		CHAOS_API int32 GetIslandIndex(const FPBDIsland* Island) const;
+		CHAOS_API int32 GetIslandArrayIndex(const FPBDIslandConstraint* Edge) const;
+		CHAOS_API TArray<const FPBDIsland*> FindParticleIslands(const FGeometryParticleHandle* Particle) const;
+		CHAOS_API TArray<const FGeometryParticleHandle*> FindParticlesInIslands(const TArray<const FPBDIsland*> Islands) const;
+		CHAOS_API TArray<const FConstraintHandle*> FindConstraintsInIslands(const TArray<const FPBDIsland*> Islands, int32 ContainerId) const;
+		CHAOS_API void SetParticleIslandIsSleeping(FGeometryParticleHandle* Particle, const bool bInIsSleeping);
 
 		// Deprecated API
 		UE_DEPRECATED(5.3, "Use Reset") void RemoveConstraints() { Reset(); }
@@ -474,75 +474,75 @@ namespace Chaos::Private
 		friend class FPBDRigidsEvolutionGBF;
 
 		// Node management
-		FPBDIslandParticle* GetGraphNode(const FGeometryParticleHandle* Particle) const;
-		FPBDIslandParticle* GetGraphNode(const FTransientGeometryParticleHandle& Particle) const;
-		FPBDIslandParticle* CreateGraphNode(FGeometryParticleHandle* Particle);
-		FPBDIslandParticle* GetOrCreateGraphNode(FGeometryParticleHandle* Particle);
-		void DestroyGraphNode(FPBDIslandParticle* Node);
+		CHAOS_API FPBDIslandParticle* GetGraphNode(const FGeometryParticleHandle* Particle) const;
+		CHAOS_API FPBDIslandParticle* GetGraphNode(const FTransientGeometryParticleHandle& Particle) const;
+		CHAOS_API FPBDIslandParticle* CreateGraphNode(FGeometryParticleHandle* Particle);
+		CHAOS_API FPBDIslandParticle* GetOrCreateGraphNode(FGeometryParticleHandle* Particle);
+		CHAOS_API void DestroyGraphNode(FPBDIslandParticle* Node);
 
 		// Edge management
-		FPBDIslandConstraint* GetGraphEdge(const FConstraintHandle* Constraint);
-		const FPBDIslandConstraint* GetGraphEdge(const FConstraintHandle* Constraint) const;
-		FPBDIslandConstraint* CreateGraphEdge(const int32 ContainerId, FConstraintHandle* Constraint, const TVec2<FGeometryParticleHandle*>& ConstrainedParticles);
-		void DestroyGraphEdge(FPBDIslandConstraint* Edge);
-		void BindEdgeToNodes(FPBDIslandConstraint* Edge, FPBDIslandParticle* Node0, FPBDIslandParticle* Node1);
-		void UnbindEdgeFromNodes(FPBDIslandConstraint* Edge);
-		void UpdateGraphNodeSleepSettings(FPBDIslandParticle* Node);
+		CHAOS_API FPBDIslandConstraint* GetGraphEdge(const FConstraintHandle* Constraint);
+		CHAOS_API const FPBDIslandConstraint* GetGraphEdge(const FConstraintHandle* Constraint) const;
+		CHAOS_API FPBDIslandConstraint* CreateGraphEdge(const int32 ContainerId, FConstraintHandle* Constraint, const TVec2<FGeometryParticleHandle*>& ConstrainedParticles);
+		CHAOS_API void DestroyGraphEdge(FPBDIslandConstraint* Edge);
+		CHAOS_API void BindEdgeToNodes(FPBDIslandConstraint* Edge, FPBDIslandParticle* Node0, FPBDIslandParticle* Node1);
+		CHAOS_API void UnbindEdgeFromNodes(FPBDIslandConstraint* Edge);
+		CHAOS_API void UpdateGraphNodeSleepSettings(FPBDIslandParticle* Node);
 
 		// Per-tick state transfer to nodes/edges
-		void UpdateGraphNode(FPBDIslandParticle* Node);
-		void UpdateGraphEdge(FPBDIslandConstraint* Edge);
+		CHAOS_API void UpdateGraphNode(FPBDIslandParticle* Node);
+		CHAOS_API void UpdateGraphEdge(FPBDIslandConstraint* Edge);
 
 		// Island managament
-		FPBDIsland* CreateIsland();
-		void DestroyIsland(FPBDIsland* Island);
-		void AssignNodeIsland(FPBDIslandParticle* Node);
-		void AssignEdgeIsland(FPBDIslandConstraint* Edge);
-		void AddNodeToIsland(FPBDIslandParticle* Node, FPBDIsland* Island);
-		void RemoveNodeFromIsland(FPBDIslandParticle* Node);
-		void AddEdgeToIsland(FPBDIslandConstraint* Edge, FPBDIsland* Island);
-		void RemoveEdgeFromIsland(FPBDIslandConstraint* Edge);
-		void EnqueueIslandCheckSleep(FPBDIsland* Island, const bool bIsSleepAllowed);
+		CHAOS_API FPBDIsland* CreateIsland();
+		CHAOS_API void DestroyIsland(FPBDIsland* Island);
+		CHAOS_API void AssignNodeIsland(FPBDIslandParticle* Node);
+		CHAOS_API void AssignEdgeIsland(FPBDIslandConstraint* Edge);
+		CHAOS_API void AddNodeToIsland(FPBDIslandParticle* Node, FPBDIsland* Island);
+		CHAOS_API void RemoveNodeFromIsland(FPBDIslandParticle* Node);
+		CHAOS_API void AddEdgeToIsland(FPBDIslandConstraint* Edge, FPBDIsland* Island);
+		CHAOS_API void RemoveEdgeFromIsland(FPBDIslandConstraint* Edge);
+		CHAOS_API void EnqueueIslandCheckSleep(FPBDIsland* Island, const bool bIsSleepAllowed);
 
 		// Island merge set management
-		FPBDIsland* MergeNodeIslands(FPBDIslandParticle* Node0, FPBDIslandParticle* Node1);
-		FPBDIsland* MergeIslands(FPBDIsland* Island0, FPBDIsland* Island1);
-		FPBDIslandMergeSet* CreateMergeSet(FPBDIsland* Island0, FPBDIsland* Island1);
-		void DestroyMergeSet(FPBDIslandMergeSet* IslandMergeSet);
-		void AddIslandToMergeSet(FPBDIsland* Island, FPBDIslandMergeSet* IslandMergeSet);
-		void RemoveIslandFromMergeSet(FPBDIsland* Island);
-		FPBDIslandMergeSet* CombineMergeSets(FPBDIslandMergeSet* IslandMergeSetParent, FPBDIslandMergeSet* IslandMergeSetChild);
-		FPBDIsland* GetMergeSetParentIsland(FPBDIslandMergeSet* MergeSet, int32& OutNumEdges, const TArrayView<int32>& OutNumContainerEdges);
+		CHAOS_API FPBDIsland* MergeNodeIslands(FPBDIslandParticle* Node0, FPBDIslandParticle* Node1);
+		CHAOS_API FPBDIsland* MergeIslands(FPBDIsland* Island0, FPBDIsland* Island1);
+		CHAOS_API FPBDIslandMergeSet* CreateMergeSet(FPBDIsland* Island0, FPBDIsland* Island1);
+		CHAOS_API void DestroyMergeSet(FPBDIslandMergeSet* IslandMergeSet);
+		CHAOS_API void AddIslandToMergeSet(FPBDIsland* Island, FPBDIslandMergeSet* IslandMergeSet);
+		CHAOS_API void RemoveIslandFromMergeSet(FPBDIsland* Island);
+		CHAOS_API FPBDIslandMergeSet* CombineMergeSets(FPBDIslandMergeSet* IslandMergeSetParent, FPBDIslandMergeSet* IslandMergeSetChild);
+		CHAOS_API FPBDIsland* GetMergeSetParentIsland(FPBDIslandMergeSet* MergeSet, int32& OutNumEdges, const TArrayView<int32>& OutNumContainerEdges);
 
 		// Merge/split processing
-		void ProcessIslands();
-		void ProcessWakes();
-		void ProcessMerges();
-		void ProcessSplits();
-		void AssignLevels();
-		void ProcessIslandMerge(FPBDIsland* ParentIsland, FPBDIsland* ChildIsland);
-		void ProcessIslandSplits(FPBDIsland* Island);
+		CHAOS_API void ProcessIslands();
+		CHAOS_API void ProcessWakes();
+		CHAOS_API void ProcessMerges();
+		CHAOS_API void ProcessSplits();
+		CHAOS_API void AssignLevels();
+		CHAOS_API void ProcessIslandMerge(FPBDIsland* ParentIsland, FPBDIsland* ChildIsland);
+		CHAOS_API void ProcessIslandSplits(FPBDIsland* Island);
 
 		// Sorting
-		void AssignIslandLevels(FPBDIsland* Island);
-		void SortIslandEdges(FPBDIsland* Island);
+		CHAOS_API void AssignIslandLevels(FPBDIsland* Island);
+		CHAOS_API void SortIslandEdges(FPBDIsland* Island);
 
 		// Sleeping
-		void ProcessSleep();
-		void ProcessParticlesSleep();
-		void ProcessIslandSleep(FPBDIsland* Island);
-		void PropagateIslandSleep(FPBDIsland* Island);
+		CHAOS_API void ProcessSleep();
+		CHAOS_API void ProcessParticlesSleep();
+		CHAOS_API void ProcessIslandSleep(FPBDIsland* Island);
+		CHAOS_API void PropagateIslandSleep(FPBDIsland* Island);
 
 		// Disabling
-		void ProcessDisable(TFunctionRef<void(FPBDRigidParticleHandle*)> ParticleDisableFunctor);
-		void ProcessParticlesDisable(TArray<FPBDRigidParticleHandle*>& OutDisableParticles);
-		void ProcessIslandDisable(FPBDIsland* Island, TArray<FPBDRigidParticleHandle*>& OutDisableParticles);
+		CHAOS_API void ProcessDisable(TFunctionRef<void(FPBDRigidParticleHandle*)> ParticleDisableFunctor);
+		CHAOS_API void ProcessParticlesDisable(TArray<FPBDRigidParticleHandle*>& OutDisableParticles);
+		CHAOS_API void ProcessIslandDisable(FPBDIsland* Island, TArray<FPBDRigidParticleHandle*>& OutDisableParticles);
 
-		void FinalizeIslands();
-		void ApplyDeterminism();
-		bool Validate() const;
+		CHAOS_API void FinalizeIslands();
+		CHAOS_API void ApplyDeterminism();
+		CHAOS_API bool Validate() const;
 
-		int32 GetNextVisitEpoch();
+		CHAOS_API int32 GetNextVisitEpoch();
 
 		// Ideally the island manager would not know about the particles, materials etc, but
 		// it is currently required for particle sleep management
diff --git a/Engine/Source/Runtime/Experimental/Chaos/Public/Chaos/Levelset.h b/Engine/Source/Runtime/Experimental/Chaos/Public/Chaos/Levelset.h
--- a/Engine/Source/Runtime/Experimental/Chaos/Public/Chaos/Levelset.h
+++ b/Engine/Source/Runtime/Experimental/Chaos/Public/Chaos/Levelset.h
@@ -26,31 +26,31 @@ class FCapsule;
 
 class FConvex;
 
-class CHAOS_API FLevelSet final : public FImplicitObject
+class FLevelSet final : public FImplicitObject
 {
   public:
 	using FImplicitObject::SignedDistance;
 
-	FLevelSet(FErrorReporter& ErrorReporter, const TUniformGrid<FReal, 3>& InGrid, const FParticles& InParticles, const FTriangleMesh& Mesh, const int32 BandWidth = 0);
-	FLevelSet(FErrorReporter& ErrorReporter, const TUniformGrid<FReal, 3>& InGrid, const FImplicitObject& InObject, const int32 BandWidth = 0, const bool bUseObjectPhi = false);
+	CHAOS_API FLevelSet(FErrorReporter& ErrorReporter, const TUniformGrid<FReal, 3>& InGrid, const FParticles& InParticles, const FTriangleMesh& Mesh, const int32 BandWidth = 0);
+	CHAOS_API FLevelSet(FErrorReporter& ErrorReporter, const TUniformGrid<FReal, 3>& InGrid, const FImplicitObject& InObject, const int32 BandWidth = 0, const bool bUseObjectPhi = false);
 #if COMPILE_WITHOUT_UNREAL_SUPPORT
 	FLevelSet(std::istream& Stream);
 #endif
-	FLevelSet(TUniformGrid<FReal, 3>&& Grid, TArrayND<FReal, 3>&& Phi, int32 BandWidth);
-	FLevelSet(const FLevelSet& Other) = delete;
-	FLevelSet(FLevelSet&& Other);
-	virtual ~FLevelSet();
+	CHAOS_API FLevelSet(TUniformGrid<FReal, 3>&& Grid, TArrayND<FReal, 3>&& Phi, int32 BandWidth);
+	CHAOS_API FLevelSet(const FLevelSet& Other) = delete;
+	CHAOS_API FLevelSet(FLevelSet&& Other);
+	CHAOS_API virtual ~FLevelSet();
 
 	virtual TUniquePtr<FImplicitObject> Copy() const override { return DeepCopy(); }
 	virtual TUniquePtr<FImplicitObject> CopyWithScale(const FVec3& Scale) const { return DeepCopyWithScale(Scale); }
-	virtual TUniquePtr<FImplicitObject> DeepCopy() const;
-	virtual TUniquePtr<FImplicitObject> DeepCopyWithScale(const FVec3& Scale) const;
+	CHAOS_API virtual TUniquePtr<FImplicitObject> DeepCopy() const;
+	CHAOS_API virtual TUniquePtr<FImplicitObject> DeepCopyWithScale(const FVec3& Scale) const;
 
 #if COMPILE_WITHOUT_UNREAL_SUPPORT
-	void Write(std::ostream& Stream) const;
+	CHAOS_API void Write(std::ostream& Stream) const;
 #endif
-	virtual FReal PhiWithNormal(const FVec3& x, FVec3& Normal) const override;
-	FReal SignedDistance(const FVec3& x) const;
+	CHAOS_API virtual FReal PhiWithNormal(const FVec3& x, FVec3& Normal) const override;
+	CHAOS_API FReal SignedDistance(const FVec3& x) const;
 
 	virtual const FAABB3 BoundingBox() const override { return MLocalBoundingBox; }
 
@@ -129,18 +129,18 @@ class CHAOS_API FLevelSet final : public FImplicitObject
 		return Volume;
 	}
 
-	bool ComputeMassProperties(FReal& OutVolume, FVec3& OutCOM, FMatrix33& OutInertia, FRotation3& OutRotationOfMass) const;
+	CHAOS_API bool ComputeMassProperties(FReal& OutVolume, FVec3& OutCOM, FMatrix33& OutInertia, FRotation3& OutRotationOfMass) const;
 
-	FReal ComputeLevelSetError(const FParticles& InParticles, const TArray<FVec3>& Normals, const FTriangleMesh& Mesh, FReal& AngleError, FReal& MaxDistError);
+	CHAOS_API FReal ComputeLevelSetError(const FParticles& InParticles, const TArray<FVec3>& Normals, const FTriangleMesh& Mesh, FReal& AngleError, FReal& MaxDistError);
 
 	// Output a mesh and level set as obj files
-	void OutputDebugData(FErrorReporter& ErrorReporter, const FParticles& InParticles, const TArray<FVec3>& Normals, const FTriangleMesh& Mesh, const FString FileName);
+	CHAOS_API void OutputDebugData(FErrorReporter& ErrorReporter, const FParticles& InParticles, const TArray<FVec3>& Normals, const FTriangleMesh& Mesh, const FString FileName);
 
-	bool CheckData(FErrorReporter& ErrorReporter, const FParticles& InParticles, const FTriangleMesh& Mesh, const TArray<FVec3> &Normals);
+	CHAOS_API bool CheckData(FErrorReporter& ErrorReporter, const FParticles& InParticles, const FTriangleMesh& Mesh, const TArray<FVec3> &Normals);
 
 	// Used to generate a simple debug surface
-	void GetZeroIsosurfaceGridCellFaces(TArray<FVector3f>& Vertices, TArray<FIntVector>& Tris) const;
-	void GetInteriorCells(TArray<TVec3<int32>>& InteriorCells, const FReal InteriorThreshold) const;
+	CHAOS_API void GetZeroIsosurfaceGridCellFaces(TArray<FVector3f>& Vertices, TArray<FIntVector>& Tris) const;
+	CHAOS_API void GetInteriorCells(TArray<TVec3<int32>>& InteriorCells, const FReal InteriorThreshold) const;
 
 	virtual uint32 GetTypeHash() const override
 	{
@@ -157,45 +157,45 @@ class CHAOS_API FLevelSet final : public FImplicitObject
 	}
 
 
-	bool SweepGeom(const TSphere<FReal, 3>& QueryGeom, const FRigidTransform3& StartTM, const FVec3& Dir, const FReal Length, FReal& OutTime, FVec3& OutPosition, FVec3& OutNormal, int32& OutFaceIndex, const FReal Thickness = 0, const bool bComputeMTD = false) const;
-	bool SweepGeom(const TBox<FReal, 3>& QueryGeom, const FRigidTransform3& StartTM, const FVec3& Dir, const FReal Length, FReal& OutTime, FVec3& OutPosition, FVec3& OutNormal, int32& OutFaceIndex, const FReal Thickness = 0, const bool bComputeMTD = false) const;
-	bool SweepGeom(const FCapsule& QueryGeom, const FRigidTransform3& StartTM, const FVec3& Dir, const FReal Length, FReal& OutTime, FVec3& OutPosition, FVec3& OutNormal, int32& OutFaceIndex, const FReal Thickness = 0, const bool bComputeMTD = false) const;
-	bool SweepGeom(const FConvex& QueryGeom, const FRigidTransform3& StartTM, const FVec3& Dir, const FReal Length, FReal& OutTime, FVec3& OutPosition, FVec3& OutNormal, int32& OutFaceIndex, const FReal Thickness = 0, const bool bComputeMTD = false) const;
+	CHAOS_API bool SweepGeom(const TSphere<FReal, 3>& QueryGeom, const FRigidTransform3& StartTM, const FVec3& Dir, const FReal Length, FReal& OutTime, FVec3& OutPosition, FVec3& OutNormal, int32& OutFaceIndex, const FReal Thickness = 0, const bool bComputeMTD = false) const;
+	CHAOS_API bool SweepGeom(const TBox<FReal, 3>& QueryGeom, const FRigidTransform3& StartTM, const FVec3& Dir, const FReal Length, FReal& OutTime, FVec3& OutPosition, FVec3& OutNormal, int32& OutFaceIndex, const FReal Thickness = 0, const bool bComputeMTD = false) const;
+	CHAOS_API bool SweepGeom(const FCapsule& QueryGeom, const FRigidTransform3& StartTM, const FVec3& Dir, const FReal Length, FReal& OutTime, FVec3& OutPosition, FVec3& OutNormal, int32& OutFaceIndex, const FReal Thickness = 0, const bool bComputeMTD = false) const;
+	CHAOS_API bool SweepGeom(const FConvex& QueryGeom, const FRigidTransform3& StartTM, const FVec3& Dir, const FReal Length, FReal& OutTime, FVec3& OutPosition, FVec3& OutNormal, int32& OutFaceIndex, const FReal Thickness = 0, const bool bComputeMTD = false) const;
 
-	bool SweepGeom(const TImplicitObjectScaled<TSphere<FReal, 3>>& QueryGeom, const FRigidTransform3& StartTM, const FVec3& Dir, const FReal Length, FReal& OutTime, FVec3& OutPosition, FVec3& OutNormal, int32& OutFaceIndex, const FReal Thickness = 0, const bool bComputeMTD = false) const;
-	bool SweepGeom(const TImplicitObjectScaled<TBox<FReal, 3>>& QueryGeom, const FRigidTransform3& StartTM, const FVec3& Dir, const FReal Length, FReal& OutTime, FVec3& OutPosition, FVec3& OutNormal, int32& OutFaceIndex, const FReal Thickness = 0, const bool bComputeMTD = false) const;
-	bool SweepGeom(const TImplicitObjectScaled<FCapsule>& QueryGeom, const FRigidTransform3& StartTM, const FVec3& Dir, const FReal Length, FReal& OutTime, FVec3& OutPosition, FVec3& OutNormal, int32& OutFaceIndex, const FReal Thickness = 0, const bool bComputeMTD = false) const;
-	bool SweepGeom(const TImplicitObjectScaled<FConvex>& QueryGeom, const FRigidTransform3& StartTM, const FVec3& Dir, const FReal Length, FReal& OutTime, FVec3& OutPosition, FVec3& OutNormal, int32& OutFaceIndex, const FReal Thickness = 0, const bool bComputeMTD = false) const;
+	CHAOS_API bool SweepGeom(const TImplicitObjectScaled<TSphere<FReal, 3>>& QueryGeom, const FRigidTransform3& StartTM, const FVec3& Dir, const FReal Length, FReal& OutTime, FVec3& OutPosition, FVec3& OutNormal, int32& OutFaceIndex, const FReal Thickness = 0, const bool bComputeMTD = false) const;
+	CHAOS_API bool SweepGeom(const TImplicitObjectScaled<TBox<FReal, 3>>& QueryGeom, const FRigidTransform3& StartTM, const FVec3& Dir, const FReal Length, FReal& OutTime, FVec3& OutPosition, FVec3& OutNormal, int32& OutFaceIndex, const FReal Thickness = 0, const bool bComputeMTD = false) const;
+	CHAOS_API bool SweepGeom(const TImplicitObjectScaled<FCapsule>& QueryGeom, const FRigidTransform3& StartTM, const FVec3& Dir, const FReal Length, FReal& OutTime, FVec3& OutPosition, FVec3& OutNormal, int32& OutFaceIndex, const FReal Thickness = 0, const bool bComputeMTD = false) const;
+	CHAOS_API bool SweepGeom(const TImplicitObjectScaled<FConvex>& QueryGeom, const FRigidTransform3& StartTM, const FVec3& Dir, const FReal Length, FReal& OutTime, FVec3& OutPosition, FVec3& OutNormal, int32& OutFaceIndex, const FReal Thickness = 0, const bool bComputeMTD = false) const;
 
 	template<typename QueryGeomType>
 	bool SweepGeomImp(const QueryGeomType& QueryGeom, const FRigidTransform3& StartTM, const FVec3& Dir, const FReal Length, FReal& OutTime, FVec3& OutPosition, FVec3& OutNormal, int32& OutFaceIndex, const FReal Thickness, const bool bComputeMTD) const;
 
-	bool OverlapGeom(const TSphere<FReal, 3>& QueryGeom, const FRigidTransform3& QueryTM, const FReal Thickness, FMTDInfo* OutMTD = nullptr) const;
-	bool OverlapGeom(const TBox<FReal, 3>& QueryGeom, const FRigidTransform3& QueryTM, const FReal Thickness, FMTDInfo* OutMTD = nullptr) const;
-	bool OverlapGeom(const FCapsule& QueryGeom, const FRigidTransform3& QueryTM, const FReal Thickness, FMTDInfo* OutMTD = nullptr) const;
-	bool OverlapGeom(const FConvex& QueryGeom, const FRigidTransform3& QueryTM, const FReal Thickness, FMTDInfo* OutMTD = nullptr) const;
+	CHAOS_API bool OverlapGeom(const TSphere<FReal, 3>& QueryGeom, const FRigidTransform3& QueryTM, const FReal Thickness, FMTDInfo* OutMTD = nullptr) const;
+	CHAOS_API bool OverlapGeom(const TBox<FReal, 3>& QueryGeom, const FRigidTransform3& QueryTM, const FReal Thickness, FMTDInfo* OutMTD = nullptr) const;
+	CHAOS_API bool OverlapGeom(const FCapsule& QueryGeom, const FRigidTransform3& QueryTM, const FReal Thickness, FMTDInfo* OutMTD = nullptr) const;
+	CHAOS_API bool OverlapGeom(const FConvex& QueryGeom, const FRigidTransform3& QueryTM, const FReal Thickness, FMTDInfo* OutMTD = nullptr) const;
 
-	bool OverlapGeom(const TImplicitObjectScaled<TSphere<FReal, 3>>& QueryGeom, const FRigidTransform3& QueryTM, const FReal Thickness, FMTDInfo* OutMTD = nullptr) const;
-	bool OverlapGeom(const TImplicitObjectScaled<TBox<FReal, 3>>& QueryGeom, const FRigidTransform3& QueryTM, const FReal Thickness, FMTDInfo* OutMTD = nullptr) const;
-	bool OverlapGeom(const TImplicitObjectScaled<FCapsule>& QueryGeom, const FRigidTransform3& QueryTM, const FReal Thickness, FMTDInfo* OutMTD = nullptr) const;
-	bool OverlapGeom(const TImplicitObjectScaled<FConvex>& QueryGeom, const FRigidTransform3& QueryTM, const FReal Thickness, FMTDInfo* OutMTD = nullptr) const;
+	CHAOS_API bool OverlapGeom(const TImplicitObjectScaled<TSphere<FReal, 3>>& QueryGeom, const FRigidTransform3& QueryTM, const FReal Thickness, FMTDInfo* OutMTD = nullptr) const;
+	CHAOS_API bool OverlapGeom(const TImplicitObjectScaled<TBox<FReal, 3>>& QueryGeom, const FRigidTransform3& QueryTM, const FReal Thickness, FMTDInfo* OutMTD = nullptr) const;
+	CHAOS_API bool OverlapGeom(const TImplicitObjectScaled<FCapsule>& QueryGeom, const FRigidTransform3& QueryTM, const FReal Thickness, FMTDInfo* OutMTD = nullptr) const;
+	CHAOS_API bool OverlapGeom(const TImplicitObjectScaled<FConvex>& QueryGeom, const FRigidTransform3& QueryTM, const FReal Thickness, FMTDInfo* OutMTD = nullptr) const;
 
 	template<typename QueryGeomType>
 	bool OverlapGeomImp(const QueryGeomType& QueryGeom, const FRigidTransform3& QueryTM, const FReal Thickness, FMTDInfo* OutMTD = nullptr) const;
 
   private:
-	bool ComputeDistancesNearZeroIsocontour(FErrorReporter& ErrorReporter, const FParticles& InParticles, const TArray<FVec3> &Normals, const FTriangleMesh& Mesh, TArrayND<bool, 3>& BlockedFaceX, TArrayND<bool, 3>& BlockedFaceY, TArrayND<bool, 3>& BlockedFaceZ, TArray<TVec3<int32>>& InterfaceIndices);
-	void ComputeDistancesNearZeroIsocontour(const FImplicitObject& Object, const TArrayND<FReal, 3>& ObjectPhi, TArray<TVec3<int32>>& InterfaceIndices);
-	void CorrectSign(const TArrayND<bool, 3>& BlockedFaceX, const TArrayND<bool, 3>& BlockedFaceY, const TArrayND<bool, 3>& BlockedFaceZ, TArray<TVec3<int32>>& InterfaceIndices);
-	FReal ComputePhi(const TArrayND<bool, 3>& Done, const TVec3<int32>& CellIndex);
-	void FillWithFastMarchingMethod(const FReal StoppingDistance, const TArray<TVec3<int32>>& InterfaceIndices);
-	void FloodFill(const TArrayND<bool, 3>& BlockedFaceX, const TArrayND<bool, 3>& BlockedFaceY, const TArrayND<bool, 3>& BlockedFaceZ, TArrayND<int32, 3>& Color, int32& NextColor);
-	void FloodFillFromCell(const TVec3<int32> CellIndex, const int32 NextColor, const TArrayND<bool, 3>& BlockedFaceX, const TArrayND<bool, 3>& BlockedFaceY, const TArrayND<bool, 3>& BlockedFaceZ, TArrayND<int32, 3>& Color);
-	bool IsIntersectingWithTriangle(const FParticles& Particles, const TVec3<int32>& Elements, const TPlane<FReal, 3>& TrianglePlane, const TVec3<int32>& CellIndex, const TVec3<int32>& PrevCellIndex);
-	void ComputeNormals();
-	void ComputeConvexity(const TArray<TVec3<int32>>& InterfaceIndices);
+	CHAOS_API bool ComputeDistancesNearZeroIsocontour(FErrorReporter& ErrorReporter, const FParticles& InParticles, const TArray<FVec3> &Normals, const FTriangleMesh& Mesh, TArrayND<bool, 3>& BlockedFaceX, TArrayND<bool, 3>& BlockedFaceY, TArrayND<bool, 3>& BlockedFaceZ, TArray<TVec3<int32>>& InterfaceIndices);
+	CHAOS_API void ComputeDistancesNearZeroIsocontour(const FImplicitObject& Object, const TArrayND<FReal, 3>& ObjectPhi, TArray<TVec3<int32>>& InterfaceIndices);
+	CHAOS_API void CorrectSign(const TArrayND<bool, 3>& BlockedFaceX, const TArrayND<bool, 3>& BlockedFaceY, const TArrayND<bool, 3>& BlockedFaceZ, TArray<TVec3<int32>>& InterfaceIndices);
+	CHAOS_API FReal ComputePhi(const TArrayND<bool, 3>& Done, const TVec3<int32>& CellIndex);
+	CHAOS_API void FillWithFastMarchingMethod(const FReal StoppingDistance, const TArray<TVec3<int32>>& InterfaceIndices);
+	CHAOS_API void FloodFill(const TArrayND<bool, 3>& BlockedFaceX, const TArrayND<bool, 3>& BlockedFaceY, const TArrayND<bool, 3>& BlockedFaceZ, TArrayND<int32, 3>& Color, int32& NextColor);
+	CHAOS_API void FloodFillFromCell(const TVec3<int32> CellIndex, const int32 NextColor, const TArrayND<bool, 3>& BlockedFaceX, const TArrayND<bool, 3>& BlockedFaceY, const TArrayND<bool, 3>& BlockedFaceZ, TArrayND<int32, 3>& Color);
+	CHAOS_API bool IsIntersectingWithTriangle(const FParticles& Particles, const TVec3<int32>& Elements, const TPlane<FReal, 3>& TrianglePlane, const TVec3<int32>& CellIndex, const TVec3<int32>& PrevCellIndex);
+	CHAOS_API void ComputeNormals();
+	CHAOS_API void ComputeConvexity(const TArray<TVec3<int32>>& InterfaceIndices);
 	
-	void ComputeNormals(const FParticles& InParticles, const FTriangleMesh& Mesh, const TArray<TVec3<int32>>& InterfaceIndices);
+	CHAOS_API void ComputeNormals(const FParticles& InParticles, const FTriangleMesh& Mesh, const TArray<TVec3<int32>>& InterfaceIndices);
 
 	TUniformGrid<FReal, 3> MGrid;
 	TArrayND<FReal, 3> MPhi;
@@ -211,4 +211,4 @@ private:
 
 template <typename T, int d>
 using TLevelSet UE_DEPRECATED(4.27, "Deprecated. this class is to be deleted, use FLevelSet instead") = FLevelSet;
-}
\ No newline at end of file
+}
diff --git a/Engine/Source/Runtime/Experimental/Chaos/Public/Chaos/MassProperties.h b/Engine/Source/Runtime/Experimental/Chaos/Public/Chaos/MassProperties.h
--- a/Engine/Source/Runtime/Experimental/Chaos/Public/Chaos/MassProperties.h
+++ b/Engine/Source/Runtime/Experimental/Chaos/Public/Chaos/MassProperties.h
@@ -20,7 +20,7 @@ namespace Chaos
 	class TParticles;
 	using FParticles = TParticles<FReal, 3>;
 
-	struct CHAOS_API FMassProperties
+	struct FMassProperties
 	{
 		FMassProperties()
 			: Mass(0)
diff --git a/Engine/Source/Runtime/Experimental/Chaos/Public/Chaos/MidPhaseModification.h b/Engine/Source/Runtime/Experimental/Chaos/Public/Chaos/MidPhaseModification.h
--- a/Engine/Source/Runtime/Experimental/Chaos/Public/Chaos/MidPhaseModification.h
+++ b/Engine/Source/Runtime/Experimental/Chaos/Public/Chaos/MidPhaseModification.h
@@ -17,7 +17,7 @@ namespace Chaos
 	 * Also contains the Modifier, which is responsible for iterating over PairModifiers, and
 	 * storing and executing requested modifications.
 	 */
-	class CHAOS_API FMidPhaseModifier
+	class FMidPhaseModifier
 	{
 	public:
 		FMidPhaseModifier()
@@ -44,18 +44,18 @@ namespace Chaos
 		//
 
 		// Disable this midphase entirely
-		void Disable();
+		CHAOS_API void Disable();
 
 		// Disable CCD for this pair
-		void DisableCCD();
+		CHAOS_API void DisableCCD();
 
 		//
 		// Accessor functions
 		//
 
-		void GetParticles(const FGeometryParticleHandle** Particle0, const FGeometryParticleHandle** Particle1) const;
+		CHAOS_API void GetParticles(const FGeometryParticleHandle** Particle0, const FGeometryParticleHandle** Particle1) const;
 
-		const FGeometryParticleHandle* GetOtherParticle(const FGeometryParticleHandle* InParticle) const;
+		CHAOS_API const FGeometryParticleHandle* GetOtherParticle(const FGeometryParticleHandle* InParticle) const;
 
 
 	private:
@@ -67,7 +67,7 @@ namespace Chaos
 	/*
 	 * Class for iterating over midphases involving a specific particle
 	 */
-	class CHAOS_API FMidPhaseModifierParticleIterator
+	class FMidPhaseModifierParticleIterator
 	{
 	public:
 		FMidPhaseModifierParticleIterator(
@@ -121,15 +121,15 @@ namespace Chaos
 		friend class FMidPhaseModifierParticleRange;
 	};
 
-	class CHAOS_API FMidPhaseModifierParticleRange
+	class FMidPhaseModifierParticleRange
 	{
 	public:
 		FMidPhaseModifierParticleRange(
 			FMidPhaseModifierAccessor* InAccessor, FGeometryParticleHandle* InParticle)
 			: Accessor(InAccessor)
 			, Particle(InParticle) { }
-		FMidPhaseModifierParticleIterator begin() const;
-		FMidPhaseModifierParticleIterator end() const;
+		CHAOS_API FMidPhaseModifierParticleIterator begin() const;
+		CHAOS_API FMidPhaseModifierParticleIterator end() const;
 	private:
 		FMidPhaseModifierAccessor* Accessor;
 		FGeometryParticleHandle* Particle;
@@ -138,14 +138,14 @@ namespace Chaos
 	/*
 	 * Provides interface for accessing midphase pair modifiers
 	 */
-	class CHAOS_API FMidPhaseModifierAccessor
+	class FMidPhaseModifierAccessor
 	{
 	public:
 		// Get an object which allows for range iteration over the list of
 		// midphases for a particle
-		FMidPhaseModifierParticleRange GetMidPhases(FGeometryParticleHandle* Particle);
+		CHAOS_API FMidPhaseModifierParticleRange GetMidPhases(FGeometryParticleHandle* Particle);
 
 		// Get a midphase modifier for a particular object pair
-		FMidPhaseModifier GetMidPhase(FGeometryParticleHandle* Particle0, FGeometryParticleHandle* Particle1);
+		CHAOS_API FMidPhaseModifier GetMidPhase(FGeometryParticleHandle* Particle0, FGeometryParticleHandle* Particle1);
 	};
 }
diff --git a/Engine/Source/Runtime/Experimental/Chaos/Public/Chaos/NewtonEvolution.h b/Engine/Source/Runtime/Experimental/Chaos/Public/Chaos/NewtonEvolution.h
--- a/Engine/Source/Runtime/Experimental/Chaos/Public/Chaos/NewtonEvolution.h
+++ b/Engine/Source/Runtime/Experimental/Chaos/Public/Chaos/NewtonEvolution.h
@@ -19,11 +19,11 @@ struct ConsModelCaches {
 	//TArray<CorotatedFiberCache> CorotatedFiberCache;
 };
 
-class CHAOS_API FNewtonEvolution : public TArrayCollection
+class FNewtonEvolution : public TArrayCollection
 {
 public:
 	// TODO: Tidy up this constructor (and update Headless Chaos)
-	FNewtonEvolution(
+	CHAOS_API FNewtonEvolution(
 		FSolverParticles&& InParticles,
 		FSolverRigidParticles&& InGeometryParticles,
 		TArray<TVec3<int32>>&& CollisionTriangles,
@@ -46,13 +46,13 @@ public:
 	~FNewtonEvolution() {}
 
 	// Advance one time step. Filter the input time step if specified.
-	void AdvanceOneTimeStep(const FSolverReal Dt, const bool bSmoothDt = true);
+	CHAOS_API void AdvanceOneTimeStep(const FSolverReal Dt, const bool bSmoothDt = true);
 
 	// Remove all particles, will also reset all rules
-	void ResetParticles();
+	CHAOS_API void ResetParticles();
 
 	// Add particles and initialize group ids. Return the index of the first added particle.
-	int32 AddParticleRange(int32 NumParticles, uint32 GroupId, bool bActivate);
+	CHAOS_API int32 AddParticleRange(int32 NumParticles, uint32 GroupId, bool bActivate);
 
 	// Return the number of particles of the block starting at Offset
 	int32 GetParticleRangeSize(int32 Offset) const { return MParticlesActiveView.GetRangeSize(Offset); }
@@ -70,11 +70,11 @@ public:
 	const TArray<uint32>& ParticleGroupIds() const { return MParticleGroupIds; }
 
 	// Remove all collision particles
-	void ResetCollisionParticles(int32 NumParticles = 0);
+	CHAOS_API void ResetCollisionParticles(int32 NumParticles = 0);
 
 	// Add collision particles and initialize group ids. Return the index of the first added particle.
 	// Use INDEX_NONE as GroupId for collision particles that affect all particle groups.
-	int32 AddCollisionParticleRange(int32 NumParticles, uint32 GroupId, bool bActivate);
+	CHAOS_API int32 AddCollisionParticleRange(int32 NumParticles, uint32 GroupId, bool bActivate);
 
 	// Set a block of collision particles active or inactive, using the index of the first added particle to identify the block.
 	void ActivateCollisionParticleRange(int32 Offset, bool bActivate) { MCollisionParticlesActiveView.ActivateRange(Offset, bActivate); }
@@ -101,9 +101,9 @@ public:
 	}
 
 	// Add constraints. Return the index of the first added constraint.
-	int32 AddConstraintInitRange(int32 NumConstraints, bool bActivate);
-	int32 AddConstraintRuleRange(int32 NumConstraints, bool bActivate);
-	int32 AddPostCollisionConstraintRuleRange(int32 NumConstraints, bool bActivate);
+	CHAOS_API int32 AddConstraintInitRange(int32 NumConstraints, bool bActivate);
+	CHAOS_API int32 AddConstraintRuleRange(int32 NumConstraints, bool bActivate);
+	CHAOS_API int32 AddPostCollisionConstraintRuleRange(int32 NumConstraints, bool bActivate);
 
 	// Return the number of particles of the block starting at Offset
 	int32 GetConstraintInitRangeSize(int32 Offset) const { return MConstraintInitsActiveView.GetRangeSize(Offset); }
@@ -178,15 +178,15 @@ public:
 	template<typename Func1, typename Func2, typename Func3, typename Func4, typename Func5, typename Func6>
 	void DoNewtonStep(const int32 max_it_newton, const FSolverReal newton_tol, const int32 max_it_cg, const FSolverReal cg_tol, Func1 P, Func2 dP, const FSolverReal time, const FSolverReal dt, const TArray<TArray<TVector<int32, 2>>>& incident_elements, const TArray<FSolverReal>& nodal_mass, Func3 set_bcs, Func4 project_bcs, Func5 add_external_force, Func6 update_position_based_state, FSolverParticles& InParticles, TArray<FSolverReal>& residual_norm, bool use_cg = false, FSolverReal cg_prctg_reduce = 0, bool no_verbose = false);
 
-	void InitFEM();
+	CHAOS_API void InitFEM();
 
-	void WriteOutputLog(const int32 Frame);
+	CHAOS_API void WriteOutputLog(const int32 Frame);
 
 private:
 	// Add simulation groups and set default values
-	void AddGroups(int32 NumGroups);
+	CHAOS_API void AddGroups(int32 NumGroups);
 	// Reset simulation groups
-	void ResetGroups();
+	CHAOS_API void ResetGroups();
 	// Selected versions of the pre-iteration updates (euler step, force, velocity field. damping updates)..
 	template<bool bForceRule, bool bVelocityField, bool bDampVelocityRule>
 	void PreIterationUpdate(const FSolverReal Dt, const int32 Offset, const int32 Range, const int32 MinParallelBatchSize);
diff --git a/Engine/Source/Runtime/Experimental/Chaos/Public/Chaos/PBDAxialSpringConstraints.h b/Engine/Source/Runtime/Experimental/Chaos/Public/Chaos/PBDAxialSpringConstraints.h
--- a/Engine/Source/Runtime/Experimental/Chaos/Public/Chaos/PBDAxialSpringConstraints.h
+++ b/Engine/Source/Runtime/Experimental/Chaos/Public/Chaos/PBDAxialSpringConstraints.h
@@ -9,7 +9,7 @@
 namespace Chaos::Softs
 {
 
-class CHAOS_API FPBDAxialSpringConstraints : public FPBDAxialSpringConstraintsBase
+class FPBDAxialSpringConstraints : public FPBDAxialSpringConstraintsBase
 {
 	typedef FPBDAxialSpringConstraintsBase Base;
 	using Base::Barys;
@@ -37,7 +37,7 @@ public:
 
 	virtual ~FPBDAxialSpringConstraints() override {}
 
-	void Apply(FSolverParticles& InParticles, const FSolverReal Dt) const;
+	CHAOS_API void Apply(FSolverParticles& InParticles, const FSolverReal Dt) const;
 
 protected:
 	using Base::Constraints;
@@ -46,13 +46,13 @@ protected:
 	using Base::ParticleCount;
 
 private:
-	void InitColor(const FSolverParticles& InParticles);
-	void ApplyHelper(FSolverParticles& InParticles, const FSolverReal Dt, const int32 ConstraintIndex, const FSolverReal ExpStiffnessValue) const;
+	CHAOS_API void InitColor(const FSolverParticles& InParticles);
+	CHAOS_API void ApplyHelper(FSolverParticles& InParticles, const FSolverReal Dt, const int32 ConstraintIndex, const FSolverReal ExpStiffnessValue) const;
 
 	TArray<int32> ConstraintsPerColorStartIndex; // Constraints are ordered so each batch is contiguous. This is ColorNum + 1 length so it can be used as start and end.
 };
 
-class CHAOS_API FPBDAreaSpringConstraints final : public FPBDAxialSpringConstraints
+class FPBDAreaSpringConstraints final : public FPBDAxialSpringConstraints
 {
 public:
 	static bool IsEnabled(const FCollectionPropertyConstFacade& PropertyCollection)
@@ -99,7 +99,7 @@ public:
 
 	virtual ~FPBDAreaSpringConstraints() override = default;
 
-	void SetProperties(
+	CHAOS_API void SetProperties(
 		const FCollectionPropertyConstFacade& PropertyCollection,
 		const TMap<FString, TConstArrayView<FRealSingle>>& WeightMaps);
 
diff --git a/Engine/Source/Runtime/Experimental/Chaos/Public/Chaos/PBDBendingConstraints.h b/Engine/Source/Runtime/Experimental/Chaos/Public/Chaos/PBDBendingConstraints.h
--- a/Engine/Source/Runtime/Experimental/Chaos/Public/Chaos/PBDBendingConstraints.h
+++ b/Engine/Source/Runtime/Experimental/Chaos/Public/Chaos/PBDBendingConstraints.h
@@ -7,7 +7,7 @@
 namespace Chaos::Softs
 {
 
-class CHAOS_API FPBDBendingConstraints : public FPBDBendingConstraintsBase
+class FPBDBendingConstraints : public FPBDBendingConstraintsBase
 {
 	typedef FPBDBendingConstraintsBase Base;
 
@@ -89,7 +89,7 @@ public:
 	}
 
 	UE_DEPRECATED(5.2, "Use one of the other constructors instead.")
-	FPBDBendingConstraints(
+	CHAOS_API FPBDBendingConstraints(
 		const FSolverParticles& InParticles,
 		TArray<TVec4<int32>>&& InConstraints,
 		const FSolverReal InStiffness = (FSolverReal)1.)
@@ -101,7 +101,7 @@ PRAGMA_ENABLE_DEPRECATION_WARNINGS
 
 	using Base::SetProperties;
 
-	void SetProperties(
+	CHAOS_API void SetProperties(
 		const FCollectionPropertyConstFacade& PropertyCollection,
 		const TMap<FString, TConstArrayView<FRealSingle>>& WeightMaps);
 
@@ -111,13 +111,13 @@ PRAGMA_ENABLE_DEPRECATION_WARNINGS
 		SetProperties(PropertyCollection, TMap<FString, TConstArrayView<FRealSingle>>());
 	}
 
-	void Apply(FSolverParticles& InParticles, const FSolverReal Dt) const;
+	CHAOS_API void Apply(FSolverParticles& InParticles, const FSolverReal Dt) const;
 
 	const TArray<int32>& GetConstraintsPerColorStartIndex() const { return ConstraintsPerColorStartIndex; }
 
 private:
-	void InitColor(const FSolverParticles& InParticles);
-	void ApplyHelper(FSolverParticles& InParticles, const FSolverReal Dt, const int32 ConstraintIndex, const FSolverReal ExpStiffnessValue, const FSolverReal ExpBucklingValue) const;
+	CHAOS_API void InitColor(const FSolverParticles& InParticles);
+	CHAOS_API void ApplyHelper(FSolverParticles& InParticles, const FSolverReal Dt, const int32 ConstraintIndex, const FSolverReal ExpStiffnessValue, const FSolverReal ExpBucklingValue) const;
 
 	using Base::Constraints;
 	using Base::ParticleOffset;
