[Chaos]

diff --git a/Engine/Source/Runtime/Experimental/Chaos/Public/Chaos/PBDCollisionConstraints.h b/Engine/Source/Runtime/Experimental/Chaos/Public/Chaos/PBDCollisionConstraints.h
--- a/Engine/Source/Runtime/Experimental/Chaos/Public/Chaos/PBDCollisionConstraints.h
+++ b/Engine/Source/Runtime/Experimental/Chaos/Public/Chaos/PBDCollisionConstraints.h
@@ -45,7 +45,7 @@ namespace Private
  * 
  * @todo(chaos): remove handles array
  */
-class CHAOS_API FPBDCollisionConstraints : public FPBDConstraintContainer
+class FPBDCollisionConstraints : public FPBDConstraintContainer
 {
 public:
 	friend class FPBDCollisionConstraintHandle;
@@ -59,7 +59,7 @@ public:
 	// For use by dependent types
 	using FConstraintContainerHandle = FPBDCollisionConstraintHandle;		// Used by constraint rules
 
-	FPBDCollisionConstraints(const FPBDRigidsSOAs& InParticles, 
+	CHAOS_API FPBDCollisionConstraints(const FPBDRigidsSOAs& InParticles, 
 		TArrayCollectionArray<bool>& Collided, 
 		const TArrayCollectionArray<TSerializablePtr<FChaosPhysicsMaterial>>& PhysicsMaterials, 
 		const TArrayCollectionArray<TUniquePtr<FChaosPhysicsMaterial>>& PerParticlePhysicsMaterials,
@@ -67,7 +67,7 @@ public:
 		const int32 NumCollisionsPerBlock = 1000,
 		const FReal RestitutionThreshold = FReal(2000));
 
-	virtual ~FPBDCollisionConstraints();
+	CHAOS_API virtual ~FPBDCollisionConstraints();
 
 	/**
 	 * Whether this container provides constraint handles (simple solvers do not need them)
@@ -77,7 +77,7 @@ public:
 	/**
 	 * Put the container in "no handles" mode for use with simple solver. Must be called when empty of constraints (ideally right after creation).
 	 */
-	void DisableHandles();
+	CHAOS_API void DisableHandles();
 
 	/**
 	 * @brief Enable or disable determinism.
@@ -92,38 +92,38 @@ public:
 	 *  Clears the list of active constraints.
 	 * @todo(chaos): This is only required because of the way events work (see AdvanceOneTimeStepTask::DoWork)
 	*/
-	void BeginFrame();
+	CHAOS_API void BeginFrame();
 
 	/**
 	*  Destroy all constraints 
 	*/
-	void Reset();
+	CHAOS_API void Reset();
 
 
 	/**
 	 * @brief Called before collision detection to reset contacts
 	*/
-	void BeginDetectCollisions();
+	CHAOS_API void BeginDetectCollisions();
 
 	/**
 	 * @brief Called after collision detection to finalize the contacts
 	*/
-	void EndDetectCollisions();
+	CHAOS_API void EndDetectCollisions();
 
 	/**
 	 * @brief Called after collision resolution in order to detect probes
 	 */
-	void DetectProbeCollisions(FReal Dt);
+	CHAOS_API void DetectProbeCollisions(FReal Dt);
 
 	/**
 	 * Apply modifiers to particle pair midphases
 	 */
-	void ApplyMidPhaseModifier(const TArray<ISimCallbackObject*>& MidPhaseModifiers, FReal Dt);
+	CHAOS_API void ApplyMidPhaseModifier(const TArray<ISimCallbackObject*>& MidPhaseModifiers, FReal Dt);
 
 	/**
 	 * Apply modifiers to CCD results
 	 */
-	void ApplyCCDModifier(const TArray<ISimCallbackObject*>& CCDModifiers, FReal Dt);
+	CHAOS_API void ApplyCCDModifier(const TArray<ISimCallbackObject*>& CCDModifiers, FReal Dt);
 
 
 	/**
@@ -131,18 +131,18 @@ public:
 	 * You would probably call this in the PostComputeCallback. Prefer this to calling RemoveConstraints in a loop,
 	 * so you don't have to worry about constraint iterator/indices changing.
 	 */
-	void ApplyCollisionModifier(const TArray<ISimCallbackObject*>& CollisionModifiers, FReal Dt);
+	CHAOS_API void ApplyCollisionModifier(const TArray<ISimCallbackObject*>& CollisionModifiers, FReal Dt);
 
 
 	/**
 	* Remove the constraints associated with the ParticleHandle.
 	*/
-	void RemoveConstraints(const TSet<FGeometryParticleHandle*>&  ParticleHandle);
+	CHAOS_API void RemoveConstraints(const TSet<FGeometryParticleHandle*>&  ParticleHandle);
 
 	/**
 	 * @brief Remove all constraints associated with the particles - called when particles are destroyed
 	*/
-	virtual void DisconnectConstraints(const TSet<FGeometryParticleHandle*>& ParticleHandles) override;
+	CHAOS_API virtual void DisconnectConstraints(const TSet<FGeometryParticleHandle*>& ParticleHandles) override;
 
 	/**
 	* Disable the constraints associated with the ParticleHandle.
@@ -155,13 +155,13 @@ public:
 	//
 	virtual int32 GetNumConstraints() const override final { return NumConstraints(); }
 	virtual void ResetConstraints() override final { Reset(); }
-	virtual void AddConstraintsToGraph(Private::FPBDIslandManager& IslandManager) override final;
+	CHAOS_API virtual void AddConstraintsToGraph(Private::FPBDIslandManager& IslandManager) override final;
 	virtual void PrepareTick() override final {}
 	virtual void UnprepareTick() override final {}
 
-	virtual TUniquePtr<FConstraintContainerSolver> CreateSceneSolver(const int32 Priority) override final;
+	CHAOS_API virtual TUniquePtr<FConstraintContainerSolver> CreateSceneSolver(const int32 Priority) override final;
 
-	virtual TUniquePtr<FConstraintContainerSolver> CreateGroupSolver(const int32 Priority) override final;
+	CHAOS_API virtual TUniquePtr<FConstraintContainerSolver> CreateGroupSolver(const int32 Priority) override final;
 
 	// The type of solver we are creating
 	Private::ECollisionSolverType GetSolverType() const
@@ -270,14 +270,14 @@ public:
 		return ConstraintAllocator.GetConstraints();
 	}
 
-	FHandles GetConstraintHandles() const;
-	FConstHandles GetConstConstraintHandles() const;
+	CHAOS_API FHandles GetConstraintHandles() const;
+	CHAOS_API FConstHandles GetConstConstraintHandles() const;
 
-	const FPBDCollisionConstraint& GetConstraint(int32 Index) const;
+	CHAOS_API const FPBDCollisionConstraint& GetConstraint(int32 Index) const;
 
 	Private::FCollisionConstraintAllocator& GetConstraintAllocator() { return ConstraintAllocator; }
 
-	void UpdateConstraintMaterialProperties(FPBDCollisionConstraint& Contact);
+	CHAOS_API void UpdateConstraintMaterialProperties(FPBDCollisionConstraint& Contact);
 
 	const FPBDCollisionSolverSettings& GetSolverSettings() const { return SolverSettings; }
 
@@ -300,10 +300,10 @@ public:
 	}
 
 protected:
-	FPBDCollisionConstraint& GetConstraint(int32 Index);
+	CHAOS_API FPBDCollisionConstraint& GetConstraint(int32 Index);
 
 	// Call PruneParticleEdgeCollisions on all particles with ECollisionConstraintFlags::CCF_SmoothEdgeCollisions set in CollisionFlags
-	void PruneEdgeCollisions();
+	CHAOS_API void PruneEdgeCollisions();
 
 private:
 	const FPBDRigidsSOAs& Particles;
diff --git a/Engine/Source/Runtime/Experimental/Chaos/Public/Chaos/PBDCollisionSpringConstraintsBase.h b/Engine/Source/Runtime/Experimental/Chaos/Public/Chaos/PBDCollisionSpringConstraintsBase.h
--- a/Engine/Source/Runtime/Experimental/Chaos/Public/Chaos/PBDCollisionSpringConstraintsBase.h
+++ b/Engine/Source/Runtime/Experimental/Chaos/Public/Chaos/PBDCollisionSpringConstraintsBase.h
@@ -18,14 +18,14 @@ namespace Chaos::Softs
 {
 
 // This is an invertible spring class, typical springs are not invertible aware
-class CHAOS_API FPBDCollisionSpringConstraintsBase
+class FPBDCollisionSpringConstraintsBase
 {
 public:
 	static constexpr FSolverReal BackCompatThickness = (FSolverReal)1.f;
 	static constexpr FSolverReal BackCompatStiffness = (FSolverReal)0.5f;
 	static constexpr FSolverReal BackCompatFrictionCoefficient = (FSolverReal)0.f;
 
-	FPBDCollisionSpringConstraintsBase(
+	CHAOS_API FPBDCollisionSpringConstraintsBase(
 		const int32 InOffset,
 		const int32 InNumParticles,
 		const FTriangleMesh& InTriangleMesh,
@@ -38,12 +38,12 @@ public:
 	virtual ~FPBDCollisionSpringConstraintsBase() {}
 
 	UE_DEPRECATED(5.0, "Use Init(Particles, Spatial, GIAColors) instead.")
-	void Init(const FSolverParticles& Particles);
+	CHAOS_API void Init(const FSolverParticles& Particles);
 
 	template<typename SpatialAccelerator>
 	void Init(const FSolverParticles& Particles, const SpatialAccelerator& Spatial, const TConstArrayView<FPBDTriangleMeshCollisions::FGIAColor>& VertexGIAColors, const TArray<FPBDTriangleMeshCollisions::FGIAColor>& TriangleGIAColors);
 
-	FSolverVec3 GetDelta(const FSolverParticles& InParticles, const int32 i) const;
+	CHAOS_API FSolverVec3 GetDelta(const FSolverParticles& InParticles, const int32 i) const;
 
 	const TArray<TVec4<int32>>& GetConstraints() const { return Constraints;  }
 	const TArray<FSolverVec3>& GetBarys() const { return Barys; }
diff --git a/Engine/Source/Runtime/Experimental/Chaos/Public/Chaos/PBDConstraintContainer.h b/Engine/Source/Runtime/Experimental/Chaos/Public/Chaos/PBDConstraintContainer.h
--- a/Engine/Source/Runtime/Experimental/Chaos/Public/Chaos/PBDConstraintContainer.h
+++ b/Engine/Source/Runtime/Experimental/Chaos/Public/Chaos/PBDConstraintContainer.h
@@ -15,12 +15,12 @@ namespace Chaos
 	 * A Constraint Container holds an array of constraints and provides methods to allocate and deallocate constraints
 	 *as well as the API required to plug into Constraint Rules.
 	 */
-	class CHAOS_API FPBDConstraintContainer
+	class FPBDConstraintContainer
 	{
 	public:
-		FPBDConstraintContainer(FConstraintHandleTypeID InConstraintHandleType);
+		CHAOS_API FPBDConstraintContainer(FConstraintHandleTypeID InConstraintHandleType);
 
-		virtual ~FPBDConstraintContainer();
+		CHAOS_API virtual ~FPBDConstraintContainer();
 
 		/**
 		 * The ContainerId is used to map constraint handles back to their constraint container.
@@ -139,4 +139,4 @@ namespace Chaos
 	}
 
 
-}
\ No newline at end of file
+}
diff --git a/Engine/Source/Runtime/Experimental/Chaos/Public/Chaos/PBDEvolution.h b/Engine/Source/Runtime/Experimental/Chaos/Public/Chaos/PBDEvolution.h
--- a/Engine/Source/Runtime/Experimental/Chaos/Public/Chaos/PBDEvolution.h
+++ b/Engine/Source/Runtime/Experimental/Chaos/Public/Chaos/PBDEvolution.h
@@ -12,11 +12,11 @@
 namespace Chaos::Softs
 {
 
-class CHAOS_API FPBDEvolution : public TArrayCollection
+class FPBDEvolution : public TArrayCollection
 {
  public:
 	// TODO: Tidy up this constructor (and update Headless Chaos)
-	FPBDEvolution(
+	CHAOS_API FPBDEvolution(
 		FSolverParticles&& InParticles,
 		FSolverRigidParticles&& InGeometryParticles,
 		TArray<TVec3<int32>>&& CollisionTriangles,
@@ -35,13 +35,13 @@ class CHAOS_API FPBDEvolution : public TArrayCollection
 	void AdvanceOneTimeStep(const FSolverReal Dt, const bool bSmoothDt) { AdvanceOneTimeStep(Dt); }
 
 	// Advance one time step. Filter the input time step if specified.
-	void AdvanceOneTimeStep(const FSolverReal Dt);
+	CHAOS_API void AdvanceOneTimeStep(const FSolverReal Dt);
 
 	// Remove all particles, will also reset all rules
-	void ResetParticles();
+	CHAOS_API void ResetParticles();
 
 	// Add particles and initialize group ids. Return the index of the first added particle.
-	int32 AddParticleRange(int32 NumParticles, uint32 GroupId, bool bActivate);
+	CHAOS_API int32 AddParticleRange(int32 NumParticles, uint32 GroupId, bool bActivate);
 
 	// Return the number of particles of the block starting at Offset
 	int32 GetParticleRangeSize(int32 Offset) const { return MParticlesActiveView.GetRangeSize(Offset); }
@@ -64,17 +64,17 @@ class CHAOS_API FPBDEvolution : public TArrayCollection
 	const TArray<uint32>& ParticleGroupIds() const { return MParticleGroupIds; }
 
 	// Remove all collision particles
-	void ResetCollisionParticles(int32 NumParticles = 0);
+	CHAOS_API void ResetCollisionParticles(int32 NumParticles = 0);
 
 	// Add collision particles and initialize group ids. Return the index of the first added particle.
 	// Use INDEX_NONE as GroupId for collision particles that affect all particle groups.
-	int32 AddCollisionParticleRange(int32 NumParticles, uint32 GroupId, bool bActivate);
+	CHAOS_API int32 AddCollisionParticleRange(int32 NumParticles, uint32 GroupId, bool bActivate);
 
 	// Add a single collision body particle to the solver. 
-	int32 AddCollisionParticle(uint32 GroupId, bool bActivate);
+	CHAOS_API int32 AddCollisionParticle(uint32 GroupId, bool bActivate);
 
 	// Remove a collision body from the active view list, and save its particle in the remvoed collision particle list
-	void RemoveCollisionParticle(int32 CollisionParticleIndex, int32 CollisionParticleViewIndex);
+	CHAOS_API void RemoveCollisionParticle(int32 CollisionParticleIndex, int32 CollisionParticleViewIndex);
 
 	// Set a block of collision particles active or inactive, using the index of the first added particle to identify the block.
 	void ActivateCollisionParticleRange(int32 Offset, bool bActivate) { MCollisionParticlesActiveView.ActivateRange(Offset, bActivate); }
@@ -104,10 +104,10 @@ class CHAOS_API FPBDEvolution : public TArrayCollection
 	}
 
 	// Add constraints. Return the index of the first added constraint.
-	int32 AddConstraintInitRange(int32 NumConstraints, bool bActivate);
-	int32 AddConstraintRuleRange(int32 NumConstraints, bool bActivate);
-	int32 AddPostCollisionConstraintRuleRange(int32 NumConstraints, bool bActivate);
-	int32 AddConstraintPostprocessingsRange(int32 NumConstraints, bool bActivate);
+	CHAOS_API int32 AddConstraintInitRange(int32 NumConstraints, bool bActivate);
+	CHAOS_API int32 AddConstraintRuleRange(int32 NumConstraints, bool bActivate);
+	CHAOS_API int32 AddPostCollisionConstraintRuleRange(int32 NumConstraints, bool bActivate);
+	CHAOS_API int32 AddConstraintPostprocessingsRange(int32 NumConstraints, bool bActivate);
 	
 	// Return the number of particles of the block starting at Offset
 	int32 GetConstraintInitRangeSize(int32 Offset) const { return MConstraintInitsActiveView.GetRangeSize(Offset); }
@@ -187,9 +187,9 @@ class CHAOS_API FPBDEvolution : public TArrayCollection
 
  private:
 	// Add simulation groups and set default values
-	void AddGroups(int32 NumGroups);
+	CHAOS_API void AddGroups(int32 NumGroups);
 	// Reset simulation groups
-	void ResetGroups();
+	CHAOS_API void ResetGroups();
 	// Selected versions of the pre-iteration updates (euler step, force, velocity field. damping updates)..
 	template<bool bForceRule, bool bVelocityField, bool bDampVelocityRule>
 	void PreIterationUpdate(const FSolverReal Dt, const int32 Offset, const int32 Range, const int32 MinParallelBatchSize);
diff --git a/Engine/Source/Runtime/Experimental/Chaos/Public/Chaos/PBDJointConstraintData.h b/Engine/Source/Runtime/Experimental/Chaos/Public/Chaos/PBDJointConstraintData.h
--- a/Engine/Source/Runtime/Experimental/Chaos/Public/Chaos/PBDJointConstraintData.h
+++ b/Engine/Source/Runtime/Experimental/Chaos/Public/Chaos/PBDJointConstraintData.h
@@ -33,31 +33,31 @@ namespace Chaos
 	using FJointConstraintDirtyFlags = TDirtyFlags<EJointConstraintFlags>;
 	typedef TVector<FTransform, 2> FTransformPair;
 
-	class CHAOS_API FJointConstraint : public FConstraintBase
+	class FJointConstraint : public FConstraintBase
 	{
 	public:
 		using Base = FConstraintBase;
 
 		friend class FPBDRigidsSolver; // friend so we can call ReleaseKinematicEndPoint when unregistering joint.
 
-		FJointConstraint();
+		CHAOS_API FJointConstraint();
 		virtual ~FJointConstraint() override {}
 
 		const FPBDJointSettings& GetJointSettings()const { return JointSettings.Read(); }
 
 		// If we created particle to serve as kinematic endpoint, track so we can release later. This will add particle to solver.
-		void SetKinematicEndPoint(FSingleParticlePhysicsProxy* InParticle, FPBDRigidsSolver* Solver);
+		CHAOS_API void SetKinematicEndPoint(FSingleParticlePhysicsProxy* InParticle, FPBDRigidsSolver* Solver);
 
-		FSingleParticlePhysicsProxy* GetKinematicEndPoint() const;
+		CHAOS_API FSingleParticlePhysicsProxy* GetKinematicEndPoint() const;
 
 		//See JointProperties for API
 		//Each CHAOS_INNER_JOINT_PROPERTY entry will have a Get* and Set*
 #define CHAOS_INNER_JOINT_PROPERTY(OuterProp, FuncName, Inner, InnerType) CONSTRAINT_JOINT_PROPERPETY_IMPL2(InnerType, FuncName, OuterProp, Inner)
 #include "Chaos/JointProperties.inl"
 
-		void SetLinearPositionDriveEnabled(TVector<bool, 3> Enabled);
+		CHAOS_API void SetLinearPositionDriveEnabled(TVector<bool, 3> Enabled);
 
-		void SetLinearVelocityDriveEnabled(TVector<bool, 3> Enabled);
+		CHAOS_API void SetLinearVelocityDriveEnabled(TVector<bool, 3> Enabled);
 
 		struct FOutputData
 		{
@@ -77,15 +77,15 @@ namespace Chaos
 			JointSettings.SyncRemote(Manager, DataIdx, RemoteData);
 		}
 
-		void SetParticleProxies(const FProxyBasePair& InJointParticles);
-		void SetPhysicsBodies(const FPhysicsObjectPair& InBodies);
+		CHAOS_API void SetParticleProxies(const FProxyBasePair& InJointParticles);
+		CHAOS_API void SetPhysicsBodies(const FPhysicsObjectPair& InBodies);
 
 		const FProxyBasePair& GetParticleProxies() const { return  JointProxies.Read().ParticleProxies; }
 		const FPhysicsObjectPair& GetPhysicsBodies() const { return  JointBodies.Read().PhysicsBodies; }
 
 	protected:
 
-		void ReleaseKinematicEndPoint(FPBDRigidsSolver* Solver);
+		CHAOS_API void ReleaseKinematicEndPoint(FPBDRigidsSolver* Solver);
 		
 		TChaosProperty<FProxyBasePairProperty, EChaosProperty::JointParticleProxies> JointProxies;
 		TChaosProperty<FPhysicsObjectPairProperty, EChaosProperty::JointPhysicsObjects> JointBodies;
diff --git a/Engine/Source/Runtime/Experimental/Chaos/Public/Chaos/PBDJointConstraintTypes.h b/Engine/Source/Runtime/Experimental/Chaos/Public/Chaos/PBDJointConstraintTypes.h
--- a/Engine/Source/Runtime/Experimental/Chaos/Public/Chaos/PBDJointConstraintTypes.h
+++ b/Engine/Source/Runtime/Experimental/Chaos/Public/Chaos/PBDJointConstraintTypes.h
@@ -110,16 +110,16 @@ namespace Chaos
 		}
 	};
 
-	class CHAOS_API FPBDJointSettings
+	class FPBDJointSettings
 	{
 	public:
 		using FTransformPair = TVector<FRigidTransform3, 2>;
 
-		FPBDJointSettings();
+		CHAOS_API FPBDJointSettings();
 
 		// Ensure that settings are consistent and within valid ranges. Should be called
 		// whenever settings change.
-		void Sanitize();
+		CHAOS_API void Sanitize();
 
 		bool operator==(const FPBDJointSettings& Other) const
 		{
@@ -202,10 +202,10 @@ namespace Chaos
 		void* UserData;
 	};
 
-	class CHAOS_API FPBDJointSolverSettings
+	class FPBDJointSolverSettings
 	{
 	public:
-		FPBDJointSolverSettings();
+		CHAOS_API FPBDJointSolverSettings();
 
 		// Tolerances
 		FReal SwingTwistAngleTolerance;
@@ -257,4 +257,4 @@ namespace Chaos
 		FReal SoftSwingDampingOverride;
 	};
 
-}
\ No newline at end of file
+}
diff --git a/Engine/Source/Runtime/Experimental/Chaos/Public/Chaos/PBDJointConstraints.h b/Engine/Source/Runtime/Experimental/Chaos/Public/Chaos/PBDJointConstraints.h
--- a/Engine/Source/Runtime/Experimental/Chaos/Public/Chaos/PBDJointConstraints.h
+++ b/Engine/Source/Runtime/Experimental/Chaos/Public/Chaos/PBDJointConstraints.h
@@ -22,7 +22,7 @@ namespace Chaos
 	/**
 	 * A handle to a joint constraint held in a joint container (FPBDJointConstraints) by index.
 	*/
-	class CHAOS_API FPBDJointConstraintHandle final : public TIndexedContainerConstraintHandle<FPBDJointConstraints>
+	class FPBDJointConstraintHandle final : public TIndexedContainerConstraintHandle<FPBDJointConstraints>
 	{
 	public:
 		using Base = TIndexedContainerConstraintHandle<FPBDJointConstraints>;
@@ -31,26 +31,26 @@ namespace Chaos
 		FPBDJointConstraintHandle();
 		FPBDJointConstraintHandle(FConstraintContainer* InConstraintContainer, int32 InConstraintIndex);
 
-		void SetConstraintEnabled(bool bEnabled);
+		CHAOS_API void SetConstraintEnabled(bool bEnabled);
 
-		void CalculateConstraintSpace(FVec3& OutXa, FMatrix33& OutRa, FVec3& OutXb, FMatrix33& OutRb) const;
+		CHAOS_API void CalculateConstraintSpace(FVec3& OutXa, FMatrix33& OutRa, FVec3& OutXb, FMatrix33& OutRb) const;
 
-		bool IsConstraintEnabled() const;
-		bool IsConstraintBroken() const;
-		bool IsConstraintBreaking() const;
-		void ClearConstraintBreaking();
-		bool IsDriveTargetChanged() const;
-		void ClearDriveTargetChanged();
-		FVec3 GetLinearImpulse() const;
-		FVec3 GetAngularImpulse() const;
+		CHAOS_API bool IsConstraintEnabled() const;
+		CHAOS_API bool IsConstraintBroken() const;
+		CHAOS_API bool IsConstraintBreaking() const;
+		CHAOS_API void ClearConstraintBreaking();
+		CHAOS_API bool IsDriveTargetChanged() const;
+		CHAOS_API void ClearDriveTargetChanged();
+		CHAOS_API FVec3 GetLinearImpulse() const;
+		CHAOS_API FVec3 GetAngularImpulse() const;
 
-		const FPBDJointSettings& GetSettings() const;
+		CHAOS_API const FPBDJointSettings& GetSettings() const;
 		const FPBDJointSettings& GetJointSettings() const { return GetSettings(); }	//needed for property macros
 
 		// Note that setting drive targets etc can be done through the settings
-		void SetSettings(const FPBDJointSettings& Settings);
+		CHAOS_API void SetSettings(const FPBDJointSettings& Settings);
 
-		TVec2<FGeometryParticleHandle*> GetConstrainedParticles() const override final;
+		CHAOS_API TVec2<FGeometryParticleHandle*> GetConstrainedParticles() const override final;
 
 		static const FConstraintHandleTypeID& StaticType()
 		{
@@ -58,18 +58,18 @@ namespace Chaos
 			return STypeID;
 		}
 
-		ESyncState SyncState() const;
-		void SetSyncState(ESyncState SyncState);
+		CHAOS_API ESyncState SyncState() const;
+		CHAOS_API void SetSyncState(ESyncState SyncState);
 
-		void SetEnabledDuringResim(bool bEnabled);
-		EResimType ResimType() const;
+		CHAOS_API void SetEnabledDuringResim(bool bEnabled);
+		CHAOS_API EResimType ResimType() const;
 
 		UE_DEPRECATED(5.2, "No longer used")
-		int32 GetConstraintIsland() const;
+		CHAOS_API int32 GetConstraintIsland() const;
 		UE_DEPRECATED(5.2, "No longer used")
-		int32 GetConstraintLevel() const;
+		CHAOS_API int32 GetConstraintLevel() const;
 		UE_DEPRECATED(5.2, "No longer used")
-		int32 GetConstraintColor() const;
+		CHAOS_API int32 GetConstraintColor() const;
 
 	protected:
 		using Base::ConstraintIndex;
@@ -82,10 +82,10 @@ namespace Chaos
 	/**
 	 * Peristent variable state for a joint
 	*/
-	class CHAOS_API FPBDJointState
+	class FPBDJointState
 	{
 	public:
-		FPBDJointState();
+		CHAOS_API FPBDJointState();
 
 		int32 Island;
 		int32 Level;
@@ -105,7 +105,7 @@ namespace Chaos
 	/**
 	 * A set of joint restricting up to 6 degrees of freedom, with linear and angular limits.
 	 */
-	class CHAOS_API FPBDJointConstraints : public FPBDConstraintContainer
+	class FPBDJointConstraints : public FPBDConstraintContainer
 	{
 	public:
 		using Base = FPBDConstraintContainer;
@@ -115,19 +115,19 @@ namespace Chaos
 		using FTransformPair = TVector<FRigidTransform3, 2>;
 		using FHandles = TArray<FConstraintContainerHandle*>;
 
-		FPBDJointConstraints();
+		CHAOS_API FPBDJointConstraints();
 
-		virtual ~FPBDJointConstraints();
+		CHAOS_API virtual ~FPBDJointConstraints();
 
 		/**
 		 * Get the solver settings (used by FPBDJointContainerSolver)
 		*/
-		const FPBDJointSolverSettings& GetSettings() const;
+		CHAOS_API const FPBDJointSolverSettings& GetSettings() const;
 
 		/*
 		* Modify the solver settings
 		*/
-		void SetSettings(const FPBDJointSolverSettings& InSettings);
+		CHAOS_API void SetSettings(const FPBDJointSolverSettings& InSettings);
 
 		/**
 		* Whether to use a linear or non-linear joint solver. Non-linear is more stable but much more expensive.
@@ -143,21 +143,21 @@ namespace Chaos
 		/**
 		 * Get the number of constraints.
 		 */
-		int32 NumConstraints() const;
+		CHAOS_API int32 NumConstraints() const;
 
 		/**
 		 * Add a constraint with particle-space constraint offsets.
 		 * @todo(chaos): clean up this set of functions (now that ConnectorTransforms is in the settings, calling AddConstraint then SetSettings leads 
 		 * to unexpected behaviour - overwriting the ConnectorTransforms with Identity)
 		 */
-		FConstraintContainerHandle* AddConstraint(const FParticlePair& InConstrainedParticles, const FRigidTransform3& WorldConstraintFrame);
-		FConstraintContainerHandle* AddConstraint(const FParticlePair& InConstrainedParticles, const FTransformPair& ConnectorTransforms);
-		FConstraintContainerHandle* AddConstraint(const FParticlePair& InConstrainedParticles, const FPBDJointSettings& InConstraintSettings);
+		CHAOS_API FConstraintContainerHandle* AddConstraint(const FParticlePair& InConstrainedParticles, const FRigidTransform3& WorldConstraintFrame);
+		CHAOS_API FConstraintContainerHandle* AddConstraint(const FParticlePair& InConstrainedParticles, const FTransformPair& ConnectorTransforms);
+		CHAOS_API FConstraintContainerHandle* AddConstraint(const FParticlePair& InConstrainedParticles, const FPBDJointSettings& InConstraintSettings);
 
 		/**
 		 * Remove the specified constraint.
 		 */
-		void RemoveConstraint(int ConstraintIndex);
+		CHAOS_API void RemoveConstraint(int ConstraintIndex);
 		void RemoveConstraints(const TSet<TGeometryParticleHandle<FReal, 3>*>& RemovedParticles) {}
 
 		/*
@@ -167,69 +167,69 @@ namespace Chaos
 		* 
 		* The constraint is unuseable at this point and pending deletion. 
 		*/
-		void DisconnectConstraints(const TSet<TGeometryParticleHandle<FReal, 3>*>& RemovedParticles);
+		CHAOS_API void DisconnectConstraints(const TSet<TGeometryParticleHandle<FReal, 3>*>& RemovedParticles);
 
 		/*
 		 * Whether the constraint is enabled
 		 */
-		bool IsConstraintEnabled(int32 ConstraintIndex) const;
+		CHAOS_API bool IsConstraintEnabled(int32 ConstraintIndex) const;
 
 		/*
 		 * Whether the constraint is broken
 		 */
-		bool IsConstraintBroken(int32 ConstraintIndex) const;
+		CHAOS_API bool IsConstraintBroken(int32 ConstraintIndex) const;
 
 		/*
 		 * Whether the constraint was broken this frame (transient flag for use by event system)
 		 */
-		bool IsConstraintBreaking(int32 ConstraintIndex) const;
+		CHAOS_API bool IsConstraintBreaking(int32 ConstraintIndex) const;
 
 		/*
 		 * Clear the transient constraint braking state (called by event system when it has used the flag)
 		 */
-		void ClearConstraintBreaking(int32 ConstraintIndex);
+		CHAOS_API void ClearConstraintBreaking(int32 ConstraintIndex);
 
 		/*
 		 * Whether the drive target has changed
 		 */
-		bool IsDriveTargetChanged(int32 ConstraintIndex) const;
+		CHAOS_API bool IsDriveTargetChanged(int32 ConstraintIndex) const;
 
 		/*
 		 * Clear the drive target state
 		 */
-		void ClearDriveTargetChanged(int32 ConstraintIndex);
+		CHAOS_API void ClearDriveTargetChanged(int32 ConstraintIndex);
 
 		/*
 		 * Enable or disable a constraints
 		 */
-		void SetConstraintEnabled(int32 ConstraintIndex, bool bEnabled);
+		CHAOS_API void SetConstraintEnabled(int32 ConstraintIndex, bool bEnabled);
 
 		/*
 		* Set Drive Target Changed State
 		*/
-		void SetDriveTargetChanged(int32 ConstraintIndex, bool bTargetChanged);
+		CHAOS_API void SetDriveTargetChanged(int32 ConstraintIndex, bool bTargetChanged);
 
 		/*
 		 * Force a constraints to break
 		 */
-		void BreakConstraint(int32 ConstraintIndex);
+		CHAOS_API void BreakConstraint(int32 ConstraintIndex);
 
 		/**
 		 * Repair a broken constraints (does not adjust particle positions)
 		 */
-		void FixConstraint(int32 ConstraintIndex);
+		CHAOS_API void FixConstraint(int32 ConstraintIndex);
 
 		/**
 		 * Set the break callback. This will be called once after the constraint solver phase on the \
 		 * tick when the constraint is broken. There is only one callback allowed - it is usually a 
 		 * method on the owner (e.g., the Evolution object) which will probably dispatch its own event.
 		*/
-		void SetBreakCallback(const FJointBreakCallback& Callback);
+		CHAOS_API void SetBreakCallback(const FJointBreakCallback& Callback);
 
 		/**
 		 * Remove the previously assigned break callback.
 		*/
-		void ClearBreakCallback();
+		CHAOS_API void ClearBreakCallback();
 
 		/**
 		 * All of the constraints in the container, including inactive
@@ -250,73 +250,73 @@ namespace Chaos
 		/**
 		 * Get a joint constraint by index
 		*/
-		const FConstraintContainerHandle* GetConstraintHandle(int32 ConstraintIndex) const;
+		CHAOS_API const FConstraintContainerHandle* GetConstraintHandle(int32 ConstraintIndex) const;
 
 		/**
 		 * Get a joint constraint by index
 		*/
-		FConstraintContainerHandle* GetConstraintHandle(int32 ConstraintIndex);
+		CHAOS_API FConstraintContainerHandle* GetConstraintHandle(int32 ConstraintIndex);
 
 		/**
 		 * Get the particles that are affected by the specified constraint.
 		 */
-		const FParticlePair& GetConstrainedParticles(int32 ConstraintIndex) const;
+		CHAOS_API const FParticlePair& GetConstrainedParticles(int32 ConstraintIndex) const;
 
 		/**
 		 * Get the settings for a joint constraint by index
 		*/
-		const FPBDJointSettings& GetConstraintSettings(int32 ConstraintIndex) const;
+		CHAOS_API const FPBDJointSettings& GetConstraintSettings(int32 ConstraintIndex) const;
 
 		/**
 		 * Set the settings for a joint constraint by index
 		*/
-		void SetConstraintSettings(int32 ConstraintIndex, const FPBDJointSettings& InConstraintSettings);
+		CHAOS_API void SetConstraintSettings(int32 ConstraintIndex, const FPBDJointSettings& InConstraintSettings);
 
 		/**
 		 * Set the linear drive target for a constraint by index
 		*/
-		void SetLinearDrivePositionTarget(int32 ConstraintIndex, FVec3 InLinearDrivePositionTarget);
+		CHAOS_API void SetLinearDrivePositionTarget(int32 ConstraintIndex, FVec3 InLinearDrivePositionTarget);
 
 		/**
 		 * Set the angular drive target for a constraint by index
 		*/
-		void SetAngularDrivePositionTarget(int32 ConstraintIndex, FRotation3 InAngularDrivePositionTarget);
+		CHAOS_API void SetAngularDrivePositionTarget(int32 ConstraintIndex, FRotation3 InAngularDrivePositionTarget);
 
 		/**
 		* The total linear impulse applied by the constraint
 		*/
-		FVec3 GetConstraintLinearImpulse(int32 ConstraintIndex) const;
+		CHAOS_API FVec3 GetConstraintLinearImpulse(int32 ConstraintIndex) const;
 
 		/**
 		* The total linear angular applied by the constraint
 		*/
-		FVec3 GetConstraintAngularImpulse(int32 ConstraintIndex) const;
+		CHAOS_API FVec3 GetConstraintAngularImpulse(int32 ConstraintIndex) const;
 
-		ESyncState GetConstraintSyncState(int32 ConstraintIndex) const;
-		void SetConstraintSyncState(int32 ConstraintIndex, ESyncState SyncState);
+		CHAOS_API ESyncState GetConstraintSyncState(int32 ConstraintIndex) const;
+		CHAOS_API void SetConstraintSyncState(int32 ConstraintIndex, ESyncState SyncState);
 		
-		void SetConstraintEnabledDuringResim(int32 ConstraintIndex, bool bEnabled);
+		CHAOS_API void SetConstraintEnabledDuringResim(int32 ConstraintIndex, bool bEnabled);
 		
-		EResimType GetConstraintResimType(int32 ConstraintIndex) const;
+		CHAOS_API EResimType GetConstraintResimType(int32 ConstraintIndex) const;
 
 		//
 		// FConstraintContainer Implementation
 		//
-		virtual TUniquePtr<FConstraintContainerSolver> CreateSceneSolver(const int32 Priority) override final;
-		virtual TUniquePtr<FConstraintContainerSolver> CreateGroupSolver(const int32 Priority) override final;
+		CHAOS_API virtual TUniquePtr<FConstraintContainerSolver> CreateSceneSolver(const int32 Priority) override final;
+		CHAOS_API virtual TUniquePtr<FConstraintContainerSolver> CreateGroupSolver(const int32 Priority) override final;
 		virtual int32 GetNumConstraints() const override final { return NumConstraints(); }
 		virtual void ResetConstraints() override final {}
-		virtual void AddConstraintsToGraph(Private::FPBDIslandManager& IslandManager) override final;
-		virtual void PrepareTick() override final;
-		virtual void UnprepareTick() override final;
+		CHAOS_API virtual void AddConstraintsToGraph(Private::FPBDIslandManager& IslandManager) override final;
+		CHAOS_API virtual void PrepareTick() override final;
+		CHAOS_API virtual void UnprepareTick() override final;
 
 		// Called by the joint solver at the end of the constraint solver phase
-		void SetSolverResults(const int32 ConstraintIndex, const FVec3& LinearImpulse, const FVec3& AngularImpulse, const bool bIsBroken, const FSolverBody* SolverBody0, const FSolverBody* SolverBody1);
+		CHAOS_API void SetSolverResults(const int32 ConstraintIndex, const FVec3& LinearImpulse, const FVec3& AngularImpulse, const bool bIsBroken, const FSolverBody* SolverBody0, const FSolverBody* SolverBody1);
 
 		// @todo(chaos): only needed for RBAN, and should be private or moved to the solver
-		int32 GetConstraintIsland(int32 ConstraintIndex) const;
-		int32 GetConstraintLevel(int32 ConstraintIndex) const;
-		int32 GetConstraintColor(int32 ConstraintIndex) const;
+		CHAOS_API int32 GetConstraintIsland(int32 ConstraintIndex) const;
+		CHAOS_API int32 GetConstraintLevel(int32 ConstraintIndex) const;
+		CHAOS_API int32 GetConstraintColor(int32 ConstraintIndex) const;
 
 		// @todo(chaos): functionality no longer supported here (see FPBDJointContainerSolver)
 		UE_DEPRECATED(5.2, "Joint Solver API moved to FPBDJointContainerSolver")
@@ -351,18 +351,18 @@ namespace Chaos
 	private:
 		friend class FPBDJointConstraintHandle;
 
-		void GetConstrainedParticleIndices(const int32 ConstraintIndex, int32& Index0, int32& Index1) const;
-		void CalculateConstraintSpace(int32 ConstraintIndex, FVec3& OutX0, FMatrix33& OutR0, FVec3& OutX1, FMatrix33& OutR1) const;
+		CHAOS_API void GetConstrainedParticleIndices(const int32 ConstraintIndex, int32& Index0, int32& Index1) const;
+		CHAOS_API void CalculateConstraintSpace(int32 ConstraintIndex, FVec3& OutX0, FMatrix33& OutR0, FVec3& OutX1, FMatrix33& OutR1) const;
 		
-		void ColorConstraints();
-		void SortConstraints();
+		CHAOS_API void ColorConstraints();
+		CHAOS_API void SortConstraints();
 
-		bool ShouldBeInGraph(const int32 ConstraintIndex) const;
+		CHAOS_API bool ShouldBeInGraph(const int32 ConstraintIndex) const;
 
-		void ApplyPlasticityLimits(const int32 ConstraintIndex, const FSolverBody& SolverBody0, const FSolverBody& SolverBody1);
+		CHAOS_API void ApplyPlasticityLimits(const int32 ConstraintIndex, const FSolverBody& SolverBody0, const FSolverBody& SolverBody1);
 
-		void SetConstraintBroken(int32 ConstraintIndex, bool bBroken);
-		void SetConstraintBreaking(int32 ConstraintIndex, bool bBreaking);
+		CHAOS_API void SetConstraintBroken(int32 ConstraintIndex, bool bBroken);
+		CHAOS_API void SetConstraintBreaking(int32 ConstraintIndex, bool bBreaking);
 
 		FPBDJointSolverSettings Settings;
 
diff --git a/Engine/Source/Runtime/Experimental/Chaos/Public/Chaos/PBDLongRangeConstraints.h b/Engine/Source/Runtime/Experimental/Chaos/Public/Chaos/PBDLongRangeConstraints.h
--- a/Engine/Source/Runtime/Experimental/Chaos/Public/Chaos/PBDLongRangeConstraints.h
+++ b/Engine/Source/Runtime/Experimental/Chaos/Public/Chaos/PBDLongRangeConstraints.h
@@ -10,7 +10,7 @@ DECLARE_CYCLE_STAT(TEXT("Chaos PBD Long Range Constraint"), STAT_PBD_LongRange,
 namespace Chaos::Softs
 {
 
-class CHAOS_API FPBDLongRangeConstraints : public FPBDLongRangeConstraintsBase
+class FPBDLongRangeConstraints : public FPBDLongRangeConstraintsBase
 {
 public:
 	typedef FPBDLongRangeConstraintsBase Base;
@@ -91,7 +91,7 @@ public:
 
 	using Base::SetProperties;
 
-	void SetProperties(
+	CHAOS_API void SetProperties(
 		const FCollectionPropertyConstFacade& PropertyCollection,
 		const TMap<FString, TConstArrayView<FRealSingle>>& WeightMaps,
 		FSolverReal MeshScale);
@@ -102,7 +102,7 @@ public:
 		SetProperties(PropertyCollection, TMap<FString, TConstArrayView<FRealSingle>>(), (FSolverReal)1.);
 	}
 
-	void Apply(FSolverParticles& Particles, const FSolverReal Dt) const;
+	CHAOS_API void Apply(FSolverParticles& Particles, const FSolverReal Dt) const;
 
 private:
 	using Base::MinTetherScale;
diff --git a/Engine/Source/Runtime/Experimental/Chaos/Public/Chaos/PBDLongRangeConstraintsBase.h b/Engine/Source/Runtime/Experimental/Chaos/Public/Chaos/PBDLongRangeConstraintsBase.h
--- a/Engine/Source/Runtime/Experimental/Chaos/Public/Chaos/PBDLongRangeConstraintsBase.h
+++ b/Engine/Source/Runtime/Experimental/Chaos/Public/Chaos/PBDLongRangeConstraintsBase.h
@@ -12,7 +12,7 @@
 namespace Chaos::Softs
 {
 
-class CHAOS_API FPBDLongRangeConstraintsBase
+class FPBDLongRangeConstraintsBase
 {
 public:
 	UE_NONCOPYABLE(FPBDLongRangeConstraintsBase);
@@ -33,7 +33,7 @@ public:
 
 	typedef TTuple<int32, int32, FRealSingle> FTether;
 
-	FPBDLongRangeConstraintsBase(
+	CHAOS_API FPBDLongRangeConstraintsBase(
 		const FSolverParticles& Particles,
 		const int32 InParticleOffset,
 		const int32 InParticleCount,
@@ -100,7 +100,7 @@ PRAGMA_ENABLE_DEPRECATION_WARNINGS
 
 protected:
 	// Return the minimum number of long range tethers in a batch to process in parallel
-	static int32 GetMinParallelBatchSize();
+	static CHAOS_API int32 GetMinParallelBatchSize();
 
 	// Return whether the constraint has been setup with a weightmap to interpolate between two low and high values of scales
 	UE_DEPRECATED(5.2, "Use TetherScale.HasWeightMap() instead")
diff --git a/Engine/Source/Runtime/Experimental/Chaos/Public/Chaos/PBDRigidClustering.h b/Engine/Source/Runtime/Experimental/Chaos/Public/Chaos/PBDRigidClustering.h
--- a/Engine/Source/Runtime/Experimental/Chaos/Public/Chaos/PBDRigidClustering.h
+++ b/Engine/Source/Runtime/Experimental/Chaos/Public/Chaos/PBDRigidClustering.h
@@ -17,7 +17,7 @@ namespace Chaos
 {
 	class FPBDCollisionConstraints;
 
-struct CHAOS_API FClusterDestoryParameters {
+struct FClusterDestoryParameters {
 	bool bReturnInternalOnly : true;
 };
 
@@ -51,7 +51,7 @@ struct CHAOS_API FClusterDestoryParameters {
 *   internal structure
 *
 */
-class CHAOS_API FRigidClustering
+class FRigidClustering
 {
 public:
 
@@ -62,8 +62,8 @@ public:
 	typedef TMap<FClusterHandle, FRigidHandleArray>				FClusterMap;
 	typedef TFunction<void(FRigidClustering&, FRigidHandle)>	FVisitorFunction;
 
-	FRigidClustering(FRigidEvolution& InEvolution, FPBDRigidClusteredParticles& InParticles, const TArray<ISimCallbackObject*>* InStrainModifiers);
-	~FRigidClustering();
+	CHAOS_API FRigidClustering(FRigidEvolution& InEvolution, FPBDRigidClusteredParticles& InParticles, const TArray<ISimCallbackObject*>* InStrainModifiers);
+	CHAOS_API ~FRigidClustering();
 
 	//
 	// Initialization
@@ -79,7 +79,7 @@ public:
 	 *		ForceMassOrientation : Inertial alignment into mass space.
 	 *	
 	 */
-	Chaos::FPBDRigidClusteredParticleHandle* CreateClusterParticle(
+	CHAOS_API Chaos::FPBDRigidClusteredParticleHandle* CreateClusterParticle(
 		const int32 ClusterGroupIndex, 
 		TArray<Chaos::FPBDRigidParticleHandle*>&& Children, 
 		const FClusterCreationParameters& Parameters = FClusterCreationParameters(),
@@ -91,7 +91,7 @@ public:
 	 *  CreateClusterParticleFromClusterChildren
 	 *    Children : Rigid body ID to include in the cluster.
 	 */
-	Chaos::FPBDRigidClusteredParticleHandle* CreateClusterParticleFromClusterChildren(
+	CHAOS_API Chaos::FPBDRigidClusteredParticleHandle* CreateClusterParticleFromClusterChildren(
 		TArray<FPBDRigidParticleHandle*>&& Children, 
 		FPBDRigidClusteredParticleHandle* Parent,
 		const FRigidTransform3& ClusterWorldTM, 
@@ -102,7 +102,7 @@ public:
 	 *	ChildToParentMap: A map that may contain a pointer to one of the child particles as a key, and a pointer to its old parent particle (prior to be being released for example).
 	 *					  If a child particle exists in this map, its parent (or whatever is specified as the value) will be used to determine the correct proxy to add to the parent cluster.
 	 */
-	void AddParticlesToCluster(
+	CHAOS_API void AddParticlesToCluster(
 		FPBDRigidClusteredParticleHandle* Cluster,
 		const TArray<FPBDRigidParticleHandle*>& InChildren,
 		const TMap<FPBDRigidParticleHandle*, FPBDRigidParticleHandle*>& ChildToParentMap);
@@ -110,7 +110,7 @@ public:
 	/**
 	 * Manually remove a set of particles from the cluster.
 	 */
-	void RemoveParticlesFromCluster(
+	CHAOS_API void RemoveParticlesFromCluster(
 		FPBDRigidClusteredParticleHandle* Cluster,
 		const TArray<FPBDRigidParticleHandle*>& InChildren);
 	/**
@@ -118,7 +118,7 @@ public:
 	 *    Clusters that share a group index should be unioned into a single cluster prior to simulation.
 	 *    The GroupIndex should be set on creation, and never touched by the client again.
 	 */
-	void UnionClusterGroups();
+	CHAOS_API void UnionClusterGroups();
 
 	//
 	// Releasing
@@ -128,7 +128,7 @@ public:
 	*  DeactivateClusterParticle
 	*    Release all the particles within the cluster particle
 	*/
-	TSet<FPBDRigidParticleHandle*> DeactivateClusterParticle(FPBDRigidClusteredParticleHandle* ClusteredParticle);
+	CHAOS_API TSet<FPBDRigidParticleHandle*> DeactivateClusterParticle(FPBDRigidClusteredParticleHandle* ClusteredParticle);
 
 	/*
 	*  ReleaseClusterParticles (BasedOnStrain)
@@ -137,11 +137,11 @@ public:
 	*    that have a strain value less than this valid will be released from the 
 	*    cluster.
 	*/
-	TSet<FPBDRigidParticleHandle*> ReleaseClusterParticles(
+	CHAOS_API TSet<FPBDRigidParticleHandle*> ReleaseClusterParticles(
 		FPBDRigidClusteredParticleHandle* ClusteredParticle, 
 		bool bForceRelease = false);
 
-	TSet<FPBDRigidParticleHandle*> ReleaseClusterParticlesNoInternalCluster(
+	CHAOS_API TSet<FPBDRigidParticleHandle*> ReleaseClusterParticlesNoInternalCluster(
 		FPBDRigidClusteredParticleHandle* ClusteredParticle,
 		bool bForceRelease = false);
 
@@ -149,14 +149,14 @@ public:
 	*  ReleaseClusterParticles
 	*    Release all rigid body IDs passed,
 	*/
-	TSet<FPBDRigidParticleHandle*> ReleaseClusterParticles(
+	CHAOS_API TSet<FPBDRigidParticleHandle*> ReleaseClusterParticles(
 		TArray<FPBDRigidParticleHandle*> ChildrenParticles, bool bTriggerBreakEvents = false);
 
 	/** 
 	* Force Release a particle at any level by making sure their parent particles are also release if necessary 
 	* @Warning this will force all particles including the ones in the parent chain to be made breakable 
 	*/
-	void ForceReleaseChildParticleAndParents(FPBDRigidClusteredParticleHandle* ChildClusteredParticle, bool bTriggerBreakEvents);
+	CHAOS_API void ForceReleaseChildParticleAndParents(FPBDRigidClusteredParticleHandle* ChildClusteredParticle, bool bTriggerBreakEvents);
 
 	/*
 	*  DestroyClusterParticle
@@ -166,7 +166,7 @@ public:
 	*    Returns the active parent cluster that might need to be rebuilt because
 	*    its geometry might be pointing to deleted particle handels. 
 	*/
-	FPBDRigidClusteredParticleHandle* DestroyClusterParticle(
+	CHAOS_API FPBDRigidClusteredParticleHandle* DestroyClusterParticle(
 		FPBDRigidClusteredParticleHandle* ClusteredParticle,
 		const FClusterDestoryParameters& Parameters = FClusterDestoryParameters());
 
@@ -177,7 +177,7 @@ public:
 	*    @param ClusteredParticle handle of the cluster to break
 	*    @return true if the cluster was successfully found and act upon 
 	*/
-	bool BreakCluster(FPBDRigidClusteredParticleHandle* ClusteredParticle);
+	CHAOS_API bool BreakCluster(FPBDRigidClusteredParticleHandle* ClusteredParticle);
 
 	/*
 	*  BreakClustersByProxy
@@ -186,7 +186,7 @@ public:
 	*    @param Proxy proxy owning the clusters to break
 	*    @return true if any cluster was successfully found and act upon
 	*/
-	bool BreakClustersByProxy(const IPhysicsProxyBase* Proxy);
+	CHAOS_API bool BreakClustersByProxy(const IPhysicsProxyBase* Proxy);
 	
 	//
 	// Operational 
@@ -199,7 +199,7 @@ public:
 	*   ... Release bodies based collision impulses.
 	*   ... Updating properties as necessary.
 	*/
-	void AdvanceClustering(const FReal dt, FPBDCollisionConstraints& CollisionRule);
+	CHAOS_API void AdvanceClustering(const FReal dt, FPBDCollisionConstraints& CollisionRule);
 
 	/**
 	*  BreakingModel
@@ -208,8 +208,8 @@ public:
 	*    encoded strain. The remainder strains are propagated back down to
 	*    the children clusters.
 	*/
-	void BreakingModel();
-	void BreakingModel(TArray<FPBDRigidClusteredParticleHandle*>& InParticles);
+	CHAOS_API void BreakingModel();
+	CHAOS_API void BreakingModel(TArray<FPBDRigidClusteredParticleHandle*>& InParticles);
 	
 	//
 	// Access
@@ -221,14 +221,14 @@ public:
 	*   Walk all the decendents of the current cluster and execute FVisitorFunction.
 	*   FVisitorFunction = [](FRigidClustering& Clustering, FRigidHandle RigidHandle){}
 	*/
-	void Visitor(FClusterHandle Cluster, FVisitorFunction Function);
+	CHAOS_API void Visitor(FClusterHandle Cluster, FVisitorFunction Function);
 
 	/*
 	*  GetActiveClusterIndex
 	*    Get the current childs active cluster. Returns INDEX_NONE if
 	*    not active or driven.
 	*/
-	FPBDRigidParticleHandle* GetActiveClusterIndex(FPBDRigidParticleHandle* Child);
+	CHAOS_API FPBDRigidParticleHandle* GetActiveClusterIndex(FPBDRigidParticleHandle* Child);
 
 	/*
 	*  GetClusterIdsArray
@@ -282,7 +282,7 @@ public:
 	/*
 	* Reset all events ( this include breaking, crumbling event and tracking data 
 	*/
-	void ResetAllEvents();
+	CHAOS_API void ResetAllEvents();
 
 	/*
 	*  Cluster Break Data
@@ -315,7 +315,7 @@ public:
 	*    @param ClusteredParticle active cluster handle to query the children from
 	*    @param WorldLocation world space location to find the closest child from   
 	*/
-	FPBDRigidParticleHandle* FindClosestChild(const FPBDRigidClusteredParticleHandle* ClusteredParticle, const FVec3& WorldLocation) const;
+	CHAOS_API FPBDRigidParticleHandle* FindClosestChild(const FPBDRigidClusteredParticleHandle* ClusteredParticle, const FVec3& WorldLocation) const;
 
 	/*
 	*  FindClosest
@@ -325,7 +325,7 @@ public:
 	*    @param Particles array of clustered particles to select from  
 	*    @param WorldLocation world space location to find the closest particle from   
 	*/
-	static FPBDRigidParticleHandle* FindClosestParticle(const TArray<FPBDRigidParticleHandle*>& Particles, const FVec3& WorldLocation);
+	static CHAOS_API FPBDRigidParticleHandle* FindClosestParticle(const TArray<FPBDRigidParticleHandle*>& Particles, const FVec3& WorldLocation);
 
 	/*
 	*  FindChildrenWithinRadius
@@ -338,7 +338,7 @@ public:
 	*    @param Radius Radius to use from the WorldLocation
 	*    @param bAlwaysReturnClosest if radius query does not return anything still return the closest from the point
 	*/
-	TArray<FPBDRigidParticleHandle*> FindChildrenWithinRadius(const FPBDRigidClusteredParticleHandle* ClusteredParticle, const FVec3& WorldLocation, FReal Radius, bool bAlwaysReturnClosest) const;
+	CHAOS_API TArray<FPBDRigidParticleHandle*> FindChildrenWithinRadius(const FPBDRigidClusteredParticleHandle* ClusteredParticle, const FVec3& WorldLocation, FReal Radius, bool bAlwaysReturnClosest) const;
 
 	/*
 	*  FindParticlesWithinRadius
@@ -350,7 +350,7 @@ public:
 	*    @param WorldLocation world space location to find the closest particle from
 	*    @param bAlwaysReturnClosest if radius query does not return anything still return the closest from the point
 	*/
-	static TArray<FPBDRigidParticleHandle*> FindParticlesWithinRadius(const TArray<FPBDRigidParticleHandle*>& Particles, const FVec3& WorldLocation, FReal Radius, bool bAlwaysReturnClosest);
+	static CHAOS_API TArray<FPBDRigidParticleHandle*> FindParticlesWithinRadius(const TArray<FPBDRigidParticleHandle*>& Particles, const FVec3& WorldLocation, FReal Radius, bool bAlwaysReturnClosest);
 	
 	/**
 	* GenerateConnectionGraph
@@ -361,17 +361,17 @@ public:
 	void SetClusterUnionConnectionType(FClusterCreationParameters::EConnectionMethod ClusterConnectionType) { MClusterUnionConnectionType = ClusterConnectionType; }
 	FClusterCreationParameters::EConnectionMethod GetClusterUnionConnectionType() const { return MClusterUnionConnectionType; }
 
-	void GenerateConnectionGraph(
+	CHAOS_API void GenerateConnectionGraph(
 		TArray<FPBDRigidParticleHandle*> Particles,
 		const FClusterCreationParameters& Parameters = FClusterCreationParameters(),
 		const TSet<FPBDRigidParticleHandle*>* FromParticles = nullptr,
 		const TSet<FPBDRigidParticleHandle*>* ToParticles = nullptr);
 
-	void GenerateConnectionGraph(
+	CHAOS_API void GenerateConnectionGraph(
 		Chaos::FPBDRigidClusteredParticleHandle* Parent,
 		const FClusterCreationParameters & Parameters = FClusterCreationParameters());
 
-	void ClearConnectionGraph(FPBDRigidClusteredParticleHandle* Parent);
+	CHAOS_API void ClearConnectionGraph(FPBDRigidClusteredParticleHandle* Parent);
 
 	const TSet<Chaos::FPBDRigidClusteredParticleHandle*>& GetTopLevelClusterParents() const { return TopLevelClusterParents; }
 	TSet<Chaos::FPBDRigidClusteredParticleHandle*>& GetTopLevelClusterParents() { return TopLevelClusterParents; }
@@ -379,10 +379,10 @@ public:
 	FRigidEvolution& GetEvolution() { return MEvolution; }
 	const FRigidEvolution& GetEvolution() const { return MEvolution; }
 
-	void SetInternalStrain(FPBDRigidClusteredParticleHandle* Particle, FRealSingle Strain);
-	void SetExternalStrain(FPBDRigidClusteredParticleHandle* Particle, FRealSingle Strain);
+	CHAOS_API void SetInternalStrain(FPBDRigidClusteredParticleHandle* Particle, FRealSingle Strain);
+	CHAOS_API void SetExternalStrain(FPBDRigidClusteredParticleHandle* Particle, FRealSingle Strain);
 
-	static bool ShouldUnionsHaveCollisionParticles();
+	static CHAOS_API bool ShouldUnionsHaveCollisionParticles();
 
 	FClusterUnionManager& GetClusterUnionManager() { return ClusterUnionManager; }
 	const FClusterUnionManager& GetClusterUnionManager() const { return ClusterUnionManager; }
@@ -390,8 +390,8 @@ public:
 	const TSet<Chaos::FPBDRigidClusteredParticleHandle*>& GetTopLevelClusterParentsStrained() const { return TopLevelClusterParentsStrained; }
 
 	// Remove connectivity edges for specified particles
-	void RemoveNodeConnections(FPBDRigidParticleHandle* Child);
-	void RemoveNodeConnections(FPBDRigidClusteredParticleHandle* Child);
+	CHAOS_API void RemoveNodeConnections(FPBDRigidParticleHandle* Child);
+	CHAOS_API void RemoveNodeConnections(FPBDRigidClusteredParticleHandle* Child);
 
 	template<typename ParticleHandleTypeA, typename ParticleHandleTypeB>
 	void CreateNodeConnection(ParticleHandleTypeA* A, ParticleHandleTypeB* B)
@@ -406,75 +406,75 @@ public:
 		}
 	}
 
-	void CreateNodeConnection(FPBDRigidClusteredParticleHandle* A, FPBDRigidClusteredParticleHandle* B);
+	CHAOS_API void CreateNodeConnection(FPBDRigidClusteredParticleHandle* A, FPBDRigidClusteredParticleHandle* B);
 
 
 	/**
 	 * CleanupInternalClustersForProxies
 	 *	For a given set of physics proxies, cleanup any tracked internal clusters that we've marked as being empty.
 	 */
-	void CleanupInternalClustersForProxies(TArrayView<IPhysicsProxyBase*> Proxies);
+	CHAOS_API void CleanupInternalClustersForProxies(TArrayView<IPhysicsProxyBase*> Proxies);
  protected:
 
-	void ComputeStrainFromCollision(const FPBDCollisionConstraints& CollisionRule, const FReal Dt);
-	void ResetCollisionImpulseArray();
-	void DisableCluster(FPBDRigidClusteredParticleHandle* ClusteredParticle);
-	void ApplyStrainModifiers();
+	CHAOS_API void ComputeStrainFromCollision(const FPBDCollisionConstraints& CollisionRule, const FReal Dt);
+	CHAOS_API void ResetCollisionImpulseArray();
+	CHAOS_API void DisableCluster(FPBDRigidClusteredParticleHandle* ClusteredParticle);
+	CHAOS_API void ApplyStrainModifiers();
 
 	/*
 	* Connectivity
 	*/
-	void UpdateConnectivityGraphUsingPointImplicit(
+	CHAOS_API void UpdateConnectivityGraphUsingPointImplicit(
 		const TArray<FPBDRigidParticleHandle*>& Particles,
 		FReal CollisionThicknessPercent,
 		const TSet<FPBDRigidParticleHandle*>* FromParticles = nullptr,
 		const TSet<FPBDRigidParticleHandle*>* ToParticles = nullptr);
-	void UpdateConnectivityGraphUsingPointImplicit(
+	CHAOS_API void UpdateConnectivityGraphUsingPointImplicit(
 		Chaos::FPBDRigidClusteredParticleHandle* Parent,
 		const FClusterCreationParameters& Parameters = FClusterCreationParameters());
 
-	void FixConnectivityGraphUsingDelaunayTriangulation(
+	CHAOS_API void FixConnectivityGraphUsingDelaunayTriangulation(
 		const TArray<FPBDRigidParticleHandle*>& Particles,
 		const FClusterCreationParameters& Parameters = FClusterCreationParameters(),
 		const TSet<FPBDRigidParticleHandle*>* FromParticles = nullptr,
 		const TSet<FPBDRigidParticleHandle*>* ToParticles = nullptr);
-	void FixConnectivityGraphUsingDelaunayTriangulation(
+	CHAOS_API void FixConnectivityGraphUsingDelaunayTriangulation(
 		Chaos::FPBDRigidClusteredParticleHandle* Parent,
 		const FClusterCreationParameters& Parameters = FClusterCreationParameters());
 
-	void UpdateConnectivityGraphUsingDelaunayTriangulation(
+	CHAOS_API void UpdateConnectivityGraphUsingDelaunayTriangulation(
 		const TArray<FPBDRigidParticleHandle*>& Particles,
 		const FClusterCreationParameters& Parameters = FClusterCreationParameters(),
 		const TSet<FPBDRigidParticleHandle*>* FromParticles = nullptr,
 		const TSet<FPBDRigidParticleHandle*>* ToParticles = nullptr);
-	void UpdateConnectivityGraphUsingDelaunayTriangulation(
+	CHAOS_API void UpdateConnectivityGraphUsingDelaunayTriangulation(
 		const Chaos::FPBDRigidClusteredParticleHandle* Parent,
 		const FClusterCreationParameters& Parameters = FClusterCreationParameters());
 
-	void UpdateConnectivityGraphUsingDelaunayTriangulationWithBoundsOverlaps(
+	CHAOS_API void UpdateConnectivityGraphUsingDelaunayTriangulationWithBoundsOverlaps(
 		const TArray<FPBDRigidParticleHandle*>& Particles,
 		const FClusterCreationParameters& Parameters = FClusterCreationParameters(),
 		const TSet<FPBDRigidParticleHandle*>* FromParticles = nullptr,
 		const TSet<FPBDRigidParticleHandle*>* ToParticles = nullptr);
-	void UpdateConnectivityGraphUsingDelaunayTriangulationWithBoundsOverlaps(
+	CHAOS_API void UpdateConnectivityGraphUsingDelaunayTriangulationWithBoundsOverlaps(
 		const Chaos::FPBDRigidClusteredParticleHandle* Parent,
 		const FClusterCreationParameters& Parameters = FClusterCreationParameters());
 	
-	void RemoveChildFromParent(FPBDRigidParticleHandle* Child, FPBDRigidClusteredParticleHandle* ClusteredParent);
-	void RemoveChildFromParentAndChildrenArray(FPBDRigidParticleHandle* Child, FPBDRigidClusteredParticleHandle* ClusteredParent);
+	CHAOS_API void RemoveChildFromParent(FPBDRigidParticleHandle* Child, FPBDRigidClusteredParticleHandle* ClusteredParent);
+	CHAOS_API void RemoveChildFromParentAndChildrenArray(FPBDRigidParticleHandle* Child, FPBDRigidClusteredParticleHandle* ClusteredParent);
 
 	// When a body has broken due to contact resolution, record an entry in a set
 	// for the collision and the particle who's momentum should be restored.
-	void TrackBreakingCollision(FPBDRigidClusteredParticleHandle* ClusteredParticle);
+	CHAOS_API void TrackBreakingCollision(FPBDRigidClusteredParticleHandle* ClusteredParticle);
 
 	// Restore some percentage of momenta for objects which were involved in collisions
 	// with destroyed GCs
-	void RestoreBreakingMomentum();
+	CHAOS_API void RestoreBreakingMomentum();
 
-	void SendBreakingEvent(FPBDRigidClusteredParticleHandle* ClusteredParticle, bool bFromCrumble);
-	void SendCrumblingEvent(FPBDRigidClusteredParticleHandle* ClusteredParticle);
+	CHAOS_API void SendBreakingEvent(FPBDRigidClusteredParticleHandle* ClusteredParticle, bool bFromCrumble);
+	CHAOS_API void SendCrumblingEvent(FPBDRigidClusteredParticleHandle* ClusteredParticle);
 
-	TSet<FPBDRigidParticleHandle*> ReleaseClusterParticlesImpl(
+	CHAOS_API TSet<FPBDRigidParticleHandle*> ReleaseClusterParticlesImpl(
 		FPBDRigidClusteredParticleHandle* ClusteredParticle, 
 		bool bForceRelease,
 		bool bCreateNewClusters);
@@ -482,20 +482,20 @@ public:
 	/**
 	 * Handles leveraging the connectivity edges on the children of the clustered particle to produce the desired effects.
 	 */
-	TSet<FPBDRigidParticleHandle*> HandleConnectivityOnReleaseClusterParticle(FPBDRigidClusteredParticleHandle* ClusteredParticle, bool bCreateNewClusters);
+	CHAOS_API TSet<FPBDRigidParticleHandle*> HandleConnectivityOnReleaseClusterParticle(FPBDRigidClusteredParticleHandle* ClusteredParticle, bool bCreateNewClusters);
 	
 	using FParticleIsland = TArray<FPBDRigidParticleHandle*>;
-	TArray<FParticleIsland> FindIslandsInChildren(const FPBDRigidClusteredParticleHandle* ClusteredParticle, bool bTraverseInterclusterEdges);
-	TArray<FPBDRigidParticleHandle*> CreateClustersFromNewIslands(TArray<FParticleIsland>& Islands, FPBDRigidClusteredParticleHandle* ClusteredParent);
+	CHAOS_API TArray<FParticleIsland> FindIslandsInChildren(const FPBDRigidClusteredParticleHandle* ClusteredParticle, bool bTraverseInterclusterEdges);
+	CHAOS_API TArray<FPBDRigidParticleHandle*> CreateClustersFromNewIslands(TArray<FParticleIsland>& Islands, FPBDRigidClusteredParticleHandle* ClusteredParent);
 
-	void UpdateTopLevelParticle(FPBDRigidClusteredParticleHandle* Particle);
+	CHAOS_API void UpdateTopLevelParticle(FPBDRigidClusteredParticleHandle* Particle);
 
 	/**
 	 * This function is a bit more versatile than the name suggests. This function can either be used to update the cluster properties
 	 * incrementally or entirely rebuild the properties all over again. This all depends on whether the input children is
 	 * either 1) the new children or 2) all the children as well as what those initial properties are set to.
 	 */
-	void UpdateClusterParticlePropertiesFromChildren(
+	CHAOS_API void UpdateClusterParticlePropertiesFromChildren(
 		FPBDRigidClusteredParticleHandle* Cluster,
 		const FRigidHandleArray& Children,
 		const TMap<FPBDRigidParticleHandle*, FPBDRigidParticleHandle*>& ChildToParentMap);
diff --git a/Engine/Source/Runtime/Experimental/Chaos/Public/Chaos/PBDRigidDynamicSpringConstraints.h b/Engine/Source/Runtime/Experimental/Chaos/Public/Chaos/PBDRigidDynamicSpringConstraints.h
--- a/Engine/Source/Runtime/Experimental/Chaos/Public/Chaos/PBDRigidDynamicSpringConstraints.h
+++ b/Engine/Source/Runtime/Experimental/Chaos/Public/Chaos/PBDRigidDynamicSpringConstraints.h
@@ -9,7 +9,7 @@ namespace Chaos
 {
 	class FPBDRigidDynamicSpringConstraints;
 
-	class CHAOS_API FPBDRigidDynamicSpringConstraintHandle final : public TIndexedContainerConstraintHandle<FPBDRigidDynamicSpringConstraints>
+	class FPBDRigidDynamicSpringConstraintHandle final : public TIndexedContainerConstraintHandle<FPBDRigidDynamicSpringConstraints>
 	{
 	public:
 		using Base = TIndexedContainerConstraintHandle<FPBDRigidDynamicSpringConstraints>;
@@ -24,7 +24,7 @@ namespace Chaos
 		{
 		}
 
-		virtual FParticlePair GetConstrainedParticles() const override final;
+		CHAOS_API virtual FParticlePair GetConstrainedParticles() const override final;
 
 		static const FConstraintHandleTypeID& StaticType()
 		{
@@ -37,7 +37,7 @@ namespace Chaos
 		using Base::ConcreteContainer;
 	};
 
-	class CHAOS_API FPBDRigidDynamicSpringConstraints : public TPBDIndexedConstraintContainer<FPBDRigidDynamicSpringConstraints>
+	class FPBDRigidDynamicSpringConstraints : public TPBDIndexedConstraintContainer<FPBDRigidDynamicSpringConstraints>
 	{
 	public:
 		using Base = TPBDIndexedConstraintContainer<FPBDRigidDynamicSpringConstraints>;
@@ -189,34 +189,34 @@ namespace Chaos
 			return Constraints[ConstraintIndex];
 		}
 
-		void UpdatePositionBasedState(const FReal Dt);
+		CHAOS_API void UpdatePositionBasedState(const FReal Dt);
 
 		//
 		// FConstraintContainer Implementation
 		//
 		virtual int32 GetNumConstraints() const override final { return NumConstraints(); }
 		virtual void ResetConstraints() override final {}
-		virtual void AddConstraintsToGraph(Private::FPBDIslandManager& IslandManager) override final;
+		CHAOS_API virtual void AddConstraintsToGraph(Private::FPBDIslandManager& IslandManager) override final;
 		virtual void PrepareTick() override final {}
 		virtual void UnprepareTick() override final {}
 
 		//
 		// TSimpleConstraintContainerSolver API - used by RBAN solvers
 		//
-		void AddBodies(FSolverBodyContainer& SolverBodyContainer);
+		CHAOS_API void AddBodies(FSolverBodyContainer& SolverBodyContainer);
 		void GatherInput(const FReal Dt) {}
-		void ScatterOutput(const FReal Dt);
-		void ApplyPositionConstraints(const FReal Dt, const int32 It, const int32 NumIts);
+		CHAOS_API void ScatterOutput(const FReal Dt);
+		CHAOS_API void ApplyPositionConstraints(const FReal Dt, const int32 It, const int32 NumIts);
 		void ApplyVelocityConstraints(const FReal Dt, const int32 It, const int32 NumIts) {}
 		void ApplyProjectionConstraints(const FReal Dt, const int32 It, const int32 NumIts) {}
 
 		//
 		// TIndexedConstraintContainerSolver API - used by World solvers
 		//
-		void AddBodies(const TArrayView<int32>& ConstraintIndices, FSolverBodyContainer& SolverBodyContainer);
+		CHAOS_API void AddBodies(const TArrayView<int32>& ConstraintIndices, FSolverBodyContainer& SolverBodyContainer);
 		void GatherInput(const TArrayView<int32>& ConstraintIndices, const FReal Dt) {}
-		void ScatterOutput(const TArrayView<int32>& ConstraintIndices, const FReal Dt);
-		void ApplyPositionConstraints(const TArrayView<int32>& ConstraintIndices, const FReal Dt, const int32 It, const int32 NumIts);
+		CHAOS_API void ScatterOutput(const TArrayView<int32>& ConstraintIndices, const FReal Dt);
+		CHAOS_API void ApplyPositionConstraints(const TArrayView<int32>& ConstraintIndices, const FReal Dt, const int32 It, const int32 NumIts);
 		void ApplyVelocityConstraints(const TArrayView<int32>& ConstraintIndices, const FReal Dt, const int32 It, const int32 NumIts) {}
 		void ApplyProjectionConstraints(const TArrayView<int32>& ConstraintIndices, const FReal Dt, const int32 It, const int32 NumIts) {}
 
@@ -225,10 +225,10 @@ namespace Chaos
 		using Base::SetConstraintIndex;
 
 	private:
-		void AddBodies(const int32 ConstraintIndex, FSolverBodyContainer& SolverBodyContainer);
-		void ApplySingle(const FReal Dt, int32 ConstraintIndex) const;
+		CHAOS_API void AddBodies(const int32 ConstraintIndex, FSolverBodyContainer& SolverBodyContainer);
+		CHAOS_API void ApplySingle(const FReal Dt, int32 ConstraintIndex) const;
 
-		FVec3 GetDelta(const FVec3& WorldSpaceX1, const FVec3& WorldSpaceX2, const int32 ConstraintIndex, const int32 SpringIndex) const;
+		CHAOS_API FVec3 GetDelta(const FVec3& WorldSpaceX1, const FVec3& WorldSpaceX2, const int32 ConstraintIndex, const int32 SpringIndex) const;
 
 		TArray<FConstrainedParticlePair> Constraints;
 		TArray<TArray<TVec2<FVec3>>> Distances;
diff --git a/Engine/Source/Runtime/Experimental/Chaos/Public/Chaos/PBDRigidParticles.h b/Engine/Source/Runtime/Experimental/Chaos/Public/Chaos/PBDRigidParticles.h
--- a/Engine/Source/Runtime/Experimental/Chaos/Public/Chaos/PBDRigidParticles.h
+++ b/Engine/Source/Runtime/Experimental/Chaos/Public/Chaos/PBDRigidParticles.h
@@ -27,14 +27,14 @@ class TPBDRigidParticles : public TRigidParticles<T, d>
 	using TRigidParticles<T, d>::RotationOfMass;
 	using TRigidParticles<T, d>::Sleeping;
 
-	CHAOS_API TPBDRigidParticles()
+	TPBDRigidParticles()
 	    : TRigidParticles<T, d>()
 	{
 		this->MParticleType = EParticleType::Rigid;
 		RegisterArrays();
 	}
 	TPBDRigidParticles(const TPBDRigidParticles<T, d>& Other) = delete;
-	CHAOS_API TPBDRigidParticles(TPBDRigidParticles<T, d>&& Other)
+	TPBDRigidParticles(TPBDRigidParticles<T, d>&& Other)
 	    : TRigidParticles<T, d>(MoveTemp(Other))
 		, MP(MoveTemp(Other.MP))
 		, MQ(MoveTemp(Other.MQ))
@@ -46,7 +46,7 @@ class TPBDRigidParticles : public TRigidParticles<T, d>
 		RegisterArrays();
 	}
 
-	CHAOS_API virtual ~TPBDRigidParticles()
+	virtual ~TPBDRigidParticles()
 	{
 	}
 
@@ -65,19 +65,19 @@ class TPBDRigidParticles : public TRigidParticles<T, d>
 	FORCEINLINE const TRotation<T, d>& Q(const int32 index) const { return MQ[index]; }
 	FORCEINLINE TRotation<T, d>& Q(const int32 index) { return MQ[index]; }
 
-	CHAOS_API const TVector<T, d>& PreV(const int32 index) const { return MPreV[index]; }
-	CHAOS_API TVector<T, d>& PreV(const int32 index) { return MPreV[index]; }
+	const TVector<T, d>& PreV(const int32 index) const { return MPreV[index]; }
+	TVector<T, d>& PreV(const int32 index) { return MPreV[index]; }
 
-	CHAOS_API const TVector<T, d>& PreW(const int32 index) const { return MPreW[index]; }
-	CHAOS_API TVector<T, d>& PreW(const int32 index) { return MPreW[index]; }
+	const TVector<T, d>& PreW(const int32 index) const { return MPreW[index]; }
+	TVector<T, d>& PreW(const int32 index) { return MPreW[index]; }
 
 	// World-space center of mass location
-	CHAOS_API const TVector<T, d> XCom(const int32 index) const { return this->X(index) + this->R(index).RotateVector(CenterOfMass(index)); }
-	CHAOS_API const TVector<T, d> PCom(const int32 index) const { return this->P(index) + this->Q(index).RotateVector(CenterOfMass(index)); }
+	const TVector<T, d> XCom(const int32 index) const { return this->X(index) + this->R(index).RotateVector(CenterOfMass(index)); }
+	const TVector<T, d> PCom(const int32 index) const { return this->P(index) + this->Q(index).RotateVector(CenterOfMass(index)); }
 
 	// World-space center of mass rotation
-	CHAOS_API const TRotation<T, d> RCom(const int32 index) const { return this->R(index) * RotationOfMass(index); }
-	CHAOS_API const TRotation<T, d> QCom(const int32 index) const { return this->Q(index) * RotationOfMass(index); }
+	const TRotation<T, d> RCom(const int32 index) const { return this->R(index) * RotationOfMass(index); }
+	const TRotation<T, d> QCom(const int32 index) const { return this->Q(index) * RotationOfMass(index); }
 
 	void SetTransformPQCom(const int32 index, const TVector<T, d>& InPCom, const TRotation<T, d>& InQCom)
 	{
@@ -87,8 +87,8 @@ class TPBDRigidParticles : public TRigidParticles<T, d>
 
 	// The index into an FSolverBodyContainer (for dynamic particles only), or INDEX_NONE.
 	// \see FSolverBodyContainer
-	CHAOS_API int32 SolverBodyIndex(const int32 index) const { return MSolverBodyIndex[index]; }
-	CHAOS_API void SetSolverBodyIndex(const int32 index, const int32 InSolverBodyIndex) { MSolverBodyIndex[index] = InSolverBodyIndex; }
+	int32 SolverBodyIndex(const int32 index) const { return MSolverBodyIndex[index]; }
+	void SetSolverBodyIndex(const int32 index, const int32 InSolverBodyIndex) { MSolverBodyIndex[index] = InSolverBodyIndex; }
 
     // Must be reinterpret cast instead of static_cast as it's a forward declare
 	typedef TPBDRigidParticleHandle<T, d> THandleType;
@@ -97,7 +97,7 @@ class TPBDRigidParticles : public TRigidParticles<T, d>
 	//cannot be reference because double pointer would allow for badness, but still useful to have non const access to handle
 	THandleType* Handle(int32 Index) { return reinterpret_cast<THandleType*>(TGeometryParticles<T, d>::Handle(Index)); }
 
-	CHAOS_API void SetSleeping(int32 Index, bool bSleeping)
+	void SetSleeping(int32 Index, bool bSleeping)
 	{
 		if (Sleeping(Index) && bSleeping == false)
 		{
@@ -136,7 +136,7 @@ class TPBDRigidParticles : public TRigidParticles<T, d>
 		}
 	}
 
-	CHAOS_API void SetObjectState(int32 Index, EObjectStateType InObjectState)
+	void SetObjectState(int32 Index, EObjectStateType InObjectState)
 	{
 		const EObjectStateType CurrentState = this->ObjectState(Index);
 
@@ -186,7 +186,7 @@ class TPBDRigidParticles : public TRigidParticles<T, d>
 		this->ObjectState(Index) = InObjectState;
 	}
 
-	CHAOS_API void SetSleepType(int32 Index, ESleepType InSleepType)
+	void SetSleepType(int32 Index, ESleepType InSleepType)
 	{
 		if (InSleepType == ESleepType::NeverSleep && this->ObjectState(Index) == EObjectStateType::Sleeping)
 		{
@@ -202,7 +202,7 @@ class TPBDRigidParticles : public TRigidParticles<T, d>
 		return FString::Printf(TEXT("%s, MP:%s, MQ:%s, MPreV:%s, MPreW:%s"), *BaseString, *P(index).ToString(), *Q(index).ToString(), *PreV(index).ToString(), *PreW(index).ToString());
 	}
 
-	CHAOS_API virtual void Serialize(FChaosArchive& Ar) override
+	virtual void Serialize(FChaosArchive& Ar) override
 	{
 		TRigidParticles<T, d>::Serialize(Ar);
 		Ar << MP << MQ << MPreV << MPreW;
diff --git a/Engine/Source/Runtime/Experimental/Chaos/Public/Chaos/PBDRigidSpringConstraints.h b/Engine/Source/Runtime/Experimental/Chaos/Public/Chaos/PBDRigidSpringConstraints.h
--- a/Engine/Source/Runtime/Experimental/Chaos/Public/Chaos/PBDRigidSpringConstraints.h
+++ b/Engine/Source/Runtime/Experimental/Chaos/Public/Chaos/PBDRigidSpringConstraints.h
@@ -9,7 +9,7 @@ namespace Chaos
 {
 	class FPBDRigidSpringConstraints;
 
-	class CHAOS_API FPBDRigidSpringConstraintHandle final : public TIndexedContainerConstraintHandle<FPBDRigidSpringConstraints>
+	class FPBDRigidSpringConstraintHandle final : public TIndexedContainerConstraintHandle<FPBDRigidSpringConstraints>
 	{
 	public:
 		using Base = TIndexedContainerConstraintHandle<FPBDRigidSpringConstraints>;
@@ -25,14 +25,14 @@ namespace Chaos
 		{
 		}
 
-		const TVector<FVec3, 2>& GetConstraintPositions() const;
-		void SetConstraintPositions(const TVector<FVec3, 2>& ConstraintPositions);
+		CHAOS_API const TVector<FVec3, 2>& GetConstraintPositions() const;
+		CHAOS_API void SetConstraintPositions(const TVector<FVec3, 2>& ConstraintPositions);
 		
-		virtual FParticlePair GetConstrainedParticles() const override final;
+		CHAOS_API virtual FParticlePair GetConstrainedParticles() const override final;
 
 		// Get the rest length of the spring
-		FReal GetRestLength() const;
-		void SetRestLength(const FReal SpringLength);
+		CHAOS_API FReal GetRestLength() const;
+		CHAOS_API void SetRestLength(const FReal SpringLength);
 
 		static const FConstraintHandleTypeID& StaticType()
 		{
diff --git a/Engine/Source/Runtime/Experimental/Chaos/Public/Chaos/PBDRigidsEvolution.h b/Engine/Source/Runtime/Experimental/Chaos/Public/Chaos/PBDRigidsEvolution.h
--- a/Engine/Source/Runtime/Experimental/Chaos/Public/Chaos/PBDRigidsEvolution.h
+++ b/Engine/Source/Runtime/Experimental/Chaos/Public/Chaos/PBDRigidsEvolution.h
@@ -78,7 +78,7 @@ class FChaosArchive;
 template <typename TPayload, typename T, int d>
 class ISpatialAccelerationCollection;
 
-struct CHAOS_API FEvolutionStats
+struct FEvolutionStats
 {
 	int32 ActiveCollisionPoints;
 	int32 ActiveShapes;
@@ -247,7 +247,7 @@ struct FSpatialAccelerationCacheHandle
 	};
 };
 
-struct CHAOS_API ISpatialAccelerationCollectionFactory
+struct ISpatialAccelerationCollectionFactory
 {
 	//Create an empty acceleration collection with the desired buckets. Chaos enqueues acceleration structure operations per bucket
 	virtual TUniquePtr<ISpatialAccelerationCollection<FAccelerationStructureHandle, FReal, 3>> CreateEmptyCollection() = 0;
@@ -281,7 +281,7 @@ public:
 	CHAOS_API FPBDRigidsEvolutionBase(FPBDRigidsSOAs& InParticles, THandleArray<FChaosPhysicsMaterial>& InSolverPhysicsMaterials, bool InIsSingleThreaded = false);
 	CHAOS_API virtual ~FPBDRigidsEvolutionBase();
 
-	CHAOS_API TArray<FGeometryParticleHandle*> CreateStaticParticles(int32 NumParticles, const FUniqueIdx* ExistingIndices = nullptr, const FGeometryParticleParameters& Params = FGeometryParticleParameters())
+	TArray<FGeometryParticleHandle*> CreateStaticParticles(int32 NumParticles, const FUniqueIdx* ExistingIndices = nullptr, const FGeometryParticleParameters& Params = FGeometryParticleParameters())
 	{
 		auto NewParticles = Particles.CreateStaticParticles(NumParticles, ExistingIndices, Params);
 		for (auto& Particle : NewParticles)
@@ -291,7 +291,7 @@ public:
 		return NewParticles;
 	}
 
-	CHAOS_API TArray<FKinematicGeometryParticleHandle*> CreateKinematicParticles(int32 NumParticles, const FUniqueIdx* ExistingIndices = nullptr, const FKinematicGeometryParticleParameters& Params = FKinematicGeometryParticleParameters())
+	TArray<FKinematicGeometryParticleHandle*> CreateKinematicParticles(int32 NumParticles, const FUniqueIdx* ExistingIndices = nullptr, const FKinematicGeometryParticleParameters& Params = FKinematicGeometryParticleParameters())
 	{
 		auto NewParticles = Particles.CreateKinematicParticles(NumParticles, ExistingIndices, Params);
 		for (auto& Particle : NewParticles)
@@ -301,7 +301,7 @@ public:
 		return NewParticles;
 	}
 
-	CHAOS_API TArray<FPBDRigidParticleHandle*> CreateDynamicParticles(int32 NumParticles, const FUniqueIdx* ExistingIndices = nullptr, const FPBDRigidParticleParameters& Params = FPBDRigidParticleParameters())
+	TArray<FPBDRigidParticleHandle*> CreateDynamicParticles(int32 NumParticles, const FUniqueIdx* ExistingIndices = nullptr, const FPBDRigidParticleParameters& Params = FPBDRigidParticleParameters())
 	{
 		auto NewParticles = Particles.CreateDynamicParticles(NumParticles, ExistingIndices, Params);
 		for (auto& Particle : NewParticles)
@@ -311,7 +311,7 @@ public:
 		return NewParticles;
 	}
 
-	CHAOS_API TArray<TPBDRigidClusteredParticleHandle<FReal, 3>*> CreateClusteredParticles(int32 NumParticles,const FUniqueIdx* ExistingIndices = nullptr,  const FPBDRigidParticleParameters& Params = FPBDRigidParticleParameters())
+	TArray<TPBDRigidClusteredParticleHandle<FReal, 3>*> CreateClusteredParticles(int32 NumParticles,const FUniqueIdx* ExistingIndices = nullptr,  const FPBDRigidParticleParameters& Params = FPBDRigidParticleParameters())
 	{
 		auto NewParticles = Particles.CreateClusteredParticles(NumParticles, ExistingIndices, Params);
 		for (auto& Particle : NewParticles)
@@ -321,7 +321,7 @@ public:
 		return NewParticles;
 	}
 
-	CHAOS_API TArray<TPBDGeometryCollectionParticleHandle<FReal, 3>*> CreateGeometryCollectionParticles(int32 NumParticles,const FUniqueIdx* ExistingIndices = nullptr,  const FPBDRigidParticleParameters& Params = FPBDRigidParticleParameters())
+	TArray<TPBDGeometryCollectionParticleHandle<FReal, 3>*> CreateGeometryCollectionParticles(int32 NumParticles,const FUniqueIdx* ExistingIndices = nullptr,  const FPBDRigidParticleParameters& Params = FPBDRigidParticleParameters())
 	{
 		auto NewParticles = Particles.CreateGeometryCollectionParticles(NumParticles, ExistingIndices, Params);
 		for (auto& Particle : NewParticles)
@@ -331,16 +331,16 @@ public:
 		return NewParticles;
 	}
 
-	CHAOS_API void AddForceFunction(FForceRule ForceFunction) { ForceRules.Add(ForceFunction); }
-	CHAOS_API void AddImpulseFunction(FForceRule ImpulseFunction) { ImpulseRules.Add(ImpulseFunction); }
-	CHAOS_API void SetParticleUpdatePositionFunction(FUpdatePositionRule ParticleUpdate) { ParticleUpdatePosition = ParticleUpdate; }
-	CHAOS_API void SetCaptureRewindDataFunction(FCaptureRewindRule Rule){ CaptureRewindData = Rule; }
+	void AddForceFunction(FForceRule ForceFunction) { ForceRules.Add(ForceFunction); }
+	void AddImpulseFunction(FForceRule ImpulseFunction) { ImpulseRules.Add(ImpulseFunction); }
+	void SetParticleUpdatePositionFunction(FUpdatePositionRule ParticleUpdate) { ParticleUpdatePosition = ParticleUpdate; }
+	void SetCaptureRewindDataFunction(FCaptureRewindRule Rule){ CaptureRewindData = Rule; }
 
-	CHAOS_API TGeometryParticleHandles<FReal, 3>& GetParticleHandles() { return Particles.GetParticleHandles(); }
-	CHAOS_API const TGeometryParticleHandles<FReal, 3>& GetParticleHandles() const { return Particles.GetParticleHandles(); }
+	TGeometryParticleHandles<FReal, 3>& GetParticleHandles() { return Particles.GetParticleHandles(); }
+	const TGeometryParticleHandles<FReal, 3>& GetParticleHandles() const { return Particles.GetParticleHandles(); }
 
-	CHAOS_API FPBDRigidsSOAs& GetParticles() { return Particles; }
-	CHAOS_API const FPBDRigidsSOAs& GetParticles() const { return Particles; }
+	FPBDRigidsSOAs& GetParticles() { return Particles; }
+	const FPBDRigidsSOAs& GetParticles() const { return Particles; }
 
 	/** Reset the collisions warm starting when resimulate. Ideally we should store
 		  that in the RewindData history but probably too expensive for now */
@@ -350,7 +350,7 @@ public:
 	* Register a constraint container with the evolution. Constraints added to the container will be solved during the tick.
 	* @note we do not currently support removing containers. In a few places we assume the ContainerId is persistent and equal to the array index.
 	*/
-	CHAOS_API void AddConstraintContainer(FPBDConstraintContainer& InContainer, const int32 Priority = 0)
+	void AddConstraintContainer(FPBDConstraintContainer& InContainer, const int32 Priority = 0)
 	{
 		const int32 ContainerId = ConstraintContainers.Add(&InContainer);
 		InContainer.SetContainerId(ContainerId);
@@ -362,7 +362,7 @@ public:
 	/**
 	* Set the number of iterations to perform in the constraint position-solve phase
 	*/
-	CHAOS_API void SetNumPositionIterations(int32 InNumIterations)
+	void SetNumPositionIterations(int32 InNumIterations)
 	{
 		IslandGroupManager.SetNumPositionIterations(InNumIterations);
 	}
@@ -370,7 +370,7 @@ public:
 	/**
 	* Get the number of position iterations the solver is running
 	*/
-	CHAOS_API int32 GetNumPositionIterations() const
+	int32 GetNumPositionIterations() const
 	{
 		return IslandGroupManager.GetIterationSettings().GetNumPositionIterations();
 	}
@@ -378,7 +378,7 @@ public:
 	/**
 	* Set the number of iterations to perform in the constraint velocity-solve phase
 	*/
-	CHAOS_API void SetNumVelocityIterations(int32 InNumIterations)
+	void SetNumVelocityIterations(int32 InNumIterations)
 	{
 		IslandGroupManager.SetNumVelocityIterations(InNumIterations);
 	}
@@ -386,7 +386,7 @@ public:
 	/**
 	* Get the number of velocity iterations the solver is running
 	*/
-	CHAOS_API int32 GetNumVelocityIterations() const
+	int32 GetNumVelocityIterations() const
 	{
 		return IslandGroupManager.GetIterationSettings().GetNumVelocityIterations();
 	}
@@ -394,7 +394,7 @@ public:
 	/**
 	* Set the number of iterations to perform in the constraint projection phase
 	*/
-	CHAOS_API void SetNumProjectionIterations(int32 InNumIterations)
+	void SetNumProjectionIterations(int32 InNumIterations)
 	{
 		IslandGroupManager.SetNumProjectionIterations(InNumIterations);
 	}
@@ -402,7 +402,7 @@ public:
 	/**
 	* Get the number of projection iterations the solver is running
 	*/
-	CHAOS_API int32 GetNumProjectionIterations() const
+	int32 GetNumProjectionIterations() const
 	{
 		return IslandGroupManager.GetIterationSettings().GetNumProjectionIterations();
 	}
@@ -410,7 +410,7 @@ public:
 	/**
 	* Set the kinematic target for a particle. This will exist for only one tick - a new target must be set for the next tick if required.
 	*/
-	CHAOS_API void SetParticleKinematicTarget(FKinematicGeometryParticleHandle* KinematicHandle, const FKinematicTarget& NewKinematicTarget)
+	void SetParticleKinematicTarget(FKinematicGeometryParticleHandle* KinematicHandle, const FKinematicTarget& NewKinematicTarget)
 	{
 		if (KinematicHandle)
 		{
@@ -429,7 +429,7 @@ public:
 	* @todo(chaos): We should add a particle creation API to the evolution
 	* @todo(chaos): This is (or could be) very similar to Enable/Disable - do we really need both?
 	*/
-	CHAOS_API void RegisterParticle(FGeometryParticleHandle* Particle)
+	void RegisterParticle(FGeometryParticleHandle* Particle)
 	{
 		// Add to the graph if necessary. Only enabled dynamic particles are added at this stage. Kinematics
 		// and statics are ignored until referenced by a constraint.
@@ -449,7 +449,7 @@ public:
 	* Enable a particle.Only enabled particles are simulated.
 	* If the particle has constraints connected to it they will also be enabled (assuming the other particles in the constraints are also enabled). 
 	*/
-	CHAOS_API void EnableParticle(FGeometryParticleHandle* Particle)
+	void EnableParticle(FGeometryParticleHandle* Particle)
 	{
 		Particles.EnableParticle(Particle);
 		EnableConstraints(Particle);
@@ -460,7 +460,7 @@ public:
 	/**
 	* Disable a particle so that it is no longer simulated. This also disables all constraints connected to the particle.
 	*/
-	CHAOS_API void DisableParticle(FGeometryParticleHandle* Particle)
+	void DisableParticle(FGeometryParticleHandle* Particle)
 	{
 		RemoveParticleFromAccelerationStructure(*Particle);
 		Particles.DisableParticle(Particle);
@@ -472,7 +472,7 @@ public:
 	/**
 	* To be called when a particle geometry changes. We must clear collisions and anything else that may reference the prior shapes.
 	*/
-	CHAOS_API void InvalidateParticle(FGeometryParticleHandle* Particle)
+	void InvalidateParticle(FGeometryParticleHandle* Particle)
 	{
 		// Remove all constraints (collisions, joints etc) from the graph
 		IslandManager.RemoveParticleConstraints(Particle);
@@ -484,7 +484,7 @@ public:
 	
 	CHAOS_API void FlushExternalAccelerationQueue(FAccelerationStructure& Acceleration,FPendingSpatialDataQueue& ExternalQueue);
 
-	CHAOS_API void DisableParticles(TSet<FGeometryParticleHandle*> &ParticlesIn)
+	void DisableParticles(TSet<FGeometryParticleHandle*> &ParticlesIn)
 	{
 		for (FGeometryParticleHandle* Particle : ParticlesIn)
 		{
@@ -563,7 +563,7 @@ public:
 		}
 	}
 
-	CHAOS_API void DestroyParticle(FGeometryParticleHandle* Particle)
+	void DestroyParticle(FGeometryParticleHandle* Particle)
 	{
 		if (MRewindData)
 		{
@@ -580,7 +580,7 @@ public:
 	/**
 	 * Preallocate buffers for creating \p Num particles.
 	 */
-	CHAOS_API void ReserveParticles(const int32 Num)
+	void ReserveParticles(const int32 Num)
 	{
 		if (const int32 NumNew = IslandManager.ReserveParticles(Num))
 		{
@@ -596,7 +596,7 @@ public:
 	CHAOS_API void DisableParticles(const TSet<FGeometryParticleHandle*>& InParticles);
 
 	/** remove a constraint from the constraint graph */
-	CHAOS_API void RemoveConstraintFromConstraintGraph(FConstraintHandle* ConstraintHandle)
+	void RemoveConstraintFromConstraintGraph(FConstraintHandle* ConstraintHandle)
 	{
 		if (ConstraintHandle->IsInConstraintGraph())
 		{
@@ -605,7 +605,7 @@ public:
 	}
 
 	/** remove a list of constraints from the constraint graph */
-	CHAOS_API void RemoveConstraintsFromConstraintGraph(const FConstraintHandleArray& Constraints)
+	void RemoveConstraintsFromConstraintGraph(const FConstraintHandleArray& Constraints)
 	{
 		for (FConstraintHandle* ConstraintHandle : Constraints)
 		{
@@ -620,7 +620,7 @@ public:
 	* @note This only applies to persistent constraints (joints etc), not transient constraints (collisions)
 	* @see DestroyTransientConstraints()
 	*/
-	CHAOS_API void DisconnectConstraints(const TSet<FGeometryParticleHandle*>& RemovedParticles)
+	void DisconnectConstraints(const TSet<FGeometryParticleHandle*>& RemovedParticles)
 	{
 		for (FGeometryParticleHandle* ParticleHandle : RemovedParticles)
 		{
@@ -640,7 +640,7 @@ public:
 	* @note This only applies to persistent constraints (joints etc), not transient constraints (collisions)
 	* @see DestroyTransientConstraints()
 	*/
-	CHAOS_API void DisableConstraints(FGeometryParticleHandle* ParticleHandle)
+	void DisableConstraints(FGeometryParticleHandle* ParticleHandle)
 	{
 		RemoveConstraintsFromConstraintGraph(ParticleHandle->ParticleConstraints());
 
@@ -657,7 +657,7 @@ public:
 	* Enable constraints (all types except collisions) from the enabled particles; constraints will only become enabled if their particle end points are valid.
 	* @note This only applies to persistent constraints (joints etc), not transient constraints (collisons)
 	*/
-	CHAOS_API void EnableConstraints(FGeometryParticleHandle* ParticleHandle)
+	void EnableConstraints(FGeometryParticleHandle* ParticleHandle)
 	{
 		for (FConstraintHandle* Constraint : ParticleHandle->ParticleConstraints())
 		{
@@ -671,7 +671,7 @@ public:
 	/** 
 	* Clear all constraints from the system reeady for shut down 
 	*/
-	CHAOS_API void ResetConstraints()
+	void ResetConstraints()
 	{
 		// Remove all the constraints from the graph
 		GetIslandManager().Reset();
@@ -700,17 +700,17 @@ public:
 
 	const TParticleView<FPBDRigidClusteredParticles>& GetNonDisabledClusteredView() const { return Particles.GetNonDisabledClusteredView(); }
 
-	CHAOS_API TSerializablePtr<FChaosPhysicsMaterial> GetPhysicsMaterial(const FGeometryParticleHandle* Particle) const { return Particle->AuxilaryValue(PhysicsMaterials); }
+	TSerializablePtr<FChaosPhysicsMaterial> GetPhysicsMaterial(const FGeometryParticleHandle* Particle) const { return Particle->AuxilaryValue(PhysicsMaterials); }
 	
-	CHAOS_API const TUniquePtr<FChaosPhysicsMaterial> &GetPerParticlePhysicsMaterial(const FGeometryParticleHandle* Particle) const { return Particle->AuxilaryValue(PerParticlePhysicsMaterials); }
+	const TUniquePtr<FChaosPhysicsMaterial> &GetPerParticlePhysicsMaterial(const FGeometryParticleHandle* Particle) const { return Particle->AuxilaryValue(PerParticlePhysicsMaterials); }
 
-	CHAOS_API void SetPerParticlePhysicsMaterial(FGeometryParticleHandle* Particle, TUniquePtr<FChaosPhysicsMaterial> &InMaterial)
+	void SetPerParticlePhysicsMaterial(FGeometryParticleHandle* Particle, TUniquePtr<FChaosPhysicsMaterial> &InMaterial)
 	{
 		Particle->AuxilaryValue(PerParticlePhysicsMaterials) = MoveTemp(InMaterial);
 		IslandManager.UpdateParticleMaterial(Particle);
 	}
 
-	CHAOS_API void SetPhysicsMaterial(FGeometryParticleHandle* Particle, TSerializablePtr<FChaosPhysicsMaterial> InMaterial)
+	void SetPhysicsMaterial(FGeometryParticleHandle* Particle, TSerializablePtr<FChaosPhysicsMaterial> InMaterial)
 	{
 		check(!Particle->AuxilaryValue(PerParticlePhysicsMaterials)); //shouldn't be setting non unique material if a unique one already exists
 		Particle->AuxilaryValue(PhysicsMaterials) = InMaterial;
@@ -879,13 +879,13 @@ public:
 	CHAOS_API void ComputeIntermediateSpatialAcceleration(bool bBlock = false);
 
 	UE_DEPRECATED(5.2, "Renamed to GetIslandManager")
-	CHAOS_API const Private::FPBDIslandManager& GetConstraintGraph() const { return IslandManager; }
+	const Private::FPBDIslandManager& GetConstraintGraph() const { return IslandManager; }
 	UE_DEPRECATED(5.2, "Renamed to GetIslandManager")
-	CHAOS_API Private::FPBDIslandManager& GetConstraintGraph() { return IslandManager; }
+	Private::FPBDIslandManager& GetConstraintGraph() { return IslandManager; }
 	
-	CHAOS_API Private::FPBDIslandManager& GetIslandManager() { return IslandManager; }
-	CHAOS_API const Private::FPBDIslandManager& GetIslandManager() const { return IslandManager; }
-	CHAOS_API const Private::FPBDIslandGroupManager& GetIslandGroupManager() const { return IslandGroupManager; }
+	Private::FPBDIslandManager& GetIslandManager() { return IslandManager; }
+	const Private::FPBDIslandManager& GetIslandManager() const { return IslandManager; }
+	const Private::FPBDIslandGroupManager& GetIslandGroupManager() const { return IslandGroupManager; }
 
 
 	void SetResim(bool bInResim) { bIsResim = bInResim; }
@@ -1120,4 +1120,4 @@ protected:
 };
 
 
-}
\ No newline at end of file
+}
diff --git a/Engine/Source/Runtime/Experimental/Chaos/Public/Chaos/PBDRigidsEvolutionGBF.h b/Engine/Source/Runtime/Experimental/Chaos/Public/Chaos/PBDRigidsEvolutionGBF.h
--- a/Engine/Source/Runtime/Experimental/Chaos/Public/Chaos/PBDRigidsEvolutionGBF.h
+++ b/Engine/Source/Runtime/Experimental/Chaos/Public/Chaos/PBDRigidsEvolutionGBF.h
@@ -168,7 +168,7 @@ namespace Chaos
 		  that in the RewindData history but probably too expensive for now */
 		virtual void ResetCollisions() override;
 
-		CHAOS_API inline void EndFrame(FReal Dt)
+		inline void EndFrame(FReal Dt)
 		{
 			Particles.GetNonDisabledDynamicView().ParallelFor([&](auto& Particle, int32 Index) {
 				Particle.Acceleration() = FVec3(0);
@@ -338,7 +338,7 @@ namespace Chaos
 		CHAOS_API TUniquePtr<IResimCacheBase> CreateExternalResimCache() const;
 		CHAOS_API void SetCurrentStepResimCache(IResimCacheBase* InCurrentStepResimCache);
 
-		CHAOS_API FSpatialAccelerationBroadPhase& GetBroadPhase() { return BroadPhase; }
+		FSpatialAccelerationBroadPhase& GetBroadPhase() { return BroadPhase; }
 
 		CHAOS_API void TransferJointConstraintCollisions();
 
diff --git a/Engine/Source/Runtime/Experimental/Chaos/Public/Chaos/PBDSphericalConstraint.h b/Engine/Source/Runtime/Experimental/Chaos/Public/Chaos/PBDSphericalConstraint.h
--- a/Engine/Source/Runtime/Experimental/Chaos/Public/Chaos/PBDSphericalConstraint.h
+++ b/Engine/Source/Runtime/Experimental/Chaos/Public/Chaos/PBDSphericalConstraint.h
@@ -26,7 +26,7 @@ extern CHAOS_API bool bChaos_Spherical_ISPC_Enabled;
 namespace Chaos::Softs
 {
 
-class CHAOS_API FPBDSphericalConstraint final
+class FPBDSphericalConstraint final
 {
 public:
 	static bool IsEnabled(const FCollectionPropertyConstFacade& PropertyCollection)
@@ -85,7 +85,7 @@ public:
 
 	~FPBDSphericalConstraint() {}
 
-	void SetProperties(
+	CHAOS_API void SetProperties(
 		const FCollectionPropertyConstFacade& PropertyCollection,
 		const TMap<FString, TConstArrayView<FRealSingle>>& WeightMaps,
 		FSolverReal MeshScale);
@@ -165,7 +165,7 @@ private:
 		});
 	}
 
-	void ApplyHelperISPC(FSolverParticles& Particles, const FSolverReal Dt) const;
+	CHAOS_API void ApplyHelperISPC(FSolverParticles& Particles, const FSolverReal Dt) const;
 
 protected:
 	const TArray<FSolverVec3>& AnimationPositions;  // Use global indexation (will need adding ParticleOffset)
@@ -183,7 +183,7 @@ private:
 	UE_CHAOS_DECLARE_PROPERTYCOLLECTION_NAME(MaxDistance, float);
 };
 
-class CHAOS_API FPBDSphericalBackstopConstraint final
+class FPBDSphericalBackstopConstraint final
 {
 public:
 	static bool IsEnabled(const FCollectionPropertyConstFacade& PropertyCollection)
@@ -261,7 +261,7 @@ public:
 	}
 	~FPBDSphericalBackstopConstraint() {}
 
-	void SetProperties(
+	CHAOS_API void SetProperties(
 		const FCollectionPropertyConstFacade& PropertyCollection,
 		const TMap<FString, TConstArrayView<FRealSingle>>& WeightMaps,
 		FSolverReal MeshScale);
@@ -456,8 +456,8 @@ private:
 		});
 	}
 
-	void ApplyLegacyHelperISPC(FSolverParticles& Particles, const FSolverReal Dt) const;
-	void ApplyHelperISPC(FSolverParticles& Particles, const FSolverReal Dt) const;
+	CHAOS_API void ApplyLegacyHelperISPC(FSolverParticles& Particles, const FSolverReal Dt) const;
+	CHAOS_API void ApplyHelperISPC(FSolverParticles& Particles, const FSolverReal Dt) const;
 
 private:
 	const TArray<FSolverVec3>& AnimationPositions;  // Positions of spheres, use global indexation (will need adding ParticleOffset)
diff --git a/Engine/Source/Runtime/Experimental/Chaos/Public/Chaos/PBDSpringConstraints.h b/Engine/Source/Runtime/Experimental/Chaos/Public/Chaos/PBDSpringConstraints.h
--- a/Engine/Source/Runtime/Experimental/Chaos/Public/Chaos/PBDSpringConstraints.h
+++ b/Engine/Source/Runtime/Experimental/Chaos/Public/Chaos/PBDSpringConstraints.h
@@ -7,7 +7,7 @@
 namespace Chaos::Softs
 {
 
-class CHAOS_API FPBDSpringConstraints : public FPBDSpringConstraintsBase
+class FPBDSpringConstraints : public FPBDSpringConstraintsBase
 {
 public:
 	template<int32 Valence>
@@ -34,7 +34,7 @@ public:
 
 	virtual ~FPBDSpringConstraints() override {}
 
-	void Apply(FSolverParticles& Particles, const FSolverReal Dt) const;
+	CHAOS_API void Apply(FSolverParticles& Particles, const FSolverReal Dt) const;
 
 	const TArray<int32>& GetConstraintsPerColorStartIndex() const { return ConstraintsPerColorStartIndex; }
 
@@ -46,14 +46,14 @@ protected:
 	using Base::ParticleCount;
 
 private:
-	void InitColor(const FSolverParticles& InParticles);
-	void ApplyHelper(FSolverParticles& Particles, const FSolverReal Dt, const int32 ConstraintIndex, const FSolverReal ExpStiffnessValue) const;
+	CHAOS_API void InitColor(const FSolverParticles& InParticles);
+	CHAOS_API void ApplyHelper(FSolverParticles& Particles, const FSolverReal Dt, const int32 ConstraintIndex, const FSolverReal ExpStiffnessValue) const;
 
 private:
 	TArray<int32> ConstraintsPerColorStartIndex; // Constraints are ordered so each batch is contiguous. This is ColorNum + 1 length so it can be used as start and end.
 };
 
-class CHAOS_API FPBDEdgeSpringConstraints final : public FPBDSpringConstraints
+class FPBDEdgeSpringConstraints final : public FPBDSpringConstraints
 {
 public:
 	static bool IsEnabled(const FCollectionPropertyConstFacade& PropertyCollection)
@@ -100,7 +100,7 @@ public:
 
 	virtual ~FPBDEdgeSpringConstraints() override = default;
 
-	void SetProperties(
+	CHAOS_API void SetProperties(
 		const FCollectionPropertyConstFacade& PropertyCollection,
 		const TMap<FString, TConstArrayView<FRealSingle>>& WeightMaps);
 
@@ -116,7 +116,7 @@ private:
 	UE_CHAOS_DECLARE_PROPERTYCOLLECTION_NAME(EdgeSpringStiffness, float);
 };
 
-class CHAOS_API FPBDBendingSpringConstraints final : public FPBDSpringConstraints
+class FPBDBendingSpringConstraints final : public FPBDSpringConstraints
 {
 public:
 	static bool IsEnabled(const FCollectionPropertyConstFacade& PropertyCollection)
@@ -163,7 +163,7 @@ public:
 
 	virtual ~FPBDBendingSpringConstraints() override = default;
 
-	void SetProperties(
+	CHAOS_API void SetProperties(
 		const FCollectionPropertyConstFacade& PropertyCollection,
 		const TMap<FString, TConstArrayView<FRealSingle>>& WeightMaps);
 
diff --git a/Engine/Source/Runtime/Experimental/Chaos/Public/Chaos/PBDSuspensionConstraintData.h b/Engine/Source/Runtime/Experimental/Chaos/Public/Chaos/PBDSuspensionConstraintData.h
--- a/Engine/Source/Runtime/Experimental/Chaos/Public/Chaos/PBDSuspensionConstraintData.h
+++ b/Engine/Source/Runtime/Experimental/Chaos/Public/Chaos/PBDSuspensionConstraintData.h
@@ -30,10 +30,10 @@ namespace Chaos
 
 	using FSuspensionConstraintDirtyFlags = TDirtyFlags<ESuspensionConstraintFlags>;
 
-	class CHAOS_API FSuspensionConstraint : public FConstraintBase
+	class FSuspensionConstraint : public FConstraintBase
 	{
 	public:
-		FSuspensionConstraint();
+		CHAOS_API FSuspensionConstraint();
 		virtual ~FSuspensionConstraint() override {}
 
 
diff --git a/Engine/Source/Runtime/Experimental/Chaos/Public/Chaos/PBDSuspensionConstraintTypes.h b/Engine/Source/Runtime/Experimental/Chaos/Public/Chaos/PBDSuspensionConstraintTypes.h
--- a/Engine/Source/Runtime/Experimental/Chaos/Public/Chaos/PBDSuspensionConstraintTypes.h
+++ b/Engine/Source/Runtime/Experimental/Chaos/Public/Chaos/PBDSuspensionConstraintTypes.h
@@ -24,7 +24,7 @@ namespace Chaos
 		FSingleParticlePhysicsProxy* ParticleProxy = nullptr;
 	};
 
-	class CHAOS_API FPBDSuspensionSettings
+	class FPBDSuspensionSettings
 	{
 	public:
 
@@ -71,7 +71,7 @@ namespace Chaos
 		FVec3 Normal;
 	};
 
-	class CHAOS_API FPBDSuspensionResults
+	class FPBDSuspensionResults
 	{
 	public:
 		FPBDSuspensionResults()
@@ -91,7 +91,7 @@ namespace Chaos
 		FReal Length;
 	};
 
-	class CHAOS_API FPBDSuspensionSolverSettings
+	class FPBDSuspensionSolverSettings
 	{
 	public:
 
@@ -102,4 +102,4 @@ namespace Chaos
 
 	};
 
-}
\ No newline at end of file
+}
