[Chaos]

diff --git a/Engine/Source/Runtime/Experimental/Chaos/Public/Chaos/PBDSuspensionConstraints.h b/Engine/Source/Runtime/Experimental/Chaos/Public/Chaos/PBDSuspensionConstraints.h
--- a/Engine/Source/Runtime/Experimental/Chaos/Public/Chaos/PBDSuspensionConstraints.h
+++ b/Engine/Source/Runtime/Experimental/Chaos/Public/Chaos/PBDSuspensionConstraints.h
@@ -23,7 +23,7 @@ namespace Chaos
 		class FPBDCollisionSolverManifoldPoint;
 	}
 
-	class CHAOS_API FPBDSuspensionConstraintHandle final : public TIndexedContainerConstraintHandle<FPBDSuspensionConstraints>
+	class FPBDSuspensionConstraintHandle final : public TIndexedContainerConstraintHandle<FPBDSuspensionConstraints>
 	{
 	public:
 		using Base = TIndexedContainerConstraintHandle<FPBDSuspensionConstraints>;
@@ -32,12 +32,12 @@ namespace Chaos
 		FPBDSuspensionConstraintHandle() {}
 		FPBDSuspensionConstraintHandle(FConstraintContainer* InConstraintContainer, int32 InConstraintIndex);
 
-		FPBDSuspensionSettings& GetSettings();
-		const FPBDSuspensionSettings& GetSettings() const;
+		CHAOS_API FPBDSuspensionSettings& GetSettings();
+		CHAOS_API const FPBDSuspensionSettings& GetSettings() const;
 
-		void SetSettings(const FPBDSuspensionSettings& Settings);
+		CHAOS_API void SetSettings(const FPBDSuspensionSettings& Settings);
 
-		virtual FParticlePair GetConstrainedParticles() const override final;
+		CHAOS_API virtual FParticlePair GetConstrainedParticles() const override final;
 
 		static const FConstraintHandleTypeID& StaticType()
 		{
@@ -49,7 +49,7 @@ namespace Chaos
 		using Base::ConcreteContainer;
 	};
 
-	class CHAOS_API FPBDSuspensionConstraints : public TPBDIndexedConstraintContainer<FPBDSuspensionConstraints>
+	class FPBDSuspensionConstraints : public TPBDIndexedConstraintContainer<FPBDSuspensionConstraints>
 	{
 	public:
 		using Base = TPBDIndexedConstraintContainer<FPBDSuspensionConstraints>;
@@ -93,12 +93,12 @@ namespace Chaos
 		/**
 		 * Add a constraint.
 		 */
-		FConstraintContainerHandle* AddConstraint(TGeometryParticleHandle<FReal, 3>* Particle, const FVec3& InConstraintFrame, const FPBDSuspensionSettings& InConstraintSettings);
+		CHAOS_API FConstraintContainerHandle* AddConstraint(TGeometryParticleHandle<FReal, 3>* Particle, const FVec3& InConstraintFrame, const FPBDSuspensionSettings& InConstraintSettings);
 
 		/**
 		 * Remove a constraint.
 		 */
-		void RemoveConstraint(int ConstraintIndex);
+		CHAOS_API void RemoveConstraint(int ConstraintIndex);
 
 
 		/*
@@ -230,28 +230,28 @@ namespace Chaos
 		//
 		virtual int32 GetNumConstraints() const override final { return NumConstraints(); }
 		virtual void ResetConstraints() override final {}
-		virtual void AddConstraintsToGraph(Private::FPBDIslandManager& IslandManager) override final;
+		CHAOS_API virtual void AddConstraintsToGraph(Private::FPBDIslandManager& IslandManager) override final;
 		virtual void PrepareTick() override final {}
 		virtual void UnprepareTick() override final {}
 
 		//
 		// TSimpleConstraintContainerSolver API - used by RBAN
 		//
-		void AddBodies(FSolverBodyContainer& SolverBodyContainer);
-		void GatherInput(const FReal Dt);
-		void ScatterOutput(const FReal Dt);
-		void ApplyPositionConstraints(const FReal Dt, const int32 It, const int32 NumIts);
-		void ApplyVelocityConstraints(const FReal Dt, const int32 It, const int32 NumIts);
+		CHAOS_API void AddBodies(FSolverBodyContainer& SolverBodyContainer);
+		CHAOS_API void GatherInput(const FReal Dt);
+		CHAOS_API void ScatterOutput(const FReal Dt);
+		CHAOS_API void ApplyPositionConstraints(const FReal Dt, const int32 It, const int32 NumIts);
+		CHAOS_API void ApplyVelocityConstraints(const FReal Dt, const int32 It, const int32 NumIts);
 		void ApplyProjectionConstraints(const FReal Dt, const int32 It, const int32 NumIts) {}
 
 		//
 		// TIndexedConstraintContainerSolver API - used by World solvers
 		//
-		void AddBodies(const TArrayView<int32>& ConstraintIndices, FSolverBodyContainer& SolverBodyContainer);
-		void GatherInput(const TArrayView<int32>& ConstraintIndices, const FReal Dt);
-		void ScatterOutput(const TArrayView<int32>& ConstraintIndices, const FReal Dt);
-		void ApplyPositionConstraints(const TArrayView<int32>& ConstraintIndices, const FReal Dt, const int32 It, const int32 NumIts);
-		void ApplyVelocityConstraints(const TArrayView<int32>& ConstraintIndices, const FReal Dt, const int32 It, const int32 NumIts);
+		CHAOS_API void AddBodies(const TArrayView<int32>& ConstraintIndices, FSolverBodyContainer& SolverBodyContainer);
+		CHAOS_API void GatherInput(const TArrayView<int32>& ConstraintIndices, const FReal Dt);
+		CHAOS_API void ScatterOutput(const TArrayView<int32>& ConstraintIndices, const FReal Dt);
+		CHAOS_API void ApplyPositionConstraints(const TArrayView<int32>& ConstraintIndices, const FReal Dt, const int32 It, const int32 NumIts);
+		CHAOS_API void ApplyVelocityConstraints(const TArrayView<int32>& ConstraintIndices, const FReal Dt, const int32 It, const int32 NumIts);
 		void ApplyProjectionConstraints(const TArrayView<int32>& ConstraintIndices, const FReal Dt, const int32 It, const int32 NumIts) {}
 
 
@@ -260,12 +260,12 @@ namespace Chaos
 		using Base::SetConstraintIndex;
 
 	private:
-		void AddBodies(const int32 ConstraintIndex, FSolverBodyContainer& SolverBodyContainer);
-		void GatherInput(const int32 ConstraintIndex, FReal Dt);
-		void ScatterOutput(const int32 ConstraintIndex, FReal Dt);
-		void ApplyPositionConstraint(const int32 ConstraintIndex, const FReal Dt, const int32 It, const int32 NumIts);
-		void ApplyVelocityConstraint(const int32 ConstraintIndex, const FReal Dt, const int32 It, const int32 NumIts);
-		void ApplySingle(int32 ConstraintIndex, const FReal Dt);
+		CHAOS_API void AddBodies(const int32 ConstraintIndex, FSolverBodyContainer& SolverBodyContainer);
+		CHAOS_API void GatherInput(const int32 ConstraintIndex, FReal Dt);
+		CHAOS_API void ScatterOutput(const int32 ConstraintIndex, FReal Dt);
+		CHAOS_API void ApplyPositionConstraint(const int32 ConstraintIndex, const FReal Dt, const int32 It, const int32 NumIts);
+		CHAOS_API void ApplyVelocityConstraint(const int32 ConstraintIndex, const FReal Dt, const int32 It, const int32 NumIts);
+		CHAOS_API void ApplySingle(int32 ConstraintIndex, const FReal Dt);
 		
 		FPBDSuspensionSolverSettings SolverSettings;
 
diff --git a/Engine/Source/Runtime/Experimental/Chaos/Public/Chaos/PBDTriangleMeshCollisions.h b/Engine/Source/Runtime/Experimental/Chaos/Public/Chaos/PBDTriangleMeshCollisions.h
--- a/Engine/Source/Runtime/Experimental/Chaos/Public/Chaos/PBDTriangleMeshCollisions.h
+++ b/Engine/Source/Runtime/Experimental/Chaos/Public/Chaos/PBDTriangleMeshCollisions.h
@@ -12,7 +12,7 @@
 //
 namespace Chaos::Softs
 {
-class CHAOS_API FPBDTriangleMeshCollisions
+class FPBDTriangleMeshCollisions
 {
 public:
 	struct FContourMinimizationIntersection
@@ -153,7 +153,7 @@ public:
 		}
 	}
 
-	void Init(const FSolverParticles& Particles, const FSolverReal MinProximityQueryRadius = (FSolverReal)0.);
+	CHAOS_API void Init(const FSolverParticles& Particles, const FSolverReal MinProximityQueryRadius = (FSolverReal)0.);
 
 	void SetGlobalIntersectionAnalysis(bool bInGlobalIntersectionAnalysis) { bGlobalIntersectionAnalysis = bInGlobalIntersectionAnalysis; }
 	void SetContourMinimization(bool bInContourMinimization) { bContourMinimization = bInContourMinimization; }
diff --git a/Engine/Source/Runtime/Experimental/Chaos/Public/Chaos/PBDTriangleMeshIntersections.h b/Engine/Source/Runtime/Experimental/Chaos/Public/Chaos/PBDTriangleMeshIntersections.h
--- a/Engine/Source/Runtime/Experimental/Chaos/Public/Chaos/PBDTriangleMeshIntersections.h
+++ b/Engine/Source/Runtime/Experimental/Chaos/Public/Chaos/PBDTriangleMeshIntersections.h
@@ -9,7 +9,7 @@
 // Resolve intersections via contour minimization. Contour Minimization data is calculated in PBDTriangleMeshCollisions.
 namespace Chaos::Softs
 {
-class CHAOS_API FPBDTriangleMeshIntersections
+class FPBDTriangleMeshIntersections
 {
 public:
 	static bool IsEnabled(const FCollectionPropertyConstFacade& PropertyCollection)
@@ -29,7 +29,7 @@ public:
 
 	~FPBDTriangleMeshIntersections() = default;
 
-	void Apply(FSolverParticles& Particles, const TArray<FPBDTriangleMeshCollisions::FContourMinimizationIntersection>& Intersections, const FSolverReal Dt) const;
+	CHAOS_API void Apply(FSolverParticles& Particles, const TArray<FPBDTriangleMeshCollisions::FContourMinimizationIntersection>& Intersections, const FSolverReal Dt) const;
 
 private:
 	const FTriangleMesh& TriangleMesh;
diff --git a/Engine/Source/Runtime/Experimental/Chaos/Public/Chaos/ParticleHandle.h b/Engine/Source/Runtime/Experimental/Chaos/Public/Chaos/ParticleHandle.h
--- a/Engine/Source/Runtime/Experimental/Chaos/Public/Chaos/ParticleHandle.h
+++ b/Engine/Source/Runtime/Experimental/Chaos/Public/Chaos/ParticleHandle.h
@@ -1521,7 +1521,7 @@ TGeometryParticleHandleImp<T,d,bPersistent>* TGeometryParticleHandleImp<T,d, bPe
 	return Ar.IsLoading() ? new TGeometryParticleHandleImp<T, d, bPersistent>() : nullptr;
 }
 
-class CHAOS_API FGenericParticleHandleImp
+class FGenericParticleHandleImp
 {
 public:
 	using FDynamicParticleHandleType = FPBDRigidParticleHandle;
@@ -2177,11 +2177,11 @@ public:
 private:
 	FGeometryParticleHandle* MHandle;
 
-	static const FVec3 ZeroVector;
-	static const FRotation3 IdentityRotation;
-	static const FMatrix33 ZeroMatrix;
-	static const TUniquePtr<FBVHParticles> NullBVHParticles;
-	static const FKinematicTarget EmptyKinematicTarget;
+	static CHAOS_API const FVec3 ZeroVector;
+	static CHAOS_API const FRotation3 IdentityRotation;
+	static CHAOS_API const FMatrix33 ZeroMatrix;
+	static CHAOS_API const TUniquePtr<FBVHParticles> NullBVHParticles;
+	static CHAOS_API const FKinematicTarget EmptyKinematicTarget;
 };
 
 template <typename T, int d>
@@ -2195,7 +2195,7 @@ using TGenericParticleHandleHandleImp UE_DEPRECATED(4.27, "Deprecated. this clas
  * you have a particle handle pointer;
 
  */
-class CHAOS_API FGenericParticleHandle
+class FGenericParticleHandle
 {
 public:
 	FGenericParticleHandle() : Imp(nullptr) {}
@@ -2233,7 +2233,7 @@ private:
 template <typename T, int d>
 using TGenericParticleHandle UE_DEPRECATED(4.27, "Deprecated. this class is to be deleted, use FGenericParticleHandle instead") = FGenericParticleHandle;
 
-class CHAOS_API FConstGenericParticleHandle
+class FConstGenericParticleHandle
 {
 public:
 	FConstGenericParticleHandle() : Imp(nullptr) {}
diff --git a/Engine/Source/Runtime/Experimental/Chaos/Public/Chaos/ParticleRule.h b/Engine/Source/Runtime/Experimental/Chaos/Public/Chaos/ParticleRule.h
--- a/Engine/Source/Runtime/Experimental/Chaos/Public/Chaos/ParticleRule.h
+++ b/Engine/Source/Runtime/Experimental/Chaos/Public/Chaos/ParticleRule.h
@@ -29,7 +29,7 @@ namespace Chaos
 /**
  * Apply an effect to all particles.
  */
-class CHAOS_API FParticleRule
+class FParticleRule
 {
   public:
 	virtual void Apply(FParticles& InParticles, const FReal Dt) const { check(0); }
diff --git a/Engine/Source/Runtime/Experimental/Chaos/Public/Chaos/PerParticlePBDCollisionConstraint.h b/Engine/Source/Runtime/Experimental/Chaos/Public/Chaos/PerParticlePBDCollisionConstraint.h
--- a/Engine/Source/Runtime/Experimental/Chaos/Public/Chaos/PerParticlePBDCollisionConstraint.h
+++ b/Engine/Source/Runtime/Experimental/Chaos/Public/Chaos/PerParticlePBDCollisionConstraint.h
@@ -24,7 +24,7 @@ extern CHAOS_API bool bChaos_PerParticleCollision_ISPC_Enabled;
 namespace Chaos::Softs
 {
 
-class CHAOS_API FPerParticlePBDCollisionConstraint final
+class FPerParticlePBDCollisionConstraint final
 {
 	struct FVelocityConstraint
 	{
@@ -246,7 +246,7 @@ private:
 		}
 	}
 
-	void ApplyHelperISPC(FSolverParticles& Particles, const FSolverReal Dt, const int32 Offset, const int32 Range) const;
+	CHAOS_API void ApplyHelperISPC(FSolverParticles& Particles, const FSolverReal Dt, const int32 Offset, const int32 Range) const;
 
 private:
 	bool bFastPositionBasedFriction;
diff --git a/Engine/Source/Runtime/Experimental/Chaos/Public/Chaos/PhysicalMaterials.h b/Engine/Source/Runtime/Experimental/Chaos/Public/Chaos/PhysicalMaterials.h
--- a/Engine/Source/Runtime/Experimental/Chaos/Public/Chaos/PhysicalMaterials.h
+++ b/Engine/Source/Runtime/Experimental/Chaos/Public/Chaos/PhysicalMaterials.h
@@ -20,9 +20,9 @@ namespace Chaos
 	 * Handles returned from the material manager are only designed to be used (resolved) on the game thread
 	 * by the physics interface.
 	 */
-	struct CHAOS_API FMaterialHandle
+	struct FMaterialHandle
 	{
-		FChaosPhysicsMaterial* Get() const;
+		CHAOS_API FChaosPhysicsMaterial* Get() const;
 		FChaosMaterialHandle InnerHandle;
 
 		bool IsValid() const { return Get() != nullptr; }
@@ -39,9 +39,9 @@ namespace Chaos
 	 * Handles returned from the material manager are only designed to be used (resolved) on the game thread
 	 * by the physics interface.
 	 */
-	struct CHAOS_API FConstMaterialHandle
+	struct FConstMaterialHandle
 	{
-		const FChaosPhysicsMaterial* Get() const;
+		CHAOS_API const FChaosPhysicsMaterial* Get() const;
 		FChaosConstMaterialHandle InnerHandle;
 
 		bool IsValid() const { return Get() != nullptr; }
@@ -56,9 +56,9 @@ namespace Chaos
 	 * Handles returned from the material manager are only designed to be used (resolved) on the game thread
 	 * by the physics interface.
 	 */
-	struct CHAOS_API FMaterialMaskHandle
+	struct FMaterialMaskHandle
 	{
-		FChaosPhysicsMaterialMask* Get() const;
+		CHAOS_API FChaosPhysicsMaterialMask* Get() const;
 		FChaosMaterialMaskHandle InnerHandle;
 
 		bool IsValid() const { return Get() != nullptr; }
@@ -73,9 +73,9 @@ namespace Chaos
 	 * Handles returned from the material manager are only designed to be used (resolved) on the game thread
 	 * by the physics interface.
 	 */
-	struct CHAOS_API FConstMaterialMaskHandle
+	struct FConstMaterialMaskHandle
 	{
-		const FChaosPhysicsMaterialMask* Get() const;
+		CHAOS_API const FChaosPhysicsMaterialMask* Get() const;
 		FChaosConstMaterialMaskHandle InnerHandle;
 
 		bool IsValid() const { return Get() != nullptr; }
@@ -111,39 +111,39 @@ namespace Chaos
 	 * material pointer. When accessing the internal material always use the handle rather than
 	 * storing the result of Get()
 	 */
-	class CHAOS_API FPhysicalMaterialManager
+	class FPhysicalMaterialManager
 	{
 	public:
 
-		static FPhysicalMaterialManager& Get();
+		static CHAOS_API FPhysicalMaterialManager& Get();
 
 		/** Create a new material, returning a stable handle to it - this should be stored and not the actual material pointer */
-		FMaterialHandle Create();
-		FMaterialMaskHandle CreateMask();
+		CHAOS_API FMaterialHandle Create();
+		CHAOS_API FMaterialMaskHandle CreateMask();
 
 		/** Destroy the material referenced by the provided handle */
-		void Destroy(FMaterialHandle InHandle);
-		void Destroy(FMaterialMaskHandle InHandle);
+		CHAOS_API void Destroy(FMaterialHandle InHandle);
+		CHAOS_API void Destroy(FMaterialMaskHandle InHandle);
 
 		/** Get the actual material from a handle */
-		FChaosPhysicsMaterial* Resolve(FChaosMaterialHandle InHandle) const;
-		const FChaosPhysicsMaterial* Resolve(FChaosConstMaterialHandle InHandle) const;
+		CHAOS_API FChaosPhysicsMaterial* Resolve(FChaosMaterialHandle InHandle) const;
+		CHAOS_API const FChaosPhysicsMaterial* Resolve(FChaosConstMaterialHandle InHandle) const;
 
-		FChaosPhysicsMaterialMask* Resolve(FChaosMaterialMaskHandle InHandle) const;
-		const FChaosPhysicsMaterialMask* Resolve(FChaosConstMaterialMaskHandle InHandle) const;
+		CHAOS_API FChaosPhysicsMaterialMask* Resolve(FChaosMaterialMaskHandle InHandle) const;
+		CHAOS_API const FChaosPhysicsMaterialMask* Resolve(FChaosConstMaterialMaskHandle InHandle) const;
 
 		/** Signals stakeholders that the stored material for the provided handle has changed */
-		void UpdateMaterial(FMaterialHandle InHandle);
-		void UpdateMaterialMask(FMaterialMaskHandle InHandle);
+		CHAOS_API void UpdateMaterial(FMaterialHandle InHandle);
+		CHAOS_API void UpdateMaterialMask(FMaterialMaskHandle InHandle);
 
 		/** Gets the internal list of primary materials representing the current user state of the material data */
 		UE_DEPRECATED(5.1, "GetMasterMaterials_External is deprecated, please use GetPrimaryMaterials_External instead")
-		const THandleArray<FChaosPhysicsMaterial>& GetMasterMaterials_External() const;
+		CHAOS_API const THandleArray<FChaosPhysicsMaterial>& GetMasterMaterials_External() const;
 		UE_DEPRECATED(5.1, "GetMasterMaterialMasks_External is deprecated, please use GetPrimaryMaterialMasks_External instead")
-		const THandleArray<FChaosPhysicsMaterialMask>& GetMasterMaterialMasks_External() const;
+		CHAOS_API const THandleArray<FChaosPhysicsMaterialMask>& GetMasterMaterialMasks_External() const;
 
-		const THandleArray<FChaosPhysicsMaterial>& GetPrimaryMaterials_External() const;
-		const THandleArray<FChaosPhysicsMaterialMask>& GetPrimaryMaterialMasks_External() const;
+		CHAOS_API const THandleArray<FChaosPhysicsMaterial>& GetPrimaryMaterials_External() const;
+		CHAOS_API const THandleArray<FChaosPhysicsMaterialMask>& GetPrimaryMaterialMasks_External() const;
 
 		/** Events */
 		FOnMaterialUpdated OnMaterialUpdated;
diff --git a/Engine/Source/Runtime/Experimental/Chaos/Public/Chaos/PhysicsObject.h b/Engine/Source/Runtime/Experimental/Chaos/Public/Chaos/PhysicsObject.h
--- a/Engine/Source/Runtime/Experimental/Chaos/Public/Chaos/PhysicsObject.h
+++ b/Engine/Source/Runtime/Experimental/Chaos/Public/Chaos/PhysicsObject.h
@@ -19,10 +19,10 @@ namespace Chaos
 	using FConstPhysicsObjectHandle = const FPhysicsObject*;
 	using FPhysicsObjectId = int32;
 
-	struct CHAOS_API FPhysicsObjectDeleter
+	struct FPhysicsObjectDeleter
 	{
-		void operator()(FPhysicsObjectHandle p);
+		CHAOS_API void operator()(FPhysicsObjectHandle p);
 	};
 	using FPhysicsObjectUniquePtr = TUniquePtr<FPhysicsObject, FPhysicsObjectDeleter>;
 
-} // namespace Chaos
\ No newline at end of file
+} // namespace Chaos
diff --git a/Engine/Source/Runtime/Experimental/Chaos/Public/Chaos/PhysicsObjectCollisionInterface.h b/Engine/Source/Runtime/Experimental/Chaos/Public/Chaos/PhysicsObjectCollisionInterface.h
--- a/Engine/Source/Runtime/Experimental/Chaos/Public/Chaos/PhysicsObjectCollisionInterface.h
+++ b/Engine/Source/Runtime/Experimental/Chaos/Public/Chaos/PhysicsObjectCollisionInterface.h
@@ -10,14 +10,14 @@
 
 namespace Chaos
 {
-	struct CHAOS_API FSweepParameters
+	struct FSweepParameters
 	{
 		bool bSweepComplex = false;
 		bool bComputeMTD = false;
 	};
 
 	template<EThreadContext Id>
-	class CHAOS_API FPhysicsObjectCollisionInterface
+	class FPhysicsObjectCollisionInterface
 	{
 	public:
 		explicit FPhysicsObjectCollisionInterface(FReadPhysicsObjectInterface<Id>& InInterface) :
@@ -25,7 +25,7 @@ namespace Chaos
 		{}
 
 		// This function will not compute any overlap heuristic.
-		bool PhysicsObjectOverlap(const FConstPhysicsObjectHandle ObjectA, const FTransform& InTransformA, const FConstPhysicsObjectHandle ObjectB, const FTransform& InTransformB, bool bTraceComplex);
+		CHAOS_API bool PhysicsObjectOverlap(const FConstPhysicsObjectHandle ObjectA, const FTransform& InTransformA, const FConstPhysicsObjectHandle ObjectB, const FTransform& InTransformB, bool bTraceComplex);
 
 		// Returns all the overlaps within A given a shape B.
 		template<typename TOverlapHit>
@@ -52,12 +52,12 @@ namespace Chaos
 		}
 
 		// This function does the same as GetPhysicsObjectOverlap but also computes the MTD metric.
-		bool PhysicsObjectOverlapWithMTD(const FConstPhysicsObjectHandle ObjectA, const FTransform& InTransformA, const FConstPhysicsObjectHandle ObjectB, const FTransform& InTransformB, bool bTraceComplex, FMTDInfo& OutMTD);
+		CHAOS_API bool PhysicsObjectOverlapWithMTD(const FConstPhysicsObjectHandle ObjectA, const FTransform& InTransformA, const FConstPhysicsObjectHandle ObjectB, const FTransform& InTransformB, bool bTraceComplex, FMTDInfo& OutMTD);
 
 		// This function does the same as GetPhysicsObjectOverlap but also computes the AABB overlap metric.
-		bool PhysicsObjectOverlapWithAABB(const FConstPhysicsObjectHandle ObjectA, const FTransform& InTransformA, const FConstPhysicsObjectHandle ObjectB, const FTransform& InTransformB, bool bTraceComplex, const FVector& Tolerance, FBox& OutOverlap);
-		bool PhysicsObjectOverlapWithAABBSize(const FConstPhysicsObjectHandle ObjectA, const FTransform& InTransformA, const FConstPhysicsObjectHandle ObjectB, const FTransform& InTransformB, bool bTraceComplex, const FVector& Tolerance, FVector& OutOverlapSize);
-		bool PhysicsObjectOverlapWithAABBIntersections(const FConstPhysicsObjectHandle ObjectA, const FTransform& InTransformA, const FConstPhysicsObjectHandle ObjectB, const FTransform& InTransformB, bool bTraceComplex, const FVector& Tolerance, TArray<FBox>& Intersections);
+		CHAOS_API bool PhysicsObjectOverlapWithAABB(const FConstPhysicsObjectHandle ObjectA, const FTransform& InTransformA, const FConstPhysicsObjectHandle ObjectB, const FTransform& InTransformB, bool bTraceComplex, const FVector& Tolerance, FBox& OutOverlap);
+		CHAOS_API bool PhysicsObjectOverlapWithAABBSize(const FConstPhysicsObjectHandle ObjectA, const FTransform& InTransformA, const FConstPhysicsObjectHandle ObjectB, const FTransform& InTransformB, bool bTraceComplex, const FVector& Tolerance, FVector& OutOverlapSize);
+		CHAOS_API bool PhysicsObjectOverlapWithAABBIntersections(const FConstPhysicsObjectHandle ObjectA, const FTransform& InTransformA, const FConstPhysicsObjectHandle ObjectB, const FTransform& InTransformB, bool bTraceComplex, const FVector& Tolerance, TArray<FBox>& Intersections);
 
 		template<typename TRaycastHit>
 		bool LineTrace(TArrayView<const FConstPhysicsObjectHandle> InObjects, const FVector& WorldStart, const FVector& WorldEnd, bool bTraceComplex, TRaycastHit& OutBestHit)
@@ -281,9 +281,9 @@ namespace Chaos
 		/**
 		 * For every pair of shapes that overlap, allows the caller to perform some computation. If additional pairs of shapes need to be examined, the input TFunction should return true.
 		 */
-		bool PairwiseShapeOverlapHelper(const FConstPhysicsObjectHandle ObjectA, const FTransform& InTransformA, const FConstPhysicsObjectHandle ObjectB, const FTransform& InTransformB, bool bTraceComplex, bool bComputeMTD, const FVector& Tolerance, const TFunction<bool(const FShapeOverlapData&, const FShapeOverlapData&, const FMTDInfo&)>& Lambda);
+		CHAOS_API bool PairwiseShapeOverlapHelper(const FConstPhysicsObjectHandle ObjectA, const FTransform& InTransformA, const FConstPhysicsObjectHandle ObjectB, const FTransform& InTransformB, bool bTraceComplex, bool bComputeMTD, const FVector& Tolerance, const TFunction<bool(const FShapeOverlapData&, const FShapeOverlapData&, const FMTDInfo&)>& Lambda);
 	};
 
 	using FPhysicsObjectCollisionInterface_External = FPhysicsObjectCollisionInterface<EThreadContext::External>;
 	using FPhysicsObjectCollisionInterface_Internal = FPhysicsObjectCollisionInterface<EThreadContext::Internal>;
-}
\ No newline at end of file
+}
diff --git a/Engine/Source/Runtime/Experimental/Chaos/Public/Chaos/PhysicsObjectInterface.h b/Engine/Source/Runtime/Experimental/Chaos/Public/Chaos/PhysicsObjectInterface.h
--- a/Engine/Source/Runtime/Experimental/Chaos/Public/Chaos/PhysicsObjectInterface.h
+++ b/Engine/Source/Runtime/Experimental/Chaos/Public/Chaos/PhysicsObjectInterface.h
@@ -18,7 +18,7 @@ class FChaosScene;
 class IPhysicsProxyBase;
 
 USTRUCT(BlueprintType)
-struct CHAOS_API FClosestPhysicsObjectResult
+struct FClosestPhysicsObjectResult
 {
 	GENERATED_BODY()
 
@@ -34,7 +34,7 @@ struct CHAOS_API FClosestPhysicsObjectResult
 		return PhysicsObject != nullptr;
 	}
 
-	FName HitName() const;
+	CHAOS_API FName HitName() const;
 };
 
 enum ERadialImpulseFalloff : int;
@@ -45,7 +45,7 @@ namespace Chaos
 	class FPerShapeData;
 	struct FMTDInfo;
 
-	struct CHAOS_API FOverlapInfo
+	struct FOverlapInfo
 	{
 		FMTDInfo* MTD = nullptr;
 		FBox* AxisOverlap = nullptr;
@@ -55,53 +55,53 @@ namespace Chaos
 	 * FReadPhysicsObjectInterface will assume that these operations are safe to call (i.e. the relevant scenes have been read locked on the game thread).
 	 */
 	template<EThreadContext Id>
-	class CHAOS_API FReadPhysicsObjectInterface
+	class FReadPhysicsObjectInterface
 	{
 	public:
-		FPhysicsObjectHandle GetRootObject(const FConstPhysicsObjectHandle Object);
-		bool HasChildren(const FConstPhysicsObjectHandle Object);
-		int32 GetClusterHierarchyLevel(const FConstPhysicsObjectHandle Object);
-
-		FTransform GetTransform(const FConstPhysicsObjectHandle Object);
-		FVector GetX(const FConstPhysicsObjectHandle Object);
-		FVector GetCoM(const FConstPhysicsObjectHandle Object);
-		FVector GetWorldCoM(const FConstPhysicsObjectHandle Object);
-		FQuat GetR(const FConstPhysicsObjectHandle Object);
-		FVector GetV(const FConstPhysicsObjectHandle Object);
-		FVector GetW(const FConstPhysicsObjectHandle Object);
-		FSpatialAccelerationIdx GetSpatialIndex(const FConstPhysicsObjectHandle Object);
-
-		TThreadParticle<Id>* GetParticle(const FConstPhysicsObjectHandle Object);
-		TThreadRigidParticle<Id>* GetRigidParticle(const FConstPhysicsObjectHandle Object);
-		TArray<TThreadParticle<Id>*> GetAllParticles(TArrayView<const FConstPhysicsObjectHandle> InObjects);
-		TArray<TThreadRigidParticle<Id>*> GetAllRigidParticles(TArrayView<const FConstPhysicsObjectHandle> InObjects);
+		CHAOS_API FPhysicsObjectHandle GetRootObject(const FConstPhysicsObjectHandle Object);
+		CHAOS_API bool HasChildren(const FConstPhysicsObjectHandle Object);
+		CHAOS_API int32 GetClusterHierarchyLevel(const FConstPhysicsObjectHandle Object);
+
+		CHAOS_API FTransform GetTransform(const FConstPhysicsObjectHandle Object);
+		CHAOS_API FVector GetX(const FConstPhysicsObjectHandle Object);
+		CHAOS_API FVector GetCoM(const FConstPhysicsObjectHandle Object);
+		CHAOS_API FVector GetWorldCoM(const FConstPhysicsObjectHandle Object);
+		CHAOS_API FQuat GetR(const FConstPhysicsObjectHandle Object);
+		CHAOS_API FVector GetV(const FConstPhysicsObjectHandle Object);
+		CHAOS_API FVector GetW(const FConstPhysicsObjectHandle Object);
+		CHAOS_API FSpatialAccelerationIdx GetSpatialIndex(const FConstPhysicsObjectHandle Object);
+
+		CHAOS_API TThreadParticle<Id>* GetParticle(const FConstPhysicsObjectHandle Object);
+		CHAOS_API TThreadRigidParticle<Id>* GetRigidParticle(const FConstPhysicsObjectHandle Object);
+		CHAOS_API TArray<TThreadParticle<Id>*> GetAllParticles(TArrayView<const FConstPhysicsObjectHandle> InObjects);
+		CHAOS_API TArray<TThreadRigidParticle<Id>*> GetAllRigidParticles(TArrayView<const FConstPhysicsObjectHandle> InObjects);
 
 		UE_DEPRECATED(5.3, "GetAllShapes has been deprecated. Please use GetAllThreadShapes instead.")
-		TArray<FPerShapeData*> GetAllShapes(TArrayView<const FConstPhysicsObjectHandle> InObjects);
+		CHAOS_API TArray<FPerShapeData*> GetAllShapes(TArrayView<const FConstPhysicsObjectHandle> InObjects);
 
-		TArray<TThreadShapeInstance<Id>*> GetAllThreadShapes(TArrayView<const FConstPhysicsObjectHandle> InObjects);
+		CHAOS_API TArray<TThreadShapeInstance<Id>*> GetAllThreadShapes(TArrayView<const FConstPhysicsObjectHandle> InObjects);
 
 		// Returns true if a shape is found and we can stop iterating.
-		void VisitEveryShape(TArrayView<const FConstPhysicsObjectHandle> InObjects, TFunctionRef<bool(const FConstPhysicsObjectHandle, TThreadShapeInstance<Id>*)> Lambda);
+		CHAOS_API void VisitEveryShape(TArrayView<const FConstPhysicsObjectHandle> InObjects, TFunctionRef<bool(const FConstPhysicsObjectHandle, TThreadShapeInstance<Id>*)> Lambda);
 
 		UE_DEPRECATED(5.3, "GetPhysicsObjectOverlap has been deprecated. Please use the function for the specific overlap metric you wish to compute instead in the FPhysicsObjectCollisionInterface.")
-		bool GetPhysicsObjectOverlap(const FConstPhysicsObjectHandle ObjectA, const FConstPhysicsObjectHandle ObjectB, bool bTraceComplex, Chaos::FOverlapInfo& OutOverlap);
+		CHAOS_API bool GetPhysicsObjectOverlap(const FConstPhysicsObjectHandle ObjectA, const FConstPhysicsObjectHandle ObjectB, bool bTraceComplex, Chaos::FOverlapInfo& OutOverlap);
 
 		UE_DEPRECATED(5.3, "GetPhysicsObjectOverlapWithTransform has been deprecated. Please use the function for the specific overlap metric you wish to compute instead in the FPhysicsObjectCollisionInterface.")
-		bool GetPhysicsObjectOverlapWithTransform(const FConstPhysicsObjectHandle ObjectA, const FTransform& InTransformA, const FConstPhysicsObjectHandle ObjectB, const FTransform& InTransformB, bool bTraceComplex, Chaos::FOverlapInfo& OutOverlap);
-
-		bool AreAllValid(TArrayView<const FConstPhysicsObjectHandle> InObjects);
-		bool AreAllKinematic(TArrayView<const FConstPhysicsObjectHandle> InObjects);
-		bool AreAllSleeping(TArrayView<const FConstPhysicsObjectHandle> InObjects);
-		bool AreAllRigidBody(TArrayView<const FConstPhysicsObjectHandle> InObjects);
-		bool AreAllDynamic(TArrayView<const FConstPhysicsObjectHandle> InObjects);
-		bool AreAllDisabled(TArrayView<const FConstPhysicsObjectHandle> InObjects);
-		bool AreAllShapesQueryEnabled(TArrayView<const FConstPhysicsObjectHandle> InObjects);
-		float GetMass(TArrayView<const FConstPhysicsObjectHandle> InObjects);
-		FBox GetBounds(TArrayView<const FConstPhysicsObjectHandle> InObjects);
-		FBox GetWorldBounds(TArrayView<const FConstPhysicsObjectHandle> InObjects);
-		FClosestPhysicsObjectResult GetClosestPhysicsBodyFromLocation(TArrayView<const FConstPhysicsObjectHandle> InObjects, const FVector& WorldLocation);
-		FAccelerationStructureHandle CreateAccelerationStructureHandle(const FConstPhysicsObjectHandle Handle);
+		CHAOS_API bool GetPhysicsObjectOverlapWithTransform(const FConstPhysicsObjectHandle ObjectA, const FTransform& InTransformA, const FConstPhysicsObjectHandle ObjectB, const FTransform& InTransformB, bool bTraceComplex, Chaos::FOverlapInfo& OutOverlap);
+
+		CHAOS_API bool AreAllValid(TArrayView<const FConstPhysicsObjectHandle> InObjects);
+		CHAOS_API bool AreAllKinematic(TArrayView<const FConstPhysicsObjectHandle> InObjects);
+		CHAOS_API bool AreAllSleeping(TArrayView<const FConstPhysicsObjectHandle> InObjects);
+		CHAOS_API bool AreAllRigidBody(TArrayView<const FConstPhysicsObjectHandle> InObjects);
+		CHAOS_API bool AreAllDynamic(TArrayView<const FConstPhysicsObjectHandle> InObjects);
+		CHAOS_API bool AreAllDisabled(TArrayView<const FConstPhysicsObjectHandle> InObjects);
+		CHAOS_API bool AreAllShapesQueryEnabled(TArrayView<const FConstPhysicsObjectHandle> InObjects);
+		CHAOS_API float GetMass(TArrayView<const FConstPhysicsObjectHandle> InObjects);
+		CHAOS_API FBox GetBounds(TArrayView<const FConstPhysicsObjectHandle> InObjects);
+		CHAOS_API FBox GetWorldBounds(TArrayView<const FConstPhysicsObjectHandle> InObjects);
+		CHAOS_API FClosestPhysicsObjectResult GetClosestPhysicsBodyFromLocation(TArrayView<const FConstPhysicsObjectHandle> InObjects, const FVector& WorldLocation);
+		CHAOS_API FAccelerationStructureHandle CreateAccelerationStructureHandle(const FConstPhysicsObjectHandle Handle);
 
 		friend class FPhysicsObjectInterface;
 	protected:
@@ -115,17 +115,17 @@ namespace Chaos
 	 * FReadPhysicsObjectInterface will assume that these operations are safe to call (i.e. the relevant scenes have been read locked on the physics thread).
 	 */
 	template<EThreadContext Id>
-	class CHAOS_API FWritePhysicsObjectInterface: public FReadPhysicsObjectInterface<Id>
+	class FWritePhysicsObjectInterface: public FReadPhysicsObjectInterface<Id>
 	{
 	public:
-		void PutToSleep(TArrayView<const FPhysicsObjectHandle> InObjects);
-		void WakeUp(TArrayView<const FPhysicsObjectHandle> InObjects);
-		void AddForce(TArrayView<const FPhysicsObjectHandle> InObjects, const FVector& Force, bool bInvalidate);
-		void AddTorque(TArrayView<const FPhysicsObjectHandle> InObjects, const FVector& Torque, bool bInvalidate);
-		void AddRadialImpulse(TArrayView<const FPhysicsObjectHandle> InObjects, FVector Origin, float Radius, float Strength, enum ERadialImpulseFalloff Falloff, bool bApplyStrain, bool bInvalidate);
+		CHAOS_API void PutToSleep(TArrayView<const FPhysicsObjectHandle> InObjects);
+		CHAOS_API void WakeUp(TArrayView<const FPhysicsObjectHandle> InObjects);
+		CHAOS_API void AddForce(TArrayView<const FPhysicsObjectHandle> InObjects, const FVector& Force, bool bInvalidate);
+		CHAOS_API void AddTorque(TArrayView<const FPhysicsObjectHandle> InObjects, const FVector& Torque, bool bInvalidate);
+		CHAOS_API void AddRadialImpulse(TArrayView<const FPhysicsObjectHandle> InObjects, FVector Origin, float Radius, float Strength, enum ERadialImpulseFalloff Falloff, bool bApplyStrain, bool bInvalidate);
 		
-		void UpdateShapeCollisionFlags(TArrayView<const FPhysicsObjectHandle> InObjects, bool bSimCollision, bool bQueryCollision);
-		void UpdateShapeFilterData(TArrayView<const FPhysicsObjectHandle> InObjects, const FCollisionFilterData& QueryData, const FCollisionFilterData& SimData);
+		CHAOS_API void UpdateShapeCollisionFlags(TArrayView<const FPhysicsObjectHandle> InObjects, bool bSimCollision, bool bQueryCollision);
+		CHAOS_API void UpdateShapeFilterData(TArrayView<const FPhysicsObjectHandle> InObjects, const FCollisionFilterData& QueryData, const FCollisionFilterData& SimData);
 
 		template<typename TPayloadType, typename T, int d>
 		void AddToSpatialAcceleration(TArrayView<const FPhysicsObjectHandle> InObjects, ISpatialAcceleration<TPayloadType, T, d>* SpatialAcceleration)
@@ -159,7 +159,7 @@ namespace Chaos
 			}
 		}
 
-		void AddConnectivityEdgesBetween(TArrayView<const FPhysicsObjectHandle> FromObjects, TArrayView<const FPhysicsObjectHandle> ToObjects, const FClusterCreationParameters& Parameters);
+		CHAOS_API void AddConnectivityEdgesBetween(TArrayView<const FPhysicsObjectHandle> FromObjects, TArrayView<const FPhysicsObjectHandle> ToObjects, const FClusterCreationParameters& Parameters);
 
 		friend class FPhysicsObjectInterface;
 	protected:
@@ -174,17 +174,17 @@ namespace Chaos
 	 * Any operations on the underlying particle/particle handle should use the FReadPhysicsObjectInterface and
 	 * FWritePhysicsObjectInterface.
 	 */
-	class CHAOS_API FPhysicsObjectInterface
+	class FPhysicsObjectInterface
 	{
 	public:
-		static void SetName(const FPhysicsObjectHandle Object, const FName& InName);
-		static FName GetName(const FConstPhysicsObjectHandle Object);
+		static CHAOS_API void SetName(const FPhysicsObjectHandle Object, const FName& InName);
+		static CHAOS_API FName GetName(const FConstPhysicsObjectHandle Object);
 
-		static void SetId(const FPhysicsObjectHandle Object, int32 InId);
-		static int32 GetId(const FConstPhysicsObjectHandle Object);
+		static CHAOS_API void SetId(const FPhysicsObjectHandle Object, int32 InId);
+		static CHAOS_API int32 GetId(const FConstPhysicsObjectHandle Object);
 
-		static FPBDRigidsSolver* GetSolver(TArrayView<const FConstPhysicsObjectHandle> InObjects);
-		static IPhysicsProxyBase* GetProxy(TArrayView<const FConstPhysicsObjectHandle> InObjects);
+		static CHAOS_API FPBDRigidsSolver* GetSolver(TArrayView<const FConstPhysicsObjectHandle> InObjects);
+		static CHAOS_API IPhysicsProxyBase* GetProxy(TArrayView<const FConstPhysicsObjectHandle> InObjects);
 
 	protected:
 		// This function should not be called without an appropriate read-lock on the relevant scene.
@@ -195,4 +195,4 @@ namespace Chaos
 		template<EThreadContext Id>
 		static FWritePhysicsObjectInterface<Id> CreateWriteInterface() { return FWritePhysicsObjectInterface<Id>{}; }
 	};
-}
\ No newline at end of file
+}
diff --git a/Engine/Source/Runtime/Experimental/Chaos/Public/Chaos/PhysicsObjectInternalInterface.h b/Engine/Source/Runtime/Experimental/Chaos/Public/Chaos/PhysicsObjectInternalInterface.h
--- a/Engine/Source/Runtime/Experimental/Chaos/Public/Chaos/PhysicsObjectInternalInterface.h
+++ b/Engine/Source/Runtime/Experimental/Chaos/Public/Chaos/PhysicsObjectInternalInterface.h
@@ -6,10 +6,10 @@
 
 namespace Chaos
 {
-	class CHAOS_API FPhysicsObjectInternalInterface : public Chaos::FPhysicsObjectInterface
+	class FPhysicsObjectInternalInterface : public Chaos::FPhysicsObjectInterface
 	{
 	public:
-		static FReadPhysicsObjectInterface_Internal GetRead();
-		static FWritePhysicsObjectInterface_Internal GetWrite();
+		static CHAOS_API FReadPhysicsObjectInterface_Internal GetRead();
+		static CHAOS_API FWritePhysicsObjectInterface_Internal GetWrite();
 	};
-}
\ No newline at end of file
+}
diff --git a/Engine/Source/Runtime/Experimental/Chaos/Public/Chaos/RigidParticles.h b/Engine/Source/Runtime/Experimental/Chaos/Public/Chaos/RigidParticles.h
--- a/Engine/Source/Runtime/Experimental/Chaos/Public/Chaos/RigidParticles.h
+++ b/Engine/Source/Runtime/Experimental/Chaos/Public/Chaos/RigidParticles.h
@@ -80,14 +80,14 @@ public:
     using TParticles<T, d>::X;
     using TGeometryParticles<T, d>::R;
 
-	CHAOS_API TRigidParticles()
+	TRigidParticles()
 	    : TKinematicGeometryParticles<T, d>()
 	{
 		RegisterArrays();
 	}
 
 	TRigidParticles(const TRigidParticles<T, d>& Other) = delete;
-	CHAOS_API TRigidParticles(TRigidParticles<T, d>&& Other)
+	TRigidParticles(TRigidParticles<T, d>&& Other)
 	    : TKinematicGeometryParticles<T, d>(MoveTemp(Other))
 		, CoreData(MoveTemp(Other.CoreData))
 		, MVSmooth(MoveTemp(Other.MVSmooth))
@@ -143,7 +143,7 @@ public:
 
 	}
 
-	CHAOS_API virtual ~TRigidParticles()
+	virtual ~TRigidParticles()
 	{}
 
 	FORCEINLINE const TVector<T, d>& VSmooth(const int32 Index) const { return MVSmooth[Index]; }
@@ -199,8 +199,8 @@ public:
 
 	FORCEINLINE int32 CollisionParticlesSize(int32 Index) const { return MCollisionParticles[Index] == nullptr ? 0 : MCollisionParticles[Index]->Size(); }
 
-	void CHAOS_API CollisionParticlesInitIfNeeded(const int32 Index);
-	void CHAOS_API SetCollisionParticles(const int32 Index, TParticles<T, d>&& Particles);
+	void CollisionParticlesInitIfNeeded(const int32 Index);
+	void SetCollisionParticles(const int32 Index, TParticles<T, d>&& Particles);
 	
 	FORCEINLINE const TUniquePtr<TBVHParticles<T, d>>& CollisionParticles(const int32 Index) const { return MCollisionParticles[Index]; }
 	FORCEINLINE TUniquePtr<TBVHParticles<T, d>>& CollisionParticles(const int32 Index) { return MCollisionParticles[Index]; }
@@ -278,7 +278,7 @@ public:
 			CollisionGroup(Index), Disabled(Index), Sleeping(Index));
 	}
 
-	CHAOS_API virtual void Serialize(FChaosArchive& Ar) override
+	virtual void Serialize(FChaosArchive& Ar) override
 	{
 		TKinematicGeometryParticles<T,d>::Serialize(Ar);
 
diff --git a/Engine/Source/Runtime/Experimental/Chaos/Public/Chaos/SegmentMesh.h b/Engine/Source/Runtime/Experimental/Chaos/Public/Chaos/SegmentMesh.h
--- a/Engine/Source/Runtime/Experimental/Chaos/Public/Chaos/SegmentMesh.h
+++ b/Engine/Source/Runtime/Experimental/Chaos/Public/Chaos/SegmentMesh.h
@@ -14,22 +14,22 @@ namespace Chaos
 	/**
 	* Mesh structure of connected particles via edges.
 	*/
-	class CHAOS_API FSegmentMesh
+	class FSegmentMesh
 	{
 	public:
 		FSegmentMesh()
 		{}
-		FSegmentMesh(TArray<TVec2<int32>>&& Elements);
+		CHAOS_API FSegmentMesh(TArray<TVec2<int32>>&& Elements);
 		FSegmentMesh(const FSegmentMesh& Other) = delete;
 		FSegmentMesh(FSegmentMesh&& Other)
 		    : MElements(MoveTemp(Other.MElements))
 		    , MPointToEdgeMap(MoveTemp(Other.MPointToEdgeMap))
 		    , MPointToNeighborsMap(MoveTemp(Other.MPointToNeighborsMap))
 		{}
-		~FSegmentMesh();
+		CHAOS_API ~FSegmentMesh();
 
-		void Init(const TArray<TVec2<int32>>& Elements);
-		void Init(TArray<TVec2<int32>>&& Elements);
+		CHAOS_API void Init(const TArray<TVec2<int32>>& Elements);
+		CHAOS_API void Init(TArray<TVec2<int32>>&& Elements);
 
 		int32 GetNumElements() const
 		{
@@ -52,28 +52,28 @@ namespace Chaos
 		/**
 		* @ret A map of each point index to the list of connected points.
 		*/
-		const TMap<int32, TSet<int32>>& GetPointToNeighborsMap() const;
+		CHAOS_API const TMap<int32, TSet<int32>>& GetPointToNeighborsMap() const;
 
 		/**
 		* @ret A map of each point index to the list of connected edges.
 		*/
-		const TMap<int32, TArray<int32>>& GetPointToEdges() const;
+		CHAOS_API const TMap<int32, TArray<int32>>& GetPointToEdges() const;
 
 		/**
 		* @ret Lengths (or lengths squared) of all edges.
 		* @param InParticles - The particle positions to use.  This routine assumes it's sized appropriately.
 		* @param lengthSquared - If true, the squared length is returned, which is faster.
 		*/
-		TArray<FReal> GetEdgeLengths(
+		CHAOS_API TArray<FReal> GetEdgeLengths(
 			const TParticles<FReal, 3>& InParticles, 
 			const bool lengthSquared = false) const;
 
 	private:
-		void _ClearAuxStructures();
+		CHAOS_API void _ClearAuxStructures();
 
-		void _UpdatePointToNeighborsMap() const;
+		CHAOS_API void _UpdatePointToNeighborsMap() const;
 
-		void _UpdatePointToEdgesMap() const;
+		CHAOS_API void _UpdatePointToEdgesMap() const;
 
 	private:
 		// We use TVector rather than FEdge to represent connectivity because
diff --git a/Engine/Source/Runtime/Experimental/Chaos/Public/Chaos/ShapeInstance.h b/Engine/Source/Runtime/Experimental/Chaos/Public/Chaos/ShapeInstance.h
--- a/Engine/Source/Runtime/Experimental/Chaos/Public/Chaos/ShapeInstance.h
+++ b/Engine/Source/Runtime/Experimental/Chaos/Public/Chaos/ShapeInstance.h
@@ -32,7 +32,7 @@ namespace Chaos
 	 * - change ShapesArray() and all code using it to use ShapeInstance or ShapeInstanceProxy as appropriate
 	 * - deprecate FPerShapeData
 	 */
-	class CHAOS_API FPerShapeData
+	class FPerShapeData
 	{
 	protected:
 		enum class EPerShapeDataType : uint8
@@ -51,15 +51,15 @@ namespace Chaos
 		template<typename TLambda> decltype(auto) DownCast(const TLambda& Lambda);
 		template<typename TLambda> decltype(auto) DownCast(const TLambda& Lambda) const;
 
-		Private::FShapeInstanceExtended* AsShapeInstanceExtended();
-		const Private::FShapeInstanceExtended* AsShapeInstanceExtended() const;
+		CHAOS_API Private::FShapeInstanceExtended* AsShapeInstanceExtended();
+		CHAOS_API const Private::FShapeInstanceExtended* AsShapeInstanceExtended() const;
 
 	public:
 		// Downcasts exposed until we deprecate and remove FPerShapeData
-		FShapeInstanceProxy* AsShapeInstanceProxy();
-		const FShapeInstanceProxy* AsShapeInstanceProxy() const;
-		FShapeInstance* AsShapeInstance();
-		const FShapeInstance* AsShapeInstance() const;
+		CHAOS_API FShapeInstanceProxy* AsShapeInstanceProxy();
+		CHAOS_API const FShapeInstanceProxy* AsShapeInstanceProxy() const;
+		CHAOS_API FShapeInstance* AsShapeInstance();
+		CHAOS_API const FShapeInstance* AsShapeInstance() const;
 
 	public:
 		static constexpr bool AlwaysSerializable = true;
@@ -68,91 +68,91 @@ namespace Chaos
 		static bool RequiresCachedLeafInfo(const FImplicitObject* Geometry) { return false; }
 
 		UE_DEPRECATED(5.3, "Call FShapeInstanceProxy::Make for game thread objects, FShapeInstance::Make for physics thread objects")
-		static TUniquePtr<FPerShapeData> CreatePerShapeData(int32 InShapeIdx, TSerializablePtr<FImplicitObject> InGeometry);
+		static CHAOS_API TUniquePtr<FPerShapeData> CreatePerShapeData(int32 InShapeIdx, TSerializablePtr<FImplicitObject> InGeometry);
 
 		UE_DEPRECATED(5.3, "Call FShapeInstanceProxy::UpdateGeometry for game thread objects, FShapeInstance::UpdateGeometry for physics thread objects")
-		static void UpdateGeometry(TUniquePtr<FPerShapeData>& InOutShapePtr, TSerializablePtr<FImplicitObject> InGeometry);
+		static CHAOS_API void UpdateGeometry(TUniquePtr<FPerShapeData>& InOutShapePtr, TSerializablePtr<FImplicitObject> InGeometry);
 
 
-		static FPerShapeData* SerializationFactory(FChaosArchive& Ar, FPerShapeData*);
+		static CHAOS_API FPerShapeData* SerializationFactory(FChaosArchive& Ar, FPerShapeData*);
 
 		virtual ~FPerShapeData() {}
 
-		virtual void Serialize(FChaosArchive& Ar);
+		CHAOS_API virtual void Serialize(FChaosArchive& Ar);
 
-		void UpdateShapeBounds(const FRigidTransform3& WorldTM, const FVec3& BoundsExpansion = FVec3(0));
+		CHAOS_API void UpdateShapeBounds(const FRigidTransform3& WorldTM, const FVec3& BoundsExpansion = FVec3(0));
 
-		void* GetUserData() const;
-		void SetUserData(void* InUserData);
+		CHAOS_API void* GetUserData() const;
+		CHAOS_API void SetUserData(void* InUserData);
 
-		const FCollisionFilterData& GetQueryData() const;
-		void SetQueryData(const FCollisionFilterData& InQueryData);
+		CHAOS_API const FCollisionFilterData& GetQueryData() const;
+		CHAOS_API void SetQueryData(const FCollisionFilterData& InQueryData);
 
-		const FCollisionFilterData& GetSimData() const;
-		void SetSimData(const FCollisionFilterData& InSimData);
+		CHAOS_API const FCollisionFilterData& GetSimData() const;
+		CHAOS_API void SetSimData(const FCollisionFilterData& InSimData);
 
-		TSerializablePtr<FImplicitObject> GetGeometry() const;
+		CHAOS_API TSerializablePtr<FImplicitObject> GetGeometry() const;
 
-		const TAABB<FReal, 3>& GetWorldSpaceInflatedShapeBounds() const;
+		CHAOS_API const TAABB<FReal, 3>& GetWorldSpaceInflatedShapeBounds() const;
 
-		void UpdateWorldSpaceState(const FRigidTransform3& WorldTransform, const FVec3& BoundsExpansion);
+		CHAOS_API void UpdateWorldSpaceState(const FRigidTransform3& WorldTransform, const FVec3& BoundsExpansion);
 
 		// The leaf shape (with transformed and implicit wrapper removed).
-		const FImplicitObject* GetLeafGeometry() const;
+		CHAOS_API const FImplicitObject* GetLeafGeometry() const;
 
 		// The actor-relative transform of the leaf geometry.
-		FRigidTransform3 GetLeafRelativeTransform() const;
+		CHAOS_API FRigidTransform3 GetLeafRelativeTransform() const;
 
 		// The world-space transform of the leaf geometry.
 		// If we have non-identity leaf relative transform, is cached from the last call to UpdateWorldSpaceState.
 		// If not cahced, is constructed from arguments.
-		FRigidTransform3 GetLeafWorldTransform(const FGeometryParticleHandle* Particle) const;
+		CHAOS_API FRigidTransform3 GetLeafWorldTransform(const FGeometryParticleHandle* Particle) const;
 
-		void UpdateLeafWorldTransform(FGeometryParticleHandle* Particle);
+		CHAOS_API void UpdateLeafWorldTransform(FGeometryParticleHandle* Particle);
 
-		int32 NumMaterials() const;
-		const FMaterialHandle& GetMaterial(const int32 Index) const;
+		CHAOS_API int32 NumMaterials() const;
+		CHAOS_API const FMaterialHandle& GetMaterial(const int32 Index) const;
 
-		const TArray<FMaterialHandle>& GetMaterials() const;
-		void SetMaterial(FMaterialHandle InMaterial);
-		void SetMaterials(const TArray<FMaterialHandle>& InMaterials);
-		void SetMaterials(TArray<FMaterialHandle>&& InMaterials);
+		CHAOS_API const TArray<FMaterialHandle>& GetMaterials() const;
+		CHAOS_API void SetMaterial(FMaterialHandle InMaterial);
+		CHAOS_API void SetMaterials(const TArray<FMaterialHandle>& InMaterials);
+		CHAOS_API void SetMaterials(TArray<FMaterialHandle>&& InMaterials);
 
-		const TArray<FMaterialMaskHandle>& GetMaterialMasks() const;
-		void SetMaterialMasks(const TArray<FMaterialMaskHandle>& InMaterialMasks);
+		CHAOS_API const TArray<FMaterialMaskHandle>& GetMaterialMasks() const;
+		CHAOS_API void SetMaterialMasks(const TArray<FMaterialMaskHandle>& InMaterialMasks);
 
-		const TArray<uint32>& GetMaterialMaskMaps() const;
-		void SetMaterialMaskMaps(const TArray<uint32>& InMaterialMaskMaps);
+		CHAOS_API const TArray<uint32>& GetMaterialMaskMaps() const;
+		CHAOS_API void SetMaterialMaskMaps(const TArray<uint32>& InMaterialMaskMaps);
 
-		const TArray<FMaterialHandle>& GetMaterialMaskMapMaterials() const;
-		void SetMaterialMaskMapMaterials(const TArray<FMaterialHandle>& InMaterialMaskMapMaterials);
+		CHAOS_API const TArray<FMaterialHandle>& GetMaterialMaskMapMaterials() const;
+		CHAOS_API void SetMaterialMaskMapMaterials(const TArray<FMaterialHandle>& InMaterialMaskMapMaterials);
 
-		const FShapeDirtyFlags GetDirtyFlags() const;
+		CHAOS_API const FShapeDirtyFlags GetDirtyFlags() const;
 
-		bool GetQueryEnabled() const;
-		void SetQueryEnabled(const bool bEnable);
+		CHAOS_API bool GetQueryEnabled() const;
+		CHAOS_API void SetQueryEnabled(const bool bEnable);
 
-		bool GetSimEnabled() const;
-		void SetSimEnabled(const bool bEnable);
+		CHAOS_API bool GetSimEnabled() const;
+		CHAOS_API void SetSimEnabled(const bool bEnable);
 
-		bool GetIsProbe() const;
-		void SetIsProbe(const bool bIsProbe);
+		CHAOS_API bool GetIsProbe() const;
+		CHAOS_API void SetIsProbe(const bool bIsProbe);
 
-		EChaosCollisionTraceFlag GetCollisionTraceType() const;
-		void SetCollisionTraceType(const EChaosCollisionTraceFlag InTraceFlag);
+		CHAOS_API EChaosCollisionTraceFlag GetCollisionTraceType() const;
+		CHAOS_API void SetCollisionTraceType(const EChaosCollisionTraceFlag InTraceFlag);
 
-		const FCollisionData& GetCollisionData() const;
-		void SetCollisionData(const FCollisionData& Data);
+		CHAOS_API const FCollisionData& GetCollisionData() const;
+		CHAOS_API void SetCollisionData(const FCollisionData& Data);
 
-		const FMaterialData& GetMaterialData() const;
-		void SetMaterialData(const FMaterialData& Data);
+		CHAOS_API const FMaterialData& GetMaterialData() const;
+		CHAOS_API void SetMaterialData(const FMaterialData& Data);
 
-		void SyncRemoteData(FDirtyPropertiesManager& Manager, int32 ShapeDataIdx, FShapeDirtyData& RemoteData);
+		CHAOS_API void SyncRemoteData(FDirtyPropertiesManager& Manager, int32 ShapeDataIdx, FShapeDirtyData& RemoteData);
 		
-		void SetProxy(IPhysicsProxyBase* InProxy);
+		CHAOS_API void SetProxy(IPhysicsProxyBase* InProxy);
 		
-		int32 GetShapeIndex() const;
-		void ModifyShapeIndex(int32 NewShapeIndex);
+		CHAOS_API int32 GetShapeIndex() const;
+		CHAOS_API void ModifyShapeIndex(int32 NewShapeIndex);
 
 		template <typename Lambda> void ModifySimData(const Lambda& LambdaFunc);
 		template <typename Lambda> void ModifyMaterials(const Lambda& LambdaFunc);
@@ -225,16 +225,16 @@ namespace Chaos
 	 *
 	 * @todo(chaos) : reduce the cost of MaterialData for shapes with one materialand no masks etc.
 	 */
-	class CHAOS_API FShapeInstanceProxy : public FPerShapeData
+	class FShapeInstanceProxy : public FPerShapeData
 	{
 	public:
 		friend class FPerShapeData;
 
-		static TUniquePtr<FShapeInstanceProxy> Make(int32 InShapeIdx, TSerializablePtr<FImplicitObject> InGeometry);
-		static void UpdateGeometry(TUniquePtr<FShapeInstanceProxy>& InOutShapePtr, TSerializablePtr<FImplicitObject> InGeometry);
-		static FShapeInstanceProxy* SerializationFactory(FChaosArchive& Ar, FShapeInstanceProxy*);
+		static CHAOS_API TUniquePtr<FShapeInstanceProxy> Make(int32 InShapeIdx, TSerializablePtr<FImplicitObject> InGeometry);
+		static CHAOS_API void UpdateGeometry(TUniquePtr<FShapeInstanceProxy>& InOutShapePtr, TSerializablePtr<FImplicitObject> InGeometry);
+		static CHAOS_API FShapeInstanceProxy* SerializationFactory(FChaosArchive& Ar, FShapeInstanceProxy*);
 
-		void UpdateShapeBounds(const FRigidTransform3& WorldTM, const FVec3& BoundsExpansion = FVec3(0));
+		CHAOS_API void UpdateShapeBounds(const FRigidTransform3& WorldTM, const FVec3& BoundsExpansion = FVec3(0));
 
 		void* GetUserData() const { return CollisionData.Read().UserData; }
 		void SetUserData(void* InUserData)
@@ -254,19 +254,19 @@ namespace Chaos
 			CollisionData.Modify(true, DirtyFlags, Proxy, ShapeIdx, [InSimData](FCollisionData& Data) { Data.SimData = InSimData; });
 		}
 
-		void UpdateWorldSpaceState(const FRigidTransform3& WorldTransform, const FVec3& BoundsExpansion);
+		CHAOS_API void UpdateWorldSpaceState(const FRigidTransform3& WorldTransform, const FVec3& BoundsExpansion);
 
 		// The leaf shape (with transformed and implicit wrapper removed).
-		const FImplicitObject* GetLeafGeometry() const;
+		CHAOS_API const FImplicitObject* GetLeafGeometry() const;
 
 		// The actor-relative transform of the leaf geometry.
-		FRigidTransform3 GetLeafRelativeTransform() const;
+		CHAOS_API FRigidTransform3 GetLeafRelativeTransform() const;
 
 		// The world-space transform of the leaf geometry.
 		// If we have non-identity leaf relative transform, is cached from the last call to UpdateWorldSpaceState.
 		// If not cahced, is constructed from arguments.
-		FRigidTransform3 GetLeafWorldTransform(const FGeometryParticleHandle* Particle) const;
-		void UpdateLeafWorldTransform(FGeometryParticleHandle* Particle);
+		CHAOS_API FRigidTransform3 GetLeafWorldTransform(const FGeometryParticleHandle* Particle) const;
+		CHAOS_API void UpdateLeafWorldTransform(FGeometryParticleHandle* Particle);
 
 		int32 NumMaterials() const { return Materials.Read(). Materials.Num(); }
 		const FMaterialHandle& GetMaterial(const int32 Index) const { return Materials.Read().Materials[Index]; }
@@ -449,7 +449,7 @@ namespace Chaos
 		{
 		}
 
-		virtual void SerializeMaterials(FChaosArchive& Ar) override final;
+		CHAOS_API virtual void SerializeMaterials(FChaosArchive& Ar) override final;
 
 
 		IPhysicsProxyBase* Proxy;
@@ -746,7 +746,7 @@ namespace Chaos
 		 * 
 		 * NOTE: keep size to a minimum. There can be millions of these in s scene.
 		 */
-		class CHAOS_API FShapeInstanceExtended : public FShapeInstance
+		class FShapeInstanceExtended : public FShapeInstance
 		{
 		public:
 			friend class FShapeInstance;
@@ -1145,4 +1145,4 @@ namespace Chaos
 	UE_DEPRECATED(5.3, "Not for external use")
 	void CHAOS_API UpdateShapesArrayFromGeometry(FShapesArray& ShapesArray, TSerializablePtr<FImplicitObject> Geometry, const FRigidTransform3& ActorTM, IPhysicsProxyBase* Proxy);
 
-}
\ No newline at end of file
+}
diff --git a/Engine/Source/Runtime/Experimental/Chaos/Public/Chaos/SimCallbackInput.h b/Engine/Source/Runtime/Experimental/Chaos/Public/Chaos/SimCallbackInput.h
--- a/Engine/Source/Runtime/Experimental/Chaos/Public/Chaos/SimCallbackInput.h
+++ b/Engine/Source/Runtime/Experimental/Chaos/Public/Chaos/SimCallbackInput.h
@@ -70,7 +70,7 @@ struct FSimCallbackNoOutput : public FSimCallbackOutput
 
 /** Handle for output that is automatically cleaned up.
 	NOTE: this should only be used on external thread as the destructor automatically frees into external structures */
-class CHAOS_API FSimCallbackOutputHandle
+class FSimCallbackOutputHandle
 {
 public:
 	FSimCallbackOutputHandle()
@@ -127,7 +127,7 @@ public:
 	const FSimCallbackOutput& operator*() const { return *SimCallbackOutput; }
 private:
 
-	void Free_External();
+	CHAOS_API void Free_External();
 
 	FSimCallbackOutput* SimCallbackOutput;
 	ISimCallbackObject* SimCallbackObject;
diff --git a/Engine/Source/Runtime/Experimental/Chaos/Public/Chaos/SimCallbackObject.h b/Engine/Source/Runtime/Experimental/Chaos/Public/Chaos/SimCallbackObject.h
--- a/Engine/Source/Runtime/Experimental/Chaos/Public/Chaos/SimCallbackObject.h
+++ b/Engine/Source/Runtime/Experimental/Chaos/Public/Chaos/SimCallbackObject.h
@@ -57,7 +57,7 @@ ENUM_CLASS_FLAGS(ESimCallbackOptions)
  * We rely on this to cache results and skip callbacks when possible during a resim.
  * See functions for more details.
  */
-class CHAOS_API ISimCallbackObject
+class ISimCallbackObject
 {
 public:
 
@@ -199,7 +199,7 @@ protected:
 	/**
 	 * Gets the current producer input data. This is what the external thread should be writing to
 	 */
-	FSimCallbackInput* GetProducerInputData_External();
+	CHAOS_API FSimCallbackInput* GetProducerInputData_External();
 
 	void SetCurrentInput_Internal(FSimCallbackInput* NewInput)
 	{
diff --git a/Engine/Source/Runtime/Experimental/Chaos/Public/Chaos/SmokeEvolution.h b/Engine/Source/Runtime/Experimental/Chaos/Public/Chaos/SmokeEvolution.h
--- a/Engine/Source/Runtime/Experimental/Chaos/Public/Chaos/SmokeEvolution.h
+++ b/Engine/Source/Runtime/Experimental/Chaos/Public/Chaos/SmokeEvolution.h
@@ -11,7 +11,7 @@
 namespace Chaos
 {
 template<class T, int d>
-class CHAOS_API TSmokeEvolution
+class TSmokeEvolution
 {
   public:
 	TSmokeEvolution(const TUniformGrid<T, d>& Grid, TKinematicGeometryParticles<T, d>&& GeometryParticles, TKinematicGeometryParticles<T, d>&& SourceParticles)
diff --git a/Engine/Source/Runtime/Experimental/Chaos/Public/Chaos/SpatialAccelerationCollection.h b/Engine/Source/Runtime/Experimental/Chaos/Public/Chaos/SpatialAccelerationCollection.h
--- a/Engine/Source/Runtime/Experimental/Chaos/Public/Chaos/SpatialAccelerationCollection.h
+++ b/Engine/Source/Runtime/Experimental/Chaos/Public/Chaos/SpatialAccelerationCollection.h
@@ -28,7 +28,7 @@ void FreeObjHelper(T*& RawPtr)
 }
 
 template <typename TPayloadType, typename T, int d>
-struct CHAOS_API TSpatialAccelerationBucketEntry
+struct TSpatialAccelerationBucketEntry
 {
 	TUniquePtr<ISpatialAcceleration<TPayloadType, T, d>> Acceleration;
 	uint16 TypeInnerIdx; // Index in bucket
@@ -78,7 +78,7 @@ TSpatialAccelerationBucketEntry<TPayloadType, T, d> CopyFromHelper(const TSpatia
 
 
 template <typename TObj>
-struct CHAOS_API TSpatialCollectionBucket
+struct TSpatialCollectionBucket
 {
 	TSpatialCollectionBucket() = default;
 	TSpatialCollectionBucket(const TSpatialCollectionBucket<TObj>& Other) = delete;
@@ -374,7 +374,7 @@ typename std::enable_if_t<!std::is_same_v<typename SpatialAccelerationCollection
 }
 
 template <typename ... TSpatialAccelerationTypes>
-class CHAOS_API TSpatialAccelerationCollection : public
+class TSpatialAccelerationCollection : public
 	ISpatialAccelerationCollection<typename std::tuple_element<0, std::tuple< TSpatialAccelerationTypes...>>::type::PayloadType,
 	typename std::tuple_element<0, std::tuple< TSpatialAccelerationTypes...>>::type::TType,
 	std::tuple_element<0, std::tuple< TSpatialAccelerationTypes...>>::type::D>
diff --git a/Engine/Source/Runtime/Experimental/Chaos/Public/Chaos/SpatialHash.h b/Engine/Source/Runtime/Experimental/Chaos/Public/Chaos/SpatialHash.h
--- a/Engine/Source/Runtime/Experimental/Chaos/Public/Chaos/SpatialHash.h
+++ b/Engine/Source/Runtime/Experimental/Chaos/Public/Chaos/SpatialHash.h
@@ -7,7 +7,7 @@
 namespace Chaos
 {
 template<class T>
-class CHAOS_API TSpatialHash
+class TSpatialHash
 {
  public:
 	TSpatialHash(const TArray<TVec3<T>>& Particles, const T Radius)
diff --git a/Engine/Source/Runtime/Experimental/Chaos/Public/Chaos/StrainModification.h b/Engine/Source/Runtime/Experimental/Chaos/Public/Chaos/StrainModification.h
--- a/Engine/Source/Runtime/Experimental/Chaos/Public/Chaos/StrainModification.h
+++ b/Engine/Source/Runtime/Experimental/Chaos/Public/Chaos/StrainModification.h
@@ -27,7 +27,7 @@ namespace Chaos
 	//
 	// User-facing api for accessing the proxy of a strained cluster. Provides const access to the
 	// proxy and limited read/write access to its internal strain data.
-	class CHAOS_API FStrainedProxyModifier
+	class FStrainedProxyModifier
 	{
 	public:
 		FStrainedProxyModifier(FRigidClustering& InRigidClustering, FGeometryCollectionPhysicsProxy* InProxy)
@@ -44,28 +44,28 @@ namespace Chaos
 		{ }
 
 		// Get the proxy that owns the strained cluster or clusters
-		const FGeometryCollectionPhysicsProxy* GetProxy() const;
+		CHAOS_API const FGeometryCollectionPhysicsProxy* GetProxy() const;
 
 		// Get the physics handle for the strained parent cluster
-		const Chaos::FPBDRigidParticleHandle* GetRootHandle() const;
+		CHAOS_API const Chaos::FPBDRigidParticleHandle* GetRootHandle() const;
 
 		// Get the number of level-1 strainable entities (number of rest-children in the per-particle
 		// strain model, or number of rest-connections in the edge/area model).
-		int32 GetNumRestBreakables() const;
+		CHAOS_API int32 GetNumRestBreakables() const;
 
 		// Get the number of breaking strains amongst the level-1 strainables
 		// If DoubleCount is true, then add N for each strain which is strong enough to
 		// break a connection N times.
-		int32 GetNumBreakingStrains(bool bDoubleCount = true, const uint8 StrainTypes = EStrainTypes::ExternalStrain | EStrainTypes::CollisionStrain) const;
+		CHAOS_API int32 GetNumBreakingStrains(bool bDoubleCount = true, const uint8 StrainTypes = EStrainTypes::ExternalStrain | EStrainTypes::CollisionStrain) const;
 
 		// Clear strains for all strained cluster children
-		void ClearStrains();
+		CHAOS_API void ClearStrains();
 
 	private:
 
-		static Chaos::FPBDRigidClusteredParticleHandle* InitRootHandle(FGeometryCollectionPhysicsProxy* Proxy);
+		static CHAOS_API Chaos::FPBDRigidClusteredParticleHandle* InitRootHandle(FGeometryCollectionPhysicsProxy* Proxy);
 
-		static const TSet<int32>* InitRestChildren(FGeometryCollectionPhysicsProxy* Proxy);
+		static CHAOS_API const TSet<int32>* InitRestChildren(FGeometryCollectionPhysicsProxy* Proxy);
 
 		FRigidClustering& RigidClustering;
 		FGeometryCollectionPhysicsProxy* Proxy;
@@ -77,7 +77,7 @@ namespace Chaos
 	//
 	// Iterator produced by FStrainedProxyRange, for looping over proxies which are
 	// associated with rigid clusters which have been strained.
-	class CHAOS_API FStrainedProxyIterator
+	class FStrainedProxyIterator
 	{
 	public:
 		FStrainedProxyIterator(FRigidClustering& InRigidClustering, TArray<FGeometryCollectionPhysicsProxy*>& InProxies, int32 InIndex)
@@ -92,9 +92,9 @@ namespace Chaos
 			, Index(Other.Index)
 		{ }
 
-		FStrainedProxyModifier operator*();
-		FStrainedProxyIterator& operator++();
-		bool operator==(const FStrainedProxyIterator& Other) const;
+		CHAOS_API FStrainedProxyModifier operator*();
+		CHAOS_API FStrainedProxyIterator& operator++();
+		CHAOS_API bool operator==(const FStrainedProxyIterator& Other) const;
 		bool operator!=(const FStrainedProxyIterator& Other) const
 		{
 			return !operator==(Other);
@@ -112,10 +112,10 @@ namespace Chaos
 	// Constructor produces filtered array of proxies, and begin and end functions produce
 	// iterators which can modify strain related properties of the clusters associated
 	// with each proxy.
-	class CHAOS_API FStrainedProxyRange
+	class FStrainedProxyRange
 	{
 	public:
-		FStrainedProxyRange(Chaos::FRigidClustering& InRigidClustering, bool bRootLevelOnly);
+		CHAOS_API FStrainedProxyRange(Chaos::FRigidClustering& InRigidClustering, bool bRootLevelOnly);
 
 		FStrainedProxyRange(const FStrainedProxyRange& Other)
 			: RigidClustering(Other.RigidClustering)
@@ -140,7 +140,7 @@ namespace Chaos
 	// FStrainModifierAccessor
 	//
 	// Provides access to strained proxies and clusters
-	class CHAOS_API FStrainModifierAccessor
+	class FStrainModifierAccessor
 	{
 	public:
 
@@ -149,7 +149,7 @@ namespace Chaos
 		// Get an iterable range of unique geometry collection proxies which
 		// correspond to all strained clusters. Optionally, only include proxies
 		// for whom the strained parent is still the original root (ie, unbroken).
-		FStrainedProxyRange GetStrainedProxies(bool bRootLevelOnly = true);
+		CHAOS_API FStrainedProxyRange GetStrainedProxies(bool bRootLevelOnly = true);
 
 	private:
 
diff --git a/Engine/Source/Runtime/Experimental/Chaos/Public/Chaos/TaperedCylinder.h b/Engine/Source/Runtime/Experimental/Chaos/Public/Chaos/TaperedCylinder.h
--- a/Engine/Source/Runtime/Experimental/Chaos/Public/Chaos/TaperedCylinder.h
+++ b/Engine/Source/Runtime/Experimental/Chaos/Public/Chaos/TaperedCylinder.h
@@ -9,7 +9,7 @@ namespace Chaos
 {
 	struct FTaperedCylinderSpecializeSamplingHelper;
 
-	class CHAOS_API FTaperedCylinder : public FImplicitObject
+	class FTaperedCylinder : public FImplicitObject
 	{
 	public:
 		FTaperedCylinder()
@@ -66,7 +66,7 @@ namespace Chaos
 		 * \p IncludeEndCaps determines whether or not points are generated on the 
 		 *    end caps of the cylinder.
 		 */
-		TArray<FVec3> ComputeLocalSamplePoints(const int32 NumPoints, const bool IncludeEndCaps = true) const;
+		CHAOS_API TArray<FVec3> ComputeLocalSamplePoints(const int32 NumPoints, const bool IncludeEndCaps = true) const;
 
 		/** 
 		 * Returns sample points centered about the origin. 
@@ -86,7 +86,7 @@ namespace Chaos
 		 * \p IncludeEndCaps determines whether or not points are generated on the 
 		 *    end caps of the cylinder.
 		 */
-		TArray<FVec3> ComputeSamplePoints(const int32 NumPoints, const bool IncludeEndCaps = true) const;
+		CHAOS_API TArray<FVec3> ComputeSamplePoints(const int32 NumPoints, const bool IncludeEndCaps = true) const;
 
 		/** 
 		 * Returns sample points at the current location of the cylinder.
@@ -363,7 +363,7 @@ namespace Chaos
 		FAABB3 MLocalBoundingBox;
 	};
 
-	struct CHAOS_API FTaperedCylinderSpecializeSamplingHelper
+	struct FTaperedCylinderSpecializeSamplingHelper
 	{
 		static FORCEINLINE void ComputeSamplePoints(
 		    TArray<FVec3>& Points, const FTaperedCylinder& Cylinder,
diff --git a/Engine/Source/Runtime/Experimental/Chaos/Public/Chaos/TrackedGeometryManager.h b/Engine/Source/Runtime/Experimental/Chaos/Public/Chaos/TrackedGeometryManager.h
--- a/Engine/Source/Runtime/Experimental/Chaos/Public/Chaos/TrackedGeometryManager.h
+++ b/Engine/Source/Runtime/Experimental/Chaos/Public/Chaos/TrackedGeometryManager.h
@@ -9,12 +9,12 @@ class FOutputDevice;
 
 namespace Chaos
 {
-	class CHAOS_API FTrackedGeometryManager
+	class FTrackedGeometryManager
 	{
 	public:
-		static FTrackedGeometryManager& Get();
+		static CHAOS_API FTrackedGeometryManager& Get();
 		
-		void DumpMemoryUsage(FOutputDevice* Ar) const;
+		CHAOS_API void DumpMemoryUsage(FOutputDevice* Ar) const;
 		
 	private:
 		TMap<TSerializablePtr<FImplicitObject>, FString> SharedGeometry;
@@ -23,8 +23,8 @@ namespace Chaos
 		friend FImplicitObject;
 		
 		//These are private because of various threading considerations. ImplicitObject does the cleanup because it needs extra information
-		void AddGeometry(TSerializablePtr<FImplicitObject> Geometry, const FString& DebugInfo);
+		CHAOS_API void AddGeometry(TSerializablePtr<FImplicitObject> Geometry, const FString& DebugInfo);
 		
-		void RemoveGeometry(const FImplicitObject* Geometry);
+		CHAOS_API void RemoveGeometry(const FImplicitObject* Geometry);
 	};
 }
diff --git a/Engine/Source/Runtime/Experimental/Chaos/Public/Chaos/TriangleCollisionPoint.h b/Engine/Source/Runtime/Experimental/Chaos/Public/Chaos/TriangleCollisionPoint.h
--- a/Engine/Source/Runtime/Experimental/Chaos/Public/Chaos/TriangleCollisionPoint.h
+++ b/Engine/Source/Runtime/Experimental/Chaos/Public/Chaos/TriangleCollisionPoint.h
@@ -13,7 +13,7 @@ namespace Chaos
 	 * @brief Data returned by TriangleMesh point-triangle and edge-edge queries.
 	 */
 	template<typename T>
-	struct CHAOS_API TTriangleCollisionPoint
+	struct TTriangleCollisionPoint
 	{
 		// Type of contact (determines meaning of some properties below such as Indices and BarycentricCoordinates)
 		enum struct EContactType : uint8
diff --git a/Engine/Source/Runtime/Experimental/Chaos/Public/Chaos/TriangleMesh.h b/Engine/Source/Runtime/Experimental/Chaos/Public/Chaos/TriangleMesh.h
--- a/Engine/Source/Runtime/Experimental/Chaos/Public/Chaos/TriangleMesh.h
+++ b/Engine/Source/Runtime/Experimental/Chaos/Public/Chaos/TriangleMesh.h
@@ -188,7 +188,7 @@ namespace Chaos
 		CHAOS_API const TArray<TVec2<int32>>& GetEdgeToFaces() const;
 
 		UE_DEPRECATED(5.1, "Non-const access to GetSegmentMesh will be removed. Use const version instead.")
-		CHAOS_API FSegmentMesh& GetSegmentMesh() { return const_cast<FSegmentMesh&>(const_cast<const FTriangleMesh*>(this)->GetSegmentMesh()); }
+		FSegmentMesh& GetSegmentMesh() { return const_cast<FSegmentMesh&>(const_cast<const FTriangleMesh*>(this)->GetSegmentMesh()); }
 
 		/**
 		 * @ret Curvature between adjacent faces, specified on edges in radians.
diff --git a/Engine/Source/Runtime/Experimental/Chaos/Public/Chaos/TriangleMeshImplicitObject.h b/Engine/Source/Runtime/Experimental/Chaos/Public/Chaos/TriangleMeshImplicitObject.h
--- a/Engine/Source/Runtime/Experimental/Chaos/Public/Chaos/TriangleMeshImplicitObject.h
+++ b/Engine/Source/Runtime/Experimental/Chaos/Public/Chaos/TriangleMeshImplicitObject.h
@@ -31,7 +31,7 @@ namespace Chaos
 	class FTriangleMeshImplicitObject;
 	FArchive& operator<<(FArchive& Ar, FAABBVectorized& Bounds);
 
-	struct CHAOS_API FTrimeshBVH
+	struct FTrimeshBVH
 	{
 		using FAABBType = TAABB<FRealSingle, 3>;
 
@@ -47,7 +47,7 @@ namespace Chaos
 			Keep,
 		};
 		
-		struct alignas(16) FChildData
+		CHAOS_API struct alignas(16) FChildData
 		{
 			FChildData()
 				: ChildOrFaceIndex(INDEX_NONE)
@@ -205,7 +205,7 @@ namespace Chaos
 		bool FindAllIntersectionsNoMTD(const Private::FOBBVectorized& Intersection, const TRigidTransform<FReal, 3>& Transform, const QueryGeomType& QueryGeom, FReal Thickness, const FVec3& TriMeshScale, const FTriangleMeshImplicitObject* TriMesh) const;
 		template <typename QueryGeomType>
 		bool FindAllIntersectionsNoMTD(const FAABB3& Intersection, const TRigidTransform<FReal, 3>& Transform, const QueryGeomType& QueryGeom, FReal Thickness, const FVec3& TriMeshScale, const FTriangleMeshImplicitObject* TriMesh) const;
-		TArray<int32> FindAllIntersections(const FAABB3& Intersection) const;
+		CHAOS_API TArray<int32> FindAllIntersections(const FAABB3& Intersection) const;
 
 		template <typename BoundsFilterType, typename FaceVisitorType>
 		FORCEINLINE_DEBUGGABLE EVisitorResult VisitFaces(int32 StartIndex, int32 IndexCount, BoundsFilterType& BoundsFilter, FaceVisitorType& FaceVisitor) const
@@ -332,7 +332,7 @@ namespace Chaos
 		return Ar;
 	}
 	
-	class CHAOS_API FTrimeshIndexBuffer
+	class FTrimeshIndexBuffer
 	{
 	public:
 		using LargeIdxType = int32;
@@ -479,7 +479,7 @@ namespace Chaos
 	}
 
 
-	class CHAOS_API FTriangleMeshImplicitObject final : public FImplicitObject
+	class FTriangleMeshImplicitObject final : public FImplicitObject
 	{
 	public:
 		using FImplicitObject::GetTypeName;
@@ -523,69 +523,69 @@ namespace Chaos
 
 		FTriangleMeshImplicitObject(const FTriangleMeshImplicitObject& Other) = delete;
 		FTriangleMeshImplicitObject(FTriangleMeshImplicitObject&& Other) = delete;
-		virtual ~FTriangleMeshImplicitObject();
+		CHAOS_API virtual ~FTriangleMeshImplicitObject();
 
 		FReal GetRadius() const
 		{
 			return 0.0f;
 		}
 
-		virtual FReal PhiWithNormal(const FVec3& x, FVec3& Normal) const;
+		CHAOS_API virtual FReal PhiWithNormal(const FVec3& x, FVec3& Normal) const;
 
-		virtual bool Raycast(const FVec3& StartPoint, const FVec3& Dir, const FReal Length, const FReal Thickness, FReal& OutTime, FVec3& OutPosition, FVec3& OutNormal, int32& OutFaceIndex) const override;
-		virtual bool Overlap(const FVec3& Point, const FReal Thickness) const override;
+		CHAOS_API virtual bool Raycast(const FVec3& StartPoint, const FVec3& Dir, const FReal Length, const FReal Thickness, FReal& OutTime, FVec3& OutPosition, FVec3& OutNormal, int32& OutFaceIndex) const override;
+		CHAOS_API virtual bool Overlap(const FVec3& Point, const FReal Thickness) const override;
 
-		bool OverlapGeom(const TSphere<FReal, 3>& QueryGeom, const FRigidTransform3& QueryTM, const FReal Thickness, FMTDInfo* OutMTD = nullptr) const;
-		bool OverlapGeom(const TBox<FReal, 3>& QueryGeom, const FRigidTransform3& QueryTM, const FReal Thickness, FMTDInfo* OutMTD = nullptr) const;
-		bool OverlapGeom(const FCapsule& QueryGeom, const FRigidTransform3& QueryTM, const FReal Thickness, FMTDInfo* OutMTD = nullptr) const;
-		bool OverlapGeom(const FConvex& QueryGeom, const FRigidTransform3& QueryTM, const FReal Thickness, FMTDInfo* OutMTD = nullptr) const;
+		CHAOS_API bool OverlapGeom(const TSphere<FReal, 3>& QueryGeom, const FRigidTransform3& QueryTM, const FReal Thickness, FMTDInfo* OutMTD = nullptr) const;
+		CHAOS_API bool OverlapGeom(const TBox<FReal, 3>& QueryGeom, const FRigidTransform3& QueryTM, const FReal Thickness, FMTDInfo* OutMTD = nullptr) const;
+		CHAOS_API bool OverlapGeom(const FCapsule& QueryGeom, const FRigidTransform3& QueryTM, const FReal Thickness, FMTDInfo* OutMTD = nullptr) const;
+		CHAOS_API bool OverlapGeom(const FConvex& QueryGeom, const FRigidTransform3& QueryTM, const FReal Thickness, FMTDInfo* OutMTD = nullptr) const;
 
-		bool OverlapGeom(const TImplicitObjectScaled<TSphere<FReal, 3>>& QueryGeom, const FRigidTransform3& QueryTM, const FReal Thickness, FMTDInfo* OutMTD = nullptr, FVec3 TriMeshScale = FVec3(1.0f)) const;
-		bool OverlapGeom(const TImplicitObjectScaled<TBox<FReal, 3>>& QueryGeom, const FRigidTransform3& QueryTM, const FReal Thickness, FMTDInfo* OutMTD = nullptr, FVec3 TriMeshScale = FVec3(1.0f)) const;
-		bool OverlapGeom(const TImplicitObjectScaled<FCapsule>& QueryGeom, const FRigidTransform3& QueryTM, const FReal Thickness, FMTDInfo* OutMTD = nullptr, FVec3 TriMeshScale = FVec3(1.0f)) const;
-		bool OverlapGeom(const TImplicitObjectScaled<FConvex>& QueryGeom, const FRigidTransform3& QueryTM, const FReal Thickness, FMTDInfo* OutMTD = nullptr, FVec3 TriMeshScale = FVec3(1.0f)) const;
+		CHAOS_API bool OverlapGeom(const TImplicitObjectScaled<TSphere<FReal, 3>>& QueryGeom, const FRigidTransform3& QueryTM, const FReal Thickness, FMTDInfo* OutMTD = nullptr, FVec3 TriMeshScale = FVec3(1.0f)) const;
+		CHAOS_API bool OverlapGeom(const TImplicitObjectScaled<TBox<FReal, 3>>& QueryGeom, const FRigidTransform3& QueryTM, const FReal Thickness, FMTDInfo* OutMTD = nullptr, FVec3 TriMeshScale = FVec3(1.0f)) const;
+		CHAOS_API bool OverlapGeom(const TImplicitObjectScaled<FCapsule>& QueryGeom, const FRigidTransform3& QueryTM, const FReal Thickness, FMTDInfo* OutMTD = nullptr, FVec3 TriMeshScale = FVec3(1.0f)) const;
+		CHAOS_API bool OverlapGeom(const TImplicitObjectScaled<FConvex>& QueryGeom, const FRigidTransform3& QueryTM, const FReal Thickness, FMTDInfo* OutMTD = nullptr, FVec3 TriMeshScale = FVec3(1.0f)) const;
 
-		bool SweepGeom(const TSphere<FReal, 3>& QueryGeom, const FRigidTransform3& StartTM, const FVec3& Dir, const FReal Length, FReal& OutTime, FVec3& OutPosition, FVec3& OutNormal, int32& OutFaceIndex, FVec3& OutFaceNormal, const FReal Thickness = 0, const bool bComputeMTD = false, FVec3 TriMeshScale = FVec3(1.0f)) const;
-		bool SweepGeom(const TBox<FReal, 3>& QueryGeom, const FRigidTransform3& StartTM, const FVec3& Dir, const FReal Length, FReal& OutTime, FVec3& OutPosition, FVec3& OutNormal, int32& OutFaceIndex, FVec3& OutFaceNormal, const FReal Thickness = 0, const bool bComputeMTD = false, FVec3 TriMeshScale = FVec3(1.0f)) const;
-		bool SweepGeom(const FCapsule& QueryGeom, const FRigidTransform3& StartTM, const FVec3& Dir, const FReal Length, FReal& OutTime, FVec3& OutPosition, FVec3& OutNormal, int32& OutFaceIndex, FVec3& OutFaceNormal, const FReal Thickness = 0, const bool bComputeMTD = false, FVec3 TriMeshScale = FVec3(1.0f)) const;
-		bool SweepGeom(const FConvex& QueryGeom, const FRigidTransform3& StartTM, const FVec3& Dir, const FReal Length, FReal& OutTime, FVec3& OutPosition, FVec3& OutNormal, int32& OutFaceIndex, FVec3& OutFaceNormal, const FReal Thickness = 0, const bool bComputeMTD = false, FVec3 TriMeshScale = FVec3(1.0f)) const;
+		CHAOS_API bool SweepGeom(const TSphere<FReal, 3>& QueryGeom, const FRigidTransform3& StartTM, const FVec3& Dir, const FReal Length, FReal& OutTime, FVec3& OutPosition, FVec3& OutNormal, int32& OutFaceIndex, FVec3& OutFaceNormal, const FReal Thickness = 0, const bool bComputeMTD = false, FVec3 TriMeshScale = FVec3(1.0f)) const;
+		CHAOS_API bool SweepGeom(const TBox<FReal, 3>& QueryGeom, const FRigidTransform3& StartTM, const FVec3& Dir, const FReal Length, FReal& OutTime, FVec3& OutPosition, FVec3& OutNormal, int32& OutFaceIndex, FVec3& OutFaceNormal, const FReal Thickness = 0, const bool bComputeMTD = false, FVec3 TriMeshScale = FVec3(1.0f)) const;
+		CHAOS_API bool SweepGeom(const FCapsule& QueryGeom, const FRigidTransform3& StartTM, const FVec3& Dir, const FReal Length, FReal& OutTime, FVec3& OutPosition, FVec3& OutNormal, int32& OutFaceIndex, FVec3& OutFaceNormal, const FReal Thickness = 0, const bool bComputeMTD = false, FVec3 TriMeshScale = FVec3(1.0f)) const;
+		CHAOS_API bool SweepGeom(const FConvex& QueryGeom, const FRigidTransform3& StartTM, const FVec3& Dir, const FReal Length, FReal& OutTime, FVec3& OutPosition, FVec3& OutNormal, int32& OutFaceIndex, FVec3& OutFaceNormal, const FReal Thickness = 0, const bool bComputeMTD = false, FVec3 TriMeshScale = FVec3(1.0f)) const;
 
-		bool SweepGeom(const TImplicitObjectScaled<TSphere<FReal, 3>>& QueryGeom, const FRigidTransform3& StartTM, const FVec3& Dir, const FReal Length, FReal& OutTime, FVec3& OutPosition, FVec3& OutNormal, int32& OutFaceIndex, FVec3& OutFaceNormal, const FReal Thickness = 0, const bool bComputeMTD = false, FVec3 TriMeshScale = FVec3(1.0f)) const;
-		bool SweepGeom(const TImplicitObjectScaled<TBox<FReal, 3>>& QueryGeom, const FRigidTransform3& StartTM, const FVec3& Dir, const FReal Length, FReal& OutTime, FVec3& OutPosition, FVec3& OutNormal, int32& OutFaceIndex, FVec3& OutFaceNormal, const FReal Thickness = 0, const bool bComputeMTD = false, FVec3 TriMeshScale = FVec3(1.0f)) const;
-		bool SweepGeom(const TImplicitObjectScaled<FCapsule>& QueryGeom, const FRigidTransform3& StartTM, const FVec3& Dir, const FReal Length, FReal& OutTime, FVec3& OutPosition, FVec3& OutNormal, int32& OutFaceIndex, FVec3& OutFaceNormal, const FReal Thickness = 0, const bool bComputeMTD = false, FVec3 TriMeshScale = FVec3(1.0f)) const;
-		bool SweepGeom(const TImplicitObjectScaled<FConvex>& QueryGeom, const FRigidTransform3& StartTM, const FVec3& Dir, const FReal Length, FReal& OutTime, FVec3& OutPosition, FVec3& OutNormal, int32& OutFaceIndex, FVec3& OutFaceNormal, const FReal Thickness = 0, const bool bComputeMTD = false, FVec3 TriMeshScale = FVec3(1.0f)) const;
+		CHAOS_API bool SweepGeom(const TImplicitObjectScaled<TSphere<FReal, 3>>& QueryGeom, const FRigidTransform3& StartTM, const FVec3& Dir, const FReal Length, FReal& OutTime, FVec3& OutPosition, FVec3& OutNormal, int32& OutFaceIndex, FVec3& OutFaceNormal, const FReal Thickness = 0, const bool bComputeMTD = false, FVec3 TriMeshScale = FVec3(1.0f)) const;
+		CHAOS_API bool SweepGeom(const TImplicitObjectScaled<TBox<FReal, 3>>& QueryGeom, const FRigidTransform3& StartTM, const FVec3& Dir, const FReal Length, FReal& OutTime, FVec3& OutPosition, FVec3& OutNormal, int32& OutFaceIndex, FVec3& OutFaceNormal, const FReal Thickness = 0, const bool bComputeMTD = false, FVec3 TriMeshScale = FVec3(1.0f)) const;
+		CHAOS_API bool SweepGeom(const TImplicitObjectScaled<FCapsule>& QueryGeom, const FRigidTransform3& StartTM, const FVec3& Dir, const FReal Length, FReal& OutTime, FVec3& OutPosition, FVec3& OutNormal, int32& OutFaceIndex, FVec3& OutFaceNormal, const FReal Thickness = 0, const bool bComputeMTD = false, FVec3 TriMeshScale = FVec3(1.0f)) const;
+		CHAOS_API bool SweepGeom(const TImplicitObjectScaled<FConvex>& QueryGeom, const FRigidTransform3& StartTM, const FVec3& Dir, const FReal Length, FReal& OutTime, FVec3& OutPosition, FVec3& OutNormal, int32& OutFaceIndex, FVec3& OutFaceNormal, const FReal Thickness = 0, const bool bComputeMTD = false, FVec3 TriMeshScale = FVec3(1.0f)) const;
 
 		// Sweep used for CCD. Ignores triangles we penetrate by less than IgnorePenetration, and calculate the TOI for a depth of TargetPenetration. If both are zero, this is equivalent to SweepGeom
-		bool SweepGeomCCD(const TSphere<FReal, 3>& QueryGeom, const FRigidTransform3& StartTM, const FVec3& Dir, const FReal Length, const FReal IgnorePenetration, const FReal TargetPenetration, FReal& OutTOI, FReal& OutPhi, FVec3& OutPosition, FVec3& OutNormal, int32& OutFaceIndex, FVec3& OutFaceNormal, const FVec3& TriMeshScale = FVec3(1.0f)) const;
-		bool SweepGeomCCD(const TBox<FReal, 3>& QueryGeom, const FRigidTransform3& StartTM, const FVec3& Dir, const FReal Length, const FReal IgnorePenetration, const FReal TargetPenetration, FReal& OutTOI, FReal& OutPhi, FVec3& OutPosition, FVec3& OutNormal, int32& OutFaceIndex, FVec3& OutFaceNormal, const FVec3& TriMeshScale = FVec3(1.0f)) const;
-		bool SweepGeomCCD(const FCapsule& QueryGeom, const FRigidTransform3& StartTM, const FVec3& Dir, const FReal Length, const FReal IgnorePenetration, const FReal TargetPenetration, FReal& OutTOI, FReal& OutPhi, FVec3& OutPosition, FVec3& OutNormal, int32& OutFaceIndex, FVec3& OutFaceNormal, const FVec3& TriMeshScale = FVec3(1.0f)) const;
-		bool SweepGeomCCD(const FConvex& QueryGeom, const FRigidTransform3& StartTM, const FVec3& Dir, const FReal Length, const FReal IgnorePenetration, const FReal TargetPenetration, FReal& OutTOI, FReal& OutPhi, FVec3& OutPosition, FVec3& OutNormal, int32& OutFaceIndex, FVec3& OutFaceNormal, const FVec3& TriMeshScale = FVec3(1.0f)) const;
-		bool SweepGeomCCD(const TImplicitObjectScaled<TSphere<FReal, 3>>& QueryGeom, const FRigidTransform3& StartTM, const FVec3& Dir, const FReal Length, const FReal IgnorePenetration, const FReal TargetPenetration, FReal& OutTOI, FReal& OutPhi, FVec3& OutPosition, FVec3& OutNormal, int32& OutFaceIndex, FVec3& OutFaceNormal, const FVec3& TriMeshScale = FVec3(1.0f)) const;
-		bool SweepGeomCCD(const TImplicitObjectScaled<TBox<FReal, 3>>& QueryGeom, const FRigidTransform3& StartTM, const FVec3& Dir, const FReal Length, const FReal IgnorePenetration, const FReal TargetPenetration, FReal& OutTOI, FReal& OutPhi, FVec3& OutPosition, FVec3& OutNormal, int32& OutFaceIndex, FVec3& OutFaceNormal, const FVec3& TriMeshScale = FVec3(1.0f)) const;
-		bool SweepGeomCCD(const TImplicitObjectScaled<FCapsule>& QueryGeom, const FRigidTransform3& StartTM, const FVec3& Dir, const FReal Length, const FReal IgnorePenetration, const FReal TargetPenetration, FReal& OutTOI, FReal& OutPhi, FVec3& OutPosition, FVec3& OutNormal, int32& OutFaceIndex, FVec3& OutFaceNormal, const FVec3& TriMeshScale = FVec3(1.0f)) const;
-		bool SweepGeomCCD(const TImplicitObjectScaled<FConvex>& QueryGeom, const FRigidTransform3& StartTM, const FVec3& Dir, const FReal Length, const FReal IgnorePenetration, const FReal TargetPenetration, FReal& OutTOI, FReal& OutPhi, FVec3& OutPosition, FVec3& OutNormal, int32& OutFaceIndex, FVec3& OutFaceNormal, const FVec3& TriMeshScale = FVec3(1.0f)) const;
-
-		bool GJKContactPoint(const TSphere<FReal, 3>& QueryGeom, const FRigidTransform3& QueryTM, const FReal Thickness, FVec3& Location, FVec3& Normal, FReal& Penetration, int32& FaceIndex) const;
-		bool GJKContactPoint(const TBox<FReal, 3>& QueryGeom, const FRigidTransform3& QueryTM, const FReal Thickness, FVec3& Location, FVec3& Normal, FReal& Penetration, int32& FaceIndex) const;
-		bool GJKContactPoint(const FCapsule& QueryGeom, const FRigidTransform3& QueryTM, const FReal Thickness, FVec3& Location, FVec3& Normal, FReal& Penetration, int32& FaceIndex) const;
-		bool GJKContactPoint(const FConvex& QueryGeom, const FRigidTransform3& QueryTM, const FReal Thickness, FVec3& Location, FVec3& Normal, FReal& Penetration, int32& FaceIndex) const;
-
-		bool GJKContactPoint(const TImplicitObjectScaled < TSphere<FReal, 3> >& QueryGeom, const FRigidTransform3& QueryTM, const FReal Thickness, FVec3& Location, FVec3& Normal, FReal& Penetration, int32& FaceIndex, FVec3 TriMeshScale = FVec3(1.0f)) const;
-		bool GJKContactPoint(const TImplicitObjectScaled < TBox<FReal, 3> >& QueryGeom, const FRigidTransform3& QueryTM, const FReal Thickness, FVec3& Location, FVec3& Normal, FReal& Penetration, int32& FaceIndex, FVec3 TriMeshScale = FVec3(1.0f)) const;
-		bool GJKContactPoint(const TImplicitObjectScaled < FCapsule >& QueryGeom, const FRigidTransform3& QueryTM, const FReal Thickness, FVec3& Location, FVec3& Normal, FReal& Penetration, int32& FaceIndex, FVec3 TriMeshScale = FVec3(1.0f)) const;
-		bool GJKContactPoint(const TImplicitObjectScaled < FConvex >& QueryGeom, const FRigidTransform3& QueryTM, const FReal Thickness, FVec3& Location, FVec3& Normal, FReal& Penetration, int32& FaceIndex, FVec3 TriMeshScale = FVec3(1.0f)) const;
+		CHAOS_API bool SweepGeomCCD(const TSphere<FReal, 3>& QueryGeom, const FRigidTransform3& StartTM, const FVec3& Dir, const FReal Length, const FReal IgnorePenetration, const FReal TargetPenetration, FReal& OutTOI, FReal& OutPhi, FVec3& OutPosition, FVec3& OutNormal, int32& OutFaceIndex, FVec3& OutFaceNormal, const FVec3& TriMeshScale = FVec3(1.0f)) const;
+		CHAOS_API bool SweepGeomCCD(const TBox<FReal, 3>& QueryGeom, const FRigidTransform3& StartTM, const FVec3& Dir, const FReal Length, const FReal IgnorePenetration, const FReal TargetPenetration, FReal& OutTOI, FReal& OutPhi, FVec3& OutPosition, FVec3& OutNormal, int32& OutFaceIndex, FVec3& OutFaceNormal, const FVec3& TriMeshScale = FVec3(1.0f)) const;
+		CHAOS_API bool SweepGeomCCD(const FCapsule& QueryGeom, const FRigidTransform3& StartTM, const FVec3& Dir, const FReal Length, const FReal IgnorePenetration, const FReal TargetPenetration, FReal& OutTOI, FReal& OutPhi, FVec3& OutPosition, FVec3& OutNormal, int32& OutFaceIndex, FVec3& OutFaceNormal, const FVec3& TriMeshScale = FVec3(1.0f)) const;
+		CHAOS_API bool SweepGeomCCD(const FConvex& QueryGeom, const FRigidTransform3& StartTM, const FVec3& Dir, const FReal Length, const FReal IgnorePenetration, const FReal TargetPenetration, FReal& OutTOI, FReal& OutPhi, FVec3& OutPosition, FVec3& OutNormal, int32& OutFaceIndex, FVec3& OutFaceNormal, const FVec3& TriMeshScale = FVec3(1.0f)) const;
+		CHAOS_API bool SweepGeomCCD(const TImplicitObjectScaled<TSphere<FReal, 3>>& QueryGeom, const FRigidTransform3& StartTM, const FVec3& Dir, const FReal Length, const FReal IgnorePenetration, const FReal TargetPenetration, FReal& OutTOI, FReal& OutPhi, FVec3& OutPosition, FVec3& OutNormal, int32& OutFaceIndex, FVec3& OutFaceNormal, const FVec3& TriMeshScale = FVec3(1.0f)) const;
+		CHAOS_API bool SweepGeomCCD(const TImplicitObjectScaled<TBox<FReal, 3>>& QueryGeom, const FRigidTransform3& StartTM, const FVec3& Dir, const FReal Length, const FReal IgnorePenetration, const FReal TargetPenetration, FReal& OutTOI, FReal& OutPhi, FVec3& OutPosition, FVec3& OutNormal, int32& OutFaceIndex, FVec3& OutFaceNormal, const FVec3& TriMeshScale = FVec3(1.0f)) const;
+		CHAOS_API bool SweepGeomCCD(const TImplicitObjectScaled<FCapsule>& QueryGeom, const FRigidTransform3& StartTM, const FVec3& Dir, const FReal Length, const FReal IgnorePenetration, const FReal TargetPenetration, FReal& OutTOI, FReal& OutPhi, FVec3& OutPosition, FVec3& OutNormal, int32& OutFaceIndex, FVec3& OutFaceNormal, const FVec3& TriMeshScale = FVec3(1.0f)) const;
+		CHAOS_API bool SweepGeomCCD(const TImplicitObjectScaled<FConvex>& QueryGeom, const FRigidTransform3& StartTM, const FVec3& Dir, const FReal Length, const FReal IgnorePenetration, const FReal TargetPenetration, FReal& OutTOI, FReal& OutPhi, FVec3& OutPosition, FVec3& OutNormal, int32& OutFaceIndex, FVec3& OutFaceNormal, const FVec3& TriMeshScale = FVec3(1.0f)) const;
+
+		CHAOS_API bool GJKContactPoint(const TSphere<FReal, 3>& QueryGeom, const FRigidTransform3& QueryTM, const FReal Thickness, FVec3& Location, FVec3& Normal, FReal& Penetration, int32& FaceIndex) const;
+		CHAOS_API bool GJKContactPoint(const TBox<FReal, 3>& QueryGeom, const FRigidTransform3& QueryTM, const FReal Thickness, FVec3& Location, FVec3& Normal, FReal& Penetration, int32& FaceIndex) const;
+		CHAOS_API bool GJKContactPoint(const FCapsule& QueryGeom, const FRigidTransform3& QueryTM, const FReal Thickness, FVec3& Location, FVec3& Normal, FReal& Penetration, int32& FaceIndex) const;
+		CHAOS_API bool GJKContactPoint(const FConvex& QueryGeom, const FRigidTransform3& QueryTM, const FReal Thickness, FVec3& Location, FVec3& Normal, FReal& Penetration, int32& FaceIndex) const;
+
+		CHAOS_API bool GJKContactPoint(const TImplicitObjectScaled < TSphere<FReal, 3> >& QueryGeom, const FRigidTransform3& QueryTM, const FReal Thickness, FVec3& Location, FVec3& Normal, FReal& Penetration, int32& FaceIndex, FVec3 TriMeshScale = FVec3(1.0f)) const;
+		CHAOS_API bool GJKContactPoint(const TImplicitObjectScaled < TBox<FReal, 3> >& QueryGeom, const FRigidTransform3& QueryTM, const FReal Thickness, FVec3& Location, FVec3& Normal, FReal& Penetration, int32& FaceIndex, FVec3 TriMeshScale = FVec3(1.0f)) const;
+		CHAOS_API bool GJKContactPoint(const TImplicitObjectScaled < FCapsule >& QueryGeom, const FRigidTransform3& QueryTM, const FReal Thickness, FVec3& Location, FVec3& Normal, FReal& Penetration, int32& FaceIndex, FVec3 TriMeshScale = FVec3(1.0f)) const;
+		CHAOS_API bool GJKContactPoint(const TImplicitObjectScaled < FConvex >& QueryGeom, const FRigidTransform3& QueryTM, const FReal Thickness, FVec3& Location, FVec3& Normal, FReal& Penetration, int32& FaceIndex, FVec3 TriMeshScale = FVec3(1.0f)) const;
 
 		// Returns -1 if InternalFaceIndex is not in map, or map is invalid.
-		int32 GetExternalFaceIndexFromInternal(int32 InternalFaceIndex) const;
+		CHAOS_API int32 GetExternalFaceIndexFromInternal(int32 InternalFaceIndex) const;
 
 		// Does Trimesh cull backfaces in raycast.
-		bool GetCullsBackFaceRaycast() const;
-		void SetCullsBackFaceRaycast(const bool bInCullsBackFace);
+		CHAOS_API bool GetCullsBackFaceRaycast() const;
+		CHAOS_API void SetCullsBackFaceRaycast(const bool bInCullsBackFace);
 
 
-		virtual int32 FindMostOpposingFace(const FVec3& Position, const FVec3& UnitDir, int32 HintFaceIndex, FReal SearchDistance) const override;
-		virtual int32 FindMostOpposingFaceScaled(const FVec3& Position, const FVec3& UnitDir, int32 HintFaceIndex, FReal SearchDist, const FVec3& Scale) const override;
-		virtual FVec3 FindGeometryOpposingNormal(const FVec3& DenormDir, int32 FaceIndex, const FVec3& OriginalNormal) const override;
+		CHAOS_API virtual int32 FindMostOpposingFace(const FVec3& Position, const FVec3& UnitDir, int32 HintFaceIndex, FReal SearchDistance) const override;
+		CHAOS_API virtual int32 FindMostOpposingFaceScaled(const FVec3& Position, const FVec3& UnitDir, int32 HintFaceIndex, FReal SearchDist, const FVec3& Scale) const override;
+		CHAOS_API virtual FVec3 FindGeometryOpposingNormal(const FVec3& DenormDir, int32 FaceIndex, const FVec3& OriginalNormal) const override;
 
 		virtual const FAABB3 BoundingBox() const
 		{
@@ -597,7 +597,7 @@ namespace Chaos
 			return ImplicitObjectType::TriangleMesh;
 		}
 
-		TUniquePtr<FTriangleMeshImplicitObject> CopySlow() const;
+		CHAOS_API TUniquePtr<FTriangleMeshImplicitObject> CopySlow() const;
 
 		void SerializeImp(FChaosArchive& Ar)
 		{
@@ -702,23 +702,23 @@ namespace Chaos
 			}
 		}
 
-		virtual void Serialize(FChaosArchive& Ar) override;
+		CHAOS_API virtual void Serialize(FChaosArchive& Ar) override;
 
 		virtual FString ToString() const
 		{
 			return FString::Printf(TEXT("TriangleMesh"));
 		}
 
-		virtual uint32 GetTypeHash() const override;
+		CHAOS_API virtual uint32 GetTypeHash() const override;
 
-		FVec3 GetFaceNormal(const int32 FaceIdx) const;
+		CHAOS_API FVec3 GetFaceNormal(const int32 FaceIdx) const;
 
-		virtual uint16 GetMaterialIndex(uint32 HintIndex) const override;
+		CHAOS_API virtual uint16 GetMaterialIndex(uint32 HintIndex) const override;
 
-		const ParticlesType& Particles() const;
-		const FTrimeshIndexBuffer& Elements() const;
+		CHAOS_API const ParticlesType& Particles() const;
+		CHAOS_API const FTrimeshIndexBuffer& Elements() const;
 
-		void UpdateVertices(const TArray<FVector>& Positions);
+		CHAOS_API void UpdateVertices(const TArray<FVector>& Positions);
 
 		/**
 		 * @brief Generate the triangle at the specified index with the specified transform (including scale)
@@ -778,8 +778,8 @@ namespace Chaos
 
 	private:
 		using BVHType = TAABBTree<int32, TAABBTreeLeafArray<int32, /*bComputeBounds=*/ false, FRealSingle>, /*bMutable=*/false, FRealSingle>;
-		void RebuildFastBVHFromTree(const BVHType& BVH);
-		void RebuildFastBVH();
+		CHAOS_API void RebuildFastBVHFromTree(const BVHType& BVH);
+		CHAOS_API void RebuildFastBVH();
 
 		ParticlesType MParticles;
 		FTrimeshIndexBuffer MElements;
diff --git a/Engine/Source/Runtime/Experimental/Chaos/Public/Chaos/UniformGrid.h b/Engine/Source/Runtime/Experimental/Chaos/Public/Chaos/UniformGrid.h
--- a/Engine/Source/Runtime/Experimental/Chaos/Public/Chaos/UniformGrid.h
+++ b/Engine/Source/Runtime/Experimental/Chaos/Public/Chaos/UniformGrid.h
@@ -25,7 +25,7 @@ template<> struct TGridPrecisionLimit<FRealSingle> { static constexpr FRealSingl
 template<> struct TGridPrecisionLimit<FRealDouble> { static constexpr FRealDouble value = 4.5035e15; };
 
 template<class T, int d>
-class CHAOS_API TUniformGridBase
+class TUniformGridBase
 {
   protected:
 	TUniformGridBase() {}
@@ -263,7 +263,7 @@ class CHAOS_API TUniformGridBase
 };
 
 template<class T, int d>
-class CHAOS_API TUniformGrid : public TUniformGridBase<T, d>
+class TUniformGrid : public TUniformGridBase<T, d>
 {
 	using TUniformGridBase<T, d>::MCells;
 	using TUniformGridBase<T, d>::MMinCorner;
@@ -416,7 +416,7 @@ public:
 
 
 template<class T>
-class CHAOS_API TUniformGrid<T, 3> : public TUniformGridBase<T, 3>
+class TUniformGrid<T, 3> : public TUniformGridBase<T, 3>
 {
 	using TUniformGridBase<T, 3>::MCells;
 	using TUniformGridBase<T, 3>::MMinCorner;
@@ -436,8 +436,8 @@ class CHAOS_API TUniformGrid<T, 3> : public TUniformGridBase<T, 3>
 	    : TUniformGridBase<T, 3>(Stream) {}
 #endif
 	~TUniformGrid() {}
-	TVector<int32, 3> GetIndex(const int32 Index) const;
-	Pair<int32, TVector<int32, 3>> GetFaceIndex(int32 Index) const;
+	CHAOS_API TVector<int32, 3> GetIndex(const int32 Index) const;
+	CHAOS_API Pair<int32, TVector<int32, 3>> GetFaceIndex(int32 Index) const;
 	int32 GetNumFaces() const
 	{
 		return GetNumCells() * 3 + MCells[0] * MCells[1] + MCells[1] * MCells[2] + MCells[0] * MCells[3];
@@ -446,12 +446,12 @@ class CHAOS_API TUniformGrid<T, 3> : public TUniformGridBase<T, 3>
 	{
 		return TUniformGridBase<T, 3>::Location(GetIndex(Index));
 	}
-	TVector<int32, 3> ClampIndex(const TVector<int32, 3>& Index) const;
-	TVector<T, 3> Clamp(const TVector<T, 3>& X) const;
-	TVector<T, 3> ClampMinusHalf(const TVector<T, 3>& X) const;
-	bool IsValid(const TVector<int32, 3>& X) const;
+	CHAOS_API TVector<int32, 3> ClampIndex(const TVector<int32, 3>& Index) const;
+	CHAOS_API TVector<T, 3> Clamp(const TVector<T, 3>& X) const;
+	CHAOS_API TVector<T, 3> ClampMinusHalf(const TVector<T, 3>& X) const;
+	CHAOS_API bool IsValid(const TVector<int32, 3>& X) const;
 
-	TUniformGrid<T, 3> SubGrid(const TVector<int32, 3>& MinCell, const TVector<int32, 3>& MaxCell) const;
+	CHAOS_API TUniformGrid<T, 3> SubGrid(const TVector<int32, 3>& MinCell, const TVector<int32, 3>& MaxCell) const;
 };
 
 
diff --git a/Engine/Source/Runtime/Experimental/Chaos/Public/Chaos/VelocityField.h b/Engine/Source/Runtime/Experimental/Chaos/Public/Chaos/VelocityField.h
--- a/Engine/Source/Runtime/Experimental/Chaos/Public/Chaos/VelocityField.h
+++ b/Engine/Source/Runtime/Experimental/Chaos/Public/Chaos/VelocityField.h
@@ -10,7 +10,7 @@ namespace Chaos::Softs
 {
 
 // Velocity field used solely for aerodynamics effects, use Chaos Fields for other types of fields.
-class CHAOS_API FVelocityAndPressureField final
+class FVelocityAndPressureField final
 {
 public:
 	static constexpr FSolverReal DefaultDragCoefficient = (FSolverReal)0.5;
@@ -48,7 +48,7 @@ public:
 
 	~FVelocityAndPressureField() {}
 
-	void UpdateForces(const FSolverParticles& InParticles, const FSolverReal /*Dt*/);
+	CHAOS_API void UpdateForces(const FSolverParticles& InParticles, const FSolverReal /*Dt*/);
 
 	inline void Apply(FSolverParticles& InParticles, const FSolverReal Dt, const int32 Index) const
 	{
@@ -61,7 +61,7 @@ public:
 		}
 	}
 
-	void SetProperties(
+	CHAOS_API void SetProperties(
 		const FCollectionPropertyConstFacade& PropertyCollection,
 		const TMap<FString, TConstArrayView<FRealSingle>>& Weightmaps,
 		FSolverReal WorldScale,
@@ -74,7 +74,7 @@ public:
 		SetProperties(PropertyCollection, TMap<FString, TConstArrayView<FRealSingle>>(), WorldScale, bEnableAerodynamics);
 	}
 
-	void SetProperties(
+	CHAOS_API void SetProperties(
 		const FSolverVec2& Drag,
 		const FSolverVec2& Lift,
 		const FSolverReal FluidDensity,
@@ -89,13 +89,13 @@ public:
 				LiftBase > (FSolverReal)0. || LiftRange != (FSolverReal)0.));
 	}
 
-	void SetGeometry(
+	CHAOS_API void SetGeometry(
 		const FTriangleMesh* TriangleMesh,
 		const FCollectionPropertyConstFacade& PropertyCollection,
 		const TMap<FString, TConstArrayView<FRealSingle>>& Weightmaps,
 		FSolverReal WorldScale);
 
-	void SetGeometry(
+	CHAOS_API void SetGeometry(
 		const FTriangleMesh* TriangleMesh,
 		const TConstArrayView<FRealSingle>& DragMultipliers,
 		const TConstArrayView<FRealSingle>& LiftMultipliers,
@@ -109,11 +109,11 @@ public:
 private:
 	bool AreAerodynamicsEnabled() const { return QuarterRho > (FSolverReal)0.; }
 
-	void SetGeometry(const FTriangleMesh* TriangleMesh);
+	CHAOS_API void SetGeometry(const FTriangleMesh* TriangleMesh);
 
-	void SetMultipliers(const FCollectionPropertyConstFacade& PropertyCollection,const TMap<FString, TConstArrayView<FRealSingle>>& Weightmaps);
+	CHAOS_API void SetMultipliers(const FCollectionPropertyConstFacade& PropertyCollection,const TMap<FString, TConstArrayView<FRealSingle>>& Weightmaps);
 
-	void SetMultipliers(
+	CHAOS_API void SetMultipliers(
 		const TConstArrayView<FRealSingle>& DragMultipliers,
 		const TConstArrayView<FRealSingle>& LiftMultipliers,
 		const TConstArrayView<FRealSingle>& PressureMultipliers);
