[Chaos]

diff --git a/Engine/Source/Runtime/Experimental/Chaos/Public/Chaos/WeightedLatticeImplicitObject.h b/Engine/Source/Runtime/Experimental/Chaos/Public/Chaos/WeightedLatticeImplicitObject.h
--- a/Engine/Source/Runtime/Experimental/Chaos/Public/Chaos/WeightedLatticeImplicitObject.h
+++ b/Engine/Source/Runtime/Experimental/Chaos/Public/Chaos/WeightedLatticeImplicitObject.h
@@ -52,16 +52,16 @@ inline FArchive& operator<<(FArchive& Ar, FWeightedLatticeInfluenceData& Value)
 /**
  * Embed the contained shape in a deformable lattice.
  */
-class CHAOS_API FWeightedLatticeImplicitObject : public FImplicitObject
+class FWeightedLatticeImplicitObject : public FImplicitObject
 {
 public:
 	using ObjectType = TUniquePtr<FImplicitObject>;
 	using FImplicitObject::GetTypeName;
 
-	FWeightedLatticeImplicitObject(int32 Flags, EImplicitObjectType InType, TUniformGrid<FReal, 3>&& InGrid,
+	CHAOS_API FWeightedLatticeImplicitObject(int32 Flags, EImplicitObjectType InType, TUniformGrid<FReal, 3>&& InGrid,
 		TArrayND<FWeightedLatticeInfluenceData, 3>&& InBoneData, TArray<FName>&& InUsedBones, TArray<FTransform>&& InReferenceRelativeTransforms);
 
-	FWeightedLatticeImplicitObject(FWeightedLatticeImplicitObject&& Other);
+	CHAOS_API FWeightedLatticeImplicitObject(FWeightedLatticeImplicitObject&& Other);
 	
 	virtual ~FWeightedLatticeImplicitObject() override = default;
 
@@ -74,14 +74,14 @@ public:
 	const TArray<int32>& GetSolverBoneIndices() const { return SolverBoneIndices; }
 	void SetSolverBoneIndices(TArray<int32>&& InSolverBoneIndices) { SolverBoneIndices = MoveTemp(InSolverBoneIndices); }
 
-	virtual void Serialize(FChaosArchive& Ar) override;
+	CHAOS_API virtual void Serialize(FChaosArchive& Ar) override;
 
-	FVec3 GetDeformedPoint(const FVec3& UndeformedPoint) const;
-	void UpdateSpatialHierarchy();
-	void FinalizeConstruction();
-	void DeformPoints(const TArray<FTransform>& RelativeTransforms);
+	CHAOS_API FVec3 GetDeformedPoint(const FVec3& UndeformedPoint) const;
+	CHAOS_API void UpdateSpatialHierarchy();
+	CHAOS_API void FinalizeConstruction();
+	CHAOS_API void DeformPoints(const TArray<FTransform>& RelativeTransforms);
 
-	struct CHAOS_API FEmbeddingCoordinate
+	struct FEmbeddingCoordinate
 	{
 		TVec3<int32> CellIndex = TVec3<int32>(INDEX_NONE, INDEX_NONE, INDEX_NONE);
 		int32 LocalTetrahedron = INDEX_NONE;
@@ -95,7 +95,7 @@ public:
 			:CellIndex(InCellIndex), LocalTetrahedron(InLocalTetrahedron), BarycentricCoordinate(InBarycentric)
 		{}
 
-		FEmbeddingCoordinate(const TVec3<int32>& InCellIndex, const FVec3& TrilinearCoordinate);
+		CHAOS_API FEmbeddingCoordinate(const TVec3<int32>& InCellIndex, const FVec3& TrilinearCoordinate);
 
 		bool IsValid() const
 		{
@@ -135,24 +135,24 @@ public:
 			}
 		}
 
-		FMatrix DeformationTransform(const TArrayND<FVec3, 3>& DeformedPoints, const TUniformGrid<FReal, 3>& Grid) const;
+		CHAOS_API FMatrix DeformationTransform(const TArrayND<FVec3, 3>& DeformedPoints, const TUniformGrid<FReal, 3>& Grid) const;
 		inline FVec3 UndeformedPosition(const TUniformGrid<FReal, 3>& InGrid) const;
 		inline FVec3 DeformedPosition(const TArrayND<FVec3, 3>& InDeformedPoints) const;
 		inline int32 GreatestInfluenceBone(const TArrayND<FWeightedLatticeInfluenceData, 3>& InBoneData) const;
 	};
 
-	bool GetEmbeddingCoordinates(const FVec3& DeformedPoint, TArray<FEmbeddingCoordinate>& CoordinatesOut, bool bFindClosest = false) const;
+	CHAOS_API bool GetEmbeddingCoordinates(const FVec3& DeformedPoint, TArray<FEmbeddingCoordinate>& CoordinatesOut, bool bFindClosest = false) const;
 
 protected:
 	FWeightedLatticeImplicitObject(int32 Flags, EImplicitObjectType InType)
 		:FImplicitObject(Flags, InType | ImplicitObjectType::IsWeightedLattice)
 	{}
 
-	FWeightedLatticeImplicitObject(const FWeightedLatticeImplicitObject& Other);
+	CHAOS_API FWeightedLatticeImplicitObject(const FWeightedLatticeImplicitObject& Other);
 
-	void InitializeDeformedPoints();
-	void SetEmptyCells();
-	uint32 GetTypeHashHelper(const uint32 InHash) const;
+	CHAOS_API void InitializeDeformedPoints();
+	CHAOS_API void SetEmptyCells();
+	CHAOS_API uint32 GetTypeHashHelper(const uint32 InHash) const;
 
 
 	// Serialized data. Only non-const because of serialization
@@ -251,12 +251,12 @@ private:
 	friend ::FKSkinnedLevelSetElem; //needed for serialization
 };
 
-class CHAOS_API FWeightedLatticeImplicitObjectBuilder
+class FWeightedLatticeImplicitObjectBuilder
 {
 public:
 
-	void GenerateGrid(const int32 GridResolution, const TAABB<FReal, 3>& ObjectBBox);
-	void AddInfluence(int32 FlatIndex, uint16 BoneIndex, float Weight, bool bIsOuterWeight);
+	CHAOS_API void GenerateGrid(const int32 GridResolution, const TAABB<FReal, 3>& ObjectBBox);
+	CHAOS_API void AddInfluence(int32 FlatIndex, uint16 BoneIndex, float Weight, bool bIsOuterWeight);
 
 	// BoneIndexToBoneName(int32 BoneIndex)->FName return FName of Bone for given BoneIndex (BoneIndex used by AddInfluence)
 	// BoneIndexToReferenceTransform(int32 BoneIndex)->FTransform return transform of Bone (BoneIndex used by AddInfluence). Should return transform for RootBone (parent of this object) when BoneIndex == INDEX_NONE
@@ -271,7 +271,7 @@ public:
 
 private:
 
-	void NormalizeBoneWeights();
+	CHAOS_API void NormalizeBoneWeights();
 
 	template<typename FBoneIndexToBoneName, typename FBoneIndexToReferenceTransform>
 	inline void CalcUsedBonesAndReIndex(const FBoneIndexToBoneName& BoneIndexToBoneName, const FBoneIndexToReferenceTransform& BoneIndexToReferenceTransform);
diff --git a/Engine/Source/Runtime/Experimental/Chaos/Public/Chaos/XPBDAnisotropicBendingConstraints.h b/Engine/Source/Runtime/Experimental/Chaos/Public/Chaos/XPBDAnisotropicBendingConstraints.h
--- a/Engine/Source/Runtime/Experimental/Chaos/Public/Chaos/XPBDAnisotropicBendingConstraints.h
+++ b/Engine/Source/Runtime/Experimental/Chaos/Public/Chaos/XPBDAnisotropicBendingConstraints.h
@@ -13,7 +13,7 @@ class FTriangleMesh;
 namespace Chaos::Softs
 {
 
-class CHAOS_API FXPBDAnisotropicBendingConstraints final : public FPBDBendingConstraintsBase
+class FXPBDAnisotropicBendingConstraints final : public FPBDBendingConstraintsBase
 {
 	typedef FPBDBendingConstraintsBase Base;
 
@@ -29,7 +29,7 @@ public:
 		return IsXPBDAnisoBendingStiffnessWarpEnabled(PropertyCollection, false);
 	}
 
-	FXPBDAnisotropicBendingConstraints(const FSolverParticles& InParticles,
+	CHAOS_API FXPBDAnisotropicBendingConstraints(const FSolverParticles& InParticles,
 		int32 InParticleOffset,
 		int32 InParticleCount,
 		const FTriangleMesh& TriangleMesh,
@@ -38,7 +38,7 @@ public:
 		const FCollectionPropertyConstFacade& PropertyCollection,
 		bool bTrimKinematicConstraints = false);
 
-	FXPBDAnisotropicBendingConstraints(const FSolverParticles& InParticles,
+	CHAOS_API FXPBDAnisotropicBendingConstraints(const FSolverParticles& InParticles,
 		int32 InParticleOffset,
 		int32 InParticleCount,
 		const FTriangleMesh& TriangleMesh,
@@ -69,7 +69,7 @@ public:
 		FPBDBendingConstraintsBase::Init(InParticles);
 	}
 
-	void SetProperties(
+	CHAOS_API void SetProperties(
 		const FCollectionPropertyConstFacade& PropertyCollection,
 		const TMap<FString, TConstArrayView<FRealSingle>>& WeightMaps);
 
@@ -91,17 +91,17 @@ public:
 		DampingRatio.ApplyValues();
 	}
 
-	void Apply(FSolverParticles& Particles, const FSolverReal Dt) const;
+	CHAOS_API void Apply(FSolverParticles& Particles, const FSolverReal Dt) const;
 
 	const TArray<int32>& GetConstraintsPerColorStartIndex() const { return ConstraintsPerColorStartIndex; }
 	const TArray<FSolverVec3>& GetWarpWeftBiasBaseMultipliers() const { return WarpWeftBiasBaseMultipliers; }
 
 private:
-	void InitColor(const FSolverParticles& InParticles);
-	void ApplyHelper(FSolverParticles& Particles, const FSolverReal Dt, const int32 ConstraintIndex, const FSolverVec3& ExpStiffnessValues, 
+	CHAOS_API void InitColor(const FSolverParticles& InParticles);
+	CHAOS_API void ApplyHelper(FSolverParticles& Particles, const FSolverReal Dt, const int32 ConstraintIndex, const FSolverVec3& ExpStiffnessValues, 
 		const FSolverVec3& ExpBucklingStiffnessValues, const FSolverReal DampingRatioValue) const;
 
-	TArray<FSolverVec3> GenerateWarpWeftBiasBaseMultipliers(const TArray<TVec3<FVec2f>>& FaceVertexPatternPositions, const FTriangleMesh& TriangleMesh) const;
+	CHAOS_API TArray<FSolverVec3> GenerateWarpWeftBiasBaseMultipliers(const TArray<TVec3<FVec2f>>& FaceVertexPatternPositions, const FTriangleMesh& TriangleMesh) const;
 
 	using Base::Constraints;
 	using Base::ParticleOffset;
diff --git a/Engine/Source/Runtime/Experimental/Chaos/Public/Chaos/XPBDBendingConstraints.h b/Engine/Source/Runtime/Experimental/Chaos/Public/Chaos/XPBDBendingConstraints.h
--- a/Engine/Source/Runtime/Experimental/Chaos/Public/Chaos/XPBDBendingConstraints.h
+++ b/Engine/Source/Runtime/Experimental/Chaos/Public/Chaos/XPBDBendingConstraints.h
@@ -14,7 +14,7 @@ static const FSolverReal XPBDBendMinStiffness = (FSolverReal)1e-4; // Stiffness
 UE_DEPRECATED(5.2, "Use FXPBDBendingConstraints::MaxStiffness instead.")
 static const FSolverReal XPBDBendMaxStiffness = (FSolverReal)1e7;
 
-class CHAOS_API FXPBDBendingConstraints final : public FPBDBendingConstraintsBase
+class FXPBDBendingConstraints final : public FPBDBendingConstraintsBase
 {
 	typedef FPBDBendingConstraintsBase Base;
 
@@ -166,7 +166,7 @@ public:
 		FPBDBendingConstraintsBase::Init(InParticles);
 	}
 
-	void SetProperties(
+	CHAOS_API void SetProperties(
 		const FCollectionPropertyConstFacade& PropertyCollection,
 		const TMap<FString, TConstArrayView<FRealSingle>>& WeightMaps);
 
@@ -192,13 +192,13 @@ public:
 		DampingRatio.ApplyValues();
 	}
 
-	void Apply(FSolverParticles& Particles, const FSolverReal Dt) const;
+	CHAOS_API void Apply(FSolverParticles& Particles, const FSolverReal Dt) const;
 
 	const TArray<int32>& GetConstraintsPerColorStartIndex() const { return ConstraintsPerColorStartIndex; }
 
 private:
-	void InitColor(const FSolverParticles& InParticles);
-	void ApplyHelper(FSolverParticles& Particles, const FSolverReal Dt, const int32 ConstraintIndex, const FSolverReal ExpStiffnessValue, const FSolverReal ExpBucklingValue, const FSolverReal DampingRatioValue) const;
+	CHAOS_API void InitColor(const FSolverParticles& InParticles);
+	CHAOS_API void ApplyHelper(FSolverParticles& Particles, const FSolverReal Dt, const int32 ConstraintIndex, const FSolverReal ExpStiffnessValue, const FSolverReal ExpBucklingValue, const FSolverReal DampingRatioValue) const;
 
 	using Base::Constraints;
 	using Base::ParticleOffset;
diff --git a/Engine/Source/Runtime/Experimental/Chaos/Public/Chaos/XPBDSpringConstraints.h b/Engine/Source/Runtime/Experimental/Chaos/Public/Chaos/XPBDSpringConstraints.h
--- a/Engine/Source/Runtime/Experimental/Chaos/Public/Chaos/XPBDSpringConstraints.h
+++ b/Engine/Source/Runtime/Experimental/Chaos/Public/Chaos/XPBDSpringConstraints.h
@@ -15,7 +15,7 @@ static const FSolverReal XPBDSpringMinStiffness = (FSolverReal)1e-4; // Stiffnes
 UE_DEPRECATED(5.2, "Use FXPBDSpringConstraints::MaxStiffness instead.")
 static const FSolverReal XPBDSpringMaxStiffness = (FSolverReal)1e7;
 
-class CHAOS_API FXPBDSpringConstraints : public FPBDSpringConstraintsBase
+class FXPBDSpringConstraints : public FPBDSpringConstraintsBase
 {
 	typedef FPBDSpringConstraintsBase Base;
 
@@ -99,13 +99,13 @@ public:
 		DampingRatio.ApplyValues();
 	}
 
-	void Apply(FSolverParticles& Particles, const FSolverReal Dt) const;
+	CHAOS_API void Apply(FSolverParticles& Particles, const FSolverReal Dt) const;
 
 	const TArray<int32>& GetConstraintsPerColorStartIndex() const { return ConstraintsPerColorStartIndex; }
 
 private:
-	void InitColor(const FSolverParticles& InParticles);
-	void ApplyHelper(FSolverParticles& Particles, const FSolverReal Dt, const int32 ConstraintIndex, const FSolverReal StiffnessValue, const FSolverReal DampingRatioValue) const;
+	CHAOS_API void InitColor(const FSolverParticles& InParticles);
+	CHAOS_API void ApplyHelper(FSolverParticles& Particles, const FSolverReal Dt, const int32 ConstraintIndex, const FSolverReal StiffnessValue, const FSolverReal DampingRatioValue) const;
 
 	FSolverVec3 GetDelta(const FSolverParticles& Particles, const FSolverReal Dt, const int32 ConstraintIndex, const FSolverReal StiffnessValue, const FSolverReal DampingRatioValue) const
 	{
@@ -162,7 +162,7 @@ private:
 	TArray<int32> ConstraintsPerColorStartIndex; // Constraints are ordered so each batch is contiguous. This is ColorNum + 1 length so it can be used as start and end.
 };
 
-class CHAOS_API FXPBDEdgeSpringConstraints final : public FXPBDSpringConstraints
+class FXPBDEdgeSpringConstraints final : public FXPBDSpringConstraints
 {
 public:
 	static bool IsEnabled(const FCollectionPropertyConstFacade& PropertyCollection)
@@ -214,7 +214,7 @@ public:
 
 	virtual ~FXPBDEdgeSpringConstraints() override = default;
 
-	void SetProperties(
+	CHAOS_API void SetProperties(
 		const FCollectionPropertyConstFacade& PropertyCollection,
 		const TMap<FString, TConstArrayView<FRealSingle>>& WeightMaps);
 
@@ -235,7 +235,7 @@ private:
 	UE_CHAOS_DECLARE_PROPERTYCOLLECTION_NAME(XPBDEdgeSpringDamping, float);
 };
 
-class CHAOS_API FXPBDBendingSpringConstraints : public FXPBDSpringConstraints
+class FXPBDBendingSpringConstraints : public FXPBDSpringConstraints
 {
 public:
 	static bool IsEnabled(const FCollectionPropertyConstFacade& PropertyCollection)
@@ -287,7 +287,7 @@ public:
 
 	virtual ~FXPBDBendingSpringConstraints() override = default;
 
-	void SetProperties(
+	CHAOS_API void SetProperties(
 		const FCollectionPropertyConstFacade& PropertyCollection,
 		const TMap<FString, TConstArrayView<FRealSingle>>& WeightMaps);
 
diff --git a/Engine/Source/Runtime/Experimental/Chaos/Public/Chaos/XPBDStretchBiasElementConstraints.h b/Engine/Source/Runtime/Experimental/Chaos/Public/Chaos/XPBDStretchBiasElementConstraints.h
--- a/Engine/Source/Runtime/Experimental/Chaos/Public/Chaos/XPBDStretchBiasElementConstraints.h
+++ b/Engine/Source/Runtime/Experimental/Chaos/Public/Chaos/XPBDStretchBiasElementConstraints.h
@@ -16,7 +16,7 @@ class FTriangleMesh;
 namespace Chaos::Softs
 {
 
-class CHAOS_API FXPBDStretchBiasElementConstraints
+class FXPBDStretchBiasElementConstraints
 {
 public:
 	// Stiffness is in kg cm / s^2 for stretch, kg cm^2 / s^2 for Bias
@@ -34,7 +34,7 @@ public:
 		return IsXPBDAnisoStretchStiffnessWarpEnabled(PropertyCollection, false);
 	}
 
-	FXPBDStretchBiasElementConstraints(const FSolverParticles& InParticles,
+	CHAOS_API FXPBDStretchBiasElementConstraints(const FSolverParticles& InParticles,
 		int32 ParticleOffset,
 		int32 ParticleCount,
 		const FTriangleMesh& TriangleMesh,
@@ -43,7 +43,7 @@ public:
 		const FCollectionPropertyConstFacade& PropertyCollection,
 		bool bTrimKinematicConstraints = false);
 
-	FXPBDStretchBiasElementConstraints(const FSolverParticles& InParticles,
+	CHAOS_API FXPBDStretchBiasElementConstraints(const FSolverParticles& InParticles,
 		int32 ParticleOffset,
 		int32 ParticleCount,
 		const FTriangleMesh& TriangleMesh,
@@ -71,7 +71,7 @@ public:
 		Lambdas.AddZeroed(Constraints.Num());
 	}
 
-	void SetProperties(
+	CHAOS_API void SetProperties(
 		const FCollectionPropertyConstFacade& PropertyCollection,
 		const TMap<FString, TConstArrayView<FRealSingle>>& WeightMaps);
 
@@ -96,8 +96,8 @@ public:
 		WeftScale.ApplyValues();
 	}
 
-	void Apply(FSolverParticles& Particles, const FSolverReal Dt) const;
-	void CalculateUVStretch(const int32 ConstraintIndex, const FSolverVec3& P0, const FSolverVec3& P1, const FSolverVec3& P2, FSolverVec3& DXDu, FSolverVec3& DXDv) const;
+	CHAOS_API void Apply(FSolverParticles& Particles, const FSolverReal Dt) const;
+	CHAOS_API void CalculateUVStretch(const int32 ConstraintIndex, const FSolverVec3& P0, const FSolverVec3& P1, const FSolverVec3& P2, FSolverVec3& DXDu, FSolverVec3& DXDv) const;
 	
 	const TArray<TVec3<int32>>& GetConstraints() const { return Constraints; }
 	const TArray<FSolverVec2> GetRestStretchLengths() const { return RestStretchLengths; }
@@ -109,10 +109,10 @@ public:
 	const TArray<int32>& GetConstraintsPerColorStartIndex() const { return ConstraintsPerColorStartIndex; }
 
 private:
-	void InitConstraintsAndRestData(const FSolverParticles& InParticles, const FTriangleMesh& TriangleMesh,
+	CHAOS_API void InitConstraintsAndRestData(const FSolverParticles& InParticles, const FTriangleMesh& TriangleMesh,
 		const TArray<TVec3<FVec2f>>& FaceVertexUVs, const bool bUse3dRestLengths, const bool bTrimKinematicConstraints);
-	void InitColor(const FSolverParticles& InParticles);
-	void ApplyHelper(FSolverParticles& Particles, const FSolverReal Dt, const int32 ConstraintIndex, const FSolverVec3& ExpStiffnessValue, const FSolverReal DampingRatioValue, const FSolverReal WarpScaleValue, const FSolverReal WeftScaleValue) const;
+	CHAOS_API void InitColor(const FSolverParticles& InParticles);
+	CHAOS_API void ApplyHelper(FSolverParticles& Particles, const FSolverReal Dt, const int32 ConstraintIndex, const FSolverVec3& ExpStiffnessValue, const FSolverReal DampingRatioValue, const FSolverReal WarpScaleValue, const FSolverReal WeftScaleValue) const;
 
 	TArray<TVec3<int32>> Constraints;
 	const int32 ParticleOffset;
diff --git a/Engine/Source/Runtime/Experimental/Chaos/Public/Chaos/XPBDWeakConstraints.h b/Engine/Source/Runtime/Experimental/Chaos/Public/Chaos/XPBDWeakConstraints.h
--- a/Engine/Source/Runtime/Experimental/Chaos/Public/Chaos/XPBDWeakConstraints.h
+++ b/Engine/Source/Runtime/Experimental/Chaos/Public/Chaos/XPBDWeakConstraints.h
@@ -21,7 +21,7 @@
 namespace Chaos::Softs
 {
 
-	struct CHAOS_API FDeformableXPBDWeakConstraintParams
+	struct FDeformableXPBDWeakConstraintParams
 	{
 		float DebugLineWidth = 5.f;
 		float DebugParticleWidth = 20.f;
diff --git a/Engine/Source/Runtime/Experimental/Chaos/Public/ChaosModule.h b/Engine/Source/Runtime/Experimental/Chaos/Public/ChaosModule.h
--- a/Engine/Source/Runtime/Experimental/Chaos/Public/ChaosModule.h
+++ b/Engine/Source/Runtime/Experimental/Chaos/Public/ChaosModule.h
@@ -4,10 +4,10 @@
 
 #include "Modules/ModuleInterface.h"
 
-class CHAOS_API FChaosEngineModule : public IModuleInterface
+class FChaosEngineModule : public IModuleInterface
 {
 public:
 
-	virtual void StartupModule() override;
-	virtual void ShutdownModule() override;
+	CHAOS_API virtual void StartupModule() override;
+	CHAOS_API virtual void ShutdownModule() override;
 };
diff --git a/Engine/Source/Runtime/Experimental/Chaos/Public/ChaosSolversModule.h b/Engine/Source/Runtime/Experimental/Chaos/Public/ChaosSolversModule.h
--- a/Engine/Source/Runtime/Experimental/Chaos/Public/ChaosSolversModule.h
+++ b/Engine/Source/Runtime/Experimental/Chaos/Public/ChaosSolversModule.h
@@ -99,25 +99,25 @@ enum class ESolverFlags : uint8
 };
 ENUM_CLASS_FLAGS(ESolverFlags);
 
-class CHAOS_API FChaosSolversModule
+class FChaosSolversModule
 {
 public:
 
-	static FChaosSolversModule* GetModule();
+	static CHAOS_API FChaosSolversModule* GetModule();
 
-	FChaosSolversModule();
+	CHAOS_API FChaosSolversModule();
 
-	void StartupModule();
-	void ShutdownModule();
+	CHAOS_API void StartupModule();
+	CHAOS_API void ShutdownModule();
 
-	void Initialize();
-	void Shutdown();
+	CHAOS_API void Initialize();
+	CHAOS_API void Shutdown();
 
 	/**
 	 * Queries for multithreaded configurations
 	 */
-	bool IsPersistentTaskEnabled() const;
-	bool IsPersistentTaskRunning() const;
+	CHAOS_API bool IsPersistentTaskEnabled() const;
+	CHAOS_API bool IsPersistentTaskRunning() const;
 
 
 	/**
@@ -125,13 +125,13 @@ public:
 	 * using methods and members that the calling context can safely access those fields
 	 * as the task will be running on its own thread.
 	 */
-	Chaos::FPersistentPhysicsTask* GetDedicatedTask() const;
+	CHAOS_API Chaos::FPersistentPhysicsTask* GetDedicatedTask() const;
 
 	/**
 	 * Called to request a sync between the game thread and the currently running physics task
 	 * @param bForceBlockingSync forces this 
 	 */
-	void SyncTask(bool bForceBlockingSync = false);
+	CHAOS_API void SyncTask(bool bForceBlockingSync = false);
 
 	/**
 	 * Create a new solver state storage object to contain a solver and proxy storage object. Intended
@@ -145,13 +145,13 @@ public:
 	 * @param ThreadingMode The desired threading mode the solver will use
 	 * @param bStandalone Whether the solver is standalone (not sent to physics thread - updating left to caller)
 	 */
-	Chaos::FPBDRigidsSolver* CreateSolver(UObject* InOwner, Chaos::FReal InAsyncDt, Chaos::EThreadingMode ThreadingMode = Chaos::EThreadingMode::SingleThread
+	CHAOS_API Chaos::FPBDRigidsSolver* CreateSolver(UObject* InOwner, Chaos::FReal InAsyncDt, Chaos::EThreadingMode ThreadingMode = Chaos::EThreadingMode::SingleThread
 #if CHAOS_DEBUG_NAME
 		, const FName& DebugName = NAME_None
 #endif
 	);
 
-	void MigrateSolver(Chaos::FPhysicsSolverBase* InSolver, const UObject* InNewOwner);
+	CHAOS_API void MigrateSolver(Chaos::FPhysicsSolverBase* InSolver, const UObject* InNewOwner);
 
 	/**
 	 * Sets the actor type which should be AChaosSolverActor::StaticClass() but that is not accessible from engine.
@@ -163,36 +163,36 @@ public:
 		check(IsValidSolverActorClass(SolverActorClass));
 	}
 
-	UClass* GetSolverActorClass() const;
+	CHAOS_API UClass* GetSolverActorClass() const;
 
-	bool IsValidSolverActorClass(UClass* Class) const;
+	CHAOS_API bool IsValidSolverActorClass(UClass* Class) const;
 
 	/**
 	 * Shuts down and destroys a solver state
 	 *
 	 * Should be called on whichever thread currently owns the solver state
 	 */
-	void DestroySolver(Chaos::FPhysicsSolverBase* InState);
+	CHAOS_API void DestroySolver(Chaos::FPhysicsSolverBase* InState);
 
 	/** Retrieve the list of all extant solvers. This contains all owned, unowned and standalone solvers */
-	const TArray<Chaos::FPhysicsSolverBase*>& GetAllSolvers() const;
+	CHAOS_API const TArray<Chaos::FPhysicsSolverBase*>& GetAllSolvers() const;
 
 	/**
 	 * Read access to the current solver-state objects, be aware which thread owns this data when
 	 * attempting to use this. Physics thread will query when spinning up to get current world state
 	 */
-	TArray<const Chaos::FPhysicsSolverBase*> GetSolvers(const UObject* InOwner) const;
-	void GetSolvers(const UObject* InOwner, TArray<const Chaos::FPhysicsSolverBase*>& OutSolvers) const;
+	CHAOS_API TArray<const Chaos::FPhysicsSolverBase*> GetSolvers(const UObject* InOwner) const;
+	CHAOS_API void GetSolvers(const UObject* InOwner, TArray<const Chaos::FPhysicsSolverBase*>& OutSolvers) const;
 
 	/** Non-const access to the solver array is private - only the module should ever modify the storage lists */
-	TArray<Chaos::FPhysicsSolverBase*> GetSolversMutable(const UObject* InOwner);
-	void GetSolversMutable(const UObject* InOwner, TArray<Chaos::FPhysicsSolverBase*>& OutSolvers);
+	CHAOS_API TArray<Chaos::FPhysicsSolverBase*> GetSolversMutable(const UObject* InOwner);
+	CHAOS_API void GetSolversMutable(const UObject* InOwner, TArray<Chaos::FPhysicsSolverBase*>& OutSolvers);
 
 	/**
 	 * Outputs statistics for the solver hierarchies. Currently engine calls into this
 	 * from a console command on demand.
 	 */
-	void DumpHierarchyStats(int32* OutOptMaxCellElements = nullptr);
+	CHAOS_API void DumpHierarchyStats(int32* OutOptMaxCellElements = nullptr);
 
 #if WITH_EDITOR
 	/**
@@ -200,25 +200,25 @@ public:
 	 * This is typically called from a playing editor to pause all solvers.
 	 * @note Game pause must use a different per solver mechanics.
 	 */
-	void PauseSolvers();
+	CHAOS_API void PauseSolvers();
 
 	/**
 	 * Resume solvers. Thread safe.
 	 * This is typically called from a paused editor to resume all solvers.
 	 * @note Game resume must use a different per solver mechanics.
 	 */
-	void ResumeSolvers();
+	CHAOS_API void ResumeSolvers();
 
 	/**
 	 * Single-step advance solvers. Thread safe.
 	 * This is typically called from a paused editor to single step all solvers.
 	 */
-	void SingleStepSolvers();
+	CHAOS_API void SingleStepSolvers();
 
 	/**
 	 * Query whether a particular solver should advance and update its single-step counter. Thread safe.
 	 */
-	bool ShouldStepSolver(int32& InOutSingleStepCounter) const;
+	CHAOS_API bool ShouldStepSolver(int32& InOutSingleStepCounter) const;
 #endif  // #if WITH_EDITOR
 
 	void RegisterSolverActorClassProvider(IChaosSolverActorClassProvider* Provider)
@@ -228,7 +228,7 @@ public:
 
 
 	/** Safe method for always getting a settings provider (from the external caller or an internal default) */
-	const IChaosSettingsProvider& GetSettingsProvider() const;
+	CHAOS_API const IChaosSettingsProvider& GetSettingsProvider() const;
 	void SetSettingsProvider(IChaosSettingsProvider* InProvider)
 	{
 		SettingsProvider = InProvider;
@@ -238,16 +238,16 @@ public:
 	void UnlockSolvers() { SolverLock.Unlock(); }
 
 	/** Events that the material manager will emit for us to update our threaded data */
-	void OnUpdateMaterial(Chaos::FMaterialHandle InHandle);
-	void OnCreateMaterial(Chaos::FMaterialHandle InHandle);
-	void OnDestroyMaterial(Chaos::FMaterialHandle InHandle);
+	CHAOS_API void OnUpdateMaterial(Chaos::FMaterialHandle InHandle);
+	CHAOS_API void OnCreateMaterial(Chaos::FMaterialHandle InHandle);
+	CHAOS_API void OnDestroyMaterial(Chaos::FMaterialHandle InHandle);
 
-	void OnUpdateMaterialMask(Chaos::FMaterialMaskHandle InHandle);
-	void OnCreateMaterialMask(Chaos::FMaterialMaskHandle InHandle);
-	void OnDestroyMaterialMask(Chaos::FMaterialMaskHandle InHandle);
+	CHAOS_API void OnUpdateMaterialMask(Chaos::FMaterialMaskHandle InHandle);
+	CHAOS_API void OnCreateMaterialMask(Chaos::FMaterialMaskHandle InHandle);
+	CHAOS_API void OnDestroyMaterialMask(Chaos::FMaterialMaskHandle InHandle);
 
 	/** Gets a list of pending prerequisites required before proceeding with a solver update */
-	void GetSolverUpdatePrerequisites(FGraphEventArray& InPrerequisiteContainer);
+	CHAOS_API void GetSolverUpdatePrerequisites(FGraphEventArray& InPrerequisiteContainer);
 
 private:
 
@@ -314,7 +314,7 @@ private:
 	FDelegateHandle OnUpdateMaterialMaskHandle;
 };
 
-struct CHAOS_API FChaosScopeSolverLock
+struct FChaosScopeSolverLock
 {
 	FChaosScopeSolverLock()
 	{
diff --git a/Engine/Source/Runtime/Experimental/Chaos/Public/ChaosVisualDebugger/ChaosVDContextProvider.h b/Engine/Source/Runtime/Experimental/Chaos/Public/ChaosVisualDebugger/ChaosVDContextProvider.h
--- a/Engine/Source/Runtime/Experimental/Chaos/Public/ChaosVisualDebugger/ChaosVDContextProvider.h
+++ b/Engine/Source/Runtime/Experimental/Chaos/Public/ChaosVisualDebugger/ChaosVDContextProvider.h
@@ -11,14 +11,14 @@
 #include "ChaosVDRuntimeModule.h"
 
 /** Chaos Visual Debugger data used to context for logging or debugging purposes */
-struct CHAOS_API FChaosVDContext
+struct FChaosVDContext
 {
 	int32 OwnerID = INDEX_NONE;
 	int32 Id = INDEX_NONE;
 };
 
 /** Singleton class that manages the thread local storage used to store CVD Context data */
-class CHAOS_API FChaosVDThreadContext : public TThreadSingleton<FChaosVDThreadContext>
+class FChaosVDThreadContext : public TThreadSingleton<FChaosVDThreadContext>
 {
 public:
 
@@ -33,19 +33,19 @@ public:
 	/** Copies the Current CVD context data into the provided struct
 	 * @return true if the copy was successful
 	 */
-	bool GetCurrentContext(FChaosVDContext& OutContext);
+	CHAOS_API bool GetCurrentContext(FChaosVDContext& OutContext);
 	
 	/** Gets the current CVD context data -
 	 * Don't use of a function that will recursively push new context data as it might invalidate the pointer
 	 * @return Ptr to the Current CVD context data
 	 */
-	const FChaosVDContext* GetCurrentContext();
+	CHAOS_API const FChaosVDContext* GetCurrentContext();
 	
 	/** Pushed a new CVD Context data to the local cvd context stack */
-	void PushContext(const FChaosVDContext& InContext);
+	CHAOS_API void PushContext(const FChaosVDContext& InContext);
 	
 	/** Removes the CVD Context data at the top of the local cvd context stack */
-	void PopContext();
+	CHAOS_API void PopContext();
 
 protected:
 	TArray<FChaosVDContext, TInlineAllocator<16>> LocalContextStack;
@@ -75,7 +75,7 @@ protected:
 /** Utility Class that will push the provided CVD Context Data to the local thread storage
  * and remove it when it goes out of scope
  */
-struct CHAOS_API FChaosVDScopeContext
+struct FChaosVDScopeContext
 {
 	FChaosVDScopeContext(const FChaosVDContext& InCVDContext)
 	{
@@ -91,7 +91,7 @@ struct CHAOS_API FChaosVDScopeContext
 /** Wrapper class that provides access to the data buffer that should be used to trace CVD (Chaos Visual Debugger) data.
  * It ensures the buffer is cleared on scope exit but keeping a minimum amount of memory allocated
  */
-struct CHAOS_API FChaosVDScopedTLSBufferAccessor
+struct FChaosVDScopedTLSBufferAccessor
 {
 	FChaosVDScopedTLSBufferAccessor() : BufferRef(FChaosVDThreadContext::Get().GetTLSDataBufferRef())
 	{
@@ -150,7 +150,7 @@ struct CHAOS_API FChaosVDScopedTLSBufferAccessor
 
 #endif // WITH_CHAOS_VISUAL_DEBUGGER
 
-struct CHAOS_API FChaosVDContextWrapper
+struct FChaosVDContextWrapper
 {
 	FChaosVDContextWrapper()
 	{
diff --git a/Engine/Source/Runtime/Experimental/Chaos/Public/EventDefaults.h b/Engine/Source/Runtime/Experimental/Chaos/Public/EventDefaults.h
--- a/Engine/Source/Runtime/Experimental/Chaos/Public/EventDefaults.h
+++ b/Engine/Source/Runtime/Experimental/Chaos/Public/EventDefaults.h
@@ -6,14 +6,14 @@
 
 namespace Chaos
 {
-	class CHAOS_API FEventDefaults
+	class FEventDefaults
 	{
 	public:
 
 		/**
 		 * Register default event types
 		 */
-		static void RegisterSystemEvents(FEventManager& EventManager);
+		static CHAOS_API void RegisterSystemEvents(FEventManager& EventManager);
 
 	private:
 
@@ -48,4 +48,4 @@ namespace Chaos
 		static void RegisterCrumblingEvent(FEventManager& EventManager);
 
 	};
-}
\ No newline at end of file
+}
diff --git a/Engine/Source/Runtime/Experimental/Chaos/Public/EventManager.h b/Engine/Source/Runtime/Experimental/Chaos/Public/EventManager.h
--- a/Engine/Source/Runtime/Experimental/Chaos/Public/EventManager.h
+++ b/Engine/Source/Runtime/Experimental/Chaos/Public/EventManager.h
@@ -343,7 +343,7 @@ namespace Chaos
 	 */
 	using FEventContainerBasePtr = FEventContainerBase*;
 
-	class CHAOS_API FEventManager
+	class FEventManager
 	{
 		friend class FPBDRigidsSolver;
 
@@ -362,7 +362,7 @@ namespace Chaos
 		/**
 		 * Clears out every handler and container calling destructors on held items
 		 */
-		void Reset();
+		CHAOS_API void Reset();
 
 		/**
 		 * Set the buffer mode to be used within the event containers
@@ -401,7 +401,7 @@ namespace Chaos
 		/**
 		 * Unregister specified event from system
 		 */
-		void UnregisterEvent(const EEventType& EventType);
+		CHAOS_API void UnregisterEvent(const EEventType& EventType);
 
 		/**
 		 * Register a handler that will receive the dispatched events
@@ -430,27 +430,27 @@ namespace Chaos
 		/**
 		 * Unregister the specified event handler
 		 */
-		void UnregisterHandler(const EEventType& EventType, const void* InHandler);
+		CHAOS_API void UnregisterHandler(const EEventType& EventType, const void* InHandler);
 
 		/**
 		 * Called by the solver to invoke the functions that fill the producer side of all the event data buffers
 		 */
-		void FillProducerData(const Chaos::FPBDRigidsSolver* Solver, bool bResetData = true);
+		CHAOS_API void FillProducerData(const Chaos::FPBDRigidsSolver* Solver, bool bResetData = true);
 
 		/**
 		 * Flips the event data buffer if it is of double or triple buffer type
 		 */
-		void FlipBuffersIfRequired();
+		CHAOS_API void FlipBuffersIfRequired();
 
 		/**
 		 * // Dispatch events to the registered handlers
 		 */
-		void DispatchEvents();
+		CHAOS_API void DispatchEvents();
 
 		/** Returns encoded collision index. */
-		static int32 EncodeCollisionIndex(int32 ActualCollisionIndex, bool bSwapOrder);
+		static CHAOS_API int32 EncodeCollisionIndex(int32 ActualCollisionIndex, bool bSwapOrder);
 		/** Returns decoded collision index. */
-		static int32 DecodeCollisionIndex(int32 EncodedCollisionIdx, bool& bSwapOrder);
+		static CHAOS_API int32 DecodeCollisionIndex(int32 EncodedCollisionIdx, bool& bSwapOrder);
 
 
 		template<typename PayloadType>
@@ -466,7 +466,7 @@ namespace Chaos
 
 	private:
 
-		void InternalRegisterInjector(const FEventID& EventID, const FEventContainerBasePtr& Container);
+		CHAOS_API void InternalRegisterInjector(const FEventID& EventID, const FEventContainerBasePtr& Container);
 
 		Chaos::EMultiBufferMode BufferMode;			// specifies the buffer type to be constructed, single, double, triple
 		TArray<FEventContainerBasePtr> EventContainers;	// Array of event types
@@ -479,4 +479,4 @@ namespace Chaos
 		TArray<TPair<FEventID, IEventHandler*>> DeferredHandlers;
 	};
 
-}
\ No newline at end of file
+}
diff --git a/Engine/Source/Runtime/Experimental/Chaos/Public/Field/FieldSystem.h b/Engine/Source/Runtime/Experimental/Chaos/Public/Field/FieldSystem.h
--- a/Engine/Source/Runtime/Experimental/Chaos/Public/Field/FieldSystem.h
+++ b/Engine/Source/Runtime/Experimental/Chaos/Public/Field/FieldSystem.h
@@ -21,7 +21,7 @@
 *   data package that the nodes will use during evaluation, the context does not assume
 *   ownership of the metadata but assumes it will remain in scope during evaluation.
 */
-struct CHAOS_API FFieldContextIndex
+struct FFieldContextIndex
 {
 	FFieldContextIndex(int32 InSample = INDEX_NONE, int32 InResult = INDEX_NONE)
 		: Sample(InSample)
@@ -75,7 +75,7 @@ enum class EFieldCommandHandlesType : uint8
 };
 
 /** List of datas that will be stored during field evaluation to avoid reallocation */
-struct CHAOS_API FFieldExecutionDatas
+struct FFieldExecutionDatas
 {
 	/** Sample positions to be used to build the context */
 	TArray<FVector> SamplePositions;
@@ -155,7 +155,7 @@ FORCEINLINE void EmptyResultsArrays(const TArray<EFieldCommandOutputType>& Field
 */
 
 
-class CHAOS_API FFieldSystemMetaData {
+class FFieldSystemMetaData {
 public:
 
 	enum EMetaType
@@ -175,7 +175,7 @@ public:
 };
 
 
-class CHAOS_API FFieldSystemMetaDataProcessingResolution : public FFieldSystemMetaData {
+class FFieldSystemMetaDataProcessingResolution : public FFieldSystemMetaData {
 public:
 	FFieldSystemMetaDataProcessingResolution(EFieldResolutionType ProcessingResolutionIn) : ProcessingResolution(ProcessingResolutionIn) {};
 	virtual ~FFieldSystemMetaDataProcessingResolution() {};
@@ -185,7 +185,7 @@ public:
 	EFieldResolutionType ProcessingResolution;
 };
 
-class CHAOS_API FFieldSystemMetaDataFilter : public FFieldSystemMetaData {
+class FFieldSystemMetaDataFilter : public FFieldSystemMetaData {
 public:
 	FFieldSystemMetaDataFilter(EFieldFilterType FilterTypeIn, EFieldObjectType ObjectTypeIn, EFieldPositionType PositionTypeIn) : FilterType(FilterTypeIn), ObjectType(ObjectTypeIn), PositionType(PositionTypeIn)  {};
 	virtual ~FFieldSystemMetaDataFilter() {};
@@ -198,7 +198,7 @@ public:
 };
 
 template<class T>
-class CHAOS_API FFieldSystemMetaDataResults : public FFieldSystemMetaData {
+class FFieldSystemMetaDataResults : public FFieldSystemMetaData {
 public:
 	FFieldSystemMetaDataResults(const TFieldArrayView<T>& ResultsIn) : Results(ResultsIn) {};
 	virtual ~FFieldSystemMetaDataResults() {};
@@ -208,7 +208,7 @@ public:
 	const TFieldArrayView<T>& Results;
 };
 
-class CHAOS_API FFieldSystemMetaDataIteration : public FFieldSystemMetaData {
+class FFieldSystemMetaDataIteration : public FFieldSystemMetaData {
 public:
 	FFieldSystemMetaDataIteration(int32 IterationsIn) : Iterations(IterationsIn) {};
 	virtual ~FFieldSystemMetaDataIteration() {};
@@ -218,7 +218,7 @@ public:
 	int32 Iterations;
 };
 
-class CHAOS_API FFieldSystemMetaDataCulling : public FFieldSystemMetaData
+class FFieldSystemMetaDataCulling : public FFieldSystemMetaData
 {
 public:
 	explicit FFieldSystemMetaDataCulling(TArray<FFieldContextIndex>& CullingIndicesIn)
@@ -241,7 +241,7 @@ public:
 	TArray<FFieldContextIndex>& CullingIndices;
 };
 
-struct CHAOS_API FFieldContext
+struct FFieldContext
 {
 	typedef  TMap<FFieldSystemMetaData::EMetaType, TUniquePtr<FFieldSystemMetaData> > UniquePointerMap;
 	typedef  TMap<FFieldSystemMetaData::EMetaType, FFieldSystemMetaData * > PointerMap;
@@ -369,7 +369,7 @@ FORCEINLINE TArray<int32>& GetResultArray<int32>(FFieldContext& FieldContext)
  * This has the effect of exposing metadata to downstream nodes but making sure
  * upstream nodes cannot see it.
  */
-class CHAOS_API FScopedFieldContextMetaData
+class FScopedFieldContextMetaData
 {
 	FScopedFieldContextMetaData() = delete;
 	FScopedFieldContextMetaData(const FScopedFieldContextMetaData&) = delete;
@@ -403,7 +403,7 @@ private:
 *  Abstract base class for the field node evaluation. 
 *
 */
-class CHAOS_API FFieldNodeBase
+class FFieldNodeBase
 {
 
 public:
@@ -522,7 +522,7 @@ template<> inline FFieldNodeBase::EFieldType FFieldNode<FVector>::StaticType() {
 *   will be passed to the evaluation of the field. 
 *
 */
-class CHAOS_API FFieldSystemCommand
+class FFieldSystemCommand
 {
 public:
 	FFieldSystemCommand()
@@ -603,8 +603,8 @@ public:
 		TimeCreation = (float)TimeSeconds;
 	}
 
-	void Serialize(FArchive& Ar);
-	bool operator==(const FFieldSystemCommand&) const;
+	CHAOS_API void Serialize(FArchive& Ar);
+	CHAOS_API bool operator==(const FFieldSystemCommand&) const;
 	bool operator!=(const FFieldSystemCommand& Other) const { return !this->operator==(Other); }
 
 	FName TargetAttribute;
diff --git a/Engine/Source/Runtime/Experimental/Chaos/Public/Field/FieldSystemNodes.h b/Engine/Source/Runtime/Experimental/Chaos/Public/Field/FieldSystemNodes.h
--- a/Engine/Source/Runtime/Experimental/Chaos/Public/Field/FieldSystemNodes.h
+++ b/Engine/Source/Runtime/Experimental/Chaos/Public/Field/FieldSystemNodes.h
@@ -10,7 +10,7 @@
 * FUniformInteger
 **/
 
-class CHAOS_API FUniformInteger : public FFieldNode<int32>
+class FUniformInteger : public FFieldNode<int32>
 {
 	typedef FFieldNode<int32> Super;
 
@@ -23,22 +23,22 @@ public:
 	virtual FFieldNodeBase * NewCopy() const override { return new FUniformInteger(Magnitude); }
 	virtual ~FUniformInteger() {}
 
-	virtual void Evaluate(FFieldContext&, TFieldArrayView<int32>& Results) const override;
-	virtual bool operator==(const FFieldNodeBase& Node) override;
+	CHAOS_API virtual void Evaluate(FFieldContext&, TFieldArrayView<int32>& Results) const override;
+	CHAOS_API virtual bool operator==(const FFieldNodeBase& Node) override;
 
 	/** Serialization API */
 	virtual FFieldNodeBase::ESerializationType SerializationType() const override { return FFieldNodeBase::ESerializationType::FieldNode_FUniformInteger; }
-	virtual void Serialize(FArchive& Ar) override;
+	CHAOS_API virtual void Serialize(FArchive& Ar) override;
 
 	/** Count the number of offsets/params that will be used by the world physics field */
-	virtual void FillSetupCount(int32& NumOffsets, int32& NumParams) const override;
+	CHAOS_API virtual void FillSetupCount(int32& NumOffsets, int32& NumParams) const override;
 
 	/** Fill the offsets/params arrays that will be used by the world physics field */
-	virtual void FillSetupDatas(TArray<int32>& NodesOffsets, TArray<float>& NodesParams,
+	CHAOS_API virtual void FillSetupDatas(TArray<int32>& NodesOffsets, TArray<float>& NodesParams,
 		const TMap<FFieldNodeBase*, float>& CommandTimes, const float PreviousTime) const override;
 
 	/** Evaluate the maximum magnitude of the field graph  */
-	virtual float EvalMaxMagnitude() const override;
+	CHAOS_API virtual float EvalMaxMagnitude() const override;
 
 	int32 Magnitude;
 };
@@ -48,7 +48,7 @@ public:
 * RadialMaskField
 **/
 
-class CHAOS_API FRadialIntMask : public FFieldNode<int32>
+class FRadialIntMask : public FFieldNode<int32>
 {
 	typedef FFieldNode<int32> Super;
 
@@ -69,25 +69,25 @@ public:
 	virtual FFieldNodeBase * NewCopy() const override { return new FRadialIntMask(Radius,Position,InteriorValue,ExteriorValue,SetMaskCondition); }
 	virtual ~FRadialIntMask() {}
 
-	virtual void Evaluate(FFieldContext&, TFieldArrayView<int32>& Results) const override;
-	virtual bool operator==(const FFieldNodeBase& Node) override;
+	CHAOS_API virtual void Evaluate(FFieldContext&, TFieldArrayView<int32>& Results) const override;
+	CHAOS_API virtual bool operator==(const FFieldNodeBase& Node) override;
 
 	/** Serialization API */
 	virtual FFieldNodeBase::ESerializationType SerializationType() const override { return FFieldNodeBase::ESerializationType::FieldNode_FRadialIntMask; }
-	virtual void Serialize(FArchive& Ar) override;
+	CHAOS_API virtual void Serialize(FArchive& Ar) override;
 
 	/** Count the number of offsets/params that will be used by the world physics field */
-	virtual void FillSetupCount(int32& NumOffsets, int32& NumParams) const override;
+	CHAOS_API virtual void FillSetupCount(int32& NumOffsets, int32& NumParams) const override;
 
 	/** Fill the offsets/params arrays that will be used by the world physics field */
-	virtual void FillSetupDatas(TArray<int32>& NodesOffsets, TArray<float>& NodesParams,
+	CHAOS_API virtual void FillSetupDatas(TArray<int32>& NodesOffsets, TArray<float>& NodesParams,
 		const TMap<FFieldNodeBase*, float>& CommandTimes, const float PreviousTime) const override;
 
 	/** Evaluate the maximum magnitude of the field graph  */
-	virtual float EvalMaxMagnitude() const override;
+	CHAOS_API virtual float EvalMaxMagnitude() const override;
 
 	/** Compute the min/max spatial bounds of the field */
-	virtual void ComputeFieldBounds(FVector& MinBounds, FVector& MaxBounds) const override;
+	CHAOS_API virtual void ComputeFieldBounds(FVector& MinBounds, FVector& MaxBounds) const override;
 
 	float Radius;
 	FVector Position;
@@ -101,7 +101,7 @@ public:
 * FUniformScalar
 **/
 
-class CHAOS_API FUniformScalar : public FFieldNode<float>
+class FUniformScalar : public FFieldNode<float>
 {
 	typedef FFieldNode<float> Super;
 
@@ -115,22 +115,22 @@ public:
 
 	virtual ~FUniformScalar() {}
 
-	void Evaluate(FFieldContext&, TFieldArrayView<float>& Results) const override;
-	virtual bool operator==(const FFieldNodeBase& Node);
+	CHAOS_API void Evaluate(FFieldContext&, TFieldArrayView<float>& Results) const override;
+	CHAOS_API virtual bool operator==(const FFieldNodeBase& Node);
 
 	/** Serialization API */
 	virtual FFieldNodeBase::ESerializationType SerializationType() const override { return FFieldNodeBase::ESerializationType::FieldNode_FUniformScalar; }
-	virtual void Serialize(FArchive& Ar) override;
+	CHAOS_API virtual void Serialize(FArchive& Ar) override;
 
 	/** Count the number of offsets/params that will be used by the world physics field */
-	virtual void FillSetupCount(int32& NumOffsets, int32& NumParams) const override;
+	CHAOS_API virtual void FillSetupCount(int32& NumOffsets, int32& NumParams) const override;
 
 	/** Fill the offsets/params arrays that will be used by the world physics field */
-	virtual void FillSetupDatas(TArray<int32>& NodesOffsets, TArray<float>& NodesParams,
+	CHAOS_API virtual void FillSetupDatas(TArray<int32>& NodesOffsets, TArray<float>& NodesParams,
 		const TMap<FFieldNodeBase*, float>& CommandTimes, const float PreviousTime) const override;
 
 	/** Evaluate the maximum magnitude of the field graph  */
-	virtual float EvalMaxMagnitude() const override;
+	CHAOS_API virtual float EvalMaxMagnitude() const override;
 
 	float Magnitude;
 };
@@ -139,7 +139,7 @@ public:
 * FWaveScalar
 **/
 
-class CHAOS_API FWaveScalar : public FFieldNode<float>
+class FWaveScalar : public FFieldNode<float>
 {
 	typedef FFieldNode<float> Super;
 
@@ -159,22 +159,22 @@ public:
 
 	virtual ~FWaveScalar() {}
 
-	virtual void Evaluate(FFieldContext&, TFieldArrayView<float>& Results) const override;
-	virtual bool operator==(const FFieldNodeBase& Node) override;
+	CHAOS_API virtual void Evaluate(FFieldContext&, TFieldArrayView<float>& Results) const override;
+	CHAOS_API virtual bool operator==(const FFieldNodeBase& Node) override;
 
 	/** Serialization API */
 	virtual FFieldNodeBase::ESerializationType SerializationType() const override { return FFieldNodeBase::ESerializationType::FieldNode_FWaveScalar; }
-	virtual void Serialize(FArchive& Ar) override;
+	CHAOS_API virtual void Serialize(FArchive& Ar) override;
 
 	/** Count the number of offsets/params that will be used by the world physics field */
-	virtual void FillSetupCount(int32& NumOffsets, int32& NumParams) const override;
+	CHAOS_API virtual void FillSetupCount(int32& NumOffsets, int32& NumParams) const override;
 
 	/** Fill the offsets/params arrays that will be used by the world physics field */
-	virtual void FillSetupDatas(TArray<int32>& NodesOffsets, TArray<float>& NodesParams,
+	CHAOS_API virtual void FillSetupDatas(TArray<int32>& NodesOffsets, TArray<float>& NodesParams,
 		const TMap<FFieldNodeBase*, float>& CommandTimes, const float PreviousTime) const override;
 
 	/** Evaluate the maximum magnitude of the field graph  */
-	virtual float EvalMaxMagnitude() const override;
+	CHAOS_API virtual float EvalMaxMagnitude() const override;
 
 
 	float Magnitude;
@@ -189,7 +189,7 @@ public:
 * RadialFalloff
 **/
 
-class CHAOS_API FRadialFalloff : public FFieldNode<float>
+class FRadialFalloff : public FFieldNode<float>
 {
 	typedef FFieldNode<float> Super;
 
@@ -216,25 +216,25 @@ public:
 
 	virtual ~FRadialFalloff() {}
 
-	virtual void Evaluate(FFieldContext&, TFieldArrayView<float>& Results) const override;
-	virtual bool operator==(const FFieldNodeBase& Node) override;
+	CHAOS_API virtual void Evaluate(FFieldContext&, TFieldArrayView<float>& Results) const override;
+	CHAOS_API virtual bool operator==(const FFieldNodeBase& Node) override;
 
 	/** Serialization API */
 	virtual FFieldNodeBase::ESerializationType SerializationType() const override { return FFieldNodeBase::ESerializationType::FieldNode_FRadialFalloff; }
-	virtual void Serialize(FArchive& Ar) override;
+	CHAOS_API virtual void Serialize(FArchive& Ar) override;
 
 	/** Count the number of offsets/params that will be used by the world physics field */
-	virtual void FillSetupCount(int32& NumOffsets, int32& NumParams) const override;
+	CHAOS_API virtual void FillSetupCount(int32& NumOffsets, int32& NumParams) const override;
 
 	/** Fill the offsets/params arrays that will be used by the world physics field */
-	virtual void FillSetupDatas(TArray<int32>& NodesOffsets, TArray<float>& NodesParams,
+	CHAOS_API virtual void FillSetupDatas(TArray<int32>& NodesOffsets, TArray<float>& NodesParams,
 		const TMap<FFieldNodeBase*, float>& CommandTimes, const float PreviousTime) const override;
 
 	/** Evaluate the maximum magnitude of the field graph  */
-	virtual float EvalMaxMagnitude() const override;
+	CHAOS_API virtual float EvalMaxMagnitude() const override;
 
 	/** Compute the min/max spatial bounds of the field */
-	virtual void ComputeFieldBounds(FVector& MinBounds, FVector& MaxBounds) const override;
+	CHAOS_API virtual void ComputeFieldBounds(FVector& MinBounds, FVector& MaxBounds) const override;
 
 	float Magnitude;
 	float MinRange;
@@ -252,7 +252,7 @@ protected:
 /**
 * FPlaneFalloff
 **/
-class CHAOS_API FPlaneFalloff : public FFieldNode<float>
+class FPlaneFalloff : public FFieldNode<float>
 {
 	typedef FFieldNode<float> Super;
 
@@ -281,25 +281,25 @@ public:
 	virtual FFieldNodeBase * NewCopy() const override { return new FPlaneFalloff(Magnitude, MinRange, MaxRange, Default, Distance, Position, Normal, Falloff); }
 	virtual ~FPlaneFalloff() {}
 
-	virtual void Evaluate(FFieldContext&, TFieldArrayView<float>& Results) const override;
-	virtual bool operator==(const FFieldNodeBase& Node) override;
+	CHAOS_API virtual void Evaluate(FFieldContext&, TFieldArrayView<float>& Results) const override;
+	CHAOS_API virtual bool operator==(const FFieldNodeBase& Node) override;
 
 	/** Serialization API */
 	virtual FFieldNodeBase::ESerializationType SerializationType() const override { return FFieldNodeBase::ESerializationType::FieldNode_FPlaneFalloff; }
-	virtual void Serialize(FArchive& Ar) override;
+	CHAOS_API virtual void Serialize(FArchive& Ar) override;
 
 	/** Count the number of offsets/params that will be used by the world physics field */
-	virtual void FillSetupCount(int32& NumOffsets, int32& NumParams) const override;
+	CHAOS_API virtual void FillSetupCount(int32& NumOffsets, int32& NumParams) const override;
 
 	/** Fill the offsets/params arrays that will be used by the world physics field */
-	virtual void FillSetupDatas(TArray<int32>& NodesOffsets, TArray<float>& NodesParams,
+	CHAOS_API virtual void FillSetupDatas(TArray<int32>& NodesOffsets, TArray<float>& NodesParams,
 		const TMap<FFieldNodeBase*, float>& CommandTimes, const float PreviousTime) const override;
 
 	/** Evaluate the maximum magnitude of the field graph  */
-	virtual float EvalMaxMagnitude() const override;
+	CHAOS_API virtual float EvalMaxMagnitude() const override;
 
 	/** Compute the min/max spatial bounds of the field */
-	virtual void ComputeFieldBounds(FVector& MinBounds, FVector& MaxBounds) const override;
+	CHAOS_API virtual void ComputeFieldBounds(FVector& MinBounds, FVector& MaxBounds) const override;
 
 
 	float Magnitude;
@@ -319,7 +319,7 @@ protected :
 /**
 * FBoxFalloff
 **/
-class CHAOS_API FBoxFalloff : public FFieldNode<float>
+class FBoxFalloff : public FFieldNode<float>
 {
 	typedef FFieldNode<float> Super;
 
@@ -344,25 +344,25 @@ public:
 	virtual FFieldNodeBase * NewCopy() const override { return new FBoxFalloff(Magnitude, MinRange, MaxRange, Default, Transform, Falloff); }
 	virtual ~FBoxFalloff() {}
 
-	virtual void Evaluate(FFieldContext&, TFieldArrayView<float>& Results) const override;
-	virtual bool operator==(const FFieldNodeBase& Node) override;
+	CHAOS_API virtual void Evaluate(FFieldContext&, TFieldArrayView<float>& Results) const override;
+	CHAOS_API virtual bool operator==(const FFieldNodeBase& Node) override;
 
 	/** Serialization API */
 	virtual FFieldNodeBase::ESerializationType SerializationType() const override { return FFieldNodeBase::ESerializationType::FieldNode_FBoxFalloff; }
-	virtual void Serialize(FArchive& Ar) override;
+	CHAOS_API virtual void Serialize(FArchive& Ar) override;
 
 	/** Count the number of offsets/params that will be used by the world physics field */
-	virtual void FillSetupCount(int32& NumOffsets, int32& NumParams) const override;
+	CHAOS_API virtual void FillSetupCount(int32& NumOffsets, int32& NumParams) const override;
 
 	/** Fill the offsets/params arrays that will be used by the world physics field */
-	virtual void FillSetupDatas(TArray<int32>& NodesOffsets, TArray<float>& NodesParams,
+	CHAOS_API virtual void FillSetupDatas(TArray<int32>& NodesOffsets, TArray<float>& NodesParams,
 		const TMap<FFieldNodeBase*, float>& CommandTimes, const float PreviousTime) const override;
 
 	/** Evaluate the maximum magnitude of the field graph  */
-	virtual float EvalMaxMagnitude() const override;
+	CHAOS_API virtual float EvalMaxMagnitude() const override;
 
 	/** Compute the min/max spatial bounds of the field */
-	virtual void ComputeFieldBounds(FVector& MinBounds, FVector& MaxBounds) const override;
+	CHAOS_API virtual void ComputeFieldBounds(FVector& MinBounds, FVector& MaxBounds) const override;
 
 	float Magnitude;
 	float MinRange;
@@ -380,7 +380,7 @@ protected:
 /**
 * NoiseField
 **/
-class CHAOS_API FNoiseField : public FFieldNode<float>
+class FNoiseField : public FFieldNode<float>
 {
 	typedef FFieldNode<float> Super;
 
@@ -395,22 +395,22 @@ public:
 	virtual FFieldNodeBase * NewCopy() const override { return new FNoiseField(MinRange,MaxRange,Transform); }
 	virtual ~FNoiseField() {}
 
-	virtual void Evaluate(FFieldContext&, TFieldArrayView<float>& Results) const override;
-	virtual bool operator==(const FFieldNodeBase& Node) override;
+	CHAOS_API virtual void Evaluate(FFieldContext&, TFieldArrayView<float>& Results) const override;
+	CHAOS_API virtual bool operator==(const FFieldNodeBase& Node) override;
 
 	/** Serialization API */
 	virtual FFieldNodeBase::ESerializationType SerializationType() const override { return FFieldNodeBase::ESerializationType::FieldNode_FNoiseField; }
-	virtual void Serialize(FArchive& Ar) override;
+	CHAOS_API virtual void Serialize(FArchive& Ar) override;
 
 	/** Count the number of offsets/params that will be used by the world physics field */
-	virtual void FillSetupCount(int32& NumOffsets, int32& NumParams) const override;
+	CHAOS_API virtual void FillSetupCount(int32& NumOffsets, int32& NumParams) const override;
 
 	/** Fill the offsets/params arrays that will be used by the world physics field */
-	virtual void FillSetupDatas(TArray<int32>& NodesOffsets, TArray<float>& NodesParams,
+	CHAOS_API virtual void FillSetupDatas(TArray<int32>& NodesOffsets, TArray<float>& NodesParams,
 		const TMap<FFieldNodeBase*, float>& CommandTimes, const float PreviousTime) const override;
 
 	/** Evaluate the maximum magnitude of the field graph  */
-	virtual float EvalMaxMagnitude() const override;
+	CHAOS_API virtual float EvalMaxMagnitude() const override;
 
 	float MinRange;
 	float MaxRange;
@@ -421,7 +421,7 @@ public:
 * UniformVector
 **/
 
-class CHAOS_API FUniformVector : public FFieldNode<FVector>
+class FUniformVector : public FFieldNode<FVector>
 {
 	typedef FFieldNode<FVector> Super;
 
@@ -436,22 +436,22 @@ public:
 	virtual FFieldNodeBase * NewCopy() const override { return new FUniformVector(Magnitude,Direction); }
 	virtual ~FUniformVector() {}
 
-	virtual void Evaluate(FFieldContext&, TFieldArrayView<FVector>& Results) const override;
-	virtual bool operator==(const FFieldNodeBase& Node) override;
+	CHAOS_API virtual void Evaluate(FFieldContext&, TFieldArrayView<FVector>& Results) const override;
+	CHAOS_API virtual bool operator==(const FFieldNodeBase& Node) override;
 
 	/** Serialization API */
 	virtual FFieldNodeBase::ESerializationType SerializationType() const override { return FFieldNodeBase::ESerializationType::FieldNode_FUniformVector; }
-	virtual void Serialize(FArchive& Ar) override;
+	CHAOS_API virtual void Serialize(FArchive& Ar) override;
 
 	/** Count the number of offsets/params that will be used by the world physics field */
-	virtual void FillSetupCount(int32& NumOffsets, int32& NumParams) const override;
+	CHAOS_API virtual void FillSetupCount(int32& NumOffsets, int32& NumParams) const override;
 
 	/** Fill the offsets/params arrays that will be used by the world physics field */
-	virtual void FillSetupDatas(TArray<int32>& NodesOffsets, TArray<float>& NodesParams,
+	CHAOS_API virtual void FillSetupDatas(TArray<int32>& NodesOffsets, TArray<float>& NodesParams,
 		const TMap<FFieldNodeBase*, float>& CommandTimes, const float PreviousTime) const override;
 
 	/** Evaluate the maximum magnitude of the field graph  */
-	virtual float EvalMaxMagnitude() const override;
+	CHAOS_API virtual float EvalMaxMagnitude() const override;
 
 	float Magnitude;
 	FVector Direction;
@@ -461,7 +461,7 @@ public:
 * RadialVector
 **/
 
-class CHAOS_API FRadialVector : public FFieldNode<FVector>
+class FRadialVector : public FFieldNode<FVector>
 {
 	typedef FFieldNode<FVector> Super;
 
@@ -476,22 +476,22 @@ public:
 	virtual FFieldNodeBase * NewCopy() const override { return new FRadialVector(Magnitude,Position); }
 	virtual ~FRadialVector() {}
 
-	virtual void Evaluate(FFieldContext&, TFieldArrayView<FVector>& Results) const override;
-	virtual bool operator==(const FFieldNodeBase& Node) override;
+	CHAOS_API virtual void Evaluate(FFieldContext&, TFieldArrayView<FVector>& Results) const override;
+	CHAOS_API virtual bool operator==(const FFieldNodeBase& Node) override;
 
 	/** Serialization API */
 	virtual FFieldNodeBase::ESerializationType SerializationType() const override { return FFieldNodeBase::ESerializationType::FieldNode_FRadialVector; }
-	virtual void Serialize(FArchive& Ar) override;
+	CHAOS_API virtual void Serialize(FArchive& Ar) override;
 
 	/** Count the number of offsets/params that will be used by the world physics field */
-	virtual void FillSetupCount(int32& NumOffsets, int32& NumParams) const override;
+	CHAOS_API virtual void FillSetupCount(int32& NumOffsets, int32& NumParams) const override;
 
 	/** Fill the offsets/params arrays that will be used by the world physics field */
-	virtual void FillSetupDatas(TArray<int32>& NodesOffsets, TArray<float>& NodesParams,
+	CHAOS_API virtual void FillSetupDatas(TArray<int32>& NodesOffsets, TArray<float>& NodesParams,
 		const TMap<FFieldNodeBase*, float>& CommandTimes, const float PreviousTime) const override;
 
 	/** Evaluate the maximum magnitude of the field graph  */
-	virtual float EvalMaxMagnitude() const override;
+	CHAOS_API virtual float EvalMaxMagnitude() const override;
 
 	float Magnitude;
 	FVector Position;
@@ -502,7 +502,7 @@ public:
 * RandomVector
 **/
 
-class CHAOS_API FRandomVector : public FFieldNode<FVector>
+class FRandomVector : public FFieldNode<FVector>
 {
 	typedef FFieldNode<FVector> Super;
 
@@ -515,22 +515,22 @@ public:
 	virtual FFieldNodeBase * NewCopy() const override { return new FRandomVector(Magnitude); }
 	virtual ~FRandomVector() {}
 
-	virtual void Evaluate(FFieldContext&, TFieldArrayView<FVector>& Results) const override;
-	virtual bool operator==(const FFieldNodeBase& Node) override;
+	CHAOS_API virtual void Evaluate(FFieldContext&, TFieldArrayView<FVector>& Results) const override;
+	CHAOS_API virtual bool operator==(const FFieldNodeBase& Node) override;
 
 	/** Serialization API */
 	virtual FFieldNodeBase::ESerializationType SerializationType() const override { return FFieldNodeBase::ESerializationType::FieldNode_FRandomVector; }
-	virtual void Serialize(FArchive& Ar) override;
+	CHAOS_API virtual void Serialize(FArchive& Ar) override;
 
 	/** Count the number of offsets/params that will be used by the world physics field */
-	virtual void FillSetupCount(int32& NumOffsets, int32& NumParams) const override;
+	CHAOS_API virtual void FillSetupCount(int32& NumOffsets, int32& NumParams) const override;
 
 	/** Fill the offsets/params arrays that will be used by the world physics field */
-	virtual void FillSetupDatas(TArray<int32>& NodesOffsets, TArray<float>& NodesParams,
+	CHAOS_API virtual void FillSetupDatas(TArray<int32>& NodesOffsets, TArray<float>& NodesParams,
 		const TMap<FFieldNodeBase*, float>& CommandTimes, const float PreviousTime) const override;
 
 	/** Evaluate the maximum magnitude of the field graph  */
-	virtual float EvalMaxMagnitude() const override;
+	CHAOS_API virtual float EvalMaxMagnitude() const override;
 
 	float Magnitude;
 };
@@ -539,7 +539,7 @@ public:
 * SumScalar
 **/
 
-class CHAOS_API FSumScalar : public FFieldNode<float>
+class FSumScalar : public FFieldNode<float>
 {
 	typedef FFieldNode<float> Super;
 
@@ -564,25 +564,25 @@ public:
 	}
 	virtual ~FSumScalar() {}
 
-	virtual void Evaluate(FFieldContext&, TFieldArrayView<float>& Results) const override;
-	virtual bool operator==(const FFieldNodeBase& Node) override;
+	CHAOS_API virtual void Evaluate(FFieldContext&, TFieldArrayView<float>& Results) const override;
+	CHAOS_API virtual bool operator==(const FFieldNodeBase& Node) override;
 
 	/** Serialization API */
 	virtual FFieldNodeBase::ESerializationType SerializationType() const override { return FFieldNodeBase::ESerializationType::FieldNode_FSumScalar; }
-	virtual void Serialize(FArchive& Ar) override;
+	CHAOS_API virtual void Serialize(FArchive& Ar) override;
 
 	/** Count the number of offsets/params that will be used by the world physics field */
-	virtual void FillSetupCount(int32& NumOffsets, int32& NumParams) const override;
+	CHAOS_API virtual void FillSetupCount(int32& NumOffsets, int32& NumParams) const override;
 
 	/** Fill the offsets/params arrays that will be used by the world physics field */
-	virtual void FillSetupDatas(TArray<int32>& NodesOffsets, TArray<float>& NodesParams,
+	CHAOS_API virtual void FillSetupDatas(TArray<int32>& NodesOffsets, TArray<float>& NodesParams,
 		const TMap<FFieldNodeBase*, float>& CommandTimes, const float PreviousTime) const override;
 
 	/** Evaluate the maximum magnitude of the field graph  */
-	virtual float EvalMaxMagnitude() const override;
+	CHAOS_API virtual float EvalMaxMagnitude() const override;
 
 	/** Compute the min/max spatial bounds of the field */
-	virtual void ComputeFieldBounds(FVector& MinBounds, FVector& MaxBounds) const override;
+	CHAOS_API virtual void ComputeFieldBounds(FVector& MinBounds, FVector& MaxBounds) const override;
 
 	float Magnitude;
 	TUniquePtr< FFieldNode<float> > ScalarRight;
@@ -594,7 +594,7 @@ public:
 * SumVector
 **/
 
-class CHAOS_API FSumVector : public FFieldNode<FVector>
+class FSumVector : public FFieldNode<FVector>
 {
 	typedef FFieldNode<FVector> Super;
 
@@ -622,25 +622,25 @@ public:
 	}
 	virtual ~FSumVector() {}
 
-	virtual void Evaluate(FFieldContext&, TFieldArrayView<FVector>& Results) const override;
-	virtual bool operator==(const FFieldNodeBase& Node) override;
+	CHAOS_API virtual void Evaluate(FFieldContext&, TFieldArrayView<FVector>& Results) const override;
+	CHAOS_API virtual bool operator==(const FFieldNodeBase& Node) override;
 
 	/** Serialization API */
 	virtual FFieldNodeBase::ESerializationType SerializationType() const override { return FFieldNodeBase::ESerializationType::FieldNode_FSumVector; }
-	virtual void Serialize(FArchive& Ar) override;
+	CHAOS_API virtual void Serialize(FArchive& Ar) override;
 
 	/** Count the number of offsets/params that will be used by the world physics field */
-	virtual void FillSetupCount(int32& NumOffsets, int32& NumParams) const override;
+	CHAOS_API virtual void FillSetupCount(int32& NumOffsets, int32& NumParams) const override;
 
 	/** Fill the offsets/params arrays that will be used by the world physics field */
-	virtual void FillSetupDatas(TArray<int32>& NodesOffsets, TArray<float>& NodesParams,
+	CHAOS_API virtual void FillSetupDatas(TArray<int32>& NodesOffsets, TArray<float>& NodesParams,
 		const TMap<FFieldNodeBase*, float>& CommandTimes, const float PreviousTime) const override;
 
 	/** Evaluate the maximum magnitude of the field graph  */
-	virtual float EvalMaxMagnitude() const override;
+	CHAOS_API virtual float EvalMaxMagnitude() const override;
 
 	/** Compute the min/max spatial bounds of the field */
-	virtual void ComputeFieldBounds(FVector& MinBounds, FVector& MaxBounds) const override;
+	CHAOS_API virtual void ComputeFieldBounds(FVector& MinBounds, FVector& MaxBounds) const override;
 
 	float Magnitude;
 	TUniquePtr< FFieldNode<float> > Scalar;
@@ -654,7 +654,7 @@ public:
 * FConversionField
 **/
 template<class InT, class OutT>
-class CHAOS_API FConversionField : public FFieldNode<OutT>
+class FConversionField : public FFieldNode<OutT>
 {
 	typedef FFieldNode<OutT> Super;
 
@@ -670,25 +670,25 @@ public:
 	}
 	virtual ~FConversionField() {}
 
-	virtual void Evaluate(FFieldContext&, TFieldArrayView<OutT>& Results) const override;
-	virtual bool operator==(const FFieldNodeBase& Node) override;
+	CHAOS_API virtual void Evaluate(FFieldContext&, TFieldArrayView<OutT>& Results) const override;
+	CHAOS_API virtual bool operator==(const FFieldNodeBase& Node) override;
 
 	/** Serialization API */
 	virtual FFieldNodeBase::ESerializationType SerializationType() const override { return FFieldNodeBase::ESerializationType::FieldNode_FConversionField; }
-	virtual void Serialize(FArchive& Ar) override;
+	CHAOS_API virtual void Serialize(FArchive& Ar) override;
 
 	/** Count the number of offsets/params that will be used by the world physics field */
-	virtual void FillSetupCount(int32& NumOffsets, int32& NumParams) const override;
+	CHAOS_API virtual void FillSetupCount(int32& NumOffsets, int32& NumParams) const override;
 
 	/** Fill the offsets/params arrays that will be used by the world physics field */
-	virtual void FillSetupDatas(TArray<int32>& NodesOffsets, TArray<float>& NodesParams,
+	CHAOS_API virtual void FillSetupDatas(TArray<int32>& NodesOffsets, TArray<float>& NodesParams,
 		const TMap<FFieldNodeBase*, float>& CommandTimes, const float PreviousTime) const override;
 
 	/** Evaluate the maximum magnitude of the field graph  */
-	virtual float EvalMaxMagnitude() const override;
+	CHAOS_API virtual float EvalMaxMagnitude() const override;
 
 	/** Compute the min/max spatial bounds of the field */
-	virtual void ComputeFieldBounds(FVector& MinBounds, FVector& MaxBounds) const override;
+	CHAOS_API virtual void ComputeFieldBounds(FVector& MinBounds, FVector& MaxBounds) const override;
 
 	TUniquePtr< FFieldNode<InT> > InputField;
 };
@@ -698,7 +698,7 @@ public:
 **/
 
 template<class T>
-class CHAOS_API FCullingField : public FFieldNode<T>
+class FCullingField : public FFieldNode<T>
 {
 	typedef FFieldNode<T> Super;
 
@@ -723,25 +723,25 @@ public:
 	}
 	virtual ~FCullingField() {}
 
-	virtual void Evaluate(FFieldContext&, TFieldArrayView<T>& Results) const override;
-	virtual bool operator==(const FFieldNodeBase& Node) override;
+	CHAOS_API virtual void Evaluate(FFieldContext&, TFieldArrayView<T>& Results) const override;
+	CHAOS_API virtual bool operator==(const FFieldNodeBase& Node) override;
 
 	/** Serialization API */
 	virtual FFieldNodeBase::ESerializationType SerializationType() const override { return FFieldNodeBase::ESerializationType::FieldNode_FCullingField; }
-	virtual void Serialize(FArchive& Ar) override;
+	CHAOS_API virtual void Serialize(FArchive& Ar) override;
 
 	/** Count the number of offsets/params that will be used by the world physics field */
-	virtual void FillSetupCount(int32& NumOffsets, int32& NumParams) const override;
+	CHAOS_API virtual void FillSetupCount(int32& NumOffsets, int32& NumParams) const override;
 
 	/** Fill the offsets/params arrays that will be used by the world physics field */
-	virtual void FillSetupDatas(TArray<int32>& NodesOffsets, TArray<float>& NodesParams,
+	CHAOS_API virtual void FillSetupDatas(TArray<int32>& NodesOffsets, TArray<float>& NodesParams,
 		const TMap<FFieldNodeBase*, float>& CommandTimes, const float PreviousTime) const override;
 
 	/** Evaluate the maximum magnitude of the field graph  */
-	virtual float EvalMaxMagnitude() const override;
+	CHAOS_API virtual float EvalMaxMagnitude() const override;
 
 	/** Compute the min/max spatial bounds of the field */
-	virtual void ComputeFieldBounds(FVector& MinBounds, FVector& MaxBounds) const override;
+	CHAOS_API virtual void ComputeFieldBounds(FVector& MinBounds, FVector& MaxBounds) const override;
 
 	TUniquePtr< FFieldNode<float> > Culling;
 	TUniquePtr< FFieldNode<T> > Input;
@@ -759,7 +759,7 @@ extern template class CHAOS_API FCullingField<FVector>;
 **/
 
 template<class T>
-class CHAOS_API FReturnResultsTerminal : public FFieldNode<T>
+class FReturnResultsTerminal : public FFieldNode<T>
 {
 	typedef FFieldNode<T> Super;
 
@@ -773,12 +773,12 @@ public:
 	}
 	virtual ~FReturnResultsTerminal() {}
 
-	void Evaluate(FFieldContext&, TFieldArrayView<T>& Results) const override;
-	virtual bool operator==(const FFieldNodeBase& Node);
+	CHAOS_API void Evaluate(FFieldContext&, TFieldArrayView<T>& Results) const override;
+	CHAOS_API virtual bool operator==(const FFieldNodeBase& Node);
 
 	/** Serialization API */
 	virtual FFieldNodeBase::ESerializationType SerializationType() const { return FFieldNodeBase::ESerializationType::FieldNode_FReturnResultsTerminal; }
-	virtual void Serialize(FArchive& Ar) override;
+	CHAOS_API virtual void Serialize(FArchive& Ar) override;
 };
 
 #if PLATFORM_MAC
diff --git a/Engine/Source/Runtime/Experimental/Chaos/Public/GeometryCollection/Facades/CollectionAnchoringFacade.h b/Engine/Source/Runtime/Experimental/Chaos/Public/GeometryCollection/Facades/CollectionAnchoringFacade.h
--- a/Engine/Source/Runtime/Experimental/Chaos/Public/GeometryCollection/Facades/CollectionAnchoringFacade.h
+++ b/Engine/Source/Runtime/Experimental/Chaos/Public/GeometryCollection/Facades/CollectionAnchoringFacade.h
@@ -11,17 +11,17 @@ namespace Chaos::Facades
 	/**
 	 * Provides an API to define anchoring properties on a collection
 	 */
-	class CHAOS_API FCollectionAnchoringFacade
+	class FCollectionAnchoringFacade
 	{
 	public:
-		FCollectionAnchoringFacade(FManagedArrayCollection& InCollection);
-		FCollectionAnchoringFacade(const FManagedArrayCollection& InCollection);
+		CHAOS_API FCollectionAnchoringFacade(FManagedArrayCollection& InCollection);
+		CHAOS_API FCollectionAnchoringFacade(const FManagedArrayCollection& InCollection);
 
 		/** Create the facade attributes. */
 		void DefineSchema() {}
 
 		/** Valid if all accessors arrays are available */
-		bool IsValid() const;
+		CHAOS_API bool IsValid() const;
 
 		/** Is the facade defined from a constant collection. */
 		bool IsConst() const { return InitialDynamicStateAttribute.IsConst(); }
@@ -30,19 +30,19 @@ namespace Chaos::Facades
 		//  Facade Functionality
 		//
 
-		bool HasInitialDynamicStateAttribute() const;
+		CHAOS_API bool HasInitialDynamicStateAttribute() const;
 		
-		EObjectStateType GetInitialDynamicState(int32 TransformIndex) const;
-		void SetInitialDynamicState(int32 TransformIndex, EObjectStateType State);
-		void SetInitialDynamicState(const TArray<int32>& TransformIndices, EObjectStateType State);
+		CHAOS_API EObjectStateType GetInitialDynamicState(int32 TransformIndex) const;
+		CHAOS_API void SetInitialDynamicState(int32 TransformIndex, EObjectStateType State);
+		CHAOS_API void SetInitialDynamicState(const TArray<int32>& TransformIndices, EObjectStateType State);
 		
-		bool HasAnchoredAttribute() const;
-		void AddAnchoredAttribute();
-		void CopyAnchoredAttribute(const FCollectionAnchoringFacade& Other);
+		CHAOS_API bool HasAnchoredAttribute() const;
+		CHAOS_API void AddAnchoredAttribute();
+		CHAOS_API void CopyAnchoredAttribute(const FCollectionAnchoringFacade& Other);
 
-		bool IsAnchored(int32 TransformIndex) const;
-		void SetAnchored(int32 TransformIndex, bool bValue);
-		void SetAnchored(const TArray<int32>& TransformIndices, bool bValue);
+		CHAOS_API bool IsAnchored(int32 TransformIndex) const;
+		CHAOS_API void SetAnchored(int32 TransformIndex, bool bValue);
+		CHAOS_API void SetAnchored(const TArray<int32>& TransformIndices, bool bValue);
 	
 	private:
 		/** Initial dynamic state of a transform ( can be changed at runtime ) */
@@ -51,4 +51,4 @@ namespace Chaos::Facades
 		/** Whether a transform is anchored ( orthogonal to the initial dynamic ) )*/
 		TManagedArrayAccessor<bool> AnchoredAttribute;
 	};
-}
\ No newline at end of file
+}
diff --git a/Engine/Source/Runtime/Experimental/Chaos/Public/GeometryCollection/Facades/CollectionBoundsFacade.h b/Engine/Source/Runtime/Experimental/Chaos/Public/GeometryCollection/Facades/CollectionBoundsFacade.h
--- a/Engine/Source/Runtime/Experimental/Chaos/Public/GeometryCollection/Facades/CollectionBoundsFacade.h
+++ b/Engine/Source/Runtime/Experimental/Chaos/Public/GeometryCollection/Facades/CollectionBoundsFacade.h
@@ -18,7 +18,7 @@ namespace GeometryCollection::Facades
 	* Defines common API for calculating the bounding box on a collection
 	* 
 	*/
-	class CHAOS_API FBoundsFacade
+	class FBoundsFacade
 	{
 	public:
 
@@ -26,32 +26,32 @@ namespace GeometryCollection::Facades
 		* FBoundsFacade Constuctor
 		* @param FManagedArrayCollection : Collection input
 		*/
-		FBoundsFacade(FManagedArrayCollection& InSelf);
-		FBoundsFacade(const FManagedArrayCollection& InSelf);
+		CHAOS_API FBoundsFacade(FManagedArrayCollection& InSelf);
+		CHAOS_API FBoundsFacade(const FManagedArrayCollection& InSelf);
 
 		/** Create the facade attributes. */
-		void DefineSchema();
+		CHAOS_API void DefineSchema();
 
 		/** Is the facade defined constant. */
 		bool IsConst() const { return BoundingBoxAttribute.IsConst(); }
 
 		/** Is the Facade defined on the collection? */
-		bool IsValid() const;		
+		CHAOS_API bool IsValid() const;		
 
 		/** UpdateBoundingBox */
-		void UpdateBoundingBox(bool bSkipCheck = false);
+		CHAOS_API void UpdateBoundingBox(bool bSkipCheck = false);
 
 		/** BoundingBox access */
 		const TManagedArray< FBox >& GetBoundingBoxes() const { return BoundingBoxAttribute.Get(); }
 
 		/** Centroids (Centers of BoundingBoxes) access */
-		TArray<FVector> GetCentroids() const;
+		CHAOS_API TArray<FVector> GetCentroids() const;
 
 		/** BoundingBox for the whole collection in Collection space */
-		FBox GetBoundingBoxInCollectionSpace() const;
+		CHAOS_API FBox GetBoundingBoxInCollectionSpace() const;
 
 		/** Returns the positions of the vertices of an FBox */
-		static TArray<FVector> GetBoundingBoxVertexPositions(const FBox& InBox);
+		static CHAOS_API TArray<FVector> GetBoundingBoxVertexPositions(const FBox& InBox);
 
 		/** TransformToGeometryIndex Access */
 		const TManagedArray< int32 >& GetTransformToGeometryIndex() const { return TransformToGeometryIndexAttribute.Get(); }
@@ -67,4 +67,4 @@ namespace GeometryCollection::Facades
 		TManagedArrayAccessor<int32>		ParentAttribute;
 	};
 
-}
\ No newline at end of file
+}
diff --git a/Engine/Source/Runtime/Experimental/Chaos/Public/GeometryCollection/Facades/CollectionConnectionGraphFacade.h b/Engine/Source/Runtime/Experimental/Chaos/Public/GeometryCollection/Facades/CollectionConnectionGraphFacade.h
--- a/Engine/Source/Runtime/Experimental/Chaos/Public/GeometryCollection/Facades/CollectionConnectionGraphFacade.h
+++ b/Engine/Source/Runtime/Experimental/Chaos/Public/GeometryCollection/Facades/CollectionConnectionGraphFacade.h
@@ -11,14 +11,14 @@ namespace GeometryCollection::Facades
 	/**
 	 * Provides an API for connection-graph related attributes
 	 */
-	class CHAOS_API FCollectionConnectionGraphFacade
+	class FCollectionConnectionGraphFacade
 	{
 	public:
-		FCollectionConnectionGraphFacade(FManagedArrayCollection& InCollection);
-		FCollectionConnectionGraphFacade(const FManagedArrayCollection& InCollection);
+		CHAOS_API FCollectionConnectionGraphFacade(FManagedArrayCollection& InCollection);
+		CHAOS_API FCollectionConnectionGraphFacade(const FManagedArrayCollection& InCollection);
 
 		/** Does the collection support the facade. */
-		bool IsValid() const;
+		CHAOS_API bool IsValid() const;
 
 		/** Is the facade defined constant. */
 		bool IsConst() const
@@ -27,42 +27,42 @@ namespace GeometryCollection::Facades
 		}
 
 		/** Create the facade attributes. */
-		void DefineSchema();
+		CHAOS_API void DefineSchema();
 
 		/** Remove the attributes */
-		void ClearAttributes();
+		CHAOS_API void ClearAttributes();
 
 		/** Connect two bones */
-		void ConnectWithContact(int32 BoneA, int32 BoneB, float ContactArea);
+		CHAOS_API void ConnectWithContact(int32 BoneA, int32 BoneB, float ContactArea);
 
 		/** Connect two bones */
-		void Connect(int32 BoneA, int32 BoneB);
+		CHAOS_API void Connect(int32 BoneA, int32 BoneB);
 
 		/** Enable or disable the Contact Area attribute */
-		void EnableContactAreas(bool bEnable, float DefaultContactArea = 1.0f);
+		CHAOS_API void EnableContactAreas(bool bEnable, float DefaultContactArea = 1.0f);
 
 		/** Reserve space for a number of additional connections */
-		void ReserveAdditionalConnections(int32 NumAdditionalConnections);
+		CHAOS_API void ReserveAdditionalConnections(int32 NumAdditionalConnections);
 
 		/**  Connections between bones that have the same parent in the hierarchy */
 		UE_DEPRECATED(5.3, "We have switched to an edge array connection representation. Please use the accessor functions (GetConnection, NumConnections, etc) to access the arrays of edge data.")
 		TManagedArrayAccessor<TSet<int32>> ConnectionsAttribute;
 
 		// Get the transform indices for the ConnectionIndex
-		TPair<int32, int32> GetConnection(int32 ConnectionIndex) const;
+		CHAOS_API TPair<int32, int32> GetConnection(int32 ConnectionIndex) const;
 		// Get the contact area for the ConnectionIndex
-		float GetConnectionContactArea(int32 ConnectionIndex) const;
+		CHAOS_API float GetConnectionContactArea(int32 ConnectionIndex) const;
 
-		bool HasContactAreas() const;
+		CHAOS_API bool HasContactAreas() const;
 
 		// Number of connection edges
-		int32 NumConnections() const;
+		CHAOS_API int32 NumConnections() const;
 
 		// Verifies the connections indices are valid indices into the Collection's Transform group
-		bool HasValidConnections() const;
+		CHAOS_API bool HasValidConnections() const;
 
 		// Remove all edge connections, but keep the connection attributes
-		void ResetConnections();
+		CHAOS_API void ResetConnections();
 
 	protected:
 		 TManagedArrayAccessor<int32> ConnectionEdgeStartAttribute;
@@ -75,4 +75,4 @@ namespace GeometryCollection::Facades
 #endif 
 	};
 
-}
\ No newline at end of file
+}
diff --git a/Engine/Source/Runtime/Experimental/Chaos/Public/GeometryCollection/Facades/CollectionConstraintOverrideFacade.h b/Engine/Source/Runtime/Experimental/Chaos/Public/GeometryCollection/Facades/CollectionConstraintOverrideFacade.h
--- a/Engine/Source/Runtime/Experimental/Chaos/Public/GeometryCollection/Facades/CollectionConstraintOverrideFacade.h
+++ b/Engine/Source/Runtime/Experimental/Chaos/Public/GeometryCollection/Facades/CollectionConstraintOverrideFacade.h
@@ -12,7 +12,7 @@ namespace GeometryCollection::Facades
 	/**
 	* Transient constraint candidate data.
 	*/
-	struct CHAOS_API FConstraintOverridesCandidateData
+	struct FConstraintOverridesCandidateData
 	{
 		FConstraintOverridesCandidateData()
 			: VertexIndex(INDEX_NONE)
@@ -26,7 +26,7 @@ namespace GeometryCollection::Facades
 	/**
 	* Transient constraint target data.
 	*/
-	struct CHAOS_API FConstraintOverridesTargetData
+	struct FConstraintOverridesTargetData
 	{
 		FConstraintOverridesTargetData()
 			: VertexIndex(INDEX_NONE)
@@ -39,23 +39,23 @@ namespace GeometryCollection::Facades
 	/**
 	* Transient constraint candidates.  Typically stored in the rest collection.
 	*/
-	class CHAOS_API FConstraintOverrideCandidateFacade
+	class FConstraintOverrideCandidateFacade
 	{
 	public:
-		static const FName GroupName;
-		static const FName TargetIndex;
-		static const FName BoneIndex;
+		static CHAOS_API const FName GroupName;
+		static CHAOS_API const FName TargetIndex;
+		static CHAOS_API const FName BoneIndex;
 
-		FConstraintOverrideCandidateFacade(FManagedArrayCollection& InCollection);
-		FConstraintOverrideCandidateFacade(const FManagedArrayCollection& InCollection);
+		CHAOS_API FConstraintOverrideCandidateFacade(FManagedArrayCollection& InCollection);
+		CHAOS_API FConstraintOverrideCandidateFacade(const FManagedArrayCollection& InCollection);
 
-		void DefineSchema();
+		CHAOS_API void DefineSchema();
 		bool IsConst() const { return Collection == nullptr; }
-		bool IsValid() const;
+		CHAOS_API bool IsValid() const;
 
-		int32 Add(FConstraintOverridesCandidateData& InputData);
-		void Clear();
-		FConstraintOverridesCandidateData Get(const int32 Index) const;
+		CHAOS_API int32 Add(FConstraintOverridesCandidateData& InputData);
+		CHAOS_API void Clear();
+		CHAOS_API FConstraintOverridesCandidateData Get(const int32 Index) const;
 		int32 Num() const { return TargetIndexAttribute.Num(); }
 
 	private:
@@ -69,25 +69,25 @@ namespace GeometryCollection::Facades
 	/**
 	* Transient constraint targets.  Typically stored in the simulation collection.
 	*/
-	class CHAOS_API FConstraintOverrideTargetFacade
+	class FConstraintOverrideTargetFacade
 	{
 	public:
-		static const FName GroupName;
-		static const FName TargetIndex;
-		static const FName TargetPosition;
+		static CHAOS_API const FName GroupName;
+		static CHAOS_API const FName TargetIndex;
+		static CHAOS_API const FName TargetPosition;
 
-		FConstraintOverrideTargetFacade(FManagedArrayCollection& InCollection);
-		FConstraintOverrideTargetFacade(const FManagedArrayCollection& InCollection);
+		CHAOS_API FConstraintOverrideTargetFacade(FManagedArrayCollection& InCollection);
+		CHAOS_API FConstraintOverrideTargetFacade(const FManagedArrayCollection& InCollection);
 
-		void DefineSchema();
+		CHAOS_API void DefineSchema();
 		bool IsConst() const { return Collection == nullptr; }
-		bool IsValid() const;
+		CHAOS_API bool IsValid() const;
 
-		int32 Add(FConstraintOverridesTargetData& InputData);
-		void Clear();
-		FConstraintOverridesTargetData Get(const int32 Index) const;
-		int32 GetIndex(const int32 Index) const;
-		const FVector3f& GetPosition(const int32 Index) const;
+		CHAOS_API int32 Add(FConstraintOverridesTargetData& InputData);
+		CHAOS_API void Clear();
+		CHAOS_API FConstraintOverridesTargetData Get(const int32 Index) const;
+		CHAOS_API int32 GetIndex(const int32 Index) const;
+		CHAOS_API const FVector3f& GetPosition(const int32 Index) const;
 		int32 Num() const { return TargetIndexAttribute.Num(); }
 
 	private:
@@ -100,4 +100,4 @@ namespace GeometryCollection::Facades
 
 
 
-} // namespace GeometryCollection::Facades
\ No newline at end of file
+} // namespace GeometryCollection::Facades
diff --git a/Engine/Source/Runtime/Experimental/Chaos/Public/GeometryCollection/Facades/CollectionExplodedVectorFacade.h b/Engine/Source/Runtime/Experimental/Chaos/Public/GeometryCollection/Facades/CollectionExplodedVectorFacade.h
--- a/Engine/Source/Runtime/Experimental/Chaos/Public/GeometryCollection/Facades/CollectionExplodedVectorFacade.h
+++ b/Engine/Source/Runtime/Experimental/Chaos/Public/GeometryCollection/Facades/CollectionExplodedVectorFacade.h
@@ -9,32 +9,32 @@
 namespace GeometryCollection::Facades
 {
 
-	class CHAOS_API FCollectionExplodedVectorFacade
+	class FCollectionExplodedVectorFacade
 	{
 	public:
-		FCollectionExplodedVectorFacade(FManagedArrayCollection& InCollection);
-		FCollectionExplodedVectorFacade(const FManagedArrayCollection& InCollection);
+		CHAOS_API FCollectionExplodedVectorFacade(FManagedArrayCollection& InCollection);
+		CHAOS_API FCollectionExplodedVectorFacade(const FManagedArrayCollection& InCollection);
 
 		/**
 		 * returns true if all the necessary attributes are present
 		 * if not then the API can be used to create
 		 */
-		bool IsValid() const;
+		CHAOS_API bool IsValid() const;
 
 		/**
 		 * Add the necessary attributes if they are missing
 		 */
-		void DefineSchema();
+		CHAOS_API void DefineSchema();
 
 		/**
 		 * Adds exploded vector to global matrices translation 
 		 */
-		void UpdateGlobalMatricesWithExplodedVectors(TArray<FMatrix>& InOutGlobalMatrices) const;
-		void UpdateGlobalMatricesWithExplodedVectors(TArray<FTransform>& InOutGlobalTransforms) const;
+		CHAOS_API void UpdateGlobalMatricesWithExplodedVectors(TArray<FMatrix>& InOutGlobalMatrices) const;
+		CHAOS_API void UpdateGlobalMatricesWithExplodedVectors(TArray<FTransform>& InOutGlobalTransforms) const;
 
 	private:
 
 		TManagedArrayAccessor<FVector3f> ExplodedVectorAttribute;
 	};
 
-}
\ No newline at end of file
+}
diff --git a/Engine/Source/Runtime/Experimental/Chaos/Public/GeometryCollection/Facades/CollectionHierarchyFacade.h b/Engine/Source/Runtime/Experimental/Chaos/Public/GeometryCollection/Facades/CollectionHierarchyFacade.h
--- a/Engine/Source/Runtime/Experimental/Chaos/Public/GeometryCollection/Facades/CollectionHierarchyFacade.h
+++ b/Engine/Source/Runtime/Experimental/Chaos/Public/GeometryCollection/Facades/CollectionHierarchyFacade.h
@@ -11,7 +11,7 @@ namespace Chaos::Facades
 	/**
 	 * Provides an API to read and manipulate hierarchy in a managed array collection
 	 */
-	class CHAOS_API FCollectionHierarchyFacade
+	class FCollectionHierarchyFacade
 	{
 	public:
 		enum class EPersistencePolicy : uint8
@@ -20,63 +20,63 @@ namespace Chaos::Facades
 			MustBePersistent
 		};
 
-		FCollectionHierarchyFacade(FManagedArrayCollection& InCollection);
-		FCollectionHierarchyFacade(const FManagedArrayCollection& InCollection);
+		CHAOS_API FCollectionHierarchyFacade(FManagedArrayCollection& InCollection);
+		CHAOS_API FCollectionHierarchyFacade(const FManagedArrayCollection& InCollection);
 
 		/** Create the facade attributes. */
 		void DefineSchema() {}
 
 		/** Valid if parent and children arrays are available */
-		bool IsValid() const;
+		CHAOS_API bool IsValid() const;
 
 		/** Is this facade const access */
 		bool IsConst() const { return ParentAttribute.IsConst(); }
 
 		/** whether the level attribute is available */
-		bool HasLevelAttribute() const;
+		CHAOS_API bool HasLevelAttribute() const;
 
 		/** whether the level attribute is persistent */
-		bool IsLevelAttributePersistent() const;
+		CHAOS_API bool IsLevelAttributePersistent() const;
 
 		/** Get the root index */
-		int32 GetRootIndex() const;
+		CHAOS_API int32 GetRootIndex() const;
 
 		/** Get the root indicies */
-		TArray<int32> GetRootIndices() const;
+		CHAOS_API TArray<int32> GetRootIndices() const;
 
 		/** Get direct children from a specific transform index. Assumes parent attribute is valid. */
-		TArray<int32> GetChildrenAsArray(int32 TransformIndex) const;
+		CHAOS_API TArray<int32> GetChildrenAsArray(int32 TransformIndex) const;
 
 		/** Get set of children from a transform index, or nullptr if not available */
-		const TSet<int32>* FindChildren(int32 TransformIndex) const;
+		CHAOS_API const TSet<int32>* FindChildren(int32 TransformIndex) const;
 
 		/** Get parent of a specific transform index */
-		int32 GetParent(int32 TransformIndex) const;
+		CHAOS_API int32 GetParent(int32 TransformIndex) const;
 
 		/** 
 		* Get initial level of a specific transform index 
 		* If the attribute is missing return INDEX_NONE
 		*/
-		int32 GetInitialLevel(int32 TransformIndex) const;
+		CHAOS_API int32 GetInitialLevel(int32 TransformIndex) const;
 
 		/** 
 		* Update level attribute for all elements (and create it if it is missing ) 
 		* @Param PersistencePolicy whether to make the attribute persistent or keep the existing state
 		*/ 
-		void GenerateLevelAttribute();
+		CHAOS_API void GenerateLevelAttribute();
 
 		/**
 		* Get transform indices in a depth first order 
 		*/
-		TArray<int32> GetTransformArrayInDepthFirstOrder() const;
+		CHAOS_API TArray<int32> GetTransformArrayInDepthFirstOrder() const;
 
 	public:
 		/** Get the root indicies */
-		static TArray<int32> GetRootIndices(const TManagedArrayAccessor<int32>& ParentAttribute);
+		static CHAOS_API TArray<int32> GetRootIndices(const TManagedArrayAccessor<int32>& ParentAttribute);
 
 	private:
 		TManagedArrayAccessor<int32>		ParentAttribute;
 		TManagedArrayAccessor<TSet<int32>>	ChildrenAttribute;
 		TManagedArrayAccessor<int32>		LevelAttribute;
 	};
-}
\ No newline at end of file
+}
diff --git a/Engine/Source/Runtime/Experimental/Chaos/Public/GeometryCollection/Facades/CollectionInstancedMeshFacade.h b/Engine/Source/Runtime/Experimental/Chaos/Public/GeometryCollection/Facades/CollectionInstancedMeshFacade.h
--- a/Engine/Source/Runtime/Experimental/Chaos/Public/GeometryCollection/Facades/CollectionInstancedMeshFacade.h
+++ b/Engine/Source/Runtime/Experimental/Chaos/Public/GeometryCollection/Facades/CollectionInstancedMeshFacade.h
@@ -11,32 +11,32 @@ namespace GeometryCollection::Facades
 	/**
 	 * Provides an API to read and manipulate hierarchy in a managed array collection
 	 */
-	class CHAOS_API FCollectionInstancedMeshFacade
+	class FCollectionInstancedMeshFacade
 	{
 	public:
-		FCollectionInstancedMeshFacade(FManagedArrayCollection& InCollection);
-		FCollectionInstancedMeshFacade(const FManagedArrayCollection& InCollection);
+		CHAOS_API FCollectionInstancedMeshFacade(FManagedArrayCollection& InCollection);
+		CHAOS_API FCollectionInstancedMeshFacade(const FManagedArrayCollection& InCollection);
 
 		/** Create the facade attributes. */
-		void DefineSchema();
+		CHAOS_API void DefineSchema();
 
 		/** Valid if parent and children arrays are available */
-		bool IsValid() const;
+		CHAOS_API bool IsValid() const;
 
 		/** Is the facade defined constant. */
-		bool IsConst() const;
+		CHAOS_API bool IsConst() const;
 
 		/** get the total number of indices */
-		int32 GetNumIndices() const;
+		CHAOS_API int32 GetNumIndices() const;
 
 		/** get the instance mesh index for specific transform index */
-		int32 GetIndex(int32 TransformIndex) const;
+		CHAOS_API int32 GetIndex(int32 TransformIndex) const;
 
 		/** set the instance mesh index for specific transform index */
-		void SetIndex(int32 TransformIndex, int32 InstanceMeshIndex);
+		CHAOS_API void SetIndex(int32 TransformIndex, int32 InstanceMeshIndex);
 
 
 	private:
 		TManagedArrayAccessor<int32> InstancedMeshIndexAttribute;
 	};
-}
\ No newline at end of file
+}
diff --git a/Engine/Source/Runtime/Experimental/Chaos/Public/GeometryCollection/Facades/CollectionKinematicBindingFacade.h b/Engine/Source/Runtime/Experimental/Chaos/Public/GeometryCollection/Facades/CollectionKinematicBindingFacade.h
--- a/Engine/Source/Runtime/Experimental/Chaos/Public/GeometryCollection/Facades/CollectionKinematicBindingFacade.h
+++ b/Engine/Source/Runtime/Experimental/Chaos/Public/GeometryCollection/Facades/CollectionKinematicBindingFacade.h
@@ -10,7 +10,7 @@
 namespace GeometryCollection::Facades
 {
 	/** Kinematic Facade */
-	class CHAOS_API FKinematicBindingFacade
+	class FKinematicBindingFacade
 	{
 	public:
 
@@ -19,31 +19,31 @@ namespace GeometryCollection::Facades
 		//
 		// Kinematics
 		//
-		static const FName KinematicGroup;
-		static const FName KinematicBoneBindingIndex;
-		static const FName KinematicBoneBindingToGroup;
+		static CHAOS_API const FName KinematicGroup;
+		static CHAOS_API const FName KinematicBoneBindingIndex;
+		static CHAOS_API const FName KinematicBoneBindingToGroup;
 
-		FKinematicBindingFacade(FManagedArrayCollection& InCollection);
-		FKinematicBindingFacade(const FManagedArrayCollection& InCollection);
+		CHAOS_API FKinematicBindingFacade(FManagedArrayCollection& InCollection);
+		CHAOS_API FKinematicBindingFacade(const FManagedArrayCollection& InCollection);
 
 		/** Create the facade attributes. */
-		void DefineSchema();
+		CHAOS_API void DefineSchema();
 
 		/** Is the facade defined constant. */
 		bool IsConst() const { return Collection==nullptr; }
 
 		/** Is the Facade defined on the collection? */
-		bool IsValid() const;
+		CHAOS_API bool IsValid() const;
 
 		//
 		//  Skeletal Mesh Bone Bindings
 		//
-		FBindingKey SetBoneBindings(const int32 BoneIndex, const TArray<int32>& Vertices, const TArray<float>& Weights);
-		void GetBoneBindings(const FBindingKey& Key, int32& OutBoneIndex, TArray<int32>& OutBoneVerts, TArray<float>& OutBoneWeights) const;
+		CHAOS_API FBindingKey SetBoneBindings(const int32 BoneIndex, const TArray<int32>& Vertices, const TArray<float>& Weights);
+		CHAOS_API void GetBoneBindings(const FBindingKey& Key, int32& OutBoneIndex, TArray<int32>& OutBoneVerts, TArray<float>& OutBoneWeights) const;
 
-		int32 AddKinematicBinding(const FBindingKey& Key);
+		CHAOS_API int32 AddKinematicBinding(const FBindingKey& Key);
 		int32 NumKinematicBindings() const { return KinemaitcBoneBindingAttribute.Num(); }
-		FBindingKey GetKinematicBindingKey(int Index) const;
+		CHAOS_API FBindingKey GetKinematicBindingKey(int Index) const;
 
 	private:
 		const FManagedArrayCollection& ConstCollection;
diff --git a/Engine/Source/Runtime/Experimental/Chaos/Public/GeometryCollection/Facades/CollectionMeshFacade.h b/Engine/Source/Runtime/Experimental/Chaos/Public/GeometryCollection/Facades/CollectionMeshFacade.h
--- a/Engine/Source/Runtime/Experimental/Chaos/Public/GeometryCollection/Facades/CollectionMeshFacade.h
+++ b/Engine/Source/Runtime/Experimental/Chaos/Public/GeometryCollection/Facades/CollectionMeshFacade.h
@@ -10,37 +10,37 @@
 namespace GeometryCollection::Facades
 {
 
-	class CHAOS_API FCollectionMeshFacade : public FCollectionUVFacade
+	class FCollectionMeshFacade : public FCollectionUVFacade
 	{
 	public:
-		FCollectionMeshFacade(FManagedArrayCollection& InCollection);
-		FCollectionMeshFacade(const FManagedArrayCollection& InCollection);
+		CHAOS_API FCollectionMeshFacade(FManagedArrayCollection& InCollection);
+		CHAOS_API FCollectionMeshFacade(const FManagedArrayCollection& InCollection);
 
 		/**
 		 * returns true if all the necessary attributes are present
 		 * if not then the API can be used to create
 		 */
-		bool IsValid() const;
+		CHAOS_API bool IsValid() const;
 
 		/**
 		 * Add the necessary attributes if they are missing
 		 */
-		void DefineSchema();
+		CHAOS_API void DefineSchema();
 
 		/**
 		 * Returns the vertex indicies for the bone
 		 */
-		const TArray<int32> GetVertexIndices(int32 BoneIdx) const;
+		CHAOS_API const TArray<int32> GetVertexIndices(int32 BoneIdx) const;
 
 		/**
 		 * Returns the vertex positions for the bone in bone space
 		 */
-		const TArrayView<const FVector3f> GetVertexPositions(int32 BoneIdx) const;
+		CHAOS_API const TArrayView<const FVector3f> GetVertexPositions(int32 BoneIdx) const;
 
 		/**
 		 * Returns the face indices for the bone
 		 */
-		const TArray<int32> GetFaceIndices(int32 BoneIdx) const;
+		CHAOS_API const TArray<int32> GetFaceIndices(int32 BoneIdx) const;
 
 		/**
 		 * Returns the vertex indices of the face for the bone
@@ -50,7 +50,7 @@ namespace GeometryCollection::Facades
 		/**
 		 * Bakes the transforms into the vertex positions and sets the bone transforms to identity
 		 */
-		 void BakeTransform(int32 TransformIdx, const FTransform& InTransform);
+		 CHAOS_API void BakeTransform(int32 TransformIdx, const FTransform& InTransform);
 
 		TManagedArrayAccessor<int32> TransformToGeometryIndexAttribute;
 		TManagedArrayAccessor<int32> TransformIndexAttribute;
@@ -71,4 +71,4 @@ namespace GeometryCollection::Facades
 		TManagedArrayAccessor<int32> FaceCountAttribute;
 	};
 
-}
\ No newline at end of file
+}
diff --git a/Engine/Source/Runtime/Experimental/Chaos/Public/GeometryCollection/Facades/CollectionPositionTargetFacade.h b/Engine/Source/Runtime/Experimental/Chaos/Public/GeometryCollection/Facades/CollectionPositionTargetFacade.h
--- a/Engine/Source/Runtime/Experimental/Chaos/Public/GeometryCollection/Facades/CollectionPositionTargetFacade.h
+++ b/Engine/Source/Runtime/Experimental/Chaos/Public/GeometryCollection/Facades/CollectionPositionTargetFacade.h
@@ -9,7 +9,7 @@
 
 namespace GeometryCollection::Facades
 {
-	struct CHAOS_API FPositionTargetsData
+	struct FPositionTargetsData
 	{
 		TArray<int32> TargetIndex;
 		TArray<int32> SourceIndex;
@@ -23,7 +23,7 @@ namespace GeometryCollection::Facades
 	};
 
 	/** Kinematic Facade */
-	class CHAOS_API FPositionTargetFacade
+	class FPositionTargetFacade
 	{
 	public:
 
@@ -32,33 +32,33 @@ namespace GeometryCollection::Facades
 		//
 		// Kinematics
 		//
-		static const FName GroupName;
-		static const FName TargetIndex;
-		static const FName SourceIndex;
-		static const FName Stiffness;
-		static const FName Damping;
-		static const FName SourceName;
-		static const FName TargetName;
-		static const FName TargetWeights;
-		static const FName SourceWeights;
+		static CHAOS_API const FName GroupName;
+		static CHAOS_API const FName TargetIndex;
+		static CHAOS_API const FName SourceIndex;
+		static CHAOS_API const FName Stiffness;
+		static CHAOS_API const FName Damping;
+		static CHAOS_API const FName SourceName;
+		static CHAOS_API const FName TargetName;
+		static CHAOS_API const FName TargetWeights;
+		static CHAOS_API const FName SourceWeights;
 
-		FPositionTargetFacade(FManagedArrayCollection& InCollection);
-		FPositionTargetFacade(const FManagedArrayCollection& InCollection);
+		CHAOS_API FPositionTargetFacade(FManagedArrayCollection& InCollection);
+		CHAOS_API FPositionTargetFacade(const FManagedArrayCollection& InCollection);
 
 		/** Create the facade attributes. */
-		void DefineSchema();
+		CHAOS_API void DefineSchema();
 
 		/** Is the facade defined constant. */
 		bool IsConst() const { return Collection == nullptr; }
 
 		/** Is the Facade defined on the collection? */
-		bool IsValid() const;
+		CHAOS_API bool IsValid() const;
 
 		//
 		//  Skeletal Mesh Bone Bindings
 		//
-		int32 AddPositionTarget(const FPositionTargetsData& InputData);
-		FPositionTargetsData GetPositionTarget(const int32 DataIndex) const;
+		CHAOS_API int32 AddPositionTarget(const FPositionTargetsData& InputData);
+		CHAOS_API FPositionTargetsData GetPositionTarget(const int32 DataIndex) const;
 		int32 NumPositionTargets() const { return TargetIndexAttribute.Num(); }
 
 	private:
diff --git a/Engine/Source/Runtime/Experimental/Chaos/Public/GeometryCollection/Facades/CollectionRemoveOnBreakFacade.h b/Engine/Source/Runtime/Experimental/Chaos/Public/GeometryCollection/Facades/CollectionRemoveOnBreakFacade.h
--- a/Engine/Source/Runtime/Experimental/Chaos/Public/GeometryCollection/Facades/CollectionRemoveOnBreakFacade.h
+++ b/Engine/Source/Runtime/Experimental/Chaos/Public/GeometryCollection/Facades/CollectionRemoveOnBreakFacade.h
@@ -12,26 +12,26 @@ namespace GeometryCollection::Facades
 	/**
 	 * this class helps reading and writing the packed data used in the managed array
 	 */
-	struct CHAOS_API FRemoveOnBreakData
+	struct FRemoveOnBreakData
 	{
 	public:
 		inline static const FVector4f DisabledPackedData{ -1, 0, 0, 0 };
 
-		FRemoveOnBreakData();
-		FRemoveOnBreakData(const FVector4f& InPackedData);
-		FRemoveOnBreakData(bool bEnable, const FVector2f& BreakTimer, bool bClusterCrumbling, const FVector2f& RemovalTimer);
+		CHAOS_API FRemoveOnBreakData();
+		CHAOS_API FRemoveOnBreakData(const FVector4f& InPackedData);
+		CHAOS_API FRemoveOnBreakData(bool bEnable, const FVector2f& BreakTimer, bool bClusterCrumbling, const FVector2f& RemovalTimer);
 
-		const FVector4f& GetPackedData() const;
+		CHAOS_API const FVector4f& GetPackedData() const;
 
-		bool IsEnabled() const;
-		bool GetClusterCrumbling() const;
-		FVector2f GetBreakTimer() const;
-		FVector2f GetRemovalTimer() const;
+		CHAOS_API bool IsEnabled() const;
+		CHAOS_API bool GetClusterCrumbling() const;
+		CHAOS_API FVector2f GetBreakTimer() const;
+		CHAOS_API FVector2f GetRemovalTimer() const;
 		
-		void SetEnabled(bool bEnable);
-		void SetClusterCrumbling(bool bClusterCrumbling);
-		void SetBreakTimer(float MinTime, float MaxTime);
-		void SetRemovalTimer(float MinTime, float MaxTime);
+		CHAOS_API void SetEnabled(bool bEnable);
+		CHAOS_API void SetClusterCrumbling(bool bClusterCrumbling);
+		CHAOS_API void SetBreakTimer(float MinTime, float MaxTime);
+		CHAOS_API void SetRemovalTimer(float MinTime, float MaxTime);
 
 	private:
 		FVector4f PackedData;
@@ -40,34 +40,34 @@ namespace GeometryCollection::Facades
 	/**
 	 * Provide an API to set and get removal on break data on a collection
 	 */
-	class CHAOS_API FCollectionRemoveOnBreakFacade
+	class FCollectionRemoveOnBreakFacade
 	{
 	public:
-		FCollectionRemoveOnBreakFacade(FManagedArrayCollection& InCollection);
-		FCollectionRemoveOnBreakFacade(const FManagedArrayCollection& InCollection);
+		CHAOS_API FCollectionRemoveOnBreakFacade(FManagedArrayCollection& InCollection);
+		CHAOS_API FCollectionRemoveOnBreakFacade(const FManagedArrayCollection& InCollection);
 
 		/** Create the facade attributes. */
-		void DefineSchema();
+		CHAOS_API void DefineSchema();
 
 		/** remove the facade attributes */
-		void RemoveSchema();
+		CHAOS_API void RemoveSchema();
 
 		/** Valid if related attributes are available */
-		bool IsValid() const;
+		CHAOS_API bool IsValid() const;
 
 		/** Is this facade const access */
-		bool IsConst() const;
+		CHAOS_API bool IsConst() const;
 
 		/** set a specific value to a array of transforms */
-		void SetFromIndexArray(const TArray<int32>& Indices, const FRemoveOnBreakData& Data);
+		CHAOS_API void SetFromIndexArray(const TArray<int32>& Indices, const FRemoveOnBreakData& Data);
 
 		/** set a specific value to all transforms */
-		void SetToAll(const FRemoveOnBreakData& Data);
+		CHAOS_API void SetToAll(const FRemoveOnBreakData& Data);
 
 		/** Get a specific value by index */
-		const FRemoveOnBreakData GetData(int32 Index) const;
+		CHAOS_API const FRemoveOnBreakData GetData(int32 Index) const;
 
 	private:
 		TManagedArrayAccessor<FVector4f> RemoveOnBreakAttribute;
 	};
-}
\ No newline at end of file
+}
