[Chaos]

diff --git a/Engine/Source/Runtime/Experimental/Chaos/Public/GeometryCollection/Facades/CollectionRenderingFacade.h b/Engine/Source/Runtime/Experimental/Chaos/Public/GeometryCollection/Facades/CollectionRenderingFacade.h
--- a/Engine/Source/Runtime/Experimental/Chaos/Public/GeometryCollection/Facades/CollectionRenderingFacade.h
+++ b/Engine/Source/Runtime/Experimental/Chaos/Public/GeometryCollection/Facades/CollectionRenderingFacade.h
@@ -18,7 +18,7 @@ namespace GeometryCollection::Facades
 	* Defines common API for storing rendering data.
 	*
 	*/
-	class CHAOS_API FRenderingFacade
+	class FRenderingFacade
 	{
 	public:
 		typedef FGeometryCollectionSection FTriangleSection;
@@ -27,33 +27,33 @@ namespace GeometryCollection::Facades
 		* FRenderingFacade Constuctor
 		* @param VertexDependencyGroup : GroupName the index attribute is dependent on.
 		*/
-		FRenderingFacade(FManagedArrayCollection& InSelf);
-		FRenderingFacade(const FManagedArrayCollection& InSelf);
+		CHAOS_API FRenderingFacade(FManagedArrayCollection& InSelf);
+		CHAOS_API FRenderingFacade(const FManagedArrayCollection& InSelf);
 
 		/**Create the facade.*/
-		void DefineSchema();
+		CHAOS_API void DefineSchema();
 
 		/** Is the facade defined constant. */
 		bool IsConst() const { return Collection==nullptr; }
 
 		/**Is the Facade defined on the collection?*/
-		bool IsValid() const;
+		CHAOS_API bool IsValid() const;
 
 		/**Does it support rendering surfaces.*/
-		bool CanRenderSurface() const;
+		CHAOS_API bool CanRenderSurface() const;
 
 		//
 		// Facade API
 		//
 
 		/**Number of triangles to render.*/
-		int32 NumTriangles() const;
+		CHAOS_API int32 NumTriangles() const;
 
 		/**Add a triangle to the rendering view.*/
-		void AddTriangle(const Chaos::FTriangle& InTriangle);
+		CHAOS_API void AddTriangle(const Chaos::FTriangle& InTriangle);
 
 		/** Add a surface to the rendering view.*/
-		void AddSurface(TArray<FVector3f>&& InVertices, TArray<FIntVector>&& InIndices, TArray<FVector3f>&& InNormals, TArray<FLinearColor>&& InColors);
+		CHAOS_API void AddSurface(TArray<FVector3f>&& InVertices, TArray<FIntVector>&& InIndices, TArray<FVector3f>&& InNormals, TArray<FLinearColor>&& InColors);
 		
 		/** GetIndices */
 		const TManagedArray< FIntVector >& GetIndices() const { return IndicesAttribute.Get(); }
@@ -65,7 +65,7 @@ namespace GeometryCollection::Facades
 		const TManagedArray< FTriangleSection >& GetTriangleSections() const { return TriangleSectionAttribute.Get(); }
 
 		/** BuildMeshSections */
-		TArray<FTriangleSection> BuildMeshSections(const TArray<FIntVector>& Indices, TArray<int32> BaseMeshOriginalIndicesIndex, TArray<FIntVector>& RetIndices) const;
+		CHAOS_API TArray<FTriangleSection> BuildMeshSections(const TArray<FIntVector>& Indices, TArray<int32> BaseMeshOriginalIndicesIndex, TArray<FIntVector>& RetIndices) const;
 
 
 		//
@@ -98,10 +98,10 @@ namespace GeometryCollection::Facades
 		//
 
 		/** Geometry Group Start : */
-		int32 StartGeometryGroup(FString InName);
+		CHAOS_API int32 StartGeometryGroup(FString InName);
 
 		/** Geometry Group End : */
-		void EndGeometryGroup(int32 InGeometryGroupIndex);
+		CHAOS_API void EndGeometryGroup(int32 InGeometryGroupIndex);
 
 		int32 NumGeometry() const { return GeometryNameAttribute.Num(); }
 
@@ -129,10 +129,10 @@ namespace GeometryCollection::Facades
 		      TManagedArray< int32 >& ModifySelectionState() { check(!IsConst()); return GeometrySelectionAttribute.Modify(); }
 
 		/** NumVerticesOnSelectedGeometry */
-		int32 NumVerticesOnSelectedGeometry() const;
+		CHAOS_API int32 NumVerticesOnSelectedGeometry() const;
 
 		/** GetGeometryNameToIndexMap */
-		FStringIntMap GetGeometryNameToIndexMap() const;
+		CHAOS_API FStringIntMap GetGeometryNameToIndexMap() const;
 
 	private : 
 		const FManagedArrayCollection& ConstCollection;
@@ -158,4 +158,4 @@ namespace GeometryCollection::Facades
 		TManagedArrayAccessor<int32> GeometrySelectionAttribute;
 	};
 
-}
\ No newline at end of file
+}
diff --git a/Engine/Source/Runtime/Experimental/Chaos/Public/GeometryCollection/Facades/CollectionSelectionFacade.h b/Engine/Source/Runtime/Experimental/Chaos/Public/GeometryCollection/Facades/CollectionSelectionFacade.h
--- a/Engine/Source/Runtime/Experimental/Chaos/Public/GeometryCollection/Facades/CollectionSelectionFacade.h
+++ b/Engine/Source/Runtime/Experimental/Chaos/Public/GeometryCollection/Facades/CollectionSelectionFacade.h
@@ -37,25 +37,25 @@ namespace GeometryCollection::Facades
 	*	- FindAttribute<int32>(FVertexSetInterface::BoneIndexAttribute, <Group>);
 	* 
 	*/
-	class CHAOS_API FSelectionFacade
+	class FSelectionFacade
  	{
-		void InitUnboundedGroup(FName GroupName, FName DependencyGroup);
-		void InitWeightedUnboundedGroup(FName GroupName, FName DependencyGroup);
-		void InitBoundedGroup(FName GroupName, FName DependencyGroup, FName BoneDependencyGroup);
-		void InitWeightedBoundedGroup(FName GroupName, FName DependencyGroup, FName BoneDependencyGroup);
+		CHAOS_API void InitUnboundedGroup(FName GroupName, FName DependencyGroup);
+		CHAOS_API void InitWeightedUnboundedGroup(FName GroupName, FName DependencyGroup);
+		CHAOS_API void InitBoundedGroup(FName GroupName, FName DependencyGroup, FName BoneDependencyGroup);
+		CHAOS_API void InitWeightedBoundedGroup(FName GroupName, FName DependencyGroup, FName BoneDependencyGroup);
 
 	public:
 
 		// groups
-		static const FName UnboundGroup;
-		static const FName WeightedUnboundGroup;
-		static const FName BoundGroup;
-		static const FName WeightedBoundGroup;
+		static CHAOS_API const FName UnboundGroup;
+		static CHAOS_API const FName WeightedUnboundGroup;
+		static CHAOS_API const FName BoundGroup;
+		static CHAOS_API const FName WeightedBoundGroup;
 
 		// Attributes
-		static const FName IndexAttribute;
-		static const FName WeightAttribute;
-		static const FName BoneIndexAttribute;
+		static CHAOS_API const FName IndexAttribute;
+		static CHAOS_API const FName WeightAttribute;
+		static CHAOS_API const FName BoneIndexAttribute;
 
 		struct FSelectionKey {
 			FSelectionKey(int32 InIndex = INDEX_NONE, FName InGroupName = "") :
@@ -67,8 +67,8 @@ namespace GeometryCollection::Facades
 		/**
 		* FSelectionFacade Constuctor
 		*/
-		FSelectionFacade(FManagedArrayCollection& InSelf);
-		FSelectionFacade(const FManagedArrayCollection& InSelf);
+		CHAOS_API FSelectionFacade(FManagedArrayCollection& InSelf);
+		CHAOS_API FSelectionFacade(const FManagedArrayCollection& InSelf);
 
 		/** Is the facade defined constant. */
 		bool IsConst() const { return Collection==nullptr; }
@@ -76,29 +76,29 @@ namespace GeometryCollection::Facades
 		/**
 		* Add the indices to the FVertexSetInterface::UnboundGroup 
 		*/
-		FSelectionKey AddSelection(const TArray<int32>& Indices, FName DependencyGroup);
+		CHAOS_API FSelectionKey AddSelection(const TArray<int32>& Indices, FName DependencyGroup);
 
 		/**
 		* Add the indices to the FVertexSetInterface::WeightedUnboundGroup
 		*/
-		FSelectionKey AddSelection(const TArray<int32>& Indices, const TArray<float>& Weights, FName DependencyGroup);
+		CHAOS_API FSelectionKey AddSelection(const TArray<int32>& Indices, const TArray<float>& Weights, FName DependencyGroup);
 
 		/**
 		* Add the indices to the FVertexSetInterface::BoundGroup
 		*/
-		FSelectionKey AddSelection(const int32 BoneIndex, const TArray<int32>& Indices, FName DependencyGroup, FName BoneDependencyGroup = FName(""));
+		CHAOS_API FSelectionKey AddSelection(const int32 BoneIndex, const TArray<int32>& Indices, FName DependencyGroup, FName BoneDependencyGroup = FName(""));
 
 		/**
 		* Add the indices to the FVertexSetInterface::BoundGroup
 		*/
-		FSelectionKey AddSelection(const int32 BoneIndex, const TArray<int32>& Indices, const TArray<float>& Weights, FName DependencyGroup, FName BoneDependencyGroup = FName(""));
+		CHAOS_API FSelectionKey AddSelection(const int32 BoneIndex, const TArray<int32>& Indices, const TArray<float>& Weights, FName DependencyGroup, FName BoneDependencyGroup = FName(""));
 
 		/**
 		* Return the vertex list from the given key 
 		* @param Key : <GroupName and Index>
 		* @param Indices : Return indices, empty if not found. 
 		*/
-		void GetSelection(const FSelectionKey& Key, TArray<int32>& Indices) const;
+		CHAOS_API void GetSelection(const FSelectionKey& Key, TArray<int32>& Indices) const;
 
 		/**
 		* Return the vertex list, bone index, and weights from the given key
@@ -107,7 +107,7 @@ namespace GeometryCollection::Facades
 		* @param Indices : Return indices, empty if not found.
 		* @param Weights : Return vertex weights, empty if not found.
 		*/
-		void GetSelection(const FSelectionKey& Key, TArray<int32>& Indices, TArray<float>& Weights) const;
+		CHAOS_API void GetSelection(const FSelectionKey& Key, TArray<int32>& Indices, TArray<float>& Weights) const;
 
 		/**
 		* Return the vertex list and bone index from the given key
@@ -115,7 +115,7 @@ namespace GeometryCollection::Facades
 		* @param BoneIndex : Return BoneIndex, INDEX_NONE if not found.
 		* @param Indices : Return indices, empty if not found.
 		*/
-		void GetSelection(const FSelectionKey& Key, int32& BoneIndex, TArray<int32>& Indices) const;
+		CHAOS_API void GetSelection(const FSelectionKey& Key, int32& BoneIndex, TArray<int32>& Indices) const;
 
 		/**
 		* Return the vertex list, bone index, and weights from the given key
@@ -124,7 +124,7 @@ namespace GeometryCollection::Facades
 		* @param Indices : Return indices, empty if not found.
 		* @param Weights : Return vertex weights, empty if not found.
 		*/
-		void GetSelection(const FSelectionKey& Key, int32& BoneIndex, TArray<int32>& Indices, TArray<float>& Weights) const;
+		CHAOS_API void GetSelection(const FSelectionKey& Key, int32& BoneIndex, TArray<int32>& Indices, TArray<float>& Weights) const;
 
 	private:
 
@@ -134,4 +134,4 @@ namespace GeometryCollection::Facades
 		FManagedArrayCollection* Collection = nullptr;
 	};
 
-}
\ No newline at end of file
+}
diff --git a/Engine/Source/Runtime/Experimental/Chaos/Public/GeometryCollection/Facades/CollectionTetrahedralBindingsFacade.h b/Engine/Source/Runtime/Experimental/Chaos/Public/GeometryCollection/Facades/CollectionTetrahedralBindingsFacade.h
--- a/Engine/Source/Runtime/Experimental/Chaos/Public/GeometryCollection/Facades/CollectionTetrahedralBindingsFacade.h
+++ b/Engine/Source/Runtime/Experimental/Chaos/Public/GeometryCollection/Facades/CollectionTetrahedralBindingsFacade.h
@@ -15,40 +15,40 @@ namespace GeometryCollection::Facades
 	* Interface for storing and retrieving bindings of surfaces (typically SkeletalMesh or StaticMesh) to
 	* tetrahedral meshes.  Bindings data for each surface is grouped by a mesh id and a level of detail.
 	*/
-	class CHAOS_API FTetrahedralBindings
+	class FTetrahedralBindings
  	{
 	public:
 
 		// groups
-		static const FName MeshBindingsGroupName;
+		static CHAOS_API const FName MeshBindingsGroupName;
 
 		//
 		// Attributes
 		//
 
-		static const FName MeshIdAttributeName;
+		static CHAOS_API const FName MeshIdAttributeName;
 
 		//! Tet or Tri vertex indices.
-		static const FName ParentsAttributeName;
+		static CHAOS_API const FName ParentsAttributeName;
 		//! Barycentric weight of each tet/tri vertex.
-		static const FName WeightsAttributeName;
+		static CHAOS_API const FName WeightsAttributeName;
 		//! Offset vector from barycentric tri position.
-		static const FName OffsetsAttributeName;
+		static CHAOS_API const FName OffsetsAttributeName;
 		//! Per vertex amount for deformer masking.
-		static const FName MaskAttributeName;
+		static CHAOS_API const FName MaskAttributeName;
 
 		/**
 		* FSelectionFacade Constuctor
 		* @param VertixDependencyGroup : GroupName the index attribute is dependent on. 
 		*/
-		FTetrahedralBindings(FManagedArrayCollection& InSelf);
-		FTetrahedralBindings(const FManagedArrayCollection& InSelf);
-		virtual ~FTetrahedralBindings();
+		CHAOS_API FTetrahedralBindings(FManagedArrayCollection& InSelf);
+		CHAOS_API FTetrahedralBindings(const FManagedArrayCollection& InSelf);
+		CHAOS_API virtual ~FTetrahedralBindings();
 
 		/** 
 		* Create the facade schema. 
 		*/
-		void DefineSchema();
+		CHAOS_API void DefineSchema();
 
 		/** Returns \c true if the facade is operating on a read-only geometry collection. */
 		bool IsConst() const { return MeshIdAttribute.IsConst(); }
@@ -57,44 +57,44 @@ namespace GeometryCollection::Facades
 		* Returns \c true if the Facade defined on the collection, and is initialized to
 		* a valid bindings group.
 		*/
-		bool IsValid() const;
+		CHAOS_API bool IsValid() const;
 
 		/**
 		* Given a \p MeshId (by convention \code Mesh->GetPrimaryAssetId() \endcode) and
 		* a Level Of Detail rank, generate the associated bindings group name.
 		*/
-		static FName GenerateMeshGroupName(const int32 TetMeshIdx, const FName& MeshId, const int32 LOD);
+		static CHAOS_API FName GenerateMeshGroupName(const int32 TetMeshIdx, const FName& MeshId, const int32 LOD);
 
 		/**
 		* For a given \p MeshId and \p LOD, return the associated tetrahedral mesh index.
 		*/
-		int32 GetTetMeshIndex(const FName& MeshId, const int32 LOD) const;
+		CHAOS_API int32 GetTetMeshIndex(const FName& MeshId, const int32 LOD) const;
 
 		/**
 		* Returns \c true if the specified bindings group exists.
 		*/
-		bool ContainsBindingsGroup(const int32 TetMeshIdx, const FName& MeshId, const int32 LOD) const;
-		bool ContainsBindingsGroup(const FName& GroupName) const;
+		CHAOS_API bool ContainsBindingsGroup(const int32 TetMeshIdx, const FName& MeshId, const int32 LOD) const;
+		CHAOS_API bool ContainsBindingsGroup(const FName& GroupName) const;
 
 		/**
 		* Create a new bindings group, allocating new arrays.
 		*/
-		void AddBindingsGroup(const int32 TetMeshIdx, const FName& MeshId, const int32 LOD);
-		void AddBindingsGroup(const FName& GroupName);
+		CHAOS_API void AddBindingsGroup(const int32 TetMeshIdx, const FName& MeshId, const int32 LOD);
+		CHAOS_API void AddBindingsGroup(const FName& GroupName);
 		
 		/**
 		* Initialize local arrays to point at a bindings group associated with \p MeshId 
 		* and \p LOD.  Returns \c false if it doesn't exist.
 		*/
-		bool ReadBindingsGroup(const int32 TetMeshIdx, const FName& MeshId, const int32 LOD);
-		bool ReadBindingsGroup(const FName& GroupName);
+		CHAOS_API bool ReadBindingsGroup(const int32 TetMeshIdx, const FName& MeshId, const int32 LOD);
+		CHAOS_API bool ReadBindingsGroup(const FName& GroupName);
 
 		/**
 		* Removes a group from the list of bindings groups, removes the bindings arrays 
 		* from the geometry collection, and removes the group if it's empty.
 		*/
-		void RemoveBindingsGroup(const int32 TetMeshIdx, const FName& MeshId, const int32 LOD);
-		void RemoveBindingsGroup(const FName& GroupName);
+		CHAOS_API void RemoveBindingsGroup(const int32 TetMeshIdx, const FName& MeshId, const int32 LOD);
+		CHAOS_API void RemoveBindingsGroup(const FName& GroupName);
 
 		/**
 		* Authors bindings data.
@@ -107,7 +107,7 @@ namespace GeometryCollection::Facades
 		* \p Mask are per-vertex multipliers on the deformer, 0 for no deformation, 1.0 for 
 		*    full deformation.
 		*/
-		void SetBindingsData(const TArray<FIntVector4>& ParentsIn, const TArray<FVector4f>& WeightsIn, const TArray<FVector3f>& OffsetsIn, const TArray<float>& MaskIn);
+		CHAOS_API void SetBindingsData(const TArray<FIntVector4>& ParentsIn, const TArray<FVector4f>& WeightsIn, const TArray<FVector3f>& OffsetsIn, const TArray<float>& MaskIn);
 		void SetBindingsData(const TArray<FIntVector4>& ParentsIn, const TArray<FVector4f>& WeightsIn, const TArray<FVector3f>& OffsetsIn)
 		{
 			TArray<float> MaskTmp; MaskTmp.SetNum(ParentsIn.Num());
diff --git a/Engine/Source/Runtime/Experimental/Chaos/Public/GeometryCollection/Facades/CollectionTetrahedralMetricsFacade.h b/Engine/Source/Runtime/Experimental/Chaos/Public/GeometryCollection/Facades/CollectionTetrahedralMetricsFacade.h
--- a/Engine/Source/Runtime/Experimental/Chaos/Public/GeometryCollection/Facades/CollectionTetrahedralMetricsFacade.h
+++ b/Engine/Source/Runtime/Experimental/Chaos/Public/GeometryCollection/Facades/CollectionTetrahedralMetricsFacade.h
@@ -10,22 +10,22 @@ namespace GeometryCollection::Facades
 	/**
 	*
 	*/
-	class CHAOS_API FTetrahedralMetrics
+	class FTetrahedralMetrics
 	{
 	public:
 		// Groups
 
 		// Attributes
-		static const FName SignedVolumeAttributeName;
-		static const FName AspectRatioAttributeName;
+		static CHAOS_API const FName SignedVolumeAttributeName;
+		static CHAOS_API const FName AspectRatioAttributeName;
 
-		FTetrahedralMetrics(FManagedArrayCollection& InCollection);
-		FTetrahedralMetrics(const FManagedArrayCollection& InCollection);
-		virtual ~FTetrahedralMetrics();
+		CHAOS_API FTetrahedralMetrics(FManagedArrayCollection& InCollection);
+		CHAOS_API FTetrahedralMetrics(const FManagedArrayCollection& InCollection);
+		CHAOS_API virtual ~FTetrahedralMetrics();
 
-		void DefineSchema();
+		CHAOS_API void DefineSchema();
 		bool IsConst() const { return SignedVolumeAttribute.IsConst(); }
-		bool IsValid() const;
+		CHAOS_API bool IsValid() const;
 
 		const TManagedArrayAccessor<float>& GetSignedVolumeRO() const { return SignedVolumeAttribute; }
 		TManagedArrayAccessor<float>& GetSignedVolume() { check(!IsConst()); return SignedVolumeAttribute; }
@@ -38,4 +38,4 @@ namespace GeometryCollection::Facades
 		TManagedArrayAccessor<float> AspectRatioAttribute;
 	};
 
-} // namespace GeometryCollection::Facades
\ No newline at end of file
+} // namespace GeometryCollection::Facades
diff --git a/Engine/Source/Runtime/Experimental/Chaos/Public/GeometryCollection/Facades/CollectionTetrahedralSkeletalBindingsFacade.h b/Engine/Source/Runtime/Experimental/Chaos/Public/GeometryCollection/Facades/CollectionTetrahedralSkeletalBindingsFacade.h
--- a/Engine/Source/Runtime/Experimental/Chaos/Public/GeometryCollection/Facades/CollectionTetrahedralSkeletalBindingsFacade.h
+++ b/Engine/Source/Runtime/Experimental/Chaos/Public/GeometryCollection/Facades/CollectionTetrahedralSkeletalBindingsFacade.h
@@ -16,32 +16,32 @@ namespace GeometryCollection::Facades
 	* tetrahedral meshes. All bindings are stored in a single group, and a lookup table
 	* is aviable to match bindings to specifi meshs by name. 
 	*/
-	class CHAOS_API FTetrahedralSkeletalBindings
+	class FTetrahedralSkeletalBindings
  	{
 	public:
 		// Groups 
-		static const FName MeshBindingsGroupName;
-		static const FName MeshBindingsIdGroupName;
+		static CHAOS_API const FName MeshBindingsGroupName;
+		static CHAOS_API const FName MeshBindingsIdGroupName;
 
 		// Attributes
-		static const FName MeshIdAttributeName;
-		static const FName MeshIndexAttributeName;
-		static const FName TetrahedronIndexAttributeName;
-		static const FName WeightsAttributeName;
-		static const FName SkeletalIndexAttributeName;
+		static CHAOS_API const FName MeshIdAttributeName;
+		static CHAOS_API const FName MeshIndexAttributeName;
+		static CHAOS_API const FName TetrahedronIndexAttributeName;
+		static CHAOS_API const FName WeightsAttributeName;
+		static CHAOS_API const FName SkeletalIndexAttributeName;
 
 		/**
 		* FSelectionFacade Constuctor
 		* @param VertixDependencyGroup : GroupName the index attribute is dependent on. 
 		*/
-		FTetrahedralSkeletalBindings(FManagedArrayCollection& InSelf);
-		FTetrahedralSkeletalBindings(const FManagedArrayCollection& InSelf);
-		virtual ~FTetrahedralSkeletalBindings();
+		CHAOS_API FTetrahedralSkeletalBindings(FManagedArrayCollection& InSelf);
+		CHAOS_API FTetrahedralSkeletalBindings(const FManagedArrayCollection& InSelf);
+		CHAOS_API virtual ~FTetrahedralSkeletalBindings();
 
 		/** 
 		* Create the facade schema. 
 		*/
-		void DefineSchema();
+		CHAOS_API void DefineSchema();
 
 		/** Returns \c true if the facade is operating on a read-only geometry collection. */
 		bool IsConst() const { return MeshIdAttribute.IsConst(); }
@@ -50,16 +50,16 @@ namespace GeometryCollection::Facades
 		* Returns \c true if the Facade defined on the collection, and is initialized to
 		* a valid bindings group.
 		*/
-		bool IsValid() const;
+		CHAOS_API bool IsValid() const;
 
 		//
 		// API
 		//
-		static FString GenerateMeshGroupName(const int32 TetMeshIdx, const FName& SkeletalMeshName, const int32 LOD = INDEX_NONE);
-		void SetBindings(const FString& InMeshGroupIndex, const TArray<int32>& InTetrahedronIndex, const TArray<FVector4f>& WeightsIn, const TArray<int32>& InSkeletalIndex);
+		static CHAOS_API FString GenerateMeshGroupName(const int32 TetMeshIdx, const FName& SkeletalMeshName, const int32 LOD = INDEX_NONE);
+		CHAOS_API void SetBindings(const FString& InMeshGroupIndex, const TArray<int32>& InTetrahedronIndex, const TArray<FVector4f>& WeightsIn, const TArray<int32>& InSkeletalIndex);
 		
 		/** Returns \c false if bindings data was not found, or there was another error; \c true otherwise. */
-		bool CalculateBindings(const FString& InKey, const TArray<FVector3f>& InVertices, TArray<FVector>& OutPosition, TArray<bool>* OutInfluence = nullptr) const;
+		CHAOS_API bool CalculateBindings(const FString& InKey, const TArray<FVector3f>& InVertices, TArray<FVector>& OutPosition, TArray<bool>* OutInfluence = nullptr) const;
 		
 	private:
 		TManagedArrayAccessor<FString> MeshIdAttribute;
diff --git a/Engine/Source/Runtime/Experimental/Chaos/Public/GeometryCollection/Facades/CollectionTransformFacade.h b/Engine/Source/Runtime/Experimental/Chaos/Public/GeometryCollection/Facades/CollectionTransformFacade.h
--- a/Engine/Source/Runtime/Experimental/Chaos/Public/GeometryCollection/Facades/CollectionTransformFacade.h
+++ b/Engine/Source/Runtime/Experimental/Chaos/Public/GeometryCollection/Facades/CollectionTransformFacade.h
@@ -11,23 +11,23 @@ namespace GeometryCollection::Facades
 	/**
 	 * Provides an API to read and manipulate hierarchy in a managed array collection
 	 */
-	class CHAOS_API FCollectionTransformFacade
+	class FCollectionTransformFacade
 	{
 	public:
-		FCollectionTransformFacade(FManagedArrayCollection& InCollection);
-		FCollectionTransformFacade(const FManagedArrayCollection& InCollection);
+		CHAOS_API FCollectionTransformFacade(FManagedArrayCollection& InCollection);
+		CHAOS_API FCollectionTransformFacade(const FManagedArrayCollection& InCollection);
 
 		/** Creates the facade attributes. */
 		void DefineSchema() {}
 
 		/** Valid if parent and children arrays are available */
-		bool IsValid() const;
+		CHAOS_API bool IsValid() const;
 
 		/** Is the facade defined constant. */
 		bool IsConst() const { return ParentAttribute.IsConst(); }
 
 		/** Gets the root index */
-		TArray<int32> GetRootIndices() const;
+		CHAOS_API TArray<int32> GetRootIndices() const;
 
 		/**
 		* Returns the parent indices from the collection. Null if not initialized.
@@ -48,25 +48,25 @@ namespace GeometryCollection::Facades
 		* Returns array of transforms for transforming from bone space to collection space
 		* Vertex(inCollectionSpace) = TransformComputed.TransformPosition(Vertex(inBoneSpace))
 		*/
-		TArray<FTransform> ComputeCollectionSpaceTransforms() const;
+		CHAOS_API TArray<FTransform> ComputeCollectionSpaceTransforms() const;
 
 		/**
 		* Returns the transform for transforming from bone space to collection space for specified bone
 		* Vertex(inCollectionSpace) = TransformComputed.TransformPosition(Vertex(inBoneSpace))
 		*/
-		FTransform ComputeCollectionSpaceTransform(int32 BoneIdx) const;
+		CHAOS_API FTransform ComputeCollectionSpaceTransform(int32 BoneIdx) const;
 
 		/** Transforms the pivot of a collection */
-		void SetPivot(const FTransform& InTransform);
+		CHAOS_API void SetPivot(const FTransform& InTransform);
 
 		/** Transforms collection */
-		void Transform(const FTransform& InTransform);
+		CHAOS_API void Transform(const FTransform& InTransform);
 
 		/** Transforms selected bones in the collection */
-		void Transform(const FTransform& InTransform, const TArray<int32>& InSelection);
+		CHAOS_API void Transform(const FTransform& InTransform, const TArray<int32>& InSelection);
 
 		/** Builds a FMatrix from all the components */
-		static FMatrix BuildMatrix(const FVector& Translate,
+		static CHAOS_API FMatrix BuildMatrix(const FVector& Translate,
 			const uint8 RotationOrder,
 			const FVector& Rotate,
 			const FVector& Scale,
@@ -77,7 +77,7 @@ namespace GeometryCollection::Facades
 			const bool InvertTransformation);
 
 		/** Builds a FTransform from all the components */
-		static FTransform BuildTransform(const FVector& Translate,
+		static CHAOS_API FTransform BuildTransform(const FVector& Translate,
 			const uint8 RotationOrder,
 			const FVector& Rotate,
 			const FVector& Scale,
@@ -87,11 +87,11 @@ namespace GeometryCollection::Facades
 			const bool InvertTransformation);
 
 		/** Sets the selected bone's transform to identity */
-		void SetBoneTransformToIdentity(int32 BoneIdx);
+		CHAOS_API void SetBoneTransformToIdentity(int32 BoneIdx);
 
 	private:
 		TManagedArrayAccessor<int32>		ParentAttribute;
 		TManagedArrayAccessor<TSet<int32>>	ChildrenAttribute;
 		TManagedArrayAccessor<FTransform>	TransformAttribute;
 	};
-}
\ No newline at end of file
+}
diff --git a/Engine/Source/Runtime/Experimental/Chaos/Public/GeometryCollection/Facades/CollectionTransformSelectionFacade.h b/Engine/Source/Runtime/Experimental/Chaos/Public/GeometryCollection/Facades/CollectionTransformSelectionFacade.h
--- a/Engine/Source/Runtime/Experimental/Chaos/Public/GeometryCollection/Facades/CollectionTransformSelectionFacade.h
+++ b/Engine/Source/Runtime/Experimental/Chaos/Public/GeometryCollection/Facades/CollectionTransformSelectionFacade.h
@@ -20,7 +20,7 @@ namespace GeometryCollection::Facades
 	* Defines common API for selecting transforms on a collection
 	* 
 	*/
-	class CHAOS_API FCollectionTransformSelectionFacade
+	class FCollectionTransformSelectionFacade
 	{
 	public:
 
@@ -28,138 +28,138 @@ namespace GeometryCollection::Facades
 		* FCollectionTransformSelectionFacade Constuctor
 		* @param FManagedArrayCollection : Collection input
 		*/
-		FCollectionTransformSelectionFacade(const FManagedArrayCollection& InSelf);
+		CHAOS_API FCollectionTransformSelectionFacade(const FManagedArrayCollection& InSelf);
 
 		/** Create the facade attributes. */
-		void DefineSchema();
+		CHAOS_API void DefineSchema();
 
 		/** Is the facade defined constant. */
 		bool IsConst() const { return ParentAttribute.IsConst(); }
 
 		/** Is the Facade defined on the collection? */
-		bool IsValid() const;		
+		CHAOS_API bool IsValid() const;		
 
 		/**  */
 		bool IsValidBone(const int32 Index) const { return Index >= 0 && Index < ParentAttribute.Num(); }
 
 		/**  */
-		bool IsARootBone(const int32 Index) const;
+		CHAOS_API bool IsARootBone(const int32 Index) const;
 
 		/**  */
-		bool HasSelectedAncestor(const TArray<int32>& InSelection, const int32 Index) const;
+		CHAOS_API bool HasSelectedAncestor(const TArray<int32>& InSelection, const int32 Index) const;
 
 		/** Get bones filtered by the level as they would be in the editor; i.e., the bones with level <= the target Level, but with some additional logic for filtering, etc */
-		TArray<int32> GetBonesByLevel(const int32 Level, bool bOnlyClusteredOrRigid, bool bSkipFiltered) const;
+		CHAOS_API TArray<int32> GetBonesByLevel(const int32 Level, bool bOnlyClusteredOrRigid, bool bSkipFiltered) const;
 
 		/** Get bones exactly at a target level. A target level of -1 will return all bones. */
-		TArray<int32> GetBonesExactlyAtLevel(const int32 TargetLevel, bool bOnlyClusteredOrRigid) const;
+		CHAOS_API TArray<int32> GetBonesExactlyAtLevel(const int32 TargetLevel, bool bOnlyClusteredOrRigid) const;
 
 		/**  */
-		void Sanitize(TArray<int32>& InSelection, bool bFavorParents = true) const;
+		CHAOS_API void Sanitize(TArray<int32>& InSelection, bool bFavorParents = true) const;
 
 		/**  */
-		void RemoveRootNodes(TArray<int32>& InOutSelection) const;
+		CHAOS_API void RemoveRootNodes(TArray<int32>& InOutSelection) const;
 
 		/**  */
-		void ConvertSelectionToRigidNodes(const int32 Index, TArray<int32>& InSelection) const;
-		void ConvertSelectionToRigidNodes(TArray<int32>& InSelection) const;
+		CHAOS_API void ConvertSelectionToRigidNodes(const int32 Index, TArray<int32>& InSelection) const;
+		CHAOS_API void ConvertSelectionToRigidNodes(TArray<int32>& InSelection) const;
 
 		/** Replace any Embedded or Rigid transforms in InOutSelection with their parent cluster. If there is just a rigid root node selected, optionally keep it if bLeaveRigidRoot is true. */
-		void ConvertSelectionToClusterNodes(TArray<int32>& InOutSelection, bool bLeaveRigidRoots = true) const;
+		CHAOS_API void ConvertSelectionToClusterNodes(TArray<int32>& InOutSelection, bool bLeaveRigidRoots = true) const;
 
-		void ConvertEmbeddedSelectionToParents(TArray<int32>& InOutSelection) const;
+		CHAOS_API void ConvertEmbeddedSelectionToParents(TArray<int32>& InOutSelection) const;
 
 		/**  */
 		bool CanSelectRootBones() const { return ParentAttribute.IsValid(); }
-		TArray<int32> SelectRootBones() const;
+		CHAOS_API TArray<int32> SelectRootBones() const;
 
 		/**  */
-		TArray<int32> SelectNone() const;
+		CHAOS_API TArray<int32> SelectNone() const;
 
 		/**  */
 		bool CanSelectAll() const { return ParentAttribute.IsValid(); }
-		TArray<int32> SelectAll() const;
+		CHAOS_API TArray<int32> SelectAll() const;
 
 		/**  */
 		bool CanSelectInverse() const { return ParentAttribute.IsValid(); }
-		void SelectInverse(TArray<int32>& InOutSelection) const;
+		CHAOS_API void SelectInverse(TArray<int32>& InOutSelection) const;
 
 		/**  */
 		bool CanSelectRandom() const { return ParentAttribute.IsValid(); }
-		TArray<int32> SelectRandom(bool bDeterministic, float RandomSeed, float RandomThresholdVal) const;
+		CHAOS_API TArray<int32> SelectRandom(bool bDeterministic, float RandomSeed, float RandomThresholdVal) const;
 
 		/**  */
 		bool CanSelectLeaf() const { return ParentAttribute.IsValid() && LevelAttribute.IsValid() && SimulationTypeAttribute.IsValid(); }
-		TArray<int32> SelectLeaf() const;
+		CHAOS_API TArray<int32> SelectLeaf() const;
 
 		/**  */
 		bool CanSelectCluster() const { return ParentAttribute.IsValid() && LevelAttribute.IsValid() && SimulationTypeAttribute.IsValid(); }
-		TArray<int32> SelectCluster() const;
+		CHAOS_API TArray<int32> SelectCluster() const;
 
 		/**  */
 		bool CanSelectContact() const { return TransformIndexAttribute.IsValid() && TransformToGeometryIndexAttribute.IsValid(); }
-		void SelectContact(TArray<int32>& InOutSelection) const;
+		CHAOS_API void SelectContact(TArray<int32>& InOutSelection) const;
 
 		/**  */
 		bool CanSelectParent() const { return ParentAttribute.IsValid(); }
-		void SelectParent(TArray<int32>& InOutSelection) const;
+		CHAOS_API void SelectParent(TArray<int32>& InOutSelection) const;
 
 		/**  */
 		bool CanSelectChildren() const { return ChildrenAttribute.IsValid(); }
-		void SelectChildren(TArray<int32>& InOutSelection) const;
+		CHAOS_API void SelectChildren(TArray<int32>& InOutSelection) const;
 
 		/**  */
 		bool CanSelectSiblings() const { return ParentAttribute.IsValid() && ChildrenAttribute.IsValid(); }
-		void SelectSiblings(TArray<int32>& InOutSelection) const;
+		CHAOS_API void SelectSiblings(TArray<int32>& InOutSelection) const;
 
 		/**  */
 		bool CanSelectLevel() const { return ParentAttribute.IsValid() && LevelAttribute.IsValid(); }
-		void SelectLevel(TArray<int32>& InOutSelection) const;
+		CHAOS_API void SelectLevel(TArray<int32>& InOutSelection) const;
 
 		/**  */
-		static void SelectByPercentage(TArray<int32>& InOutSelection, int32 Percentage, bool Deterministic, float RandomSeed);
+		static CHAOS_API void SelectByPercentage(TArray<int32>& InOutSelection, int32 Percentage, bool Deterministic, float RandomSeed);
 
 		/**  */
-		TArray<int32> SelectBySize(float SizeMin, float SizeMax, bool bInclusive, bool bInsideRange) const;
+		CHAOS_API TArray<int32> SelectBySize(float SizeMin, float SizeMax, bool bInclusive, bool bInsideRange) const;
 
 		/**  */
-		TArray<int32> SelectByVolume(float VolumeMin, float VolumeMax, bool bInclusive, bool bInsideRange) const;
+		CHAOS_API TArray<int32> SelectByVolume(float VolumeMin, float VolumeMax, bool bInclusive, bool bInsideRange) const;
 
 		/** 
 		* Returns indices of bones where all the vertices or at least one vertex of the bone is inside InBox 
 		*/
 		bool CanSelectVerticesInBox() const { return ParentAttribute.IsValid(); }
-		TArray<int32> SelectVerticesInBox(const FBox& InBox, const FTransform& InBoxTransform, bool bAllVerticesInBox) const;
+		CHAOS_API TArray<int32> SelectVerticesInBox(const FBox& InBox, const FTransform& InBoxTransform, bool bAllVerticesInBox) const;
 
 		/**
 		* Returns indices of bones where centroid of the bone is inside InBox
 		*/
 		bool CanSelectCentroidInBox() const { return ParentAttribute.IsValid(); }
-		TArray<int32> SelectCentroidInBox(const FBox& InBox, const FTransform& InBoxTransform) const;
+		CHAOS_API TArray<int32> SelectCentroidInBox(const FBox& InBox, const FTransform& InBoxTransform) const;
 
 		/**
 		* Returns indices of bones where BoundingBox of the bone is inside InBox
 		*/
 		bool CanSelectBoundingBoxInBox() const { return ParentAttribute.IsValid(); }
-		TArray<int32> SelectBoundingBoxInBox(const FBox& InBox, const FTransform& InBoxTransform) const;
+		CHAOS_API TArray<int32> SelectBoundingBoxInBox(const FBox& InBox, const FTransform& InBoxTransform) const;
 
 		/**
 		* Returns indices of bones where all the vertices or at least one vertex of the bone is inside a sphere
 		*/
 		bool CanSelectVerticesInSphere() const { return ParentAttribute.IsValid(); }
-		TArray<int32> SelectVerticesInSphere(const FSphere& InSphere, const FTransform& InSphereTransform, bool bAllVerticesInSphere) const;
+		CHAOS_API TArray<int32> SelectVerticesInSphere(const FSphere& InSphere, const FTransform& InSphereTransform, bool bAllVerticesInSphere) const;
 
 		/**
 		* Returns indices of bones where BoundingBox of the bone is inside a sphere
 		*/
 		bool CanSelectBoundingBoxInSphere() const { return ParentAttribute.IsValid(); }
-		TArray<int32> SelectBoundingBoxInSphere(const FSphere& InSphere, const FTransform& InSphereTransform) const;
+		CHAOS_API TArray<int32> SelectBoundingBoxInSphere(const FSphere& InSphere, const FTransform& InSphereTransform) const;
 
 		/**
 		* Returns indices of bones where centroid of the bone is inside a sphere
 		*/
 		bool CanSelectCentroidInSphere() const { return ParentAttribute.IsValid(); }
-		TArray<int32> SelectCentroidInSphere(const FSphere& InSphere, const FTransform& InSphereTransform) const;
+		CHAOS_API TArray<int32> SelectCentroidInSphere(const FSphere& InSphere, const FTransform& InSphereTransform) const;
 
 		/** 
 		* Returns indices of bones where centroid of the bone is inside a convex hull
@@ -169,42 +169,42 @@ namespace GeometryCollection::Facades
 		/**
 		* Select by a float attribute in the collection, it can be any existing float attribute specified by GroupName/AttrName
 		*/
-		TArray<int32> SelectByFloatAttribute(FString GroupName, FString AttrName, float Min, float Max, bool bInclusive, bool bInsideRange) const;
+		CHAOS_API TArray<int32> SelectByFloatAttribute(FString GroupName, FString AttrName, float Min, float Max, bool bInclusive, bool bInsideRange) const;
 
 		/**
 		* Select by an int attribute in the collection, it can be any existing int attribute specified by GroupName/AttrName 
 		*/
-		TArray<int32> SelectByIntAttribute(FString GroupName, FString AttrName, int32 Min, int32 Max, bool bInclusive, bool bInsideRange) const;
+		CHAOS_API TArray<int32> SelectByIntAttribute(FString GroupName, FString AttrName, int32 Min, int32 Max, bool bInclusive, bool bInsideRange) const;
 
 		/**
 		* Convert a VertexSelection to a TransformSelection, if bAllElementsMustBeSelected is true then all vertices of a transform must be selected for the transform to be selected 
 		*/
-		TArray<int32> ConvertVertexSelectionToTransformSelection(const TArray<int32>& InVertexSelection, bool bAllElementsMustBeSelected) const;
+		CHAOS_API TArray<int32> ConvertVertexSelectionToTransformSelection(const TArray<int32>& InVertexSelection, bool bAllElementsMustBeSelected) const;
 
 		/**
 		* Convert a FaceSelection to a TransformSelection, if bAllElementsMustBeSelected is true then all faces of a transform must be selected for the transform to be selected
 		*/
-		TArray<int32> ConvertFaceSelectionToTransformSelection(const TArray<int32>& InFaceSelection, bool bAllElementsMustBeSelected) const;
+		CHAOS_API TArray<int32> ConvertFaceSelectionToTransformSelection(const TArray<int32>& InFaceSelection, bool bAllElementsMustBeSelected) const;
 
 		/**
 		* Convert a VertexSelection to a FaceSelection, if bAllElementsMustBeSelected is true then all vertices of a face must be selectod for the face to be selected 
 		*/
-		TArray<int32> ConvertVertexSelectionToFaceSelection(const TArray<int32>& InVertexSelection, bool bAllElementsMustBeSelected) const;
+		CHAOS_API TArray<int32> ConvertVertexSelectionToFaceSelection(const TArray<int32>& InVertexSelection, bool bAllElementsMustBeSelected) const;
 
 		/**
 		* Convert a TransformSelection to a FaceSelection 
 		*/
-		TArray<int32> ConvertTransformSelectionToFaceSelection(const TArray<int32>& InTransformSelection) const;
+		CHAOS_API TArray<int32> ConvertTransformSelectionToFaceSelection(const TArray<int32>& InTransformSelection) const;
 
 		/**
 		* Convert a FaceSelection to a VertexSelection 
 		*/
-		TArray<int32> ConvertFaceSelectionToVertexSelection(const TArray<int32>& InFaceSelection) const;
+		CHAOS_API TArray<int32> ConvertFaceSelectionToVertexSelection(const TArray<int32>& InFaceSelection) const;
 
 		/**
 		* Convert a TransformSelection to a VertexSelection 
 		*/
-		TArray<int32> ConvertTransformSelectionToVertexSelection(const TArray<int32>& InTransformSelection) const;
+		CHAOS_API TArray<int32> ConvertTransformSelectionToVertexSelection(const TArray<int32>& InTransformSelection) const;
 
 
 	private:
@@ -224,4 +224,4 @@ namespace GeometryCollection::Facades
 		TManagedArrayAccessor<FIntVector>	IndicesAttribute;
 	};
 
-}
\ No newline at end of file
+}
diff --git a/Engine/Source/Runtime/Experimental/Chaos/Public/GeometryCollection/Facades/CollectionTransformSourceFacade.h b/Engine/Source/Runtime/Experimental/Chaos/Public/GeometryCollection/Facades/CollectionTransformSourceFacade.h
--- a/Engine/Source/Runtime/Experimental/Chaos/Public/GeometryCollection/Facades/CollectionTransformSourceFacade.h
+++ b/Engine/Source/Runtime/Experimental/Chaos/Public/GeometryCollection/Facades/CollectionTransformSourceFacade.h
@@ -29,33 +29,33 @@ namespace GeometryCollection::Facades
 	*	- FindAttribute<TSet<int32>>(FTransformSource::SourceRootsAttribute, <Group>);
 	* 
 	*/
-	class CHAOS_API FTransformSource
+	class FTransformSource
  	{
 	public:
 
 		// groups
-		static const FName TransformSourceGroupName;
+		static CHAOS_API const FName TransformSourceGroupName;
 
 		// Attributes
-		static const FName SourceNameAttributeName;
-		static const FName SourceGuidAttributeName;
-		static const FName SourceRootsAttributeName;
+		static CHAOS_API const FName SourceNameAttributeName;
+		static CHAOS_API const FName SourceGuidAttributeName;
+		static CHAOS_API const FName SourceRootsAttributeName;
 
 		/**
 		* FSelectionFacade Constuctor
 		* @param VertixDependencyGroup : GroupName the index attribute is dependent on. 
 		*/
-		FTransformSource(FManagedArrayCollection& InSelf);
-		FTransformSource(const FManagedArrayCollection& InSelf);
+		CHAOS_API FTransformSource(FManagedArrayCollection& InSelf);
+		CHAOS_API FTransformSource(const FManagedArrayCollection& InSelf);
 
 		/** Create the facade. */
-		void DefineSchema();
+		CHAOS_API void DefineSchema();
 
 		/** Is the facade defined constant. */
 		bool IsConst() const { return SourceNameAttribute.IsConst(); }
 
 		/** Is the Facade defined on the collection? */
-		bool IsValid() const;
+		CHAOS_API bool IsValid() const;
 
 		/**
 		* Add a transform root mapping.
@@ -63,14 +63,14 @@ namespace GeometryCollection::Facades
 		* @param Guid : Guid of the owner of the transform set.
 		* @param Roots : Root indices of the transform set.
 		*/
-		void AddTransformSource(const FString& Name, const FString& Guid, const TSet<int32>& Roots);
+		CHAOS_API void AddTransformSource(const FString& Name, const FString& Guid, const TSet<int32>& Roots);
 
 		/**
 		* Query for root indices.  
 		* @param Name : Name of the owner of the transform set.  
 		* @param Guid : Guid of the owner of the transform set.
 		*/
-		TSet<int32> GetTransformSource(const FString& Name, const FString& Guid) const;
+		CHAOS_API TSet<int32> GetTransformSource(const FString& Name, const FString& Guid) const;
 
 	private:
 		TManagedArrayAccessor<FString> SourceNameAttribute;
@@ -79,4 +79,4 @@ namespace GeometryCollection::Facades
 		
 	};
 
-}
\ No newline at end of file
+}
diff --git a/Engine/Source/Runtime/Experimental/Chaos/Public/GeometryCollection/Facades/CollectionVertexBoneWeightsFacade.h b/Engine/Source/Runtime/Experimental/Chaos/Public/GeometryCollection/Facades/CollectionVertexBoneWeightsFacade.h
--- a/Engine/Source/Runtime/Experimental/Chaos/Public/GeometryCollection/Facades/CollectionVertexBoneWeightsFacade.h
+++ b/Engine/Source/Runtime/Experimental/Chaos/Public/GeometryCollection/Facades/CollectionVertexBoneWeightsFacade.h
@@ -29,34 +29,34 @@ namespace GeometryCollection::Facades
 	*	- FindAttribute<TArray<float>>(FVertexSetInterface::WeightAttribute, FGeometryCollection::VerticesGroup);
 	* 
 	*/
-	class CHAOS_API FVertexBoneWeightsFacade
+	class FVertexBoneWeightsFacade
 	{
 	public:
 
 		// Attributes
-		static const FName BoneIndexAttributeName;
-		static const FName BoneWeightAttributeName;
+		static CHAOS_API const FName BoneIndexAttributeName;
+		static CHAOS_API const FName BoneWeightAttributeName;
 
 		/**
 		* FVertexBoneWeightsFacade Constuctor
 		*/
-		FVertexBoneWeightsFacade(FManagedArrayCollection& InSelf);
-		FVertexBoneWeightsFacade(const FManagedArrayCollection& InSelf);
+		CHAOS_API FVertexBoneWeightsFacade(FManagedArrayCollection& InSelf);
+		CHAOS_API FVertexBoneWeightsFacade(const FManagedArrayCollection& InSelf);
 
 		/** Define the facade */
-		void DefineSchema();
+		CHAOS_API void DefineSchema();
 
 		/** Is the Facade const */
 		bool IsConst() const { return Collection == nullptr; }
 
 		/** Is the Facade defined on the collection? */
-		bool IsValid() const;
+		CHAOS_API bool IsValid() const;
 
 		/** Add bone weight based on the kinematic bindings. */
-		void AddBoneWeightsFromKinematicBindings();
+		CHAOS_API void AddBoneWeightsFromKinematicBindings();
 
 		/** Add bone weight based on the kinematic bindings. */
-		void AddBoneWeight(int32 VertexIndex, int32 BoneIndex, float BoneWeight);
+		CHAOS_API void AddBoneWeight(int32 VertexIndex, int32 BoneIndex, float BoneWeight);
 
 		/** Return the vertex bone indices from the collection. Null if not initialized.  */
 		const TManagedArray< TArray<int32> >* FindBoneIndices()  const { return BoneIndexAttribute.Find(); }
@@ -78,4 +78,4 @@ namespace GeometryCollection::Facades
 
 	};
 
-}
\ No newline at end of file
+}
diff --git a/Engine/Source/Runtime/Experimental/Chaos/Public/GeometryCollection/GeometryCollection.h b/Engine/Source/Runtime/Experimental/Chaos/Public/GeometryCollection/GeometryCollection.h
--- a/Engine/Source/Runtime/Experimental/Chaos/Public/GeometryCollection/GeometryCollection.h
+++ b/Engine/Source/Runtime/Experimental/Chaos/Public/GeometryCollection/GeometryCollection.h
@@ -22,7 +22,7 @@ class FGeometryCollectionProximityPropertiesInterface;
 /**
 * FGeometryCollection (FTransformCollection)
 */
-class CHAOS_API FGeometryCollection : public FTransformCollection, 
+class FGeometryCollection : public FTransformCollection, 
 	public FGeometryCollectionConvexPropertiesInterface,
 	public FGeometryCollectionProximityPropertiesInterface
 {
@@ -30,7 +30,7 @@ class CHAOS_API FGeometryCollection : public FTransformCollection,
 public:
 	typedef FTransformCollection Super;
 
-	FGeometryCollection();
+	CHAOS_API FGeometryCollection();
 	FGeometryCollection(FGeometryCollection &) = delete;
 	FGeometryCollection& operator=(const FGeometryCollection &) = delete;
 	FGeometryCollection(FGeometryCollection &&) = default;
@@ -91,15 +91,15 @@ public:
 		*		 The set of triangles which are rendered with the same material
 		*/
 
-	static const FName VerticesGroup; // Vertices
-	static const FName FacesGroup;	  // Faces
-	static const FName GeometryGroup; // Geometry
-	static const FName BreakingGroup; // Breaking
-	static const FName MaterialGroup; // Materials
+	static CHAOS_API const FName VerticesGroup; // Vertices
+	static CHAOS_API const FName FacesGroup;	  // Faces
+	static CHAOS_API const FName GeometryGroup; // Geometry
+	static CHAOS_API const FName BreakingGroup; // Breaking
+	static CHAOS_API const FName MaterialGroup; // Materials
 
-	static const FName SimulatableParticlesAttribute;
-	static const FName SimulationTypeAttribute;
-	static const FName StatusFlagsAttribute;
+	static CHAOS_API const FName SimulatableParticlesAttribute;
+	static CHAOS_API const FName SimulationTypeAttribute;
+	static CHAOS_API const FName StatusFlagsAttribute;
 
 	enum ESimulationTypes : uint8
 	{
@@ -122,19 +122,19 @@ public:
 
 	};
 
-	static bool AreCollisionParticlesEnabled();
+	static CHAOS_API bool AreCollisionParticlesEnabled();
 
 	/**
 	 * Create a GeometryCollection from Vertex and Indices arrays
 	 */
-	static FGeometryCollection* NewGeometryCollection(const TArray<float>& RawVertexArray, const TArray<int32>& RawIndicesArray, bool ReverseVertexOrder = true);
-	static void Init(FGeometryCollection* Collection, const TArray<float>& RawVertexArray, const TArray<int32>& RawIndicesArray, bool ReverseVertexOrder = true);
-	static void DefineGeometrySchema(FManagedArrayCollection&);
+	static CHAOS_API FGeometryCollection* NewGeometryCollection(const TArray<float>& RawVertexArray, const TArray<int32>& RawIndicesArray, bool ReverseVertexOrder = true);
+	static CHAOS_API void Init(FGeometryCollection* Collection, const TArray<float>& RawVertexArray, const TArray<int32>& RawIndicesArray, bool ReverseVertexOrder = true);
+	static CHAOS_API void DefineGeometrySchema(FManagedArrayCollection&);
 
 	/**
 	* Create a GeometryCollection from Vertex, Indices, BoneMap, Transform, BoneHierarchy arrays
 	*/
-	static FGeometryCollection* NewGeometryCollection(const TArray<float>& RawVertexArray,
+	static CHAOS_API FGeometryCollection* NewGeometryCollection(const TArray<float>& RawVertexArray,
 		const TArray<int32>& RawIndicesArray,
 		const TArray<int32>& RawBoneMapArray,
 		const TArray<FTransform>& RawTransformArray,
@@ -152,95 +152,95 @@ public:
 	/** 
 	* Append a single geometric object to a FGeometryCollection 
 	*/
-	int32 AppendGeometry(const FGeometryCollection & GeometryCollection, int32 MaterialIDOffset = 0, bool ReindexAllMaterials = true, const FTransform& TransformRoot = FTransform::Identity);
+	CHAOS_API int32 AppendGeometry(const FGeometryCollection & GeometryCollection, int32 MaterialIDOffset = 0, bool ReindexAllMaterials = true, const FTransform& TransformRoot = FTransform::Identity);
 
 	/**
 	* Append single embedded geometry. Returns true if the operation succeeds.
 	*/
-	bool AppendEmbeddedInstance(int32 InExemplarIndex, int32 InParentIndex, const FTransform& InTransform = FTransform::Identity);
+	CHAOS_API bool AppendEmbeddedInstance(int32 InExemplarIndex, int32 InParentIndex, const FTransform& InTransform = FTransform::Identity);
 
 	/**
 	 * Reindex exemplar indices to reflect removed exemplars.
 	 */
-	void ReindexExemplarIndices(TArray<int32>& SortedRemovedIndices);
+	CHAOS_API void ReindexExemplarIndices(TArray<int32>& SortedRemovedIndices);
 
 	/**
 	* Remove Geometry and update dependent elements
 	*/
-	virtual void RemoveElements(const FName & Group, const TArray<int32>& DeletionList, FProcessingParameters Params = FProcessingParameters()) override;
+	CHAOS_API virtual void RemoveElements(const FName & Group, const TArray<int32>& DeletionList, FProcessingParameters Params = FProcessingParameters()) override;
 
 	/**
 	* Empty each managed array in each group in order to reset the collection to an initial (empty) state. 
 	*/
-	void Empty();
+	CHAOS_API void Empty();
 
 	/**
 	* reset internal state
 	*/
-	virtual void Reset() override;
+	CHAOS_API virtual void Reset() override;
 
 	/**
 	* Reorders elements in a group. NewOrder must be the same length as the group.
 	*/
-	virtual void ReorderElements(FName Group, const TArray<int32>& NewOrder) override;
+	CHAOS_API virtual void ReorderElements(FName Group, const TArray<int32>& NewOrder) override;
 
 	//
 	//
 	//
 
 	// Initialize any interfaces on the geometry collection (i.e., the FGeometryCollectionConvexPropertiesInterface)
-	virtual void InitializeInterfaces();
+	CHAOS_API virtual void InitializeInterfaces();
 
 
 	/**
 	*  Update bounding box entries for the geometry
 	*/
-	void UpdateBoundingBox();
-	static void UpdateBoundingBox(FManagedArrayCollection&, bool bSkipCheck=false);
+	CHAOS_API void UpdateBoundingBox();
+	static CHAOS_API void UpdateBoundingBox(FManagedArrayCollection&, bool bSkipCheck=false);
 
 	/**  
 	* GetBoundingBox 
 	*/
-	FBoxSphereBounds GetBoundingBox() const;
+	CHAOS_API FBoxSphereBounds GetBoundingBox() const;
 
 	/**
 	 * Update the visibility of specified geometry nodes
 	 */
-	void UpdateGeometryVisibility(const TArray<int32>& NodeList, bool VisibilityState);
+	CHAOS_API void UpdateGeometryVisibility(const TArray<int32>& NodeList, bool VisibilityState);
 
 	/**
 	* Reindex sections to keep polys with same materials together to reduce the number of draw calls
 	*/
-	void ReindexMaterials();
-	static void ReindexMaterials(FManagedArrayCollection&);
+	CHAOS_API void ReindexMaterials();
+	static CHAOS_API void ReindexMaterials(FManagedArrayCollection&);
 
 	/**
 	* Builds mesh sections for a given index buffer that could be a subset.
 	* Currently, this call assumes that the indices are ordered by MaterialID
 	* #todo(dmp): Refactor this and ReindexMaterials to share code
 	*/
-	TArray<FGeometryCollectionSection> BuildMeshSections(const TArray<FIntVector> &Indices, const TArray<int32>& BaseMeshOriginalIndicesIndex, TArray<FIntVector> &RetIndices) const;
+	CHAOS_API TArray<FGeometryCollectionSection> BuildMeshSections(const TArray<FIntVector> &Indices, const TArray<int32>& BaseMeshOriginalIndicesIndex, TArray<FIntVector> &RetIndices) const;
 	//
 	//
 	//
 
 	/** Returns true if there is anything to render */
-	bool HasVisibleGeometry() const;
+	CHAOS_API bool HasVisibleGeometry() const;
 
 	/** Returns true if the vertices are contiguous*/
-	bool HasContiguousVertices() const;
+	CHAOS_API bool HasContiguousVertices() const;
 
 	/** Returns true if the faces are contiguous*/
-	bool HasContiguousFaces() const;
+	CHAOS_API bool HasContiguousFaces() const;
 
 	/** Returns true if the render faces are contiguous*/
-	bool HasContiguousRenderFaces() const;
+	CHAOS_API bool HasContiguousRenderFaces() const;
 
 	/** Returns number of UV layers represented by UV array. A Valid Geometry Collection has the same count for every vertex */
-	int32 NumUVLayers() const;
+	CHAOS_API int32 NumUVLayers() const;
 
 	/** Update a geometry collection to have the target number of UV layers (must be in the range [1, MAX_UV_LAYERS)) */
-	bool SetNumUVLayers(int32 NumLayers);
+	CHAOS_API bool SetNumUVLayers(int32 NumLayers);
 
 	FORCEINLINE bool IsGeometry(int32 Element) const { return TransformToGeometryIndex[Element] != INDEX_NONE; }
 	FORCEINLINE bool IsClustered(int32 Element) const { const TManagedArray<int32>& SimType = SimulationType;  return !!(SimType[Element] == ESimulationTypes::FST_Clustered); }
@@ -251,29 +251,29 @@ public:
 	FORCEINLINE bool HasFlags(int32 Element, int32 InFlags) const { const TManagedArray<int32>& Status = StatusFlags; return (Status[Element] & InFlags) != 0; }
 
 	/** Return true if the Element contains any visible faces. */
-	bool IsVisible(int32 Element) const;
+	CHAOS_API bool IsVisible(int32 Element) const;
 
 	/** Connection of leaf geometry */
-	TArray<TArray<int32>> ConnectionGraph();
+	CHAOS_API TArray<TArray<int32>> ConnectionGraph();
 
 	//
 	//
 	//
 
 	/** Serialize */
-	void Serialize(Chaos::FChaosArchive& Ar);
+	CHAOS_API void Serialize(Chaos::FChaosArchive& Ar);
 
 	/**   */
-	void WriteDataToHeaderFile(const FString &Name, const FString &Path);
+	CHAOS_API void WriteDataToHeaderFile(const FString &Name, const FString &Path);
 
 	/**  */
-	void WriteDataToOBJFile(const FString &Name, const FString &Path, const bool WriteTopology=true, const bool WriteAuxStructures=true);
+	CHAOS_API void WriteDataToOBJFile(const FString &Name, const FString &Path, const bool WriteTopology=true, const bool WriteAuxStructures=true);
 
 	//
 	//
 	//
 
-	virtual void SetDefaults(FName Group, uint32 StartSize, uint32 NumElements) override;
+	CHAOS_API virtual void SetDefaults(FName Group, uint32 StartSize, uint32 NumElements) override;
 
 	// Transform Group
 	TManagedArray<int32>		TransformToGeometryIndex;
@@ -343,34 +343,34 @@ protected:
 		GeometryCollection::UV::MatchUVLayerCount(*this, InCollection);
 	}
 
-	void Construct();
+	CHAOS_API void Construct();
 
 	/**
 	* Remove Geometry elements i.e. verts, faces, etc, leaving the transform nodes intact
 	*/
-	void RemoveGeometryElements(const TArray<int32>& SortedGeometryIndicesToDelete);
+	CHAOS_API void RemoveGeometryElements(const TArray<int32>& SortedGeometryIndicesToDelete);
 
 	/**
 	* Update Face Attributes based on changes in the group.
 	*/
-	bool BuildFaceToGeometryMapping(bool InSaved=false);
-	void UpdateFaceGroupElements();
+	CHAOS_API bool BuildFaceToGeometryMapping(bool InSaved=false);
+	CHAOS_API void UpdateFaceGroupElements();
 
 	/**
 	* Build and Update Vertex Attributes based on changes in the group.
 	*/
-	bool BuildVertexToGeometryMapping(bool InSaved = false);
-	void UpdateVerticesGroupElements();
+	CHAOS_API bool BuildVertexToGeometryMapping(bool InSaved = false);
+	CHAOS_API void UpdateVerticesGroupElements();
 
 	/** 
 	* Reorder geometry elements. i.e. verts faces etc are reordered so we can get contiguous memory access
 	*/
-	void ReorderGeometryElements(const TArray<int32>& NewOrder);
+	CHAOS_API void ReorderGeometryElements(const TArray<int32>& NewOrder);
 
 	/**
 	* Reorder geometry elements based on the new transform order. i.e. verts faces etc are reordered so we can get contiguous memory access
 	*/
-	void ReorderTransformElements(const TArray<int32>& NewOrder);
+	CHAOS_API void ReorderTransformElements(const TArray<int32>& NewOrder);
 
 	/**
 	 * @return the latest Version number used in serialization
@@ -383,7 +383,7 @@ protected:
 
 public:
 	/* Backwards compatibility */
-	void UpdateOldAttributeNames();
+	CHAOS_API void UpdateOldAttributeNames();
 
 
 };
diff --git a/Engine/Source/Runtime/Experimental/Chaos/Public/GeometryCollection/GeometryCollectionClusteringUtility.h b/Engine/Source/Runtime/Experimental/Chaos/Public/GeometryCollection/GeometryCollectionClusteringUtility.h
--- a/Engine/Source/Runtime/Experimental/Chaos/Public/GeometryCollection/GeometryCollectionClusteringUtility.h
+++ b/Engine/Source/Runtime/Experimental/Chaos/Public/GeometryCollection/GeometryCollectionClusteringUtility.h
@@ -7,7 +7,7 @@
 
 class FGeometryCollection;
 
-class CHAOS_API FGeometryCollectionClusteringUtility
+class FGeometryCollectionClusteringUtility
 {
 public:
 	/**
@@ -37,52 +37,52 @@ public:
 	* 
 	* @return Index of the New Node
 	*/
-	static int32 ClusterBonesUnderNewNode(FGeometryCollection* GeometryCollection, const int32 InsertAtIndex, const TArray<int32>& SelectedBones, bool CalcNewLocalTransform, bool Validate = true);
+	static CHAOS_API int32 ClusterBonesUnderNewNode(FGeometryCollection* GeometryCollection, const int32 InsertAtIndex, const TArray<int32>& SelectedBones, bool CalcNewLocalTransform, bool Validate = true);
 
 	// Same as ClusterBonesUnderNewNode, but specify the parent of the new node instead of a sibling
-	static int32 ClusterBonesUnderNewNodeWithParent(FGeometryCollection* GeometryCollection, const int32 ParentOfNewNode, const TArray<int32>& SelectedBones, bool CalcNewLocalTransform, bool Validate = true);
+	static CHAOS_API int32 ClusterBonesUnderNewNodeWithParent(FGeometryCollection* GeometryCollection, const int32 ParentOfNewNode, const TArray<int32>& SelectedBones, bool CalcNewLocalTransform, bool Validate = true);
 
 	/** Cluster all existing bones under a new root node, so there is now only one root node and a completely flat hierarchy underneath it */
-	static void ClusterAllBonesUnderNewRoot(FGeometryCollection* GeometryCollection);
+	static CHAOS_API void ClusterAllBonesUnderNewRoot(FGeometryCollection* GeometryCollection);
 
 	/** Cluster all source bones under an existing node, algorithm chooses best node to add to 'closest to root' */
-	static void ClusterBonesUnderExistingNode(FGeometryCollection* GeometryCollection, const TArray<int32>& SourceElements);
+	static CHAOS_API void ClusterBonesUnderExistingNode(FGeometryCollection* GeometryCollection, const TArray<int32>& SourceElements);
 
 	/** Cluster all source bones under an existing node, existing node is specifiedas MergeNode */
-	static void ClusterBonesUnderExistingNode(FGeometryCollection* GeometryCollection, int32 MergeNode, const TArray<int32>& SourceElements);
+	static CHAOS_API void ClusterBonesUnderExistingNode(FGeometryCollection* GeometryCollection, int32 MergeNode, const TArray<int32>& SourceElements);
 
 	/** Clusters using one of ClusterBonesUnderNewNode or ClusterBonesUnderExistingNode based on the context of what is being clustered */
-	static void ClusterBonesByContext(FGeometryCollection* GeometryCollection, int32 MergeNode, const TArray<int32>& SourceElementsIn);
+	static CHAOS_API void ClusterBonesByContext(FGeometryCollection* GeometryCollection, int32 MergeNode, const TArray<int32>& SourceElementsIn);
 
 	/** Returns true if bone hierarchy contains more than one root node */
-	static bool ContainsMultipleRootBones(FGeometryCollection* GeometryCollection);
+	static CHAOS_API bool ContainsMultipleRootBones(FGeometryCollection* GeometryCollection);
 
 	/** Finds the root bone in the hierarchy, the one with an invalid parent bone index */
-	static void GetRootBones(const FGeometryCollection* GeometryCollection, TArray<int32>& RootBonesOut);
+	static CHAOS_API void GetRootBones(const FGeometryCollection* GeometryCollection, TArray<int32>& RootBonesOut);
 
 	/** Return true if the specified bone is a root bone */
-	static bool IsARootBone(const FGeometryCollection* GeometryCollection, int32 InBone);
+	static CHAOS_API bool IsARootBone(const FGeometryCollection* GeometryCollection, int32 InBone);
 	
 	/** Finds all Bones in same cluster as the one specified */
-	static void GetClusteredBonesWithCommonParent(const FGeometryCollection* GeometryCollection, int32 SourceBone, TArray<int32>& BonesOut);
+	static CHAOS_API void GetClusteredBonesWithCommonParent(const FGeometryCollection* GeometryCollection, int32 SourceBone, TArray<int32>& BonesOut);
 
 	/** Get list of all bones above the specified hierarchy level */
-	static void GetBonesToLevel(const FGeometryCollection* GeometryCollection, int32 Level, TArray<int32>& BonesOut, bool bOnlyClusteredOrRigid = true, bool bSkipFiltered = true);
+	static CHAOS_API void GetBonesToLevel(const FGeometryCollection* GeometryCollection, int32 Level, TArray<int32>& BonesOut, bool bOnlyClusteredOrRigid = true, bool bSkipFiltered = true);
 	
 	/** Get list of child bones down from the source bone below the specified hierarchy level */
-	static void GetChildBonesFromLevel(const FGeometryCollection* GeometryCollection, int32 SourceBone, int32 Level, TArray<int32>& BonesOut);
+	static CHAOS_API void GetChildBonesFromLevel(const FGeometryCollection* GeometryCollection, int32 SourceBone, int32 Level, TArray<int32>& BonesOut);
 
 	/** Get list of child bones down from the source bone at the specified hierarchy level */
-	static void GetChildBonesAtLevel(const FGeometryCollection* GeometryCollection, int32 SourceBone, int32 Level, TArray<int32>& BonesOut);
+	static CHAOS_API void GetChildBonesAtLevel(const FGeometryCollection* GeometryCollection, int32 SourceBone, int32 Level, TArray<int32>& BonesOut);
 
 	/** Recursively Add all children to output bone list from source bone down to the leaf nodes */
-	static void RecursiveAddAllChildren(const TManagedArray<TSet<int32>>& Children, int32 SourceBone, TArray<int32>& BonesOut);
+	static CHAOS_API void RecursiveAddAllChildren(const TManagedArray<TSet<int32>>& Children, int32 SourceBone, TArray<int32>& BonesOut);
 
 	/**
 	 * Search hierarchy for the parent of the specified bone, where the parent exists at the given level in the hierarchy.
 	 * If bSkipFiltered, also skip to parents of any rigid/embedded nodes that will be filtered from the outliner.
 	 */
-	static int32 GetParentOfBoneAtSpecifiedLevel(const FGeometryCollection* GeometryCollection, int32 SourceBone, int32 Level, bool bSkipFiltered = false);
+	static CHAOS_API int32 GetParentOfBoneAtSpecifiedLevel(const FGeometryCollection* GeometryCollection, int32 SourceBone, int32 Level, bool bSkipFiltered = false);
 
 	/**
 	* Maintains the bone naming convention of
@@ -91,48 +91,48 @@ public:
 	*  Level 2 children of "Name_0001" are "Name_0001_0001", "Name_0001_0002",.. etc
 	* from the given bone index down through the hierarchy to the leaf nodes
 	*/
-	static void RecursivelyUpdateChildBoneNames(int32 BoneIndex, const TManagedArray<TSet<int32>>& Children, TManagedArray<FString>& BoneNames, bool OverrideBoneNames = false);
+	static CHAOS_API void RecursivelyUpdateChildBoneNames(int32 BoneIndex, const TManagedArray<TSet<int32>>& Children, TManagedArray<FString>& BoneNames, bool OverrideBoneNames = false);
 	
 	/** Recursively update the hierarchy level of all the children below this bone */
-	static void UpdateHierarchyLevelOfChildren(FGeometryCollection* GeometryCollection, int32 ParentElement);
-	static void UpdateHierarchyLevelOfChildren(FManagedArrayCollection& InCollection, int32 ParentElement);
+	static CHAOS_API void UpdateHierarchyLevelOfChildren(FGeometryCollection* GeometryCollection, int32 ParentElement);
+	static CHAOS_API void UpdateHierarchyLevelOfChildren(FManagedArrayCollection& InCollection, int32 ParentElement);
 
 	/** Collapse hierarchy at specified level */
-	static void CollapseLevelHierarchy(int8 Level, FGeometryCollection* GeometryCollection);
+	static CHAOS_API void CollapseLevelHierarchy(int8 Level, FGeometryCollection* GeometryCollection);
 
 	/** Collapse hierarchy of selected bones at specified level */
-	static void CollapseSelectedHierarchy(int8 Level, const TArray<int32>& SelectedBones, FGeometryCollection* GeometryCollection);
+	static CHAOS_API void CollapseSelectedHierarchy(int8 Level, const TArray<int32>& SelectedBones, FGeometryCollection* GeometryCollection);
 
 	/** reparent source bones under root bone */
-	static void ClusterBonesUnderExistingRoot(FGeometryCollection* GeometryCollection, const TArray<int32>& SourceElements);
+	static CHAOS_API void ClusterBonesUnderExistingRoot(FGeometryCollection* GeometryCollection, const TArray<int32>& SourceElements);
 
 	/** moved the selected bones closer to the root */
-	static void CollapseHierarchyOneLevel(FGeometryCollection* GeometryCollection, TArray<int32>& SourceElements);
+	static CHAOS_API void CollapseHierarchyOneLevel(FGeometryCollection* GeometryCollection, TArray<int32>& SourceElements);
 
 	/** return true of Test Node exists on tree branch specified by TreeElement Node*/
-	static bool NodeExistsOnThisBranch(const FGeometryCollection* GeometryCollection, int32 TestNode, int32 TreeElement);
+	static CHAOS_API bool NodeExistsOnThisBranch(const FGeometryCollection* GeometryCollection, int32 TestNode, int32 TreeElement);
 
 	/** Rename a bone node, will automatically update all child node names if requested (current default) */
-	static void RenameBone(FGeometryCollection* GeometryCollection, int32 BoneIndex, const FString& NewName, bool UpdateChildren = true);
+	static CHAOS_API void RenameBone(FGeometryCollection* GeometryCollection, int32 BoneIndex, const FString& NewName, bool UpdateChildren = true);
 
 	/** return an array of all child leaf nodes below the specified node. If bOnlyRigids is true, the first Rigid node dound is considered a leaf, regardless of an children it might have. */
-	static void GetLeafBones(const FManagedArrayCollection* GeometryCollection, int BoneIndex, bool bOnlyRigids, TArray<int32>& LeafBonesOut);
+	static CHAOS_API void GetLeafBones(const FManagedArrayCollection* GeometryCollection, int BoneIndex, bool bOnlyRigids, TArray<int32>& LeafBonesOut);
 
 	/** move the selected node up a level in direction of root */
-	static void MoveUpOneHierarchyLevel(FGeometryCollection* GeometryCollection, const TArray<int32>& SelectedBones);
+	static CHAOS_API void MoveUpOneHierarchyLevel(FGeometryCollection* GeometryCollection, const TArray<int32>& SelectedBones);
 
 	/** Find the lowest common ancestor index of currently selected nodes. Returns INDEX_NODE if there is no common ancestor. */
-	static int32 FindLowestCommonAncestor(const FManagedArrayCollection* GeometryCollection, const TArray<int32>& SelectedBones);
+	static CHAOS_API int32 FindLowestCommonAncestor(const FManagedArrayCollection* GeometryCollection, const TArray<int32>& SelectedBones);
 
 	/** Delete any cluster nodes discovered to have no children. Returns true if any nodes were removed. */
-	static bool RemoveDanglingClusters(FGeometryCollection* GeometryCollection);
+	static CHAOS_API bool RemoveDanglingClusters(FGeometryCollection* GeometryCollection);
 
 	/** Simplify the geometry collection by removing cluster nodes w/ a parent and only one child, re-parenting the child to its grand-parent. Returns true if any nodes were removed. */
-	static bool RemoveClustersOfOnlyOneChild(FGeometryCollection* GeometryCollection);
+	static CHAOS_API bool RemoveClustersOfOnlyOneChild(FGeometryCollection* GeometryCollection);
 
-	static void ValidateResults(FGeometryCollection* GeometryCollection);
+	static CHAOS_API void ValidateResults(FGeometryCollection* GeometryCollection);
 
-	static int32 PickBestNodeToMergeTo(const FManagedArrayCollection* Collection, const TArray<int32>& SourceElements);
+	static CHAOS_API int32 PickBestNodeToMergeTo(const FManagedArrayCollection* Collection, const TArray<int32>& SourceElements);
 
 private:
 	// #todo: intend to remove reliance on custom attributes for slider by making use of Rest/Dynamic collections
diff --git a/Engine/Source/Runtime/Experimental/Chaos/Public/GeometryCollection/GeometryCollectionCollisionStructureManager.h b/Engine/Source/Runtime/Experimental/Chaos/Public/GeometryCollection/GeometryCollectionCollisionStructureManager.h
--- a/Engine/Source/Runtime/Experimental/Chaos/Public/GeometryCollection/GeometryCollectionCollisionStructureManager.h
+++ b/Engine/Source/Runtime/Experimental/Chaos/Public/GeometryCollection/GeometryCollectionCollisionStructureManager.h
@@ -22,10 +22,10 @@ namespace Chaos
 
 
 
-class CHAOS_API FCollisionStructureManager
+class FCollisionStructureManager
 {
 public:
-	FCollisionStructureManager();
+	CHAOS_API FCollisionStructureManager();
 	virtual ~FCollisionStructureManager() {}
 
 	typedef TArray<Chaos::FVec3> FPoints;
@@ -38,7 +38,7 @@ public:
 	 * at \p CollisionParticlesMaxInput, and returns a bounding volume hierarchy
 	 * of the remainder.
 	 */
-	static FSimplicial* NewSimplicial(
+	static CHAOS_API FSimplicial* NewSimplicial(
 		const Chaos::FParticles& Vertices,
 		const Chaos::FTriangleMesh& TriMesh,
 		const Chaos::FImplicitObject* Implicit,
@@ -48,7 +48,7 @@ public:
 	 * Culls particles by importance (See \c FTriangleMesh::GetVertexImportanceOrdering()),
 	 * and returns a bounding volume hierarchy of the remainder.
 	 */
-	static FSimplicial* NewSimplicial(
+	static CHAOS_API FSimplicial* NewSimplicial(
 		const Chaos::FParticles& AllParticles,
 		const TManagedArray<int32>& BoneMap,
 		const ECollisionTypeEnum CollisionType,
@@ -59,7 +59,7 @@ public:
 	 * Calls \c SetDoCollide(false) and \c SetConvex(false) on \p Implicit if 
 	 * \p CollisionType is \c ECollisionTypeEnum::Chaos_Surface_Volumetric.
 	 */
-	static void UpdateImplicitFlags(
+	static CHAOS_API void UpdateImplicitFlags(
 		FImplicit* Implicit, 
 		const ECollisionTypeEnum CollisionType);
 
@@ -78,7 +78,7 @@ public:
 	 *	\p CollisionType - param forwarded to \c UpdateImplictFlags().
 	 *	\p ImplicitType - type of implicit shape to build.
 	 */
-	static FImplicit * NewImplicit(
+	static CHAOS_API FImplicit * NewImplicit(
 		Chaos::FErrorReporter ErrorReporter,
 		const Chaos::FParticles& MeshParticles,
 		const Chaos::FTriangleMesh& TriMesh,
@@ -90,28 +90,28 @@ public:
 		const ECollisionTypeEnum CollisionType,
 		const EImplicitTypeEnum ImplicitType);
 
-	static FImplicit* NewImplicitBox(
+	static CHAOS_API FImplicit* NewImplicitBox(
 		const FBox& CollisionBounds,
 		const float CollisionObjectReduction,
 		const ECollisionTypeEnum CollisionType);
 
-	static FImplicit* NewImplicitSphere(
+	static CHAOS_API FImplicit* NewImplicitSphere(
 		const Chaos::FReal Radius,
 		const float CollisionObjectReduction,
 		const ECollisionTypeEnum CollisionType);
 
-	static FImplicit* NewImplicitCapsule(
+	static CHAOS_API FImplicit* NewImplicitCapsule(
 		const Chaos::FReal Radius,
 		const Chaos::FReal Length,
 		const float CollisionObjectReduction,
 		const ECollisionTypeEnum CollisionType);
 
-	static FImplicit* NewImplicitCapsule(
+	static CHAOS_API FImplicit* NewImplicitCapsule(
 		const FBox& CollisionBounds,
 		const float CollisionObjectReduction,
 		const ECollisionTypeEnum CollisionType);
 
-	static FImplicit* NewImplicitConvex(
+	static CHAOS_API FImplicit* NewImplicitConvex(
 		const TArray<int32>& ConvexIndices,
 		const TManagedArray<TUniquePtr<Chaos::FConvex>>* ConvexGeometry,
 		const ECollisionTypeEnum CollisionType,
@@ -119,7 +119,7 @@ public:
 		const Chaos::FReal CollisionMarginFraction,
 		const float CollisionObjectReduction);
 
-	static FImplicit* NewImplicitLevelset(
+	static CHAOS_API FImplicit* NewImplicitLevelset(
 		Chaos::FErrorReporter ErrorReporter,
 		const Chaos::FParticles& MeshParticles,
 		const Chaos::FTriangleMesh& TriMesh,
@@ -129,7 +129,7 @@ public:
 		const float CollisionObjectReduction,
 		const ECollisionTypeEnum CollisionType);
 
-	static Chaos::FLevelSet* NewLevelset(
+	static CHAOS_API Chaos::FLevelSet* NewLevelset(
 		Chaos::FErrorReporter ErrorReporter,
 		const Chaos::FParticles& MeshParticles,
 		const Chaos::FTriangleMesh& TriMesh,
@@ -138,12 +138,12 @@ public:
 		const int32 MaxRes,
 		const ECollisionTypeEnum CollisionType);
 
-	static FVector CalculateUnitMassInertiaTensor(
+	static CHAOS_API FVector CalculateUnitMassInertiaTensor(
 		const FBox& BoundingBox,
 		const Chaos::FReal Radius,
 		const EImplicitTypeEnum ImplicitType);
 
-	static Chaos::FReal CalculateVolume(
+	static CHAOS_API Chaos::FReal CalculateVolume(
 		const FBox& BoundingBox,
 		const Chaos::FReal Radius,
 		const EImplicitTypeEnum ImplicitType);
diff --git a/Engine/Source/Runtime/Experimental/Chaos/Public/GeometryCollection/GeometryCollectionConvexPropertiesInterface.h b/Engine/Source/Runtime/Experimental/Chaos/Public/GeometryCollection/GeometryCollectionConvexPropertiesInterface.h
--- a/Engine/Source/Runtime/Experimental/Chaos/Public/GeometryCollection/GeometryCollectionConvexPropertiesInterface.h
+++ b/Engine/Source/Runtime/Experimental/Chaos/Public/GeometryCollection/GeometryCollectionConvexPropertiesInterface.h
@@ -8,28 +8,28 @@
 
 class FGeometryCollection;
 
-class CHAOS_API FGeometryCollectionConvexPropertiesInterface : public FManagedArrayInterface
+class FGeometryCollectionConvexPropertiesInterface : public FManagedArrayInterface
 {
 public :
 	typedef FManagedArrayInterface Super;
 	using FManagedArrayInterface::ManagedCollection;
 
 	// Convex Properties Group Name
-	static const FName ConvexPropertiesGroup; 
+	static CHAOS_API const FName ConvexPropertiesGroup; 
 	// Attribute
-	static const FName ConvexIndexAttribute;
+	static CHAOS_API const FName ConvexIndexAttribute;
 	// Attribute
-	static const FName ConvexEnable;
+	static CHAOS_API const FName ConvexEnable;
 	// Attribute
-	static const FName ConvexFractionRemoveAttribute;
+	static CHAOS_API const FName ConvexFractionRemoveAttribute;
 	// Attribute
-	static const FName ConvexSimplificationThresholdAttribute;
+	static CHAOS_API const FName ConvexSimplificationThresholdAttribute;
 	// Attribute
-	static const FName ConvexCanExceedFractionAttribute;
+	static CHAOS_API const FName ConvexCanExceedFractionAttribute;
 	// Attribute
-	static const FName ConvexRemoveOverlapsMethodAttribute;
+	static CHAOS_API const FName ConvexRemoveOverlapsMethodAttribute;
 	// Attribute
-	static const FName ConvexRemoveOverlapsShrinkAttribute;
+	static CHAOS_API const FName ConvexRemoveOverlapsShrinkAttribute;
 
 
 	struct FConvexCreationProperties {
@@ -41,16 +41,16 @@ public :
 		float OverlapRemovalShrinkPercent = 0.0f;
 	};
 
-	FGeometryCollectionConvexPropertiesInterface(FGeometryCollection* InGeometryCollection);
+	CHAOS_API FGeometryCollectionConvexPropertiesInterface(FGeometryCollection* InGeometryCollection);
 
-	void InitializeInterface() override;
+	CHAOS_API void InitializeInterface() override;
 
-	void CleanInterfaceForCook() override;
+	CHAOS_API void CleanInterfaceForCook() override;
 	
-	void RemoveInterfaceAttributes() override;
+	CHAOS_API void RemoveInterfaceAttributes() override;
 
-	FConvexCreationProperties GetConvexProperties(int TransformGroupIndex = INDEX_NONE) const;
-	void SetConvexProperties(const FConvexCreationProperties&, int TransformGroupIndex = INDEX_NONE);
+	CHAOS_API FConvexCreationProperties GetConvexProperties(int TransformGroupIndex = INDEX_NONE) const;
+	CHAOS_API void SetConvexProperties(const FConvexCreationProperties&, int TransformGroupIndex = INDEX_NONE);
 private:
 
 	void SetDefaultProperty();
diff --git a/Engine/Source/Runtime/Experimental/Chaos/Public/GeometryCollection/GeometryCollectionConvexUtility.h b/Engine/Source/Runtime/Experimental/Chaos/Public/GeometryCollection/GeometryCollectionConvexUtility.h
--- a/Engine/Source/Runtime/Experimental/Chaos/Public/GeometryCollection/GeometryCollectionConvexUtility.h
+++ b/Engine/Source/Runtime/Experimental/Chaos/Public/GeometryCollection/GeometryCollectionConvexUtility.h
@@ -80,7 +80,7 @@ namespace UE::GeometryCollectionConvexUtility
 
 }
 
-class CHAOS_API FGeometryCollectionConvexUtility
+class FGeometryCollectionConvexUtility
 {
 public:
 
@@ -91,13 +91,13 @@ public:
 	};
 
 	/** Ensure that convex hull data exists for the Geometry Collection and construct it if not (or if some data is missing. */
-	static FGeometryCollectionConvexData GetValidConvexHullData(FGeometryCollection* GeometryCollection);
+	static CHAOS_API FGeometryCollectionConvexData GetValidConvexHullData(FGeometryCollection* GeometryCollection);
 
 	/** Get convex hull data for the Geometry Collection if it is present */
-	static TOptional<FGeometryCollectionConvexData> GetConvexHullDataIfPresent(FManagedArrayCollection* GeometryCollection);
+	static CHAOS_API TOptional<FGeometryCollectionConvexData> GetConvexHullDataIfPresent(FManagedArrayCollection* GeometryCollection);
 
 	/** @return true if convex hull data is present */
-	static bool HasConvexHullData(const FManagedArrayCollection* GeometryCollection);
+	static CHAOS_API bool HasConvexHullData(const FManagedArrayCollection* GeometryCollection);
 
 	/**
 	 Create non-overlapping convex hull data for all transforms in the geometry collection (except transforms where it would be better to just use the hulls of the children) 
@@ -110,7 +110,7 @@ public:
 	 @param OverlapRemovalShrinkPercent			Compute overlaps based on objects shrunk by this percentage, so objects that would not overlap with this value set as their 'Collision Object Reduction Percentage' will not be cut
 	 @param ComputedLeafHullsToModify			Optional pre-computed hulls for geometry of rigid leaves, in the top-level coordinate space of the geometry collection. If passed in, the data will be updated and moved.
 	 */
-	static FGeometryCollectionConvexData CreateNonOverlappingConvexHullData(FGeometryCollection* GeometryCollection, double FractionAllowRemove = .3, double SimplificationDistanceThreshold = 0.0, double CanExceedFraction = .5, 
+	static CHAOS_API FGeometryCollectionConvexData CreateNonOverlappingConvexHullData(FGeometryCollection* GeometryCollection, double FractionAllowRemove = .3, double SimplificationDistanceThreshold = 0.0, double CanExceedFraction = .5, 
 		EConvexOverlapRemoval OverlapRemovalMethod = EConvexOverlapRemoval::All, double OverlapRemovalShrinkPercent = 0.0, UE::GeometryCollectionConvexUtility::FConvexHulls* ComputedLeafHullsToModify = nullptr);
 
 	struct FClusterConvexHullSettings
@@ -127,10 +127,10 @@ public:
 		UE::Geometry::FSphereCovering* EmptySpace = nullptr;
 	};
 
-	static void GenerateClusterConvexHullsFromChildrenHulls(FGeometryCollection& Collection, const FClusterConvexHullSettings& Settings, const TArrayView<const int32> TransformSubset);
-	static void GenerateClusterConvexHullsFromChildrenHulls(FGeometryCollection& Collection, const FClusterConvexHullSettings& Settings);
-	static void GenerateClusterConvexHullsFromLeafHulls(FGeometryCollection& Collection, const FClusterConvexHullSettings& Settings, const TArrayView<const int32> OptionalTransformSubset);
-	static void GenerateClusterConvexHullsFromLeafHulls(FGeometryCollection& Collection, const FClusterConvexHullSettings& Settings);
+	static CHAOS_API void GenerateClusterConvexHullsFromChildrenHulls(FGeometryCollection& Collection, const FClusterConvexHullSettings& Settings, const TArrayView<const int32> TransformSubset);
+	static CHAOS_API void GenerateClusterConvexHullsFromChildrenHulls(FGeometryCollection& Collection, const FClusterConvexHullSettings& Settings);
+	static CHAOS_API void GenerateClusterConvexHullsFromLeafHulls(FGeometryCollection& Collection, const FClusterConvexHullSettings& Settings, const TArrayView<const int32> OptionalTransformSubset);
+	static CHAOS_API void GenerateClusterConvexHullsFromLeafHulls(FGeometryCollection& Collection, const FClusterConvexHullSettings& Settings);
 
 	// Additional settings for filtering when the EGenerateConvexMethod::IntersectExternalWithComputed is applied
 	struct FIntersectionFilters
@@ -203,28 +203,28 @@ public:
 		FConvexDecompositionSettings DecompositionSettings;
 	};
 	
-	static void GenerateLeafConvexHulls(FGeometryCollection& Collection, bool bRestrictToSelection, const TArrayView<const int32> TransformSubset, const FLeafConvexHullSettings& Settings);
+	static CHAOS_API void GenerateLeafConvexHulls(FGeometryCollection& Collection, bool bRestrictToSelection, const TArrayView<const int32> TransformSubset, const FLeafConvexHullSettings& Settings);
 
 	/** Returns the convex hull of the vertices contained in the specified geometry. */
-	static TUniquePtr<Chaos::FConvex> FindConvexHull(const FGeometryCollection* GeometryCollection, int32 GeometryIndex);
+	static CHAOS_API TUniquePtr<Chaos::FConvex> FindConvexHull(const FGeometryCollection* GeometryCollection, int32 GeometryIndex);
 
 	/** Delete the convex hulls pointed at by the transform indices provided. */
-	static void RemoveConvexHulls(FGeometryCollection* GeometryCollection, const TArray<int32>& SortedTransformDeletes);
+	static CHAOS_API void RemoveConvexHulls(FGeometryCollection* GeometryCollection, const TArray<int32>& SortedTransformDeletes);
 
 	/** Delete the convex hulls that are null */
-	static void RemoveEmptyConvexHulls(FManagedArrayCollection& GeometryCollection);
+	static CHAOS_API void RemoveEmptyConvexHulls(FManagedArrayCollection& GeometryCollection);
 
 	/** Set default values for convex hull related managed arrays. */
-	static void SetDefaults(FGeometryCollection* GeometryCollection, FName Group, uint32 StartSize, uint32 NumElements);
+	static CHAOS_API void SetDefaults(FGeometryCollection* GeometryCollection, FName Group, uint32 StartSize, uint32 NumElements);
 
 	/** Get the HasCustomConvex flags.  If they're missing, either add them (if bAddIfMissing) or return nullptr */
-	static TManagedArray<int32>* GetCustomConvexFlags(FGeometryCollection* GeometryCollection, bool bAddIfMissing = false);
+	static CHAOS_API TManagedArray<int32>* GetCustomConvexFlags(FGeometryCollection* GeometryCollection, bool bAddIfMissing = false);
 
 	/** @return true if the GeometryCollection has convex data with no null pointers and no invalid indices */
-	static bool ValidateConvexData(const FGeometryCollection* GeometryCollection);
+	static CHAOS_API bool ValidateConvexData(const FGeometryCollection* GeometryCollection);
 
 	/** Set Volume and Size attributes on the Collection (will be called by CreateNonOverlappingConvexHullData -- Volumes must be up to date for convex calc) */
-	static void SetVolumeAttributes(FManagedArrayCollection* Collection);
+	static CHAOS_API void SetVolumeAttributes(FManagedArrayCollection* Collection);
 
 	/**
 	 * Copy convex hulls from *below* FromTransformIdx over to all live at ToTransformIdx.
@@ -237,13 +237,13 @@ public:
 	 * @param ToTransformIdx	The transform indices whose convexes will be replaced with *copies* of the child convexes.  Must be same length as FromTransformIdx; can be the same array.
 	 * @param bLeafOnly			If true, we will only collect convexes from leaf bones, not from clusters.
 	 */
-	static void CopyChildConvexes(const FGeometryCollection* FromCollection, const TArrayView<const int32>& FromTransformIdx, FGeometryCollection* ToCollection, const TArrayView<const int32>& ToTransformIdx, bool bLeafOnly);
+	static CHAOS_API void CopyChildConvexes(const FGeometryCollection* FromCollection, const TArrayView<const int32>& FromTransformIdx, FGeometryCollection* ToCollection, const TArrayView<const int32>& ToTransformIdx, bool bLeafOnly);
 
 	// Compute just the hulls of the leaf / rigid nodes that hold geometry directly, with no cluster hulls and no overlap removal by cutting
 	// This is an initial step of several algorithms: The CreateNonOverlappingConvexHullData function as well as convex-based proximity detection (TODO: and the auto-embed algorithm?)
 	// (TODO: Make auto-embed use this instead of the full hulls?)
 	// @param GlobalTransformArray		GeometryCollection's transforms to global space, as computed by GeometryCollectionAlgo::GlobalMatrices
-	static UE::GeometryCollectionConvexUtility::FConvexHulls ComputeLeafHulls(FGeometryCollection* GeometryCollection, const TArray<FTransform>& GlobalTransformArray, double SimplificationDistanceThreshold = 0.0, double OverlapRemovalShrinkPercent = 0.0,
+	static CHAOS_API UE::GeometryCollectionConvexUtility::FConvexHulls ComputeLeafHulls(FGeometryCollection* GeometryCollection, const TArray<FTransform>& GlobalTransformArray, double SimplificationDistanceThreshold = 0.0, double OverlapRemovalShrinkPercent = 0.0,
 		TFunction<bool(int32)> SkipBoneFn = nullptr, const FConvexDecompositionSettings* OptionalDecompositionSettings = nullptr);
 
 	struct FTransformedConvex
@@ -256,7 +256,7 @@ public:
 	// suported shapes are scaled / transformed implicits as well as Boxes, convexes, spheres and capsules
 	// levelset, tapered capsule, cylinder, heightfields and trimesh are not supported
 	// array will not be erased and convex will be added if any already exists
-	static void ConvertImplicitToConvexArray(const Chaos::FImplicitObject& InImplicit, const FTransform& Transform, TArray<FTransformedConvex>& InOutConvex);
+	static CHAOS_API void ConvertImplicitToConvexArray(const Chaos::FImplicitObject& InImplicit, const FTransform& Transform, TArray<FTransformedConvex>& InOutConvex);
 
 private:
 	static void ConvertScaledImplicitToConvexArray(
diff --git a/Engine/Source/Runtime/Experimental/Chaos/Public/GeometryCollection/GeometryCollectionProximityPropertiesInterface.h b/Engine/Source/Runtime/Experimental/Chaos/Public/GeometryCollection/GeometryCollectionProximityPropertiesInterface.h
--- a/Engine/Source/Runtime/Experimental/Chaos/Public/GeometryCollection/GeometryCollectionProximityPropertiesInterface.h
+++ b/Engine/Source/Runtime/Experimental/Chaos/Public/GeometryCollection/GeometryCollectionProximityPropertiesInterface.h
@@ -8,29 +8,29 @@
 
 class FGeometryCollection;
 
-class CHAOS_API FGeometryCollectionProximityPropertiesInterface : public FManagedArrayInterface
+class FGeometryCollectionProximityPropertiesInterface : public FManagedArrayInterface
 {
 public :
 	typedef FManagedArrayInterface Super;
 	using FManagedArrayInterface::ManagedCollection;
 
 	// Proximity Properties Group Name
-	static const FName ProximityPropertiesGroup;
+	static CHAOS_API const FName ProximityPropertiesGroup;
 	
 	// Attribute: Method to determine proximity
-	static const FName ProximityDetectionMethod;
+	static CHAOS_API const FName ProximityDetectionMethod;
 	// Attribute: For convex hull proximity, what separation distance can still be considered as 'proximity'
-	static const FName ProximityDistanceThreshold;
+	static CHAOS_API const FName ProximityDistanceThreshold;
 	// Attribute: Whether to use the computed proximity graph as a connection graph
-	static const FName ProximityAsConnectionGraph;
+	static CHAOS_API const FName ProximityAsConnectionGraph;
 	// Attribute: How to compute contact area for connection graph edges, to be used by the simulation
-	static const FName ProximityConnectionGraphContactAreaMethod;
+	static CHAOS_API const FName ProximityConnectionGraphContactAreaMethod;
 	// Attribute: If greater than zero, filter proximity connections by requiring an amount of 'contact' as computed by the Contact Method
 	// This is a second filter applied after initial proximity connections were determined by any Proximity Detection Method,
 	// and can be used to reduce unsupported or spurious/glancing proximities.
-	static const FName ProximityRequireContactAmount;
+	static CHAOS_API const FName ProximityRequireContactAmount;
 	// Attribute: How to define 'contact' for the above Require Contact Amount.
-	static const FName ProximityContactMethod;
+	static CHAOS_API const FName ProximityContactMethod;
 
 	struct FProximityProperties
 	{
@@ -42,15 +42,15 @@ public :
 		bool bUseAsConnectionGraph = false;
 	};
 
-	FGeometryCollectionProximityPropertiesInterface(FGeometryCollection* InGeometryCollection);
+	CHAOS_API FGeometryCollectionProximityPropertiesInterface(FGeometryCollection* InGeometryCollection);
 
-	void InitializeInterface() override;
+	CHAOS_API void InitializeInterface() override;
 
-	void CleanInterfaceForCook() override;
+	CHAOS_API void CleanInterfaceForCook() override;
 	
-	void RemoveInterfaceAttributes() override;
+	CHAOS_API void RemoveInterfaceAttributes() override;
 
-	FProximityProperties GetProximityProperties() const;
-	void SetProximityProperties(const FProximityProperties&);
+	CHAOS_API FProximityProperties GetProximityProperties() const;
+	CHAOS_API void SetProximityProperties(const FProximityProperties&);
 };
 
diff --git a/Engine/Source/Runtime/Experimental/Chaos/Public/GeometryCollection/GeometryCollectionProximityUtility.h b/Engine/Source/Runtime/Experimental/Chaos/Public/GeometryCollection/GeometryCollectionProximityUtility.h
--- a/Engine/Source/Runtime/Experimental/Chaos/Public/GeometryCollection/GeometryCollectionProximityUtility.h
+++ b/Engine/Source/Runtime/Experimental/Chaos/Public/GeometryCollection/GeometryCollectionProximityUtility.h
@@ -44,17 +44,17 @@ enum class EConnectionContactMethod : uint8
 	ConvexHullContactArea
 };
 
-class CHAOS_API FGeometryCollectionProximityUtility
+class FGeometryCollectionProximityUtility
 {
 public:
-	FGeometryCollectionProximityUtility(FGeometryCollection* InCollection);
+	CHAOS_API FGeometryCollectionProximityUtility(FGeometryCollection* InCollection);
 
-	void UpdateProximity(UE::GeometryCollectionConvexUtility::FConvexHulls* OptionalComputedHulls = nullptr);
+	CHAOS_API void UpdateProximity(UE::GeometryCollectionConvexUtility::FConvexHulls* OptionalComputedHulls = nullptr);
 
 	// Update proximity data if it is not already present
-	void RequireProximity(UE::GeometryCollectionConvexUtility::FConvexHulls* OptionalComputedHulls = nullptr);
+	CHAOS_API void RequireProximity(UE::GeometryCollectionConvexUtility::FConvexHulls* OptionalComputedHulls = nullptr);
 
-	void InvalidateProximity();
+	CHAOS_API void InvalidateProximity();
 
 	// Stores stats about the contact between two geometries
 	struct FGeometryContactEdge
@@ -75,11 +75,11 @@ public:
 		{}
 	};
 	// Note: This computes connections from lower to higher geometry indices, assuming connections are symmetric
-	static TArray<FGeometryContactEdge> ComputeConvexGeometryContactFromProximity(FGeometryCollection* Collection, float DistanceTolerance, UE::GeometryCollectionConvexUtility::FConvexHulls& LocalHulls);
+	static CHAOS_API TArray<FGeometryContactEdge> ComputeConvexGeometryContactFromProximity(FGeometryCollection* Collection, float DistanceTolerance, UE::GeometryCollectionConvexUtility::FConvexHulls& LocalHulls);
 
 	// @param ContactEdges	Optional pre-computed proximity contact edges, to be used for computing contact areas on the connection graph edges
-	void CopyProximityToConnectionGraph(const TArray<FGeometryContactEdge>* ContactEdges = nullptr);
-	void ClearConnectionGraph();
+	CHAOS_API void CopyProximityToConnectionGraph(const TArray<FGeometryContactEdge>* ContactEdges = nullptr);
+	CHAOS_API void ClearConnectionGraph();
 
 private:
 	FGeometryCollection* Collection;
diff --git a/Engine/Source/Runtime/Experimental/Chaos/Public/GeometryCollection/GeometryCollectionSection.h b/Engine/Source/Runtime/Experimental/Chaos/Public/GeometryCollection/GeometryCollectionSection.h
--- a/Engine/Source/Runtime/Experimental/Chaos/Public/GeometryCollection/GeometryCollectionSection.h
+++ b/Engine/Source/Runtime/Experimental/Chaos/Public/GeometryCollection/GeometryCollectionSection.h
@@ -9,7 +9,7 @@ struct FManagedArrayCollection;
 /**
 * A set of triangles which are rendered with the same material.
 */
-struct CHAOS_API FGeometryCollectionSection
+struct FGeometryCollectionSection
 {
 	/** Constructor. */
 	FGeometryCollectionSection()
@@ -33,7 +33,7 @@ struct CHAOS_API FGeometryCollectionSection
 		return true;
 	}
 
-	static TArray<FGeometryCollectionSection>
+	static CHAOS_API TArray<FGeometryCollectionSection>
 	BuildMeshSections(const FManagedArrayCollection& InCollection, const TArray<FIntVector>& InputIndices, const TArray<int32>& BaseMeshOriginalIndicesIndex, TArray<FIntVector>& RetIndices);
 
 	/** The index of the material with which to render this section. */
@@ -48,4 +48,4 @@ struct CHAOS_API FGeometryCollectionSection
 	HHitProxy* HitProxy;
 
 	static const int InvalidIndex = -1;
-};
\ No newline at end of file
+};
