[Chaos]

diff --git a/Engine/Source/Runtime/Experimental/Chaos/Public/GeometryCollection/ManagedArrayCollection.h b/Engine/Source/Runtime/Experimental/Chaos/Public/GeometryCollection/ManagedArrayCollection.h
--- a/Engine/Source/Runtime/Experimental/Chaos/Public/GeometryCollection/ManagedArrayCollection.h
+++ b/Engine/Source/Runtime/Experimental/Chaos/Public/GeometryCollection/ManagedArrayCollection.h
@@ -41,7 +41,7 @@ namespace Chaos
 *
 */
 USTRUCT()
-struct CHAOS_API FManagedArrayCollection
+struct FManagedArrayCollection
 {
 	GENERATED_USTRUCT_BODY()
 	friend FSimulationProperties;
@@ -49,7 +49,7 @@ struct CHAOS_API FManagedArrayCollection
 public:
 
 
-	FManagedArrayCollection();
+	CHAOS_API FManagedArrayCollection();
 	virtual ~FManagedArrayCollection() {}
 
 	FManagedArrayCollection(const FManagedArrayCollection& In) { In.CopyTo(this); }
@@ -57,7 +57,7 @@ public:
 	FManagedArrayCollection(FManagedArrayCollection&&) = default;
 	FManagedArrayCollection& operator=(FManagedArrayCollection&&)= default;
 
-	static int8 Invalid;
+	static CHAOS_API int8 Invalid;
 	typedef EManagedArrayType EArrayType;
 	/**
 	*
@@ -230,18 +230,18 @@ public:
 	* Create a group on the collection. Adding attribute will also create unknown groups.
 	* @param Group - The group name
 	*/
-	void AddGroup(FName Group);
+	CHAOS_API void AddGroup(FName Group);
 
 	/**
 	* Returns the number of attributes in a group.
 	*/
-	int32 NumAttributes(FName Group) const;
+	CHAOS_API int32 NumAttributes(FName Group) const;
 
 	/**
 	* List all the attributes in a group names.
 	* @return list of group names
 	*/
-	TArray<FName> AttributeNames(FName Group) const;
+	CHAOS_API TArray<FName> AttributeNames(FName Group) const;
 
 	/**
 	* Add elements to a group
@@ -249,7 +249,7 @@ public:
 	* @param Group - The group to append entries to.
 	* @return starting index of the new ManagedArray<T> entries.
 	*/
-	int32 AddElements(int32 NumberElements, FName Group);
+	CHAOS_API int32 AddElements(int32 NumberElements, FName Group);
 
 	/**
 	* Insert elements to a group
@@ -258,14 +258,14 @@ public:
 	* @param Group - The group to append entries to.
 	* @return starting index of the new ManagedArray<T> entries (same as Position).
 	*/
-	int32 InsertElements(int32 NumberElements, int32 Position, FName Group);
+	CHAOS_API int32 InsertElements(int32 NumberElements, int32 Position, FName Group);
 
 
 	/**
 	* Append Collection and reindex dependencies on this collection. 
 	* @param InCollection : Collection to add. 
 	*/
-	virtual void Append(const FManagedArrayCollection& Collection);
+	CHAOS_API virtual void Append(const FManagedArrayCollection& Collection);
 
 	/**
 	* Returns attribute(Name) of Type(T) from the group
@@ -380,7 +380,7 @@ public:
 	/**
 	* Remove the element at index and reindex the dependent arrays 
 	*/
-	virtual void RemoveElements(const FName & Group, const TArray<int32> & SortedDeletionList, FProcessingParameters Params = FProcessingParameters());
+	CHAOS_API virtual void RemoveElements(const FName & Group, const TArray<int32> & SortedDeletionList, FProcessingParameters Params = FProcessingParameters());
 
 	/**
 	* Remove the elements at Position and reindex the dependent arrays 
@@ -388,7 +388,7 @@ public:
 	* @param NumberElements - The number of array entries to remove
 	* @param Position - The position from which to remove entries
 	*/
-	virtual void RemoveElements(const FName& Group, int32 NumberElements, int32 Position);
+	CHAOS_API virtual void RemoveElements(const FName& Group, int32 NumberElements, int32 Position);
 
 
 	/**
@@ -396,32 +396,32 @@ public:
 	* @param Name - The name of the attribute
 	* @param Group - The group that manages the attribute
 	*/
-	void RemoveAttribute(FName Name, FName Group);
+	CHAOS_API void RemoveAttribute(FName Name, FName Group);
 
 
 	/**
 	* Remove the group from the collection.
 	* @param Group - The group that manages the attribute
 	*/
-	void RemoveGroup(FName Group);
+	CHAOS_API void RemoveGroup(FName Group);
 
 	/**
 	* List all the group names.
 	*/
-	TArray<FName> GroupNames() const;
+	CHAOS_API TArray<FName> GroupNames() const;
 
 	/**
 	* Check for the existence of a attribute
 	* @param Name - The name of the attribute
 	* @param Group - The group that manages the attribute
 	*/
-	bool HasAttribute(FName Name, FName Group) const;
+	CHAOS_API bool HasAttribute(FName Name, FName Group) const;
 
 	/**
 	* Check for the existence of a attribute
 	* @param TArray<ManagedType> - Attributes to search for
 	*/
-	bool HasAttributes(const TArray<FManagedType>& Types) const;
+	CHAOS_API bool HasAttributes(const TArray<FManagedType>& Types) const;
 
 	/**
 	* Check for the existence of a group
@@ -434,14 +434,14 @@ public:
 	* @param Name - The name of the attribute
 	* @param Group - The group that manages the attribute
 	*/
-	EArrayType GetAttributeType(FName Name, FName Group) const;
+	CHAOS_API EArrayType GetAttributeType(FName Name, FName Group) const;
 
 	/**
 	* Check if an attribute is dirty
 	* @param Name - The name of the attribute
 	* @param Group - The group that manages the attribute
 	*/
-	bool IsAttributeDirty(FName Name, FName Group) const;
+	CHAOS_API bool IsAttributeDirty(FName Name, FName Group) const;
 
 	/**
 	* Check if an attribute is persistent - ie its data will be serialized
@@ -449,24 +449,24 @@ public:
 	* @param Name - The name of the attribute
 	* @param Group - The group that manages the attribute
 	*/
-	bool IsAttributePersistent(FName Name, FName Group) const;
+	CHAOS_API bool IsAttributePersistent(FName Name, FName Group) const;
 	
 	/**
 	*
 	*/
-	void SetDependency(FName Name, FName Group, FName DependencyGroup, bool bAllowCircularDependency = false);
+	CHAOS_API void SetDependency(FName Name, FName Group, FName DependencyGroup, bool bAllowCircularDependency = false);
 
 	/**
 	*
 	*/
-	void RemoveDependencyFor(FName Group);
+	CHAOS_API void RemoveDependencyFor(FName Group);
 
 	/**
 	* Copy an attribute. Will perform an implicit group sync. Attribute must exist in both InCollection and this collection
 	* @param Name - The name of the attribute
 	* @param Group - The group that manages the attribute
 	*/
-	void CopyAttribute(const FManagedArrayCollection& InCollection, FName Name, FName Group);
+	CHAOS_API void CopyAttribute(const FManagedArrayCollection& InCollection, FName Name, FName Group);
 
 	/**
 	* Copy an attribute. Will perform an implicit group sync.
@@ -474,45 +474,45 @@ public:
 	* @param DestName - The name of the attribute being copied to
 	* @param Group - The group that manages the attribute
 	*/
-	void CopyAttribute(const FManagedArrayCollection& InCollection, FName SrcName, FName DestName, FName Group);
+	CHAOS_API void CopyAttribute(const FManagedArrayCollection& InCollection, FName SrcName, FName DestName, FName Group);
 
 	/**
 	* Copy attributes that match the input collection. This is a utility to easily sync collections
 	* @param InCollection - All groups from this collection found in the input will be sized accordingly
 	* @param SkipList - Group/Attrs to skip. Keys are group names, values are attributes in those groups.
 	*/
-	void CopyMatchingAttributesFrom(const FManagedArrayCollection& InCollection, const TMap<FName, TSet<FName>>* SkipList=nullptr);
+	CHAOS_API void CopyMatchingAttributesFrom(const FManagedArrayCollection& InCollection, const TMap<FName, TSet<FName>>* SkipList=nullptr);
 
 	/**
 	* Number of elements in a group
 	* @return the group size, and Zero(0) if group does not exist.
 	*/
-	int32 NumElements(FName Group) const;
+	CHAOS_API int32 NumElements(FName Group) const;
 
 	/**
 	* Resize a group
 	* @param Size - The size of the group
 	* @param Group - The group to resize
 	*/
-	void Resize(int32 Size, FName Group);
+	CHAOS_API void Resize(int32 Size, FName Group);
 
 	/**
 	* Reserve a group
 	* @param Size - The size of the group
 	* @param Group - The group to reserve
 	*/
-	void Reserve(int32 Size, FName Group);
+	CHAOS_API void Reserve(int32 Size, FName Group);
 
 	/**
 	* Empty the group 
 	*  @param Group - The group to empty
 	*/
-	void EmptyGroup(FName Group);
+	CHAOS_API void EmptyGroup(FName Group);
 
 	/**
 	* Reorders elements in a group. NewOrder must be the same length as the group.
 	*/
-	virtual void ReorderElements(FName Group, const TArray<int32>& NewOrder);
+	CHAOS_API virtual void ReorderElements(FName Group, const TArray<int32>& NewOrder);
 #if 0	//not needed until we support per instance serialization
 	/**
 	* Swap elements within a group
@@ -520,7 +520,7 @@ public:
 	* @param Index2 - The second location to be swapped
 	* @param Group - The group to resize
 	*/
-	void SwapElements(int32 Index1, int32 Index2, FName Group);
+	CHAOS_API void SwapElements(int32 Index1, int32 Index2, FName Group);
 #endif
 
 	/**
@@ -533,29 +533,29 @@ public:
 	/**
 	* Serialize
 	*/
-	virtual void Serialize(Chaos::FChaosArchive& Ar);
+	CHAOS_API virtual void Serialize(Chaos::FChaosArchive& Ar);
 
 	/**
 	* Cycle Checking. 
 	* Search for TargetNode from the StartingNode. 
 	*/
-	bool IsConnected(FName StartingNode, FName TargetNode);
+	CHAOS_API bool IsConnected(FName StartingNode, FName TargetNode);
 
 
 	/**
 	* Dump the contents to a FString
 	*/
-	FString ToString() const;
+	CHAOS_API FString ToString() const;
 
 	/**
 	* Get allocated memory size 
 	*/
-	SIZE_T GetAllocatedSize() const;
+	CHAOS_API SIZE_T GetAllocatedSize() const;
 
 	/**
 	 * Get information on the total storage required for each element in each group
 	 */
-	 void GetElementSizeInfoForGroups(TArray<TPair<FName, SIZE_T>>& OutSizeInfo) const;
+	 CHAOS_API void GetElementSizeInfoForGroups(TArray<TPair<FName, SIZE_T>>& OutSizeInfo) const;
 
 private:
 
@@ -670,7 +670,7 @@ private:
 
 protected:
 
-	virtual void SetDefaults(FName Group, uint32 StartSize, uint32 NumElements);
+	CHAOS_API virtual void SetDefaults(FName Group, uint32 StartSize, uint32 NumElements);
 
 	/**
 	 * Virtual helper function called by CopyMatchingAttributesFrom; adds attributes 'default, but optional' attributes that are present in InCollection
@@ -684,7 +684,7 @@ protected:
 	* @param InCollection - The collection we are ordering our group against. 
 	* @param Group - The group that manages the attribute
 	*/
-	void SyncGroupSizeFrom(const FManagedArrayCollection& InCollection, FName Group);
+	CHAOS_API void SyncGroupSizeFrom(const FManagedArrayCollection& InCollection, FName Group);
 
 	/**
 	 * Version to indicate need for conditioning to current expected data layout during serialization loading.
@@ -708,7 +708,7 @@ protected:
 /*
 *  FManagedArrayInterface
 */
-class CHAOS_API FManagedArrayInterface
+class FManagedArrayInterface
 {
 public:
 	FManagedArrayInterface() : ManagedCollection(nullptr) {}
diff --git a/Engine/Source/Runtime/Experimental/Chaos/Public/GeometryCollection/RecordedTransformTrack.h b/Engine/Source/Runtime/Experimental/Chaos/Public/GeometryCollection/RecordedTransformTrack.h
--- a/Engine/Source/Runtime/Experimental/Chaos/Public/GeometryCollection/RecordedTransformTrack.h
+++ b/Engine/Source/Runtime/Experimental/Chaos/Public/GeometryCollection/RecordedTransformTrack.h
@@ -269,7 +269,7 @@ struct FRecordedFrame
 };
 
 USTRUCT()
-struct CHAOS_API FRecordedTransformTrack
+struct FRecordedTransformTrack
 {
 	GENERATED_BODY()
 
@@ -550,6 +550,6 @@ struct CHAOS_API FRecordedTransformTrack
 		return (Axis * Angle) / SampleWidth;
 	}
 
-	static FRecordedTransformTrack ProcessRawRecordedData(const FRecordedTransformTrack& InCache);
+	static CHAOS_API FRecordedTransformTrack ProcessRawRecordedData(const FRecordedTransformTrack& InCache);
 
 };
diff --git a/Engine/Source/Runtime/Experimental/Chaos/Public/GeometryCollection/TransformCollection.h b/Engine/Source/Runtime/Experimental/Chaos/Public/GeometryCollection/TransformCollection.h
--- a/Engine/Source/Runtime/Experimental/Chaos/Public/GeometryCollection/TransformCollection.h
+++ b/Engine/Source/Runtime/Experimental/Chaos/Public/GeometryCollection/TransformCollection.h
@@ -13,12 +13,12 @@
 *
 * @see FTransformCollectionComponent
 */
-class CHAOS_API FTransformCollection : public FManagedArrayCollection
+class FTransformCollection : public FManagedArrayCollection
 {
 public:
 	typedef FManagedArrayCollection Super;
 
-	FTransformCollection();
+	CHAOS_API FTransformCollection();
 	FTransformCollection(FTransformCollection &) = delete;
 	FTransformCollection& operator=(const FTransformCollection&) = delete;
 	FTransformCollection(FTransformCollection&&) = default;
@@ -45,54 +45,54 @@ public:
 	*       Children defines the child indices of the transform node in the transform hierarchy (leaf nodes will have no children)
 	*       Level is the distance from the root node at level 0. Leaf nodes will have the highest level number.
 	*/
-	static const FName TransformGroup;
-	static const FName TransformAttribute;
-	static const FName ParentAttribute;
-	static const FName ChildrenAttribute;
-	static const FName ParticlesAttribute;
-	static const FName LevelAttribute;
+	static CHAOS_API const FName TransformGroup;
+	static CHAOS_API const FName TransformAttribute;
+	static CHAOS_API const FName ParentAttribute;
+	static CHAOS_API const FName ChildrenAttribute;
+	static CHAOS_API const FName ParticlesAttribute;
+	static CHAOS_API const FName LevelAttribute;
 
 
 	/** Schema */
-	static void DefineTransformSchema(FManagedArrayCollection&);
+	static CHAOS_API void DefineTransformSchema(FManagedArrayCollection&);
 
 
 	/** Serialize */
-	virtual void Serialize(Chaos::FChaosArchive& Ar) override;
+	CHAOS_API virtual void Serialize(Chaos::FChaosArchive& Ar) override;
 
 	/*
 	* SingleTransform:
 	*   Create a single transform.
 	*/
-	static FTransformCollection SingleTransform(const FTransform& TransformRoot = FTransform::Identity);
+	static CHAOS_API FTransformCollection SingleTransform(const FTransform& TransformRoot = FTransform::Identity);
 
 	/**
 	* Append Collection and reindex dependencies on this collection.
 	* @param InCollection : Collection to add.
 	*/
-	virtual void Append(const FManagedArrayCollection& Collection);
+	CHAOS_API virtual void Append(const FManagedArrayCollection& Collection);
 
 	/*
 	* AppendTransform:
 	*   Append a transform at the end of the collection without
 	*   parenting. 
 	*/
-	int32 AppendTransform(const FTransformCollection & GeometryCollection, const FTransform& TransformRoot = FTransform::Identity);
+	CHAOS_API int32 AppendTransform(const FTransformCollection & GeometryCollection, const FTransform& TransformRoot = FTransform::Identity);
 
 	/*
 	* ParentTransforms
 	*   Parent Transforms under the specified node using local parent
 	*   hierarchy compensation. .
 	*/
-	void ParentTransforms(const int32 TransformIndex, const int32 ChildIndex);
-	void ParentTransforms(const int32 TransformIndex, const TArray<int32>& SelectedBones);
-	void UnparentTransform(const int32 ChildIndex);
+	CHAOS_API void ParentTransforms(const int32 TransformIndex, const int32 ChildIndex);
+	CHAOS_API void ParentTransforms(const int32 TransformIndex, const TArray<int32>& SelectedBones);
+	CHAOS_API void UnparentTransform(const int32 ChildIndex);
 
 	/*
 	* RelativeTransformation
 	*   Modify the specified index by the local matrix offset. 
 	*/
-	void RelativeTransformation(const int32& Index, const FTransform& LocalOffset); 
+	CHAOS_API void RelativeTransformation(const int32& Index, const FTransform& LocalOffset); 
 
 	/**
 	* RemoveElements
@@ -100,12 +100,12 @@ public:
 	*   under the deleted elements parent using local parent compensation [relative local matrices].
 	* 
 	*/
-	virtual void RemoveElements(const FName & Group, const TArray<int32> & SortedDeletionList, FProcessingParameters Params = FProcessingParameters()) override;
+	CHAOS_API virtual void RemoveElements(const FName & Group, const TArray<int32> & SortedDeletionList, FProcessingParameters Params = FProcessingParameters()) override;
 
 	/**
 	* reset internal state
 	*/
-	virtual void Reset() override;
+	CHAOS_API virtual void Reset() override;
 
 	// Transform Group
 	TManagedArray<FTransform>   Transform;
@@ -118,5 +118,5 @@ public:
 protected:
 
 	/** Construct */
-	void Construct();
-};
\ No newline at end of file
+	CHAOS_API void Construct();
+};
diff --git a/Engine/Source/Runtime/Experimental/Chaos/Public/GeometryCollectionProxyData.h b/Engine/Source/Runtime/Experimental/Chaos/Public/GeometryCollectionProxyData.h
--- a/Engine/Source/Runtime/Experimental/Chaos/Public/GeometryCollectionProxyData.h
+++ b/Engine/Source/Runtime/Experimental/Chaos/Public/GeometryCollectionProxyData.h
@@ -17,12 +17,12 @@
 *
 * Stores per instance data for transforms and hierarchy information
 */
-class CHAOS_API FTransformDynamicCollection : public FManagedArrayCollection
+class FTransformDynamicCollection : public FManagedArrayCollection
 {
 public:
 	typedef FManagedArrayCollection Super;
 
-	FTransformDynamicCollection();
+	CHAOS_API FTransformDynamicCollection();
 	FTransformDynamicCollection(FTransformDynamicCollection&) = delete;
 	FTransformDynamicCollection& operator=(const FTransformDynamicCollection&) = delete;
 	FTransformDynamicCollection(FTransformDynamicCollection&&) = delete;
@@ -38,7 +38,7 @@ public:
 protected:
 
 	/** Construct */
-	void Construct();
+	CHAOS_API void Construct();
 };
 
 
@@ -48,10 +48,10 @@ protected:
 * Stores per instance data for simulation level information
 */
 
-class CHAOS_API FGeometryDynamicCollection : public FTransformDynamicCollection
+class FGeometryDynamicCollection : public FTransformDynamicCollection
 {
 public:
-	FGeometryDynamicCollection();
+	CHAOS_API FGeometryDynamicCollection();
 	FGeometryDynamicCollection(FGeometryDynamicCollection&) = delete;
 	FGeometryDynamicCollection& operator=(const FGeometryDynamicCollection&) = delete;
 	FGeometryDynamicCollection(FGeometryDynamicCollection&&) = delete;
@@ -60,16 +60,16 @@ public:
 	typedef FTransformDynamicCollection Super;
 	typedef TSharedPtr<Chaos::FImplicitObject, ESPMode::ThreadSafe> FSharedImplicit;
 
-	static const FName ActiveAttribute;
-	static const FName CollisionGroupAttribute;
-	static const FName CollisionMaskAttribute;
-	static const FName DynamicStateAttribute;
-	static const FName ImplicitsAttribute;
-	static const FName ShapesQueryDataAttribute;
-	static const FName ShapesSimDataAttribute;
-	static const FName SharedImplicitsAttribute;
-	static const FName SimplicialsAttribute;
-	static const FName SimulatableParticlesAttribute;
+	static CHAOS_API const FName ActiveAttribute;
+	static CHAOS_API const FName CollisionGroupAttribute;
+	static CHAOS_API const FName CollisionMaskAttribute;
+	static CHAOS_API const FName DynamicStateAttribute;
+	static CHAOS_API const FName ImplicitsAttribute;
+	static CHAOS_API const FName ShapesQueryDataAttribute;
+	static CHAOS_API const FName ShapesSimDataAttribute;
+	static CHAOS_API const FName SharedImplicitsAttribute;
+	static CHAOS_API const FName SimplicialsAttribute;
+	static CHAOS_API const FName SimulatableParticlesAttribute;
 
 	// Transform Group
 	TManagedArray<bool> Active;
@@ -100,7 +100,7 @@ public:
 	FInitialVelocityFacade GetInitialVelocityFacade() { return FInitialVelocityFacade(*this); }
 	FInitialVelocityFacade GetInitialVelocityFacade() const { return FInitialVelocityFacade(*this); }
 
-	void CopyInitialVelocityAttributesFrom(const FGeometryDynamicCollection& SourceCollection);
+	CHAOS_API void CopyInitialVelocityAttributesFrom(const FGeometryDynamicCollection& SourceCollection);
 };
 
 /**
@@ -108,31 +108,31 @@ public:
  * physics state , broken state, current parent (normal or internal clusters )
  * To be used with the dynamic collection
  */
-class CHAOS_API FGeometryCollectionDynamicStateFacade
+class FGeometryCollectionDynamicStateFacade
 {
 public:
-	FGeometryCollectionDynamicStateFacade(FManagedArrayCollection& InCollection);
+	CHAOS_API FGeometryCollectionDynamicStateFacade(FManagedArrayCollection& InCollection);
 
 	/** returns true if all the necessary attributes are present */
-	bool IsValid() const;
+	CHAOS_API bool IsValid() const;
 
 	/** return true if the transform is in a dynamic or sleeping state */
-	bool IsDynamicOrSleeping(int32 TransformIndex) const;
+	CHAOS_API bool IsDynamicOrSleeping(int32 TransformIndex) const;
 
 	/** return true if the transform is in a sleeping state */
-	bool IsSleeping(int32 TransformIndex) const;
+	CHAOS_API bool IsSleeping(int32 TransformIndex) const;
 
 	/** whether there's children attached to this transfom (Cluster) */
-	bool HasChildren(int32 TransformIndex) const;
+	CHAOS_API bool HasChildren(int32 TransformIndex) const;
 	
 	/** return true if the transform has broken off its parent */
-	bool HasBrokenOff(int32 TransformIndex) const;
+	CHAOS_API bool HasBrokenOff(int32 TransformIndex) const;
 
 	/** return true if the transform has an internal cluster parent */
-	bool HasInternalClusterParent(int32 TransformIndex) const;
+	CHAOS_API bool HasInternalClusterParent(int32 TransformIndex) const;
 
 	/** return true if the transform has an internal cluster parent in a dynamic state */
-	bool HasDynamicInternalClusterParent(int32 TransformIndex) const;
+	CHAOS_API bool HasDynamicInternalClusterParent(int32 TransformIndex) const;
 	
 private:
 	/** Active state, true means that the transform is active or broken off from its parent */
diff --git a/Engine/Source/Runtime/Experimental/Chaos/Public/PBDRigidActiveParticlesBuffer.h b/Engine/Source/Runtime/Experimental/Chaos/Public/PBDRigidActiveParticlesBuffer.h
--- a/Engine/Source/Runtime/Experimental/Chaos/Public/PBDRigidActiveParticlesBuffer.h
+++ b/Engine/Source/Runtime/Experimental/Chaos/Public/PBDRigidActiveParticlesBuffer.h
@@ -13,7 +13,7 @@ namespace Chaos
 	/**
 	 * Solver specific data buffered for use on Game thread
 	 */
-	struct CHAOS_API FPBDRigidDirtyParticlesBufferOut
+	struct FPBDRigidDirtyParticlesBufferOut
 	{
 		TArray<FSingleParticlePhysicsProxy*> DirtyGameThreadParticles;
 		// Some particle types (clustered) only exist on the game thread, but we
@@ -22,19 +22,19 @@ namespace Chaos
 	};
 
 
-	class CHAOS_API FPBDRigidDirtyParticlesBuffer
+	class FPBDRigidDirtyParticlesBuffer
 	{
 		friend class FPBDRigidDirtyParticlesBufferAccessor;
 
 	public:
-		FPBDRigidDirtyParticlesBuffer(const Chaos::EMultiBufferMode& InBufferMode, bool bInSingleThreaded);
+		CHAOS_API FPBDRigidDirtyParticlesBuffer(const Chaos::EMultiBufferMode& InBufferMode, bool bInSingleThreaded);
 
-		void CaptureSolverData(FPBDRigidsSolver* Solver);
+		CHAOS_API void CaptureSolverData(FPBDRigidsSolver* Solver);
 
-		void ReadLock();
-		void ReadUnlock();
-		void WriteLock();
-		void WriteUnlock();
+		CHAOS_API void ReadLock();
+		CHAOS_API void ReadUnlock();
+		CHAOS_API void WriteLock();
+		CHAOS_API void WriteUnlock();
 	
 	private:
 		const FPBDRigidDirtyParticlesBufferOut* GetSolverOutData() const
@@ -45,7 +45,7 @@ namespace Chaos
 		/**
 		 * Fill data from solver destined for the game thread - used to limit the number of objects updated on the game thread
 		 */
-		void BufferPhysicsResults(FPBDRigidsSolver* Solver);
+		CHAOS_API void BufferPhysicsResults(FPBDRigidsSolver* Solver);
 
 		/**
 		 * Flip should be performed on physics thread side non-game thread
diff --git a/Engine/Source/Runtime/Experimental/Chaos/Public/PBDRigidsSolver.h b/Engine/Source/Runtime/Experimental/Chaos/Public/PBDRigidsSolver.h
--- a/Engine/Source/Runtime/Experimental/Chaos/Public/PBDRigidsSolver.h
+++ b/Engine/Source/Runtime/Experimental/Chaos/Public/PBDRigidsSolver.h
@@ -79,11 +79,11 @@ namespace Chaos
 	/**
 	*
 	*/
-	class CHAOS_API FPBDRigidsSolver : public FPhysicsSolverBase
+	class FPBDRigidsSolver : public FPhysicsSolverBase
 	{
 
-		FPBDRigidsSolver(const EMultiBufferMode BufferingModeIn, UObject* InOwner, FReal AsyncDt);
-		virtual ~FPBDRigidsSolver();
+		CHAOS_API FPBDRigidsSolver(const EMultiBufferMode BufferingModeIn, UObject* InOwner, FReal AsyncDt);
+		CHAOS_API virtual ~FPBDRigidsSolver();
 
 	public:
 
@@ -118,29 +118,29 @@ namespace Chaos
 		// Execution API
 		//
 
-		void ChangeBufferMode(Chaos::EMultiBufferMode InBufferMode);
+		CHAOS_API void ChangeBufferMode(Chaos::EMultiBufferMode InBufferMode);
 
 		//
 		//  Object API
 		//
 
-		void RegisterObject(FSingleParticlePhysicsProxy* Proxy);
-		void UnregisterObject(FSingleParticlePhysicsProxy* Proxy);
+		CHAOS_API void RegisterObject(FSingleParticlePhysicsProxy* Proxy);
+		CHAOS_API void UnregisterObject(FSingleParticlePhysicsProxy* Proxy);
 
-		void RegisterObject(FGeometryCollectionPhysicsProxy* InProxy);
-		void UnregisterObject(FGeometryCollectionPhysicsProxy* InProxy);
+		CHAOS_API void RegisterObject(FGeometryCollectionPhysicsProxy* InProxy);
+		CHAOS_API void UnregisterObject(FGeometryCollectionPhysicsProxy* InProxy);
 
-		void RegisterObject(FClusterUnionPhysicsProxy* Proxy);
-		void UnregisterObject(FClusterUnionPhysicsProxy* Proxy);
+		CHAOS_API void RegisterObject(FClusterUnionPhysicsProxy* Proxy);
+		CHAOS_API void UnregisterObject(FClusterUnionPhysicsProxy* Proxy);
 
-		void RegisterObject(Chaos::FJointConstraint* GTConstraint);
-		void UnregisterObject(Chaos::FJointConstraint* GTConstraint);
+		CHAOS_API void RegisterObject(Chaos::FJointConstraint* GTConstraint);
+		CHAOS_API void UnregisterObject(Chaos::FJointConstraint* GTConstraint);
 
-		void RegisterObject(Chaos::FSuspensionConstraint* GTConstraint);
-		void UnregisterObject(Chaos::FSuspensionConstraint* GTConstraint);
+		CHAOS_API void RegisterObject(Chaos::FSuspensionConstraint* GTConstraint);
+		CHAOS_API void UnregisterObject(Chaos::FSuspensionConstraint* GTConstraint);
 
-		void RegisterObject(Chaos::FCharacterGroundConstraint* GTConstraint);
-		void UnregisterObject(Chaos::FCharacterGroundConstraint* GTConstraint);
+		CHAOS_API void RegisterObject(Chaos::FCharacterGroundConstraint* GTConstraint);
+		CHAOS_API void UnregisterObject(Chaos::FCharacterGroundConstraint* GTConstraint);
 
 		//
 		//  Simulation API
@@ -149,24 +149,24 @@ namespace Chaos
 		/**/
 		FDirtyParticlesBuffer* GetDirtyParticlesBuffer() const { return MDirtyParticlesBuffer.Get(); }
 
-		int32 NumJointConstraints() const;
-		int32 NumCollisionConstraints() const;
+		CHAOS_API int32 NumJointConstraints() const;
+		CHAOS_API int32 NumCollisionConstraints() const;
 
 		//Make friend with unit test code so we can verify some behavior
 		template <typename TSolver>
 		friend void ChaosTest::AdvanceSolverNoPushHelper(TSolver* Solver, FReal Dt);
 
 		/**/
-		void Reset();
+		CHAOS_API void Reset();
 
 		/**/
-		void StartingSceneSimulation();
+		CHAOS_API void StartingSceneSimulation();
 
 		/**/
-		void CompleteSceneSimulation();
+		CHAOS_API void CompleteSceneSimulation();
 
 		/**/
-		void UpdateGameThreadStructures();
+		CHAOS_API void UpdateGameThreadStructures();
 
 
 
@@ -201,14 +201,14 @@ namespace Chaos
 		 * This will be true if determinism is explicitly requested (via SetIsDeterministic()) or if required
 		 * by some other system like Rewind/Resim support.
 		*/
-		bool IsDetemerministic() const;
+		CHAOS_API bool IsDetemerministic() const;
 
 		/**
 		 * @brief Request that the sim be deterministic (or not)
 		 * @note Even if set to false, the sim may still be deterministic if some other feature is enabled and requires it.
 		 * @see IsDetemerministic()
 		*/
-		void SetIsDeterministic(const bool bInIsDeterministic);
+		CHAOS_API void SetIsDeterministic(const bool bInIsDeterministic);
 
 		/**/
 		FJointConstraints& GetJointConstraints() { return MEvolution->GetJointConstraints(); }
@@ -216,12 +216,12 @@ namespace Chaos
 
 		FPBDSuspensionConstraints& GetSuspensionConstraints() { return MEvolution->GetSuspensionConstraints(); }
 		const FPBDSuspensionConstraints& GetSuspensionConstraints() const { return MEvolution->GetSuspensionConstraints(); }
-		void SetSuspensionTarget(Chaos::FSuspensionConstraint* GTConstraint, const FVector& TargetPos, const FVector& Normal, bool Enabled);
+		CHAOS_API void SetSuspensionTarget(Chaos::FSuspensionConstraint* GTConstraint, const FVector& TargetPos, const FVector& Normal, bool Enabled);
 
 		FCharacterGroundConstraintContainer& GetCharacterGroundConstraints() { return MEvolution->GetCharacterGroundConstraints(); }
 		const FCharacterGroundConstraintContainer& GetCharacterGroundConstraints() const { return MEvolution->GetCharacterGroundConstraints(); }
 
-		void EnableRewindCapture(int32 NumFrames, bool InUseCollisionResimCache, TUniquePtr<IRewindCallback>&& RewindCallback = TUniquePtr<IRewindCallback>());
+		CHAOS_API void EnableRewindCapture(int32 NumFrames, bool InUseCollisionResimCache, TUniquePtr<IRewindCallback>&& RewindCallback = TUniquePtr<IRewindCallback>());
 
 		/**/
 		FPBDRigidsEvolution* GetEvolution() { return MEvolution.Get(); }
@@ -238,10 +238,10 @@ namespace Chaos
 		FSolverEventFilters* GetEventFilters() const { return MSolverEventFilters.Get(); }
 
 		/**/
-		void SyncEvents_GameThread();
+		CHAOS_API void SyncEvents_GameThread();
 
 		/**/
-		void PostTickDebugDraw(FReal Dt) const;
+		CHAOS_API void PostTickDebugDraw(FReal Dt) const;
 
 		// Visual debugger (VDB) push methods
 		UE_DEPRECATED(5.4, "This method will be removed in the future")
@@ -263,12 +263,12 @@ namespace Chaos
 		}
 
 		/** Events hooked up to the Chaos material manager */
-		void UpdateMaterial(Chaos::FMaterialHandle InHandle, const Chaos::FChaosPhysicsMaterial& InNewData);
-		void CreateMaterial(Chaos::FMaterialHandle InHandle, const Chaos::FChaosPhysicsMaterial& InNewData);
-		void DestroyMaterial(Chaos::FMaterialHandle InHandle);
-		void UpdateMaterialMask(Chaos::FMaterialMaskHandle InHandle, const Chaos::FChaosPhysicsMaterialMask& InNewData);
-		void CreateMaterialMask(Chaos::FMaterialMaskHandle InHandle, const Chaos::FChaosPhysicsMaterialMask& InNewData);
-		void DestroyMaterialMask(Chaos::FMaterialMaskHandle InHandle);
+		CHAOS_API void UpdateMaterial(Chaos::FMaterialHandle InHandle, const Chaos::FChaosPhysicsMaterial& InNewData);
+		CHAOS_API void CreateMaterial(Chaos::FMaterialHandle InHandle, const Chaos::FChaosPhysicsMaterial& InNewData);
+		CHAOS_API void DestroyMaterial(Chaos::FMaterialHandle InHandle);
+		CHAOS_API void UpdateMaterialMask(Chaos::FMaterialMaskHandle InHandle, const Chaos::FChaosPhysicsMaterialMask& InNewData);
+		CHAOS_API void CreateMaterialMask(Chaos::FMaterialMaskHandle InHandle, const Chaos::FChaosPhysicsMaterialMask& InNewData);
+		CHAOS_API void DestroyMaterialMask(Chaos::FMaterialMaskHandle InHandle);
 
 		/** Access to the internal material mirrors */
 		const THandleArray<FChaosPhysicsMaterial>& GetQueryMaterials_External() const { return QueryMaterials_External; }
@@ -277,22 +277,22 @@ namespace Chaos
 		const THandleArray<FChaosPhysicsMaterialMask>& GetSimMaterialMasks() const { return SimMaterialMasks; }
 
 		/** Copy the simulation material list to the query material list, to be done when the SQ commits an update */
-		void SyncQueryMaterials_External();
+		CHAOS_API void SyncQueryMaterials_External();
 
-		void FinalizeRewindData(const TParticleView<FPBDRigidParticles>& DirtyParticles);
+		CHAOS_API void FinalizeRewindData(const TParticleView<FPBDRigidParticles>& DirtyParticles);
 		bool RewindUsesCollisionResimCache() const { return bUseCollisionResimCache; }
 
 		FPerSolverFieldSystem& GetPerSolverField() { return *PerSolverField; }
 		const FPerSolverFieldSystem& GetPerSolverField() const { return *PerSolverField; }
 
-		void UpdateExternalAccelerationStructure_External(ISpatialAccelerationCollection<FAccelerationStructureHandle,FReal,3>*& ExternalStructure);
+		CHAOS_API void UpdateExternalAccelerationStructure_External(ISpatialAccelerationCollection<FAccelerationStructureHandle,FReal,3>*& ExternalStructure);
 		const ISpatialAccelerationCollection<FAccelerationStructureHandle, FReal, 3>* GetInternalAccelerationStructure_Internal() const
 		{
 			return MEvolution->GetSpatialAcceleration();
 		}
 
 		/** Apply a solver configuration to this solver, set externally by the owner of a solver (see UPhysicsSettings for world solver settings) */
-		void ApplyConfig(const FChaosSolverConfiguration& InConfig);
+		CHAOS_API void ApplyConfig(const FChaosSolverConfiguration& InConfig);
 
 		virtual bool AreAnyTasksPending() const override
 		{
@@ -304,44 +304,44 @@ namespace Chaos
 			return false;
 		}
 
-		void BeginDestroy();
+		CHAOS_API void BeginDestroy();
 
 		/** Update the particles parameters based on field evaluation */
-		void FieldParameterUpdateCallback(
+		CHAOS_API void FieldParameterUpdateCallback(
 			Chaos::FPBDPositionConstraints& PositionTarget,
 			TMap<int32, int32>& TargetedParticles);
 
 		/** Update the particles forces based on field evaluation */
-		void FieldForcesUpdateCallback();
+		CHAOS_API void FieldForcesUpdateCallback();
 
 		// Update the counter in Stats and the CSV profiler
-		void UpdateStatCounters() const;
-		void UpdateExpensiveStatCounters() const;
+		CHAOS_API void UpdateStatCounters() const;
+		CHAOS_API void UpdateExpensiveStatCounters() const;
 
 		// Access particle proxy from physics thread useful for cross thread communication
-		FSingleParticlePhysicsProxy* GetParticleProxy_PT(const FUniqueIdx& Idx);
-		const FSingleParticlePhysicsProxy* GetParticleProxy_PT(const FUniqueIdx& Idx) const;
-		FSingleParticlePhysicsProxy* GetParticleProxy_PT(const FGeometryParticleHandle& Handle);
-		const FSingleParticlePhysicsProxy* GetParticleProxy_PT(const FGeometryParticleHandle& Handle) const;
+		CHAOS_API FSingleParticlePhysicsProxy* GetParticleProxy_PT(const FUniqueIdx& Idx);
+		CHAOS_API const FSingleParticlePhysicsProxy* GetParticleProxy_PT(const FUniqueIdx& Idx) const;
+		CHAOS_API FSingleParticlePhysicsProxy* GetParticleProxy_PT(const FGeometryParticleHandle& Handle);
+		CHAOS_API const FSingleParticlePhysicsProxy* GetParticleProxy_PT(const FGeometryParticleHandle& Handle) const;
 
 		// Interop utilities
-		void SetParticleDynamicMisc(FPBDRigidParticleHandle* Rigid, const FParticleDynamicMisc& DynamicMisc);
+		CHAOS_API void SetParticleDynamicMisc(FPBDRigidParticleHandle* Rigid, const FParticleDynamicMisc& DynamicMisc);
 
 		// Apply callbacks internally 
-		virtual void ApplyCallbacks_Internal() override;
+		CHAOS_API virtual void ApplyCallbacks_Internal() override;
 
 	private:
 
 		/**/
-		void BufferPhysicsResults();
+		CHAOS_API void BufferPhysicsResults();
 	
 		/**/
-		virtual void PrepareAdvanceBy(const FReal DeltaTime) override;
-		virtual void AdvanceSolverBy(const FSubStepInfo& SubStepInfo) override;
-		virtual void PushPhysicsState(const FReal ExternalDt, const int32 NumSteps, const int32 NumExternalSteps) override;
-		virtual void SetExternalTimestampConsumed_Internal(const int32 Timestamp) override;
+		CHAOS_API virtual void PrepareAdvanceBy(const FReal DeltaTime) override;
+		CHAOS_API virtual void AdvanceSolverBy(const FSubStepInfo& SubStepInfo) override;
+		CHAOS_API virtual void PushPhysicsState(const FReal ExternalDt, const int32 NumSteps, const int32 NumExternalSteps) override;
+		CHAOS_API virtual void SetExternalTimestampConsumed_Internal(const int32 Timestamp) override;
 
-		void UpdateIsDeterministic();
+		CHAOS_API void UpdateIsDeterministic();
 
 		//
 		// Solver Data
@@ -398,11 +398,11 @@ namespace Chaos
 		TArray<FClusterUnionPhysicsProxy*> PendingDestroyClusterUnionProxy;
 
 
-		void ProcessSinglePushedData_Internal(FPushPhysicsData& PushData);
-		virtual void ProcessPushedData_Internal(FPushPhysicsData& PushData) override;
-		void DestroyPendingProxies_Internal();
+		CHAOS_API void ProcessSinglePushedData_Internal(FPushPhysicsData& PushData);
+		CHAOS_API virtual void ProcessPushedData_Internal(FPushPhysicsData& PushData) override;
+		CHAOS_API void DestroyPendingProxies_Internal();
 
-		virtual void ConditionalApplyRewind_Internal() override;
+		CHAOS_API virtual void ConditionalApplyRewind_Internal() override;
 
 		/** Check if we are resimming or not */
 		virtual bool IsResimming() const {return GetEvolution()->IsResimming();}
diff --git a/Engine/Source/Runtime/Experimental/Chaos/Public/PhysicsProxy/CharacterGroundConstraintProxy.h b/Engine/Source/Runtime/Experimental/Chaos/Public/PhysicsProxy/CharacterGroundConstraintProxy.h
--- a/Engine/Source/Runtime/Experimental/Chaos/Public/PhysicsProxy/CharacterGroundConstraintProxy.h
+++ b/Engine/Source/Runtime/Experimental/Chaos/Public/PhysicsProxy/CharacterGroundConstraintProxy.h
@@ -14,12 +14,12 @@ namespace Chaos
 
 	/// Proxy class to manage access to and syncing of character ground constraint data
 	/// between the game thread and physics thread
-	class CHAOS_API FCharacterGroundConstraintProxy : public IPhysicsProxyBase
+	class FCharacterGroundConstraintProxy : public IPhysicsProxyBase
 	{
 	public:
 		using Base = IPhysicsProxyBase;
 
-		FCharacterGroundConstraintProxy(FCharacterGroundConstraint* InConstraintGT, FCharacterGroundConstraintHandle* InConstraintPT = nullptr, UObject* InOwner = nullptr);
+		CHAOS_API FCharacterGroundConstraintProxy(FCharacterGroundConstraint* InConstraintGT, FCharacterGroundConstraintHandle* InConstraintPT = nullptr, UObject* InOwner = nullptr);
 	
 		//////////////////////////////////////////////////////////////////////////
 		/// Member Access
@@ -34,19 +34,19 @@ namespace Chaos
 		/// State Management
 
 		/// Gets the data from the game thread via the RemoteData and creates a constraint on the physics thread
-		void InitializeOnPhysicsThread(FPBDRigidsSolver* InSolver, FDirtyPropertiesManager& Manager, int32 DataIdx, FDirtyChaosProperties& RemoteData);
+		CHAOS_API void InitializeOnPhysicsThread(FPBDRigidsSolver* InSolver, FDirtyPropertiesManager& Manager, int32 DataIdx, FDirtyChaosProperties& RemoteData);
 
 		/// Pushes any changed data from the game thread constraint to the remote data
-		void PushStateOnGameThread(FDirtyPropertiesManager& Manager, int32 DataIdx, FDirtyChaosProperties& RemoteData);
+		CHAOS_API void PushStateOnGameThread(FDirtyPropertiesManager& Manager, int32 DataIdx, FDirtyChaosProperties& RemoteData);
 
 		/// Reads changed data from the game thread constraint into the physics thread constraint via the remote data
-		void PushStateOnPhysicsThread(FPBDRigidsSolver* InSolver, const FDirtyPropertiesManager& Manager, int32 DataIdx, const FDirtyChaosProperties& RemoteData);
+		CHAOS_API void PushStateOnPhysicsThread(FPBDRigidsSolver* InSolver, const FDirtyPropertiesManager& Manager, int32 DataIdx, const FDirtyChaosProperties& RemoteData);
 
 		/// Deletes the game thread constraint
-		void DestroyOnGameThread();
+		CHAOS_API void DestroyOnGameThread();
 
 		/// Removes references to the physics thread constraint and deletes it
-		void DestroyOnPhysicsThread(FPBDRigidsSolver* InSolver);
+		CHAOS_API void DestroyOnPhysicsThread(FPBDRigidsSolver* InSolver);
 
 		//////////////////////////////////////////////////////////////////////////
 		/// IPhysicsProxyBase Implementation
@@ -57,10 +57,10 @@ namespace Chaos
 		/// Manage Output Data
 
 		/// Write the output constraint data from the physics thread to the buffer
-		void BufferPhysicsResults(FDirtyCharacterGroundConstraintData& Buffer);
+		CHAOS_API void BufferPhysicsResults(FDirtyCharacterGroundConstraintData& Buffer);
 
 		/// Write the output constraint data from the buffer to the game thread 
-		bool PullFromPhysicsState(const FDirtyCharacterGroundConstraintData& Buffer, const int32 SolverSyncTimestamp);
+		CHAOS_API bool PullFromPhysicsState(const FDirtyCharacterGroundConstraintData& Buffer, const int32 SolverSyncTimestamp);
 
 	private:
 		FCharacterGroundConstraint* Constraint_GT;
@@ -68,4 +68,4 @@ namespace Chaos
 		bool bInitialized = false;
 	};
 
-} // namespace Chaos
\ No newline at end of file
+} // namespace Chaos
diff --git a/Engine/Source/Runtime/Experimental/Chaos/Public/PhysicsProxy/ClusterUnionPhysicsProxy.h b/Engine/Source/Runtime/Experimental/Chaos/Public/PhysicsProxy/ClusterUnionPhysicsProxy.h
--- a/Engine/Source/Runtime/Experimental/Chaos/Public/PhysicsProxy/ClusterUnionPhysicsProxy.h
+++ b/Engine/Source/Runtime/Experimental/Chaos/Public/PhysicsProxy/ClusterUnionPhysicsProxy.h
@@ -45,7 +45,7 @@ namespace Chaos
 	};
 	
 
-	class CHAOS_API FClusterUnionPhysicsProxy : public TPhysicsProxy<FClusterUnionPhysicsProxy, void, FClusterUnionProxyTimestamp>
+	class FClusterUnionPhysicsProxy : public TPhysicsProxy<FClusterUnionPhysicsProxy, void, FClusterUnionProxyTimestamp>
 	{
 		using Base = TPhysicsProxy<FClusterUnionPhysicsProxy, void, FClusterUnionProxyTimestamp>;
 	public:
@@ -54,33 +54,33 @@ namespace Chaos
 		using FInternalParticle = TPBDRigidClusteredParticleHandle<Chaos::FReal, 3>;
 
 		FClusterUnionPhysicsProxy() = delete;
-		FClusterUnionPhysicsProxy(UObject* InOwner, const FClusterCreationParameters& InParameters, const FClusterUnionInitData& InInitData);
+		CHAOS_API FClusterUnionPhysicsProxy(UObject* InOwner, const FClusterCreationParameters& InParameters, const FClusterUnionInitData& InInitData);
 
 		// Add physics objects to the cluster union. Should only be called from the game thread.
-		void AddPhysicsObjects_External(const TArray<FPhysicsObjectHandle>& Objects);
+		CHAOS_API void AddPhysicsObjects_External(const TArray<FPhysicsObjectHandle>& Objects);
 		const FClusterUnionSyncedData& GetSyncedData_External() const { return SyncedData_External; }
 
 		// Remove physics objects from the cluster union.
-		void RemovePhysicsObjects_External(const TSet<FPhysicsObjectHandle>& Objects);
+		CHAOS_API void RemovePhysicsObjects_External(const TSet<FPhysicsObjectHandle>& Objects);
 
 		// Set/Remove any anchors on the cluster particle.
-		void SetIsAnchored_External(bool bIsAnchored);
+		CHAOS_API void SetIsAnchored_External(bool bIsAnchored);
 
 		// Set Object State.
-		EObjectStateType GetObjectState_External() const;
-		void SetObjectState_External(EObjectStateType State);
+		CHAOS_API EObjectStateType GetObjectState_External() const;
+		CHAOS_API void SetObjectState_External(EObjectStateType State);
 
 		// Set GT geometry - this is only for smoothing over any changes until the PT syncs back to the GT.
-		void SetSharedGeometry_External(const TSharedPtr<Chaos::FImplicitObject, ESPMode::ThreadSafe>& Geometry, const TArray<FPBDRigidParticle*>& ShapeParticles);
+		CHAOS_API void SetSharedGeometry_External(const TSharedPtr<Chaos::FImplicitObject, ESPMode::ThreadSafe>& Geometry, const TArray<FPBDRigidParticle*>& ShapeParticles);
 
 		// Cluster union proxy initialization happens in two first on the game thread (external) then on the
 		// physics thread (internal). Cluster unions are a primarily physics concept so the things exposed to
 		// an external context is just there to be able to safely query the state on the physics thread and to
 		// be able to add/remove things to the proxy itself (rather than to move it, for example).
-		void Initialize_External();
+		CHAOS_API void Initialize_External();
 
 		bool IsInitializedOnPhysicsThread() const { return bIsInitializedOnPhysicsThread; }
-		void Initialize_Internal(FPBDRigidsSolver* RigidsSolver, FPBDRigidsSolver::FParticlesType& Particles);
+		CHAOS_API void Initialize_Internal(FPBDRigidsSolver* RigidsSolver, FPBDRigidsSolver::FParticlesType& Particles);
 
 		FExternalParticle* GetParticle_External() const { return Particle_External.Get(); }
 		FInternalParticle* GetParticle_Internal() const { return Particle_Internal; }
@@ -88,28 +88,28 @@ namespace Chaos
 		FPhysicsObjectHandle GetPhysicsObjectHandle() const { return PhysicsObject.Get(); }
 
 		bool HasChildren_External() const { return !SyncedData_External.ChildParticles.IsEmpty(); }
-		bool HasChildren_Internal() const;
+		CHAOS_API bool HasChildren_Internal() const;
 
 		bool IsAnchored_External() const { return SyncedData_External.bIsAnchored; }
 
-		void SetXR_External(const FVector& X, const FQuat& R);
-		void SetLinearVelocity_External(const FVector& V);
-		void SetAngularVelocity_External(const FVector& W);
-		void SetChildToParent_External(FPhysicsObjectHandle Child, const FTransform& RelativeTransform, bool bLock);
-		void BulkSetChildToParent_External(const TArray<FPhysicsObjectHandle>& Objects, const TArray<FTransform>& Transforms, bool bLock);
+		CHAOS_API void SetXR_External(const FVector& X, const FQuat& R);
+		CHAOS_API void SetLinearVelocity_External(const FVector& V);
+		CHAOS_API void SetAngularVelocity_External(const FVector& W);
+		CHAOS_API void SetChildToParent_External(FPhysicsObjectHandle Child, const FTransform& RelativeTransform, bool bLock);
+		CHAOS_API void BulkSetChildToParent_External(const TArray<FPhysicsObjectHandle>& Objects, const TArray<FTransform>& Transforms, bool bLock);
 
 		//
 		// These functions take care of marshaling data back and forth between the game thread
 		// and the physics thread.
 		//
-		void PushToPhysicsState(const FDirtyPropertiesManager& Manager, int32 DataIdx, const FDirtyProxy& Dirty);
-		bool PullFromPhysicsState(const FDirtyClusterUnionData& PullData, int32 SolverSyncTimestamp, const FDirtyClusterUnionData* NextPullData = nullptr, const FRealSingle* Alpha = nullptr);
+		CHAOS_API void PushToPhysicsState(const FDirtyPropertiesManager& Manager, int32 DataIdx, const FDirtyProxy& Dirty);
+		CHAOS_API bool PullFromPhysicsState(const FDirtyClusterUnionData& PullData, int32 SolverSyncTimestamp, const FDirtyClusterUnionData* NextPullData = nullptr, const FRealSingle* Alpha = nullptr);
 
-		void BufferPhysicsResults_Internal(FDirtyClusterUnionData& BufferData);
-		void BufferPhysicsResults_External(FDirtyClusterUnionData& BufferData);
+		CHAOS_API void BufferPhysicsResults_Internal(FDirtyClusterUnionData& BufferData);
+		CHAOS_API void BufferPhysicsResults_External(FDirtyClusterUnionData& BufferData);
 
-		void SyncRemoteData(FDirtyPropertiesManager& Manager, int32 DataIdx, FDirtyChaosProperties& RemoteData) const;
-		void ClearAccumulatedData();
+		CHAOS_API void SyncRemoteData(FDirtyPropertiesManager& Manager, int32 DataIdx, FDirtyChaosProperties& RemoteData) const;
+		CHAOS_API void ClearAccumulatedData();
 
 		FProxyInterpolationBase& GetInterpolationData() { return InterpolationData; }
 		const FProxyInterpolationBase& GetInterpolationData() const { return InterpolationData; }
@@ -148,4 +148,4 @@ namespace Chaos
 		//~ End TPhysicsProxy Interface
 	};
 
-}
\ No newline at end of file
+}
diff --git a/Engine/Source/Runtime/Experimental/Chaos/Public/PhysicsProxy/GeometryCollectionPhysicsProxy.h b/Engine/Source/Runtime/Experimental/Chaos/Public/PhysicsProxy/GeometryCollectionPhysicsProxy.h
--- a/Engine/Source/Runtime/Experimental/Chaos/Public/PhysicsProxy/GeometryCollectionPhysicsProxy.h
+++ b/Engine/Source/Runtime/Experimental/Chaos/Public/PhysicsProxy/GeometryCollectionPhysicsProxy.h
@@ -137,7 +137,7 @@ private:
  * Class to manage sharing data between the game thread and the simulation thread 
  * (which may not be different than the game thread) for a \c FGeometryDynamicCollection.
  */
-class CHAOS_API FGeometryCollectionPhysicsProxy : public TPhysicsProxy<FGeometryCollectionPhysicsProxy, FStubGeometryCollectionData, FGeometryCollectionProxyTimestamp>
+class FGeometryCollectionPhysicsProxy : public TPhysicsProxy<FGeometryCollectionPhysicsProxy, FStubGeometryCollectionData, FGeometryCollectionProxyTimestamp>
 {
 public:
 	typedef TPhysicsProxy<FGeometryCollectionPhysicsProxy, FStubGeometryCollectionData, FGeometryCollectionProxyTimestamp> Base;
@@ -157,7 +157,7 @@ public:
 	 * \p InCacheSyncFunc callback invoked from \c PullFromPhysicsState().
 	 * \p InFinalSyncFunc callback invoked from \c SyncBeforeDestory().
 	 */
-	FGeometryCollectionPhysicsProxy(
+	CHAOS_API FGeometryCollectionPhysicsProxy(
 		UObject* InOwner, 
 		FGeometryDynamicCollection& GameThreadCollection,
 		const FSimulationParameters& SimulationParameters,
@@ -165,13 +165,13 @@ public:
 		FCollisionFilterData InQueryFilter,
 		FGuid InCollectorGuid = FGuid::NewGuid(),
 		const Chaos::EMultiBufferMode BufferMode=Chaos::EMultiBufferMode::TripleGuarded);
-	virtual ~FGeometryCollectionPhysicsProxy();
+	CHAOS_API virtual ~FGeometryCollectionPhysicsProxy();
 
 	/**
 	 * Construct \c PTDynamicCollection, copying attributes from the game thread, 
 	 * and prepare for simulation.
 	 */
-	void Initialize(Chaos::FPBDRigidsEvolutionBase* Evolution);
+	CHAOS_API void Initialize(Chaos::FPBDRigidsEvolutionBase* Evolution);
 	void Reset() { }
 
 	bool IsInitializedOnPhysicsThread() const { return bIsInitializedOnPhysicsThread; }
@@ -181,12 +181,12 @@ public:
 	 *
 	 * Called by solver command registered by \c FPBDRigidsSolver::RegisterObject().
 	 */
-	void InitializeBodiesPT(
+	CHAOS_API void InitializeBodiesPT(
 		Chaos::FPBDRigidsSolver* RigidsSolver,
 		typename Chaos::FPBDRigidsSolver::FParticlesType& Particles);
 
 	/** */
-	static void InitializeDynamicCollection(FGeometryDynamicCollection& DynamicCollection, const FGeometryCollection& RestCollection, const FSimulationParameters& Params);
+	static CHAOS_API void InitializeDynamicCollection(FGeometryDynamicCollection& DynamicCollection, const FGeometryCollection& RestCollection, const FSimulationParameters& Params);
 
 	/** */
 	bool IsSimulating() const { return Parameters.Simulating; }
@@ -198,39 +198,39 @@ public:
 	 * data transport to the physics thread itself, without allocating memory.
 	 */
 	Chaos::FParticleData* NewData() { BufferGameState(); return nullptr; }
-	void BufferGameState();
+	CHAOS_API void BufferGameState();
 
 	/** Called at the end of \c FPBDRigidsSolver::PushPhysicsStateExec(). */
 	void ClearAccumulatedData() {}
 
 	/** Push PT state into the \c PhysToGameInterchange. */
-	void BufferPhysicsResults_Internal(Chaos::FPBDRigidsSolver* CurrentSolver, Chaos::FDirtyGeometryCollectionData& BufferData);
+	CHAOS_API void BufferPhysicsResults_Internal(Chaos::FPBDRigidsSolver* CurrentSolver, Chaos::FDirtyGeometryCollectionData& BufferData);
 
 	/** Push GT state into the \c PhysToGameInterchange for async physics */
-	void BufferPhysicsResults_External(Chaos::FDirtyGeometryCollectionData& BufferData);
+	CHAOS_API void BufferPhysicsResults_External(Chaos::FDirtyGeometryCollectionData& BufferData);
 
 	/** Push data from the game thread to the physics thread */
-	void PushStateOnGameThread(Chaos::FPBDRigidsSolver* InSolver);
+	CHAOS_API void PushStateOnGameThread(Chaos::FPBDRigidsSolver* InSolver);
 
 	/** apply the state changes on the physics thread */
-	void PushToPhysicsState();
+	CHAOS_API void PushToPhysicsState();
 	
 	/** Does nothing as \c BufferPhysicsResults() already did this. */
-	void FlipBuffer();
+	CHAOS_API void FlipBuffer();
 	
 	/** 
 	 * Pulls data out of the PhysToGameInterchange and updates \c GTDynamicCollection. 
 	 * Called from FPhysScene_ChaosInterface::SyncBodies(), NOT the solver.
 	 */
-	bool PullFromPhysicsState(const Chaos::FDirtyGeometryCollectionData& BufferData, const int32 SolverSyncTimestamp, const Chaos::FDirtyGeometryCollectionData* NextPullData = nullptr, const Chaos::FRealSingle* Alpha= nullptr);
+	CHAOS_API bool PullFromPhysicsState(const Chaos::FDirtyGeometryCollectionData& BufferData, const int32 SolverSyncTimestamp, const Chaos::FDirtyGeometryCollectionData* NextPullData = nullptr, const Chaos::FRealSingle* Alpha= nullptr);
 
 	bool IsDirty() { return false; }
 
 	static EPhysicsProxyType ConcreteType() { return EPhysicsProxyType::GeometryCollectionType; }
 
-	void SyncBeforeDestroy();
-	void OnRemoveFromSolver(Chaos::FPBDRigidsSolver *RBDSolver);
-	void OnRemoveFromScene();
+	CHAOS_API void SyncBeforeDestroy();
+	CHAOS_API void OnRemoveFromSolver(Chaos::FPBDRigidsSolver *RBDSolver);
+	CHAOS_API void OnRemoveFromScene();
 
 	void SetCollisionParticlesPerObjectFraction(float CollisionParticlesPerObjectFractionIn) 
 	{CollisionParticlesPerObjectFraction = CollisionParticlesPerObjectFractionIn;}
@@ -252,11 +252,11 @@ public:
 		Commands.Add(Command); 
 	}
 
-	static void InitializeSharedCollisionStructures(Chaos::FErrorReporter& ErrorReporter, FGeometryCollection& RestCollection, const FSharedSimulationParameters& SharedParams);
+	static CHAOS_API void InitializeSharedCollisionStructures(Chaos::FErrorReporter& ErrorReporter, FGeometryCollection& RestCollection, const FSharedSimulationParameters& SharedParams);
 
-	void FieldForcesUpdateCallback(Chaos::FPBDRigidsSolver* RigidSolver);
+	CHAOS_API void FieldForcesUpdateCallback(Chaos::FPBDRigidsSolver* RigidSolver);
 
-	void FieldParameterUpdateCallback(Chaos::FPBDRigidsSolver* RigidSolver, const bool bUpdateViews = true);
+	CHAOS_API void FieldParameterUpdateCallback(Chaos::FPBDRigidsSolver* RigidSolver, const bool bUpdateViews = true);
 
 	void UpdateKinematicBodiesCallback(const FParticlesType& InParticles, const float InDt, const float InTime, FKinematicProxy& InKinematicProxy) {}
 	void StartFrameCallback(const float InDt, const float InTime) {}
@@ -269,7 +269,7 @@ public:
 	bool IsGTCollectionDirty() const { return GameThreadCollection.IsDirty(); }
 
 	// set the world transform ( this needs to be called on the game thread ) 
-	void SetWorldTransform_External(const FTransform& WorldTransform);
+	CHAOS_API void SetWorldTransform_External(const FTransform& WorldTransform);
 
 	const TArray<FClusterHandle*> GetParticles() const
 	{
@@ -301,16 +301,16 @@ public:
 		return GTParticles;
 	}
 
-	FParticle* GetParticleByIndex_External(int32 Index);
-	const FParticle* GetParticleByIndex_External(int32 Index) const;
+	CHAOS_API FParticle* GetParticleByIndex_External(int32 Index);
+	CHAOS_API const FParticle* GetParticleByIndex_External(int32 Index) const;
 
-	FParticleHandle* GetParticleByIndex_Internal(int32 Index);
-	const FParticleHandle* GetParticleByIndex_Internal(int32 Index) const;
+	CHAOS_API FParticleHandle* GetParticleByIndex_Internal(int32 Index);
+	CHAOS_API const FParticleHandle* GetParticleByIndex_Internal(int32 Index) const;
 
 	/**
 	*  * Get all the geometry collection particle handles based on the processing resolution
 	 */
-	void GetRelevantParticleHandles(
+	CHAOS_API void GetRelevantParticleHandles(
 		TArray<Chaos::TGeometryParticleHandle<Chaos::FReal, 3>*>& Handles,
 		const Chaos::FPBDRigidsSolver* RigidSolver,
 		EFieldResolutionType ResolutionType);
@@ -318,7 +318,7 @@ public:
 	/**
 	 * Get all the geometry collection particle handles filtered by object state
 	 */
-	void GetFilteredParticleHandles(
+	CHAOS_API void GetFilteredParticleHandles(
 		TArray<Chaos::TGeometryParticleHandle<Chaos::FReal, 3>*>& Handles,
 		const Chaos::FPBDRigidsSolver* RigidSolver,
 		const EFieldFilterType FilterType,
@@ -396,38 +396,38 @@ public:
 		return FGeometryCollectionItemIndex::CreateInvalidItemIndex();
 	}
 	
-	FName GetTransformName_External(FGeometryCollectionItemIndex ItemIndex) const;
+	CHAOS_API FName GetTransformName_External(FGeometryCollectionItemIndex ItemIndex) const;
 
 	bool GetIsObjectDynamic() const { return IsObjectDynamic; }
 
-	void DisableParticles_External(TArray<int32>&& TransformGroupIndices);
+	CHAOS_API void DisableParticles_External(TArray<int32>&& TransformGroupIndices);
 
-	void ApplyForceAt_External(FVector Force, FVector WorldLocation);
-	void ApplyImpulseAt_External(FVector Force, FVector WorldLocation);
-	void BreakClusters_External(TArray<FGeometryCollectionItemIndex>&& ItemIndices);
-	void BreakActiveClusters_External();
-	void SetAnchoredByIndex_External(int32 Index, bool bAnchored);
-	void SetAnchoredByTransformedBox_External(const FBox& Box, const FTransform& Transform, bool bAnchored, int32 MaxLevel = INDEX_NONE);
-	void RemoveAllAnchors_External();
-	void ApplyExternalStrain_External(FGeometryCollectionItemIndex ItemIndex, const FVector& WorldLocation, float Radius, int32 PropagationDepth, float PropagationFactor, float StrainValue);
-	void ApplyInternalStrain_External(FGeometryCollectionItemIndex ItemIndex, const FVector& WorldLocation, float Radius, int32 PropagationDepth, float PropagationFactor, float StrainValue);
-	void ApplyBreakingLinearVelocity_External(FGeometryCollectionItemIndex ItemIndex, const FVector& LinearVelocity);
-	void ApplyBreakingAngularVelocity_External(FGeometryCollectionItemIndex ItemIndex, const FVector& AngularVelocity);
-	void ApplyLinearVelocity_External(FGeometryCollectionItemIndex ItemIndex, const FVector& LinearVelocity);
-	void ApplyAngularVelocity_External(FGeometryCollectionItemIndex ItemIndex, const FVector& AngularVelocity);
+	CHAOS_API void ApplyForceAt_External(FVector Force, FVector WorldLocation);
+	CHAOS_API void ApplyImpulseAt_External(FVector Force, FVector WorldLocation);
+	CHAOS_API void BreakClusters_External(TArray<FGeometryCollectionItemIndex>&& ItemIndices);
+	CHAOS_API void BreakActiveClusters_External();
+	CHAOS_API void SetAnchoredByIndex_External(int32 Index, bool bAnchored);
+	CHAOS_API void SetAnchoredByTransformedBox_External(const FBox& Box, const FTransform& Transform, bool bAnchored, int32 MaxLevel = INDEX_NONE);
+	CHAOS_API void RemoveAllAnchors_External();
+	CHAOS_API void ApplyExternalStrain_External(FGeometryCollectionItemIndex ItemIndex, const FVector& WorldLocation, float Radius, int32 PropagationDepth, float PropagationFactor, float StrainValue);
+	CHAOS_API void ApplyInternalStrain_External(FGeometryCollectionItemIndex ItemIndex, const FVector& WorldLocation, float Radius, int32 PropagationDepth, float PropagationFactor, float StrainValue);
+	CHAOS_API void ApplyBreakingLinearVelocity_External(FGeometryCollectionItemIndex ItemIndex, const FVector& LinearVelocity);
+	CHAOS_API void ApplyBreakingAngularVelocity_External(FGeometryCollectionItemIndex ItemIndex, const FVector& AngularVelocity);
+	CHAOS_API void ApplyLinearVelocity_External(FGeometryCollectionItemIndex ItemIndex, const FVector& LinearVelocity);
+	CHAOS_API void ApplyAngularVelocity_External(FGeometryCollectionItemIndex ItemIndex, const FVector& AngularVelocity);
 
-	void SetProxyDirty_External();
+	CHAOS_API void SetProxyDirty_External();
 
-	void SetEnableDamageFromCollision_External(bool bEnable);
-	void SetNotifyBreakings_External(bool bNotify);
-	void SetNotifyRemovals_External(bool bNotify);
-	void SetNotifyCrumblings_External(bool bNotify, bool bIncludeChildren);
+	CHAOS_API void SetEnableDamageFromCollision_External(bool bEnable);
+	CHAOS_API void SetNotifyBreakings_External(bool bNotify);
+	CHAOS_API void SetNotifyRemovals_External(bool bNotify);
+	CHAOS_API void SetNotifyCrumblings_External(bool bNotify, bool bIncludeChildren);
 
-	void SetNotifyGlobalBreakings_External(bool bNotify);
-	void SetNotifyGlobalRemovals_External(bool bNotify);
-	void SetNotifyGlobalCrumblings_External(bool bNotify, bool bIncludeChildren);
+	CHAOS_API void SetNotifyGlobalBreakings_External(bool bNotify);
+	CHAOS_API void SetNotifyGlobalRemovals_External(bool bNotify);
+	CHAOS_API void SetNotifyGlobalCrumblings_External(bool bNotify, bool bIncludeChildren);
 
-	float ComputeMaterialBasedDamageThreshold_Internal(Chaos::FPBDRigidClusteredParticleHandle& ClusteredParticle) const;
+	CHAOS_API float ComputeMaterialBasedDamageThreshold_Internal(Chaos::FPBDRigidClusteredParticleHandle& ClusteredParticle) const;
 
 	FProxyInterpolationBase& GetInterpolationData() { return InterpolationData; }
 	const FProxyInterpolationBase& GetInterpolationData() const { return InterpolationData; }
@@ -442,7 +442,7 @@ public:
 	void SetReplicationMode(EReplicationMode Mode) { ReplicationMode = Mode; }
 	EReplicationMode GetReplicationMode() const { return ReplicationMode; }
 
-	void UpdateFilterData_External(const FCollisionFilterData& NewSimFilter, const FCollisionFilterData& NewQueryFilter);
+	CHAOS_API void UpdateFilterData_External(const FCollisionFilterData& NewSimFilter, const FCollisionFilterData& NewQueryFilter);
 	
 	/** 
 	 * Traverses the parents of TransformGroupIdx counting number of levels,
@@ -450,20 +450,20 @@ public:
 	 * If level is already set, retrieve stored level.
 	 * Uninitialized level array should be correct size and defaulted to zeros.
 	 */
-	static int32 CalculateAndSetLevel(int32 TransformGroupIdx, const TManagedArray<int32>& Parent, TManagedArray<int32>& Levels);
+	static CHAOS_API int32 CalculateAndSetLevel(int32 TransformGroupIdx, const TManagedArray<int32>& Parent, TManagedArray<int32>& Levels);
 
 	void SetPostPhysicsSyncCallback(TFunction<void()> Callback)
 	{
 		PostPhysicsSyncCallback = Callback;
 	}
 	
-	TArray<Chaos::FPhysicsObjectHandle> GetAllPhysicsObjects() const ;
-	TArray<Chaos::FPhysicsObjectHandle> GetAllPhysicsObjectIncludingNulls() const;
-	Chaos::FPhysicsObjectHandle GetPhysicsObjectByIndex(int32 Index) const;
+	CHAOS_API TArray<Chaos::FPhysicsObjectHandle> GetAllPhysicsObjects() const ;
+	CHAOS_API TArray<Chaos::FPhysicsObjectHandle> GetAllPhysicsObjectIncludingNulls() const;
+	CHAOS_API Chaos::FPhysicsObjectHandle GetPhysicsObjectByIndex(int32 Index) const;
 	int32 GetNumParticles() const { return NumParticles; }
 protected:
 
-	float ComputeMaterialBasedDamageThreshold_Internal(int32 TransformIndex) const;
+	CHAOS_API float ComputeMaterialBasedDamageThreshold_Internal(int32 TransformIndex) const;
 
 	/**
 	* Compute user defined damage threshold for a specific transform
@@ -472,15 +472,15 @@ protected:
 	* @param TransformIndex index of the transform to compute the threshold for
 	* #return damage threshold value
 	*/
-	float ComputeUserDefinedDamageThreshold_Internal(int32 TransformIndex) const;
+	CHAOS_API float ComputeUserDefinedDamageThreshold_Internal(int32 TransformIndex) const;
 
 	/** adjust scalar mass to account for per component scale properties ( from material override and world transform scale ) */ 
-	float AdjustMassForScale(float Mass) const;
+	CHAOS_API float AdjustMassForScale(float Mass) const;
 
 	/** adjust inertia to account for per component scale properties ( from material override and world transform scale ) */
-	Chaos::FVec3f AdjustInertiaForScale(const Chaos::FVec3f& Inertia) const;
+	CHAOS_API Chaos::FVec3f AdjustInertiaForScale(const Chaos::FVec3f& Inertia) const;
 
-	Chaos::TPBDGeometryCollectionParticleHandle<Chaos::FReal, 3>* BuildNonClusters_Internal(const uint32 CollectionClusterIndex, Chaos::FPBDRigidsSolver* RigidsSolver, float Mass, Chaos::FVec3f Inertia);
+	CHAOS_API Chaos::TPBDGeometryCollectionParticleHandle<Chaos::FReal, 3>* BuildNonClusters_Internal(const uint32 CollectionClusterIndex, Chaos::FPBDRigidsSolver* RigidsSolver, float Mass, Chaos::FVec3f Inertia);
 
 	/**
 	 * Build a physics thread cluster parent particle.
@@ -490,31 +490,31 @@ protected:
 	 *  \P Parameters - uh, yeah...  Other parameters.
 	 */
 
-	Chaos::FPBDRigidClusteredParticleHandle* BuildClusters_Internal(
+	CHAOS_API Chaos::FPBDRigidClusteredParticleHandle* BuildClusters_Internal(
 		const uint32 CollectionClusterIndex, 
 		TArray<Chaos::FPBDRigidParticleHandle*>& ChildHandles,
 		const TArray<int32>& ChildTransformGroupIndices,
 		const Chaos::FClusterCreationParameters & Parameters,
 		const Chaos::FUniqueIdx* ExistingIndex);
 
-	void SetSleepingState(const Chaos::FPBDRigidsSolver& RigidsSolver);
-	void DirtyAllParticles(const Chaos::FPBDRigidsSolver& RigidsSolver);
+	CHAOS_API void SetSleepingState(const Chaos::FPBDRigidsSolver& RigidsSolver);
+	CHAOS_API void DirtyAllParticles(const Chaos::FPBDRigidsSolver& RigidsSolver);
 
 	/** 
 	 * Traverses the parents of \p TransformIndex in \p GeometryCollection, counting
 	 * the number of levels until the next parent is \c INDEX_NONE.
 	 */
-	static int32 CalculateHierarchyLevel(const FGeometryDynamicCollection& DynamicCollection, int32 TransformIndex);
+	static CHAOS_API int32 CalculateHierarchyLevel(const FGeometryDynamicCollection& DynamicCollection, int32 TransformIndex);
 
-	int32 CalculateEffectiveParticles(const FGeometryDynamicCollection& DynamicCollection, int32 NumTransform, TBitArray<>& EffectiveParticles) const;
+	CHAOS_API int32 CalculateEffectiveParticles(const FGeometryDynamicCollection& DynamicCollection, int32 NumTransform, TBitArray<>& EffectiveParticles) const;
 
-	void SetClusteredParticleKinematicTarget_Internal(Chaos::FPBDRigidClusteredParticleHandle* Handle, const FTransform& WorldTransform);
+	CHAOS_API void SetClusteredParticleKinematicTarget_Internal(Chaos::FPBDRigidClusteredParticleHandle* Handle, const FTransform& WorldTransform);
 
-	void PrepareBufferData(Chaos::FDirtyGeometryCollectionData& BufferData, const FGeometryDynamicCollection& ThreadCollection,  Chaos::FReal SolverLastDt = 0.0);
+	CHAOS_API void PrepareBufferData(Chaos::FDirtyGeometryCollectionData& BufferData, const FGeometryDynamicCollection& ThreadCollection,  Chaos::FReal SolverLastDt = 0.0);
 
-	void CreateNonClusteredParticles(Chaos::FPBDRigidsSolver* RigidsSolver,	const FGeometryCollection& RestCollection, const FGeometryDynamicCollection& DynamicCollection, const TBitArray<>& EffectiveParticles);
+	CHAOS_API void CreateNonClusteredParticles(Chaos::FPBDRigidsSolver* RigidsSolver,	const FGeometryCollection& RestCollection, const FGeometryDynamicCollection& DynamicCollection, const TBitArray<>& EffectiveParticles);
 
-	Chaos::FPBDRigidClusteredParticleHandle* FindClusteredParticleHandleByItemIndex_Internal(FGeometryCollectionItemIndex ItemIndex) const;
+	CHAOS_API Chaos::FPBDRigidClusteredParticleHandle* FindClusteredParticleHandleByItemIndex_Internal(FGeometryCollectionItemIndex ItemIndex) const;
 private:
 
 
@@ -623,7 +623,7 @@ private:
  * the watcher collects runtime data about
  * damage on each piece of the geometry collection 
  */
-struct CHAOS_API FDamageCollector
+struct FDamageCollector
 {
 public:
 	struct FDamageData
@@ -633,7 +633,7 @@ public:
 		bool  bIsBroken = false;
 	};
 
-	void Reset(int32 NumTransforms);
+	CHAOS_API void Reset(int32 NumTransforms);
 	int32 Num() const { return DamageData.Num(); }
 	const FDamageData& operator[](int32 TransformIndex) const 
 	{ 
@@ -642,22 +642,22 @@ public:
 			return DamageData[TransformIndex]; 
 		return DefaultDamageData;
 	}
-	void SampleDamage(int32 TransformIndex, float Damage, float DamageThreshold);
+	CHAOS_API void SampleDamage(int32 TransformIndex, float Damage, float DamageThreshold);
 	
 private:
 	TArray<FDamageData> DamageData;
 };
 
-struct CHAOS_API FRuntimeDataCollector
+struct FRuntimeDataCollector
 {
 public:
-	void Clear();
-	void AddCollector(const FGuid& Guid, int32 TransformNum);
-	void RemoveCollector(const FGuid& Guid);
+	CHAOS_API void Clear();
+	CHAOS_API void AddCollector(const FGuid& Guid, int32 TransformNum);
+	CHAOS_API void RemoveCollector(const FGuid& Guid);
 
-	FDamageCollector* Find(const FGuid& Guid);
+	CHAOS_API FDamageCollector* Find(const FGuid& Guid);
 
-	static FRuntimeDataCollector& GetInstance();
+	static CHAOS_API FRuntimeDataCollector& GetInstance();
 	
 private:
 	// collectors by geometry collection Guids
diff --git a/Engine/Source/Runtime/Experimental/Chaos/Public/PhysicsProxy/PerSolverFieldSystem.h b/Engine/Source/Runtime/Experimental/Chaos/Public/PhysicsProxy/PerSolverFieldSystem.h
--- a/Engine/Source/Runtime/Experimental/Chaos/Public/PhysicsProxy/PerSolverFieldSystem.h
+++ b/Engine/Source/Runtime/Experimental/Chaos/Public/PhysicsProxy/PerSolverFieldSystem.h
@@ -7,7 +7,7 @@
 #include "Chaos/PBDPositionConstraints.h"
 #include "Chaos/Defines.h"
 
-class CHAOS_API FPerSolverFieldSystem
+class FPerSolverFieldSystem
 {
 public:
 
@@ -30,7 +30,7 @@ public:
 	 * EFieldPhysicsType::Field_PositionAnimated
 	 * EFieldPhysicsType::Field_DynamicConstraint
 	 */
-	void FieldParameterUpdateCallback(
+	CHAOS_API void FieldParameterUpdateCallback(
 		Chaos::FPBDRigidsSolver* InSolver, 
 		Chaos::FPBDPositionConstraints& PositionTarget,
 		TMap<int32, int32>& TargetedParticles);
@@ -42,7 +42,7 @@ public:
 	 * EFieldPhysicsType::Field_LinearForce
 	 * EFieldPhysicsType::Field_AngularTorque
 	 */
-	void FieldForcesUpdateCallback(
+	CHAOS_API void FieldForcesUpdateCallback(
 		Chaos::FPBDRigidsSolver* RigidSolver);
 
 	/**
@@ -54,7 +54,7 @@ public:
 	 * EFieldPhysicsType::Field_AngularVelocity
 	 * EFieldPhysicsType::Field_AngularrTorque
 	 */
-	void ComputeFieldRigidImpulse(const Chaos::FReal SolverTime);
+	CHAOS_API void ComputeFieldRigidImpulse(const Chaos::FReal SolverTime);
 
 	/**
 	 * Compute field linear velocity/force given a list of samples (positions + indices)
@@ -63,19 +63,19 @@ public:
 	 * EFieldPhysicsType::Field_LinearVelocity
 	 * EFieldPhysicsType::Field_LinearForce
 	 */
-	void ComputeFieldLinearImpulse(const Chaos::FReal SolverTime);
+	CHAOS_API void ComputeFieldLinearImpulse(const Chaos::FReal SolverTime);
 
 	/** Add the transient field command */
-	void AddTransientCommand(const FFieldSystemCommand& FieldCommand);
+	CHAOS_API void AddTransientCommand(const FFieldSystemCommand& FieldCommand);
 
 	/** Add the persistent field command */
-	void AddPersistentCommand(const FFieldSystemCommand& FieldCommand);
+	CHAOS_API void AddPersistentCommand(const FFieldSystemCommand& FieldCommand);
 
 	/** Remove the transient field command */
-	void RemoveTransientCommand(const FFieldSystemCommand& FieldCommand);
+	CHAOS_API void RemoveTransientCommand(const FFieldSystemCommand& FieldCommand);
 
 	/** Remove the persistent field command */
-	void RemovePersistentCommand(const FFieldSystemCommand& FieldCommand);
+	CHAOS_API void RemovePersistentCommand(const FFieldSystemCommand& FieldCommand);
 
 	/** Get all the non const transient field commands */
 	TArray<FFieldSystemCommand>& GetTransientCommands() { return TransientCommands; }
@@ -98,7 +98,7 @@ public:
 	 * the active clusters are set in the \p IndicesArray.
 	 */
 
-	static void GetRelevantParticleHandles(
+	static CHAOS_API void GetRelevantParticleHandles(
 		TArray<Chaos::FGeometryParticleHandle*>& ParticleHandles,
 		const Chaos::FPBDRigidsSolver* RigidSolver,
 		const EFieldResolutionType ResolutionType);
@@ -110,7 +110,7 @@ public:
 	 * is provided from the \c Particles.X to \c Particles.Attribute.
 	 */
 
-	static void GetFilteredParticleHandles(
+	static CHAOS_API void GetFilteredParticleHandles(
 		TArray<Chaos::FGeometryParticleHandle*>& ParticleHandles,
 		const Chaos::FPBDRigidsSolver* RigidSolver,
 		const EFieldFilterType FilterType,
@@ -140,13 +140,13 @@ public:
 private:
 
 	/** Forces update callback implementation */
-	void FieldForcesUpdateInternal(
+	CHAOS_API void FieldForcesUpdateInternal(
 		Chaos::FPBDRigidsSolver* RigidSolver,
 		TArray<FFieldSystemCommand>& Commands, 
 		const bool IsTransient);
 
 	/** Parameter update callback implementation */
-	void FieldParameterUpdateInternal(
+	CHAOS_API void FieldParameterUpdateInternal(
 		Chaos::FPBDRigidsSolver* RigidSolver,
 		Chaos::FPBDPositionConstraints& PositionTarget,
 		TMap<int32, int32>& PositionTargetedParticles,
diff --git a/Engine/Source/Runtime/Experimental/Chaos/Public/PhysicsProxy/SingleParticlePhysicsProxy.h b/Engine/Source/Runtime/Experimental/Chaos/Public/PhysicsProxy/SingleParticlePhysicsProxy.h
--- a/Engine/Source/Runtime/Experimental/Chaos/Public/PhysicsProxy/SingleParticlePhysicsProxy.h
+++ b/Engine/Source/Runtime/Experimental/Chaos/Public/PhysicsProxy/SingleParticlePhysicsProxy.h
@@ -54,18 +54,18 @@ private:
 class FRigidBodyHandle_External;
 class FRigidBodyHandle_Internal;
 
-class CHAOS_API FSingleParticlePhysicsProxy : public IPhysicsProxyBase
+class FSingleParticlePhysicsProxy : public IPhysicsProxyBase
 {
 public:
 	using PARTICLE_TYPE = FGeometryParticle;
 	using FParticleHandle = FGeometryParticleHandle;
 
-	static FSingleParticlePhysicsProxy* Create(TUniquePtr<FGeometryParticle>&& Particle);
+	static CHAOS_API FSingleParticlePhysicsProxy* Create(TUniquePtr<FGeometryParticle>&& Particle);
 
 	FSingleParticlePhysicsProxy() = delete;
 	FSingleParticlePhysicsProxy(const FSingleParticlePhysicsProxy&) = delete;
 	FSingleParticlePhysicsProxy(FSingleParticlePhysicsProxy&&) = delete;
-	virtual ~FSingleParticlePhysicsProxy();
+	CHAOS_API virtual ~FSingleParticlePhysicsProxy();
 
 	const FProxyInterpolationBase& GetInterpolationData() const { return InterpolationData; }
 	FProxyInterpolationBase& GetInterpolationData() { return InterpolationData; }
@@ -116,29 +116,29 @@ public:
 
 	// Threading API
 
-	void PushToPhysicsState(const FDirtyPropertiesManager& Manager,int32 DataIdx,const FDirtyProxy& Dirty,FShapeDirtyData* ShapesData, FReal ExternalDt);
+	CHAOS_API void PushToPhysicsState(const FDirtyPropertiesManager& Manager,int32 DataIdx,const FDirtyProxy& Dirty,FShapeDirtyData* ShapesData, FReal ExternalDt);
 
 	/**/
-	void ClearAccumulatedData();
+	CHAOS_API void ClearAccumulatedData();
 
 	/**/
-	void BufferPhysicsResults(FDirtyRigidParticleData&);
+	CHAOS_API void BufferPhysicsResults(FDirtyRigidParticleData&);
 
 	/**/
-	void BufferPhysicsResults_External(FDirtyRigidParticleData&);
+	CHAOS_API void BufferPhysicsResults_External(FDirtyRigidParticleData&);
 
 	/**/
-	bool PullFromPhysicsState(const FDirtyRigidParticleData& PullData, int32 SolverSyncTimestamp, const FDirtyRigidParticleData* NextPullData = nullptr, const FRealSingle* Alpha = nullptr, const FDirtyRigidParticleReplicationErrorData* Error = nullptr, const Chaos::FReal AsyncFixedTimeStep = 0);
+	CHAOS_API bool PullFromPhysicsState(const FDirtyRigidParticleData& PullData, int32 SolverSyncTimestamp, const FDirtyRigidParticleData* NextPullData = nullptr, const FRealSingle* Alpha = nullptr, const FDirtyRigidParticleReplicationErrorData* Error = nullptr, const Chaos::FReal AsyncFixedTimeStep = 0);
 
 	/**/
-	bool IsDirty();
+	CHAOS_API bool IsDirty();
 
 
 	/**/
-	EWakeEventEntry GetWakeEvent() const;
+	CHAOS_API EWakeEventEntry GetWakeEvent() const;
 
 	/**/
-	void ClearEvents();
+	CHAOS_API void ClearEvents();
 
 	//Returns the underlying game thread particle. Note this should only be needed for internal book keeping type tasks. API may change, use GetGameThreadAPI instead
 	PARTICLE_TYPE* GetParticle_LowLevel()
@@ -185,7 +185,7 @@ private:
 #endif
 
 	//use static Create
-	FSingleParticlePhysicsProxy(TUniquePtr<PARTICLE_TYPE>&& InParticle, FParticleHandle* InHandle, UObject* InOwner = nullptr);
+	CHAOS_API FSingleParticlePhysicsProxy(TUniquePtr<PARTICLE_TYPE>&& InParticle, FParticleHandle* InHandle, UObject* InOwner = nullptr);
 };
 
 /** Wrapper class that routes all reads and writes to the appropriate particle data. This is helpful for cases where we want to both write to a particle and a network buffer for example*/
diff --git a/Engine/Source/Runtime/Experimental/Chaos/Public/PhysicsProxy/SkeletalMeshPhysicsProxy.h b/Engine/Source/Runtime/Experimental/Chaos/Public/PhysicsProxy/SkeletalMeshPhysicsProxy.h
--- a/Engine/Source/Runtime/Experimental/Chaos/Public/PhysicsProxy/SkeletalMeshPhysicsProxy.h
+++ b/Engine/Source/Runtime/Experimental/Chaos/Public/PhysicsProxy/SkeletalMeshPhysicsProxy.h
@@ -10,7 +10,7 @@
 
 // @todo(chaos): remove this file
 
-struct CHAOS_API FSkeletalMeshPhysicsProxyParams
+struct FSkeletalMeshPhysicsProxyParams
 {
 	FSkeletalMeshPhysicsProxyParams()
 		: Name("")
@@ -108,7 +108,7 @@ struct FSkeletalMeshPhysicsProxyOutputs : public Chaos::FParticleData
 };
 
 
-class CHAOS_API FSkeletalMeshPhysicsProxy : public TPhysicsProxy<FSkeletalMeshPhysicsProxy, FSkeletalMeshPhysicsProxyOutputs, FProxyTimestampBase>
+class FSkeletalMeshPhysicsProxy : public TPhysicsProxy<FSkeletalMeshPhysicsProxy, FSkeletalMeshPhysicsProxyOutputs, FProxyTimestampBase>
 {
 	typedef TPhysicsProxy<FSkeletalMeshPhysicsProxy, FSkeletalMeshPhysicsProxyOutputs, FProxyTimestampBase> Base;
 public:
@@ -119,32 +119,32 @@ public:
 
 
 	FSkeletalMeshPhysicsProxy() = delete;
-	FSkeletalMeshPhysicsProxy(UObject* InOwner, const FInitFunc& InitFunc);
-	~FSkeletalMeshPhysicsProxy();
+	CHAOS_API FSkeletalMeshPhysicsProxy(UObject* InOwner, const FInitFunc& InitFunc);
+	CHAOS_API ~FSkeletalMeshPhysicsProxy();
 
 	/** Solver Object interface */
-	void Initialize();
-	bool IsSimulating() const;
-	void UpdateKinematicBodiesCallback(const FParticlesType& Particles, const float Dt, const float Time, FKinematicProxy& Proxy);
-	void StartFrameCallback(const float InDt, const float InTime);
-	void EndFrameCallback(const float InDt);
-	void CreateRigidBodyCallback(FParticlesType& InOutParticles);
-	void ParameterUpdateCallback(FParticlesType& InParticles, const float InTime);
-	void DisableCollisionsCallback(TSet<TTuple<int32, int32>>& InPairs);
-	void AddForceCallback(FParticlesType& InParticles, const float InDt, const int32 InIndex);
-
-	void BindParticleCallbackMapping(Chaos::TArrayCollectionArray<PhysicsProxyWrapper>& PhysicsProxyReverseMap, Chaos::TArrayCollectionArray<int32>& ParticleIDReverseMap);
+	CHAOS_API void Initialize();
+	CHAOS_API bool IsSimulating() const;
+	CHAOS_API void UpdateKinematicBodiesCallback(const FParticlesType& Particles, const float Dt, const float Time, FKinematicProxy& Proxy);
+	CHAOS_API void StartFrameCallback(const float InDt, const float InTime);
+	CHAOS_API void EndFrameCallback(const float InDt);
+	CHAOS_API void CreateRigidBodyCallback(FParticlesType& InOutParticles);
+	CHAOS_API void ParameterUpdateCallback(FParticlesType& InParticles, const float InTime);
+	CHAOS_API void DisableCollisionsCallback(TSet<TTuple<int32, int32>>& InPairs);
+	CHAOS_API void AddForceCallback(FParticlesType& InParticles, const float InDt, const int32 InIndex);
+
+	CHAOS_API void BindParticleCallbackMapping(Chaos::TArrayCollectionArray<PhysicsProxyWrapper>& PhysicsProxyReverseMap, Chaos::TArrayCollectionArray<int32>& ParticleIDReverseMap);
 
 	void BufferCommand(Chaos::FPhysicsSolver* InSolver, const FFieldSystemCommand& InCommmand) {}
 
 	FSkeletalMeshPhysicsProxyOutputs* NewData() { return nullptr; }
-	void SyncBeforeDestroy();
-	void OnRemoveFromScene();
+	CHAOS_API void SyncBeforeDestroy();
+	CHAOS_API void OnRemoveFromScene();
 	void PushToPhysicsState(const Chaos::FParticleData*) {};
 	void ClearAccumulatedData() {}
-	void BufferPhysicsResults();
-	void FlipBuffer();
-	bool PullFromPhysicsState(const int32 SolverSyncTimestamp);
+	CHAOS_API void BufferPhysicsResults();
+	CHAOS_API void FlipBuffer();
+	CHAOS_API bool PullFromPhysicsState(const int32 SolverSyncTimestamp);
 	bool IsDirty() { return false; }
 	EPhysicsProxyType ConcreteType() { return EPhysicsProxyType::SkeletalMeshType; }
 	/** ----------------------- */
@@ -152,13 +152,13 @@ public:
 	/**
 	 *
 	 */
-	void Reset();
+	CHAOS_API void Reset();
 
 	/**
 	 * Capture the current animation pose for use by the physics.
 	 * Called by game thread via the owning component's tick.
 	 */
-	void CaptureInputs(const float Dt, const FInputFunc& InputFunc);
+	CHAOS_API void CaptureInputs(const float Dt, const FInputFunc& InputFunc);
 
 	/** 
 	 */
diff --git a/Engine/Source/Runtime/Experimental/Chaos/Public/PhysicsProxy/StaticMeshPhysicsProxy.h b/Engine/Source/Runtime/Experimental/Chaos/Public/PhysicsProxy/StaticMeshPhysicsProxy.h
--- a/Engine/Source/Runtime/Experimental/Chaos/Public/PhysicsProxy/StaticMeshPhysicsProxy.h
+++ b/Engine/Source/Runtime/Experimental/Chaos/Public/PhysicsProxy/StaticMeshPhysicsProxy.h
@@ -16,7 +16,7 @@ struct FStubSkeletalMeshData : public Chaos::FParticleData {
 	void Reset() { };
 };
 
-class CHAOS_API FStaticMeshPhysicsProxy : public TPhysicsProxy<FStaticMeshPhysicsProxy, FStubSkeletalMeshData, FProxyTimestampBase>
+class FStaticMeshPhysicsProxy : public TPhysicsProxy<FStaticMeshPhysicsProxy, FStubSkeletalMeshData, FProxyTimestampBase>
 {
 	typedef  TPhysicsProxy<FStaticMeshPhysicsProxy, FStubSkeletalMeshData, FProxyTimestampBase> Base;
 public:
@@ -71,33 +71,33 @@ public:
 	using FSyncDynamicFunc = TFunction<void(const FTransform&)>;
 
 	FStaticMeshPhysicsProxy() = delete;
-	FStaticMeshPhysicsProxy(UObject* InOwner, FCallbackInitFunc InInitFunc, FSyncDynamicFunc InSyncFunc);
+	CHAOS_API FStaticMeshPhysicsProxy(UObject* InOwner, FCallbackInitFunc InInitFunc, FSyncDynamicFunc InSyncFunc);
 
-	void Initialize();
-	void Reset();
+	CHAOS_API void Initialize();
+	CHAOS_API void Reset();
 
 	/** Stores latest update, to be applied at next opportunity (via UpdateKinematicBodiesCallback). */
-	void BufferKinematicUpdate(const FPhysicsProxyKinematicUpdate& InParamUpdate);
+	CHAOS_API void BufferKinematicUpdate(const FPhysicsProxyKinematicUpdate& InParamUpdate);
 
 	/** Solver Object interface */
-	bool IsSimulating() const;
-	void UpdateKinematicBodiesCallback(const FParticlesType& Particles, const float Dt, const float Time, FKinematicProxy& Proxy);
-	void StartFrameCallback(const float InDt, const float InTime);
-	void EndFrameCallback(const float InDt);
-	void BindParticleCallbackMapping(Chaos::TArrayCollectionArray<PhysicsProxyWrapper> & PhysicsProxyReverseMap, Chaos::TArrayCollectionArray<int32> & ParticleIDReverseMap);
-	void CreateRigidBodyCallback(FParticlesType& InOutParticles);
-	void ParameterUpdateCallback(FParticlesType& InParticles, const float InTime);
-	void DisableCollisionsCallback(TSet<TTuple<int32, int32>>& InPairs);
-	void AddForceCallback(FParticlesType& InParticles, const float InDt, const int32 InIndex);
+	CHAOS_API bool IsSimulating() const;
+	CHAOS_API void UpdateKinematicBodiesCallback(const FParticlesType& Particles, const float Dt, const float Time, FKinematicProxy& Proxy);
+	CHAOS_API void StartFrameCallback(const float InDt, const float InTime);
+	CHAOS_API void EndFrameCallback(const float InDt);
+	CHAOS_API void BindParticleCallbackMapping(Chaos::TArrayCollectionArray<PhysicsProxyWrapper> & PhysicsProxyReverseMap, Chaos::TArrayCollectionArray<int32> & ParticleIDReverseMap);
+	CHAOS_API void CreateRigidBodyCallback(FParticlesType& InOutParticles);
+	CHAOS_API void ParameterUpdateCallback(FParticlesType& InParticles, const float InTime);
+	CHAOS_API void DisableCollisionsCallback(TSet<TTuple<int32, int32>>& InPairs);
+	CHAOS_API void AddForceCallback(FParticlesType& InParticles, const float InDt, const int32 InIndex);
 	void BufferCommand(Chaos::FPhysicsSolver* InSolver, const FFieldSystemCommand& InCommand) {};
 
 	void SyncBeforeDestroy() {};
-	void OnRemoveFromScene();
+	CHAOS_API void OnRemoveFromScene();
 	void PushToPhysicsState(const Chaos::FParticleData*) {};
 	void ClearAccumulatedData() {}
-	void BufferPhysicsResults();
-	void FlipBuffer();
-	bool PullFromPhysicsState(const int32 SolverSyncTimestamp);
+	CHAOS_API void BufferPhysicsResults();
+	CHAOS_API void FlipBuffer();
+	CHAOS_API bool PullFromPhysicsState(const int32 SolverSyncTimestamp);
 	bool IsDirty() { return false; }
 	FStubSkeletalMeshData* NewData() { return nullptr; }
 	EPhysicsProxyType ConcreteType() { return EPhysicsProxyType::StaticMeshType; }
diff --git a/Engine/Source/Runtime/Experimental/Chaos/Public/PhysicsProxy/SuspensionConstraintProxy.h b/Engine/Source/Runtime/Experimental/Chaos/Public/PhysicsProxy/SuspensionConstraintProxy.h
--- a/Engine/Source/Runtime/Experimental/Chaos/Public/PhysicsProxy/SuspensionConstraintProxy.h
+++ b/Engine/Source/Runtime/Experimental/Chaos/Public/PhysicsProxy/SuspensionConstraintProxy.h
@@ -18,35 +18,35 @@ namespace Chaos
 
 	class FPBDRigidsEvolutionGBF;
 
-class CHAOS_API FSuspensionConstraintPhysicsProxy : public IPhysicsProxyBase
+class FSuspensionConstraintPhysicsProxy : public IPhysicsProxyBase
 {
 public:
 	using Base = IPhysicsProxyBase;
 	
 	FSuspensionConstraintPhysicsProxy() = delete;
-	FSuspensionConstraintPhysicsProxy(FSuspensionConstraint* InConstraint, FPBDSuspensionConstraintHandle* InHandle, UObject* InOwner = nullptr);
+	CHAOS_API FSuspensionConstraintPhysicsProxy(FSuspensionConstraint* InConstraint, FPBDSuspensionConstraintHandle* InHandle, UObject* InOwner = nullptr);
 
 	bool IsInitialized() const { return bInitialized; }
 	void SetInitialized() { bInitialized = true; }
 
-	static FGeometryParticleHandle* GetParticleHandleFromProxy(IPhysicsProxyBase* ProxyBase);
+	static CHAOS_API FGeometryParticleHandle* GetParticleHandleFromProxy(IPhysicsProxyBase* ProxyBase);
 
 	//
 	//  Lifespan Management
 	//
 
-	void InitializeOnPhysicsThread(FPBDRigidsSolver* InSolver, FDirtyPropertiesManager& Manager, int32 DataIdx, FDirtyChaosProperties& RemoteData);
+	CHAOS_API void InitializeOnPhysicsThread(FPBDRigidsSolver* InSolver, FDirtyPropertiesManager& Manager, int32 DataIdx, FDirtyChaosProperties& RemoteData);
 
 	// Merge to perform a remote sync
-	void PushStateOnGameThread(FDirtyPropertiesManager& Manager, int32 DataIdx, FDirtyChaosProperties& RemoteData);
+	CHAOS_API void PushStateOnGameThread(FDirtyPropertiesManager& Manager, int32 DataIdx, FDirtyChaosProperties& RemoteData);
 
-	void PushStateOnPhysicsThread(FPBDRigidsSolver* InSolver, const FDirtyPropertiesManager& Manager, int32 DataIdx, const FDirtyChaosProperties& RemoteData);
+	CHAOS_API void PushStateOnPhysicsThread(FPBDRigidsSolver* InSolver, const FDirtyPropertiesManager& Manager, int32 DataIdx, const FDirtyChaosProperties& RemoteData);
 	// Merge to perform a remote sync - END
 
-	void DestroyOnGameThread();
-	void DestroyOnPhysicsThread(FPBDRigidsSolver* InSolver);
+	CHAOS_API void DestroyOnGameThread();
+	CHAOS_API void DestroyOnPhysicsThread(FPBDRigidsSolver* InSolver);
 
-	void UpdateTargetOnPhysicsThread(FPBDRigidsSolver* InSolver, const FVector& TargetPos, const FVector& Normal, bool Enabled);
+	CHAOS_API void UpdateTargetOnPhysicsThread(FPBDRigidsSolver* InSolver, const FVector& TargetPos, const FVector& Normal, bool Enabled);
 
 
 	//
@@ -88,4 +88,4 @@ private:
 	FPBDSuspensionConstraintHandle* Constraint_PT;
 	bool bInitialized;
 };
-}
\ No newline at end of file
+}
diff --git a/Engine/Source/Runtime/Experimental/Chaos/Public/RewindData.h b/Engine/Source/Runtime/Experimental/Chaos/Public/RewindData.h
--- a/Engine/Source/Runtime/Experimental/Chaos/Public/RewindData.h
+++ b/Engine/Source/Runtime/Experimental/Chaos/Public/RewindData.h
@@ -1277,17 +1277,17 @@ public:
 		return Managers[Frame].DeltaTime;
 	}
 
-	void CHAOS_API RemoveObject(const FGeometryParticleHandle* Particle)
+	void RemoveObject(const FGeometryParticleHandle* Particle)
 	{
 		DirtyParticles.Remove(Particle);
 	}
 
-	void CHAOS_API RemoveObject(const FPBDJointConstraintHandle* Joint)
+	void RemoveObject(const FPBDJointConstraintHandle* Joint)
 	{
 		DirtyJoints.Remove(Joint);
 	}
 
-	int32 CHAOS_API GetEarliestFrame_Internal() const { return CurFrame - FramesSaved; }
+	int32 GetEarliestFrame_Internal() const { return CurFrame - FramesSaved; }
 
 	/* Extend the current history size to be sure to include the given frame */
 	void CHAOS_API ExtendHistoryWithFrame(const int32 Frame);
@@ -1303,14 +1303,14 @@ public:
 		const FVector& Position, const FQuat& Quaternion, const FVector& LinVelocity, const FVector& AngVelocity, const bool bShouldSleep);
 
 	/** Extract some history information before cleaning/pushing state*/
-	void CHAOS_API ExtractHistoryState(FGeometryParticleHandle& Handle, int32& PositionValidCount, int32& VelocityValidCount, int32& PositionNextIterator, int32& VelocityNextIterator)
+	void ExtractHistoryState(FGeometryParticleHandle& Handle, int32& PositionValidCount, int32& VelocityValidCount, int32& PositionNextIterator, int32& VelocityNextIterator)
 	{
 		FDirtyParticleInfo& Info = FindOrAddDirtyObj(Handle);
 		Info.GetHistory().ExtractHistoryState(PositionValidCount, VelocityValidCount, PositionNextIterator, VelocityNextIterator);
 	}
 
 	/** Restore some history information after cleaning/pushing state*/
-	void CHAOS_API RestoreHistoryState(FGeometryParticleHandle& Handle, const int32& PositionValidCount, const int32& VelocityValidCount, const int32& PositionNextIterator, const int32& VelocityNextIterator)
+	void RestoreHistoryState(FGeometryParticleHandle& Handle, const int32& PositionValidCount, const int32& VelocityValidCount, const int32& PositionNextIterator, const int32& VelocityNextIterator)
 	{
 		FDirtyParticleInfo& Info = FindOrAddDirtyObj(Handle);
 		Info.GetHistory().RestoreHistoryState(PositionValidCount, VelocityValidCount, PositionNextIterator, VelocityNextIterator);
@@ -1392,25 +1392,25 @@ public:
 	void CHAOS_API SpawnProxyIfNeeded(FSingleParticlePhysicsProxy& Proxy);
 
 	/** Add Inputs history to the rewind datas for future use while resimulating */
-	void CHAOS_API AddInputsHistory(const TSharedPtr<FBaseRewindHistory>& InputsHistory)
+	void AddInputsHistory(const TSharedPtr<FBaseRewindHistory>& InputsHistory)
 	{
 		InputsHistories.Add(InputsHistory.ToWeakPtr());
 	}
 
 	/** Remove Inputs history from the rewind datas */
-	void CHAOS_API RemoveInputsHistory(const TSharedPtr<FBaseRewindHistory>& InputsHistory)
+	void RemoveInputsHistory(const TSharedPtr<FBaseRewindHistory>& InputsHistory)
 	{
 		InputsHistories.Remove(InputsHistory.ToWeakPtr());
 	}
 
 	/** Add states history to the rewind datas for future use while rewinding */
-	void CHAOS_API AddStatesHistory(const TSharedPtr<FBaseRewindHistory>& StatesHistory)
+	void AddStatesHistory(const TSharedPtr<FBaseRewindHistory>& StatesHistory)
 	{
 		StatesHistories.Add(StatesHistory.ToWeakPtr());
 	}
 
 	/** Remove states history from the rewind datas */
-	void CHAOS_API RemoveStatesHistory(const TSharedPtr<FBaseRewindHistory>& StatesHistory)
+	void RemoveStatesHistory(const TSharedPtr<FBaseRewindHistory>& StatesHistory)
 	{
 		StatesHistories.Remove(StatesHistory.ToWeakPtr());
 	}
@@ -1677,4 +1677,4 @@ public:
 	/** Rewind Data holding the callback */
 	Chaos::FRewindData* RewindData = nullptr;
 };
-}
\ No newline at end of file
+}
diff --git a/Engine/Source/Runtime/Experimental/Chaos/Public/SolverEventFilters.h b/Engine/Source/Runtime/Experimental/Chaos/Public/SolverEventFilters.h
--- a/Engine/Source/Runtime/Experimental/Chaos/Public/SolverEventFilters.h
+++ b/Engine/Source/Runtime/Experimental/Chaos/Public/SolverEventFilters.h
@@ -10,7 +10,7 @@
 
 
 	USTRUCT(Blueprintable)
-	struct CHAOS_API FSolverTrailingFilterSettings
+	struct FSolverTrailingFilterSettings
 	{
 		GENERATED_USTRUCT_BODY()
 
@@ -39,7 +39,7 @@
 	};
 
 	USTRUCT(BlueprintType)
-	struct CHAOS_API FSolverCollisionFilterSettings
+	struct FSolverCollisionFilterSettings
 	{
 		GENERATED_USTRUCT_BODY()
 
@@ -69,7 +69,7 @@
 	};
 
 	USTRUCT(BlueprintType)
-	struct CHAOS_API FSolverBreakingFilterSettings
+	struct FSolverBreakingFilterSettings
 	{
 		GENERATED_USTRUCT_BODY()
 
@@ -98,7 +98,7 @@
 	};
 
 	USTRUCT(BlueprintType)
-	struct CHAOS_API FSolverRemovalFilterSettings
+	struct FSolverRemovalFilterSettings
 	{
 		GENERATED_USTRUCT_BODY()
 
@@ -133,52 +133,52 @@ namespace Chaos
 
 	struct FRemovalData;
 
-	class CHAOS_API FSolverCollisionEventFilter
+	class FSolverCollisionEventFilter
 	{
 	public:
 		FSolverCollisionEventFilter() {}
 		FSolverCollisionEventFilter(const FSolverCollisionFilterSettings& InSettings) : Settings(InSettings) {}
 
-		bool Pass(const Chaos::FCollidingData& InData) const;
+		CHAOS_API bool Pass(const Chaos::FCollidingData& InData) const;
 		bool Enabled() const { return Settings.FilterEnabled; }
 		void UpdateFilterSettings(const FSolverCollisionFilterSettings& InSettings) { Settings = InSettings; }
 
 		FSolverCollisionFilterSettings Settings;
 	};
 
-	class CHAOS_API FSolverTrailingEventFilter
+	class FSolverTrailingEventFilter
 	{
 	public:
 		FSolverTrailingEventFilter() {}
 		FSolverTrailingEventFilter(const FSolverTrailingFilterSettings &InSettings) : Settings(InSettings) {}
 
-		bool Pass(const Chaos::FTrailingData& InData) const;
+		CHAOS_API bool Pass(const Chaos::FTrailingData& InData) const;
 		bool Enabled() const { return Settings.FilterEnabled; }
 		void UpdateFilterSettings(const FSolverTrailingFilterSettings& InSettings) { Settings = InSettings; }
 
 		FSolverTrailingFilterSettings Settings;
 	};
 
-	class CHAOS_API FSolverBreakingEventFilter
+	class FSolverBreakingEventFilter
 	{
 	public:
 		FSolverBreakingEventFilter() {}
 		FSolverBreakingEventFilter(const FSolverBreakingFilterSettings& InSettings) : Settings(InSettings) {}
 
-		bool Pass(const Chaos::FBreakingData& InData) const;
+		CHAOS_API bool Pass(const Chaos::FBreakingData& InData) const;
 		bool Enabled() const { return Settings.FilterEnabled; }
 		void UpdateFilterSettings(const FSolverBreakingFilterSettings& InSettings) { Settings = InSettings; }
 
 		FSolverBreakingFilterSettings Settings;
 	};
 
-	class CHAOS_API FSolverRemovalEventFilter
+	class FSolverRemovalEventFilter
 	{
 	public:
 		FSolverRemovalEventFilter() {}
 		FSolverRemovalEventFilter(const FSolverRemovalFilterSettings& InSettings) : Settings(InSettings) {}
 
-		bool Pass(const Chaos::FRemovalData& InData) const;
+		CHAOS_API bool Pass(const Chaos::FRemovalData& InData) const;
 		bool Enabled() const { return Settings.FilterEnabled; }
 		void UpdateFilterSettings(const FSolverRemovalFilterSettings& InSettings) { Settings = InSettings; }
 
