Consolidation of cache stats

#rnx

diff --git a/Engine/Source/Runtime/Experimental/IoStoreOnDemand/Private/FileCache.cpp b/Engine/Source/Runtime/Experimental/IoStoreOnDemand/Private/FileCache.cpp
--- a/Engine/Source/Runtime/Experimental/IoStoreOnDemand/Private/FileCache.cpp
+++ b/Engine/Source/Runtime/Experimental/IoStoreOnDemand/Private/FileCache.cpp
@@ -21,15 +21,6 @@
 
 #include <atomic>
 
-TRACE_DECLARE_MEMORY_COUNTER(FFileIoCache_CachedBytes, TEXT("FileIoCache/TotalBytes"));
-TRACE_DECLARE_MEMORY_COUNTER(FFileIoCache_PendingBytes, TEXT("FileIoCache/PendingBytes"));
-TRACE_DECLARE_MEMORY_COUNTER(FFileIoCache_ReadBytes, TEXT("FileIoCache/TotalReadBytes"));
-TRACE_DECLARE_INT_COUNTER(FFileIoCache_GetCount, TEXT("FileIoCache/GetCount"));
-TRACE_DECLARE_INT_COUNTER(FFileIoCache_ErrorCount, TEXT("FileIoCache/ErrorCount"));
-TRACE_DECLARE_INT_COUNTER(FFileIoCache_PutCount, TEXT("FileIoCache/PutCount"));
-TRACE_DECLARE_INT_COUNTER(FFileIoCache_PutRejectCount, TEXT("FileIoCache/PutRejectCount"));
-TRACE_DECLARE_INT_COUNTER(FFileIoCache_PutExistingCount, TEXT("FileIoCache/PutExistingCount"));
-
 DEFINE_LOG_CATEGORY(LogIoCache);
 
 namespace UE::IO::Private
@@ -37,7 +28,90 @@ namespace UE::IO::Private
 
 using namespace UE::Tasks;
 
-///////////////////////////////////////////////////////////////////////////////
+////////////////////////////////////////////////////////////////////////////////
+#define BUILD_DATUM_X(x) \
+	x(CachedBytes) \
+	x(ErrorCount) \
+	x(GetCount) \
+	x(PendingBytes) \
+	x(PutCount) \
+	x(PutExistingCount) \
+	x(PutRejectCount) \
+	x(ReadBytes)
+
+#if !COUNTERSTRACE_ENABLED
+namespace FCountersTrace
+{
+	struct FCounterInt
+	{
+		FCounterInt(...) {}
+		void Set(int64 i) { V = i; }
+		void Add(int64 d) { V += d; }
+		int64 Get() const { return V;}
+		std::atomic<int64> V = 0;
+	};
+	using FCounterAtomicInt = FCounterInt;
+}
+#endif 
+
+class FStatistics
+{
+public:
+	enum class EDatum : int32
+	{
+#define DATUM_X(x) x,
+		BUILD_DATUM_X(DATUM_X)
+#undef DATUM_X
+		_Count,
+	};
+
+	void Add(EDatum Which, int64 Delta);
+	void Set(EDatum Which, int64 Value);
+	int64 Get(EDatum Which) const;
+
+private:
+	struct FDatum
+	{
+		FDatum(const TCHAR* Name);
+		FCountersTrace::FCounterAtomicInt Value;
+	};
+
+	FDatum Datums[int(EDatum::_Count)] = {
+#define DATUM_X(x) FDatum{ TEXT("Ias/FileCache/") TEXT(PREPROCESSOR_TO_STRING(x)) },
+		BUILD_DATUM_X(DATUM_X)
+#undef DATUM_X
+	};
+};
+
+static FStatistics GStatistics;
+
+FStatistics::FDatum::FDatum(const TCHAR* Name)
+: Value(TraceCounterNameType_Static, Name, TraceCounterDisplayHint_None)
+{
+	FCountersTrace::FCounterAtomicInt(TEXT(""), TraceCounterDisplayHint_None).Get();
+	FCountersTrace::FCounterInt(TEXT(""), TraceCounterDisplayHint_None).Get();
+}
+
+void FStatistics::Add(EDatum Which, int64 Delta)
+{
+	Datums[int(Which)].Value.Add(Delta);
+}
+
+void FStatistics::Set(EDatum Which, int64 Value)
+{
+	Datums[int(Which)].Value.Set(Value);
+}
+
+int64 FStatistics::Get(EDatum Which) const
+{
+	return Datums[int(Which)].Value.Get();
+}
+
+
+
+// {{{1 toc-file
+
+////////////////////////////////////////////////////////////////////////////////
 class FCacheFileToc
 {
 	struct FHeader
@@ -236,13 +310,13 @@ bool FCacheMap::InsertPending(FIoHash Key, FIoBuffer& Data, bool& bAdded)
 	bAdded = false;
 	if (TotalPendingBytes + Data.GetSize() > MaxPendingBytes)
 	{
-		TRACE_COUNTER_INCREMENT(FFileIoCache_PutRejectCount);
+		GStatistics.Add(FStatistics::EDatum::PutRejectCount, 1);
 		return false;
 	}
 
 	if (Lookup.Contains(Key))
 	{
-		TRACE_COUNTER_INCREMENT(FFileIoCache_PutExistingCount);
+		GStatistics.Add(FStatistics::EDatum::PutExistingCount, 1);
 		return true;
 	}
 
@@ -256,8 +330,8 @@ bool FCacheMap::InsertPending(FIoHash Key, FIoBuffer& Data, bool& bAdded)
 
 	TotalPendingBytes += Data.GetSize();
 
-	TRACE_COUNTER_ADD(FFileIoCache_PendingBytes, Data.GetSize());
-	TRACE_COUNTER_INCREMENT(FFileIoCache_PutCount);
+	GStatistics.Add(FStatistics::EDatum::PendingBytes, Data.GetSize());
+	GStatistics.Add(FStatistics::EDatum::PutCount, 1);
 
 	return bAdded = true;
 }
@@ -289,7 +363,7 @@ int32 FCacheMap::RemovePending(FCacheEntryList& OutPending, uint32 MaxSize)
 	}
 
 	TotalPendingBytes -= ReturnSize;
-	TRACE_COUNTER_SET(FFileIoCache_PendingBytes, TotalPendingBytes);
+	GStatistics.Set(FStatistics::EDatum::PendingBytes, TotalPendingBytes);
 
 	return ReturnSize;
 }
@@ -314,7 +388,7 @@ void FCacheMap::InsertPersisted(FCacheEntryList&& InPersisted, const uint64 Curs
 
 	Persisted.AddTail(MoveTemp(InPersisted));
 	TotalPersistedBytes += PersistedBytes;
-	TRACE_COUNTER_ADD(FFileIoCache_CachedBytes, PersistedBytes);
+	GStatistics.Add(FStatistics::EDatum::CachedBytes, PersistedBytes);
 }
 
 void FCacheMap::RemovePersisted(const uint64 RequiredSize)
@@ -338,7 +412,7 @@ void FCacheMap::RemovePersisted(const uint64 RequiredSize)
 	}
 
 	TotalPersistedBytes -= RemovedBytes;
-	TRACE_COUNTER_SUBTRACT(FFileIoCache_CachedBytes, RemovedBytes);
+	GStatistics.Add(FStatistics::EDatum::CachedBytes, -int64(RemovedBytes));
 }
 
 FIoStatus FCacheMap::Load(const FString& FilePath, uint64& OutCursorPos)
@@ -372,7 +446,7 @@ FIoStatus FCacheMap::Load(const FString& FilePath, uint64& OutCursorPos)
 		check(ExpectedCursosPos == OutCursorPos);
 	}
 
-	TRACE_COUNTER_SET(FFileIoCache_CachedBytes, TotalPersistedBytes);
+	GStatistics.Set(FStatistics::EDatum::CachedBytes, TotalPersistedBytes);
 
 	return FIoStatus(EIoErrorCode::Ok);
 }
@@ -594,7 +668,7 @@ TTask<TIoStatusOr<FIoBuffer>> FFileIoCache::Get(const FIoHash& Key, const FIoRea
 				const uint64 ReadSize = FMath::Min(Options.GetSize(), Entry.Data.GetSize());
 				FIoBuffer Buffer = Options.GetTargetVa() ? FIoBuffer(FIoBuffer::Wrap, Options.GetTargetVa(), ReadSize) : FIoBuffer(ReadSize);
 				Buffer.GetMutableView().CopyFrom(Entry.Data.GetView().RightChop(ReadOffset));
-				TRACE_COUNTER_INCREMENT(FFileIoCache_GetCount);
+				GStatistics.Add(FStatistics::EDatum::GetCount, 1);
 
 				return TIoStatusOr<FIoBuffer>(Buffer);
 			}
@@ -622,12 +696,12 @@ TTask<TIoStatusOr<FIoBuffer>> FFileIoCache::Get(const FIoHash& Key, const FIoRea
 					const FIoHash Hash = FIoHash::HashBuffer(Buffer.GetView());
 					if (Hash == ExpectedHash) 
 					{
-						TRACE_COUNTER_INCREMENT(FFileIoCache_GetCount);
-						TRACE_COUNTER_ADD(FFileIoCache_ReadBytes, ReadSize);
+						GStatistics.Add(FStatistics::EDatum::GetCount, 1);
+						GStatistics.Add(FStatistics::EDatum::ReadBytes, ReadSize);
 						return TIoStatusOr<FIoBuffer>(Buffer);
 					}
 
-					TRACE_COUNTER_INCREMENT(FFileIoCache_ErrorCount);
+					GStatistics.Add(FStatistics::EDatum::ErrorCount, 1);
 					UE_LOG(LogIoCache, Verbose, TEXT("Read chunk failed, hash mismatch, Key='%s', Hash='%s', ExpectedHash='%s', File='%s', Offset='%llu', Size='%llu'"),
 						*LexToString(Entry.Key), *LexToString(Hash), *LexToString(ExpectedHash), *CacheFilePath, ReadOffset, ReadSize);
 
@@ -635,7 +709,7 @@ TTask<TIoStatusOr<FIoBuffer>> FFileIoCache::Get(const FIoHash& Key, const FIoRea
 				}
 				else
 				{
-					TRACE_COUNTER_INCREMENT(FFileIoCache_ErrorCount);
+					GStatistics.Add(FStatistics::EDatum::ErrorCount, 1);
 					UE_LOG(LogIoCache, Warning, TEXT("Read chunk failed, unable to open cache file '%s' for reading"), *CacheFilePath);
 					return TIoStatusOr<FIoBuffer>(FIoStatus(EIoErrorCode::ReadError));
 				}
