Writer was missing if the cache was dropped

#rnx

diff --git a/Engine/Source/Runtime/Experimental/IoStoreOnDemand/Private/FileCache.cpp b/Engine/Source/Runtime/Experimental/IoStoreOnDemand/Private/FileCache.cpp
--- a/Engine/Source/Runtime/Experimental/IoStoreOnDemand/Private/FileCache.cpp
+++ b/Engine/Source/Runtime/Experimental/IoStoreOnDemand/Private/FileCache.cpp
@@ -737,6 +737,8 @@ void FFileIoCache::Initialize()
 	UE_LOG(LogIoCache, Log,
 		TEXT("Initializing file I/O cache, disk size %lluB, memory size %lluB"), CacheConfig.DiskQuota, CacheConfig.MemoryQuota);
 
+	WriterThread.Reset(FRunnableThread::Create(this, TEXT("Ias.FileCache"), 0, TPri_BelowNormal));
+
 	const FString CacheDir = FPaths::ProjectPersistentDownloadDir() / TEXT("IoCache");
 	const FString CacheTocPath = CacheDir / TEXT("cache.utoc");
 	CacheFilePath = CacheDir / TEXT("cache.ucas");
@@ -795,8 +797,6 @@ void FFileIoCache::Initialize()
 			FileMgr.MakeDirectory(*CacheDir, true);
 		}
 	}
-
-	WriterThread.Reset(FRunnableThread::Create(this, TEXT("File I/O Cache"), 0, TPri_BelowNormal));
 }
 
 void FFileIoCache::Shutdown()
