Simple IAS network latency simulation

#rnx

diff --git a/Engine/Source/Runtime/Experimental/IoStoreOnDemand/Private/CoreHttp/Http.cpp b/Engine/Source/Runtime/Experimental/IoStoreOnDemand/Private/CoreHttp/Http.cpp
--- a/Engine/Source/Runtime/Experimental/IoStoreOnDemand/Private/CoreHttp/Http.cpp
+++ b/Engine/Source/Runtime/Experimental/IoStoreOnDemand/Private/CoreHttp/Http.cpp
@@ -4,6 +4,7 @@
 
 #if !defined(NO_UE_INCLUDES)
 #include "IO/IoBuffer.h"
+#include "LatencyInjector.h"
 #include "Math/UnrealMathUtility.h"
 #include "Misc/ScopeExit.h"
 #include "Misc/ScopeLock.h"
@@ -66,6 +67,8 @@ enum : SocketType { InvalidSocket = -1 };
 namespace UE::HTTP
 {
 
+using FLatencyInjector = UE::IO::Private::FLatencyInjector;
+
 // {{{1 misc ...................................................................
 
 ////////////////////////////////////////////////////////////////////////////////
@@ -1633,6 +1636,8 @@ int32 FEventLoopInternal::DoRecvContent(FActivity* Activity)
 		Activity->StateParam += Result;
 	}
 
+	FLatencyInjector::Begin(FLatencyInjector::EType::Network, Activity->StateParam);
+
 	Activity->State = FActivity::EState::RecvDone;
 	return 0;
 }
@@ -1640,6 +1645,11 @@ int32 FEventLoopInternal::DoRecvContent(FActivity* Activity)
 ////////////////////////////////////////////////////////////////////////////////
 int32 FEventLoopInternal::DoRecvDone(FActivity* Activity)
 {
+	if (!FLatencyInjector::HasExpired(Activity->StateParam))
+	{
+		return 1;
+	}
+
 	// Notify the user we've received everything
 	FTicketStatus& SinkArg = *(FTicketStatus*)Activity;
 	Activity->Sink(SinkArg);
diff --git a/Engine/Source/Runtime/Experimental/IoStoreOnDemand/Private/IoStoreOnDemand.cpp b/Engine/Source/Runtime/Experimental/IoStoreOnDemand/Private/IoStoreOnDemand.cpp
--- a/Engine/Source/Runtime/Experimental/IoStoreOnDemand/Private/IoStoreOnDemand.cpp
+++ b/Engine/Source/Runtime/Experimental/IoStoreOnDemand/Private/IoStoreOnDemand.cpp
@@ -3,6 +3,7 @@
 #include "IO/IoStoreOnDemand.h"
 #include "OnDemandIoDispatcherBackend.h"
 #include "EncryptionKeyManager.h"
+#include "LatencyInjector.h"
 
 #include "FileCache.h"
 #include "HAL/PlatformMisc.h"
@@ -442,6 +443,8 @@ void FIoStoreOnDemandModule::StartupModule()
 		return;
 	}
 
+	FLatencyInjector::Initialize(CommandLine);
+
 	Endpoint.EndpointType = UE::EOnDemandEndpointType::CDN;
 
 	TSharedPtr<IIoCache> Cache;
diff --git a/Engine/Source/Runtime/Experimental/IoStoreOnDemand/Private/LatencyInjector.cpp b/Engine/Source/Runtime/Experimental/IoStoreOnDemand/Private/LatencyInjector.cpp
new file mode 100644
--- /dev/null
+++ b/Engine/Source/Runtime/Experimental/IoStoreOnDemand/Private/LatencyInjector.cpp
@@ -0,0 +1,116 @@
+// Copyright Epic Games, Inc. All Rights Reserved.
+
+#include "LatencyInjector.h"
+
+#include "String/ParseTokens.h"
+
+UE_DISABLE_OPTIMIZATION
+
+#if IAS_WITH_LATENCY_INJECTOR
+
+namespace UE::IO::Private
+{
+
+////////////////////////////////////////////////////////////////////////////////
+namespace Injector
+{
+
+static int64 CycleBase;
+static int64 CycleFreq;
+static int32 MinMs;
+static int32 MaxMs;
+
+};
+
+////////////////////////////////////////////////////////////////////////////////
+void FLatencyInjector::Initialize(const TCHAR* CommandLine)
+{
+	// -Ias.AddLatencyMs=130   : add 130ms of latency
+	// -Ias.AddLatencyMs=40/90 : random latency between 40ms and 90ms
+
+	FString Value;
+	if (!FParse::Value(CommandLine, TEXT("Ias.AddLatencyMs="), Value))
+	{
+		return;
+	}
+
+	TArray<int32> Values;
+	UE::String::ParseTokens(Value, TEXT("/"), [&Values] (const FStringView& Token)
+	{
+		int32 Int = -1;
+		LexFromString(Int, Token.GetData());
+		Values.Add((Int >= 0) ? Int : -1);
+	});
+
+	for (; Values.Num() < 2; Values.Add(-1));
+	Set(Values[0], Values[1]);
+}
+
+////////////////////////////////////////////////////////////////////////////////
+void FLatencyInjector::Set(int32 MinMs, int32 MaxMs)
+{
+	if (MinMs < 0)
+	{
+		Injector::MinMs = Injector::MaxMs = 0;
+		return;
+	}
+
+	if (MaxMs < 0)
+	{
+		MaxMs = MinMs;
+	}
+	else if (MinMs > MaxMs)
+	{
+		Swap(MinMs, MaxMs);
+	}
+
+	Injector::MinMs = MinMs;
+	Injector::MaxMs = MaxMs;
+}
+
+////////////////////////////////////////////////////////////////////////////////
+void FLatencyInjector::Begin(EType, uint32& Param)
+{
+	using namespace Injector;
+
+	if (MaxMs <= 0 || MinMs < 0)
+	{
+		Param = 0;
+		return;
+	}
+
+	int64 Cycles = FPlatformTime::Cycles64();
+
+	if (CycleBase == 0)
+	{
+		CycleBase = Cycles;
+		CycleFreq = int64(1.0 / FPlatformTime::GetSecondsPerCycle());
+	}
+
+	uint64 Bias = MaxMs - MinMs;
+	if (Bias)
+	{
+		Bias = (Cycles * 0x369dea0f31a53f85ull) % Bias;
+	}
+	Bias += MinMs;
+	Bias *= CycleFreq;
+	Bias /= 1000;
+
+	Param = uint32(Cycles + Bias - CycleBase);
+}
+
+////////////////////////////////////////////////////////////////////////////////
+bool FLatencyInjector::HasExpired(uint32 Param)
+{
+	if (Param == 0)
+	{
+		return true;
+	}
+
+	int64 Cycles = FPlatformTime::Cycles64();
+	return Cycles >= (int64(Param) + Injector::CycleBase);
+}
+
+} // namespace UE::IO::Private
+
+#endif // IAS_WITH_LATENCY_INJECTOR
diff --git a/Engine/Source/Runtime/Experimental/IoStoreOnDemand/Private/LatencyInjector.h b/Engine/Source/Runtime/Experimental/IoStoreOnDemand/Private/LatencyInjector.h
new file mode 100644
--- /dev/null
+++ b/Engine/Source/Runtime/Experimental/IoStoreOnDemand/Private/LatencyInjector.h
@@ -0,0 +1,25 @@
+// Copyright Epic Games, Inc. All Rights Reserved.
+
+#pragma once
+
+namespace UE::IO::Private
+{
+
+#if !UE_BUILD_SHIPPING
+#	define IAS_WITH_LATENCY_INJECTOR 1
+#	define IAS_DISABLED_IMPL(x)
+#else
+#	define IAS_WITH_LATENCY_INJECTOR 0
+#	define IAS_DISABLED_IMPL(x) { x; }
+#endif
+
+struct FLatencyInjector
+{
+	enum class EType { Network, File };
+	static void Initialize(const TCHAR* CommandLine)IAS_DISABLED_IMPL(return);
+	static void Set(int32 MinMs, int32 MaxMs)		IAS_DISABLED_IMPL(return);
+	static void Begin(EType, uint32& Param)			IAS_DISABLED_IMPL(return);
+	static bool HasExpired(uint32 Param)			IAS_DISABLED_IMPL(return true);
+};
+
+} // namespace UE::IO::Private
