Non-unity fix

#rnx

diff --git a/Engine/Source/Runtime/Experimental/IoStoreOnDemand/Private/LatencyInjector.cpp b/Engine/Source/Runtime/Experimental/IoStoreOnDemand/Private/LatencyInjector.cpp
--- a/Engine/Source/Runtime/Experimental/IoStoreOnDemand/Private/LatencyInjector.cpp
+++ b/Engine/Source/Runtime/Experimental/IoStoreOnDemand/Private/LatencyInjector.cpp
@@ -2,10 +2,11 @@
 
 #include "LatencyInjector.h"
 
+#include "Containers/UnrealString.h"
+#include "HAL/PlatformTime.h"
+#include "Misc/Parse.h"
 #include "String/ParseTokens.h"
 
-UE_DISABLE_OPTIMIZATION
-
 #if IAS_WITH_LATENCY_INJECTOR
 
 namespace UE::IO::Private
diff --git a/Engine/Source/Runtime/Experimental/IoStoreOnDemand/Private/LatencyInjector.h b/Engine/Source/Runtime/Experimental/IoStoreOnDemand/Private/LatencyInjector.h
--- a/Engine/Source/Runtime/Experimental/IoStoreOnDemand/Private/LatencyInjector.h
+++ b/Engine/Source/Runtime/Experimental/IoStoreOnDemand/Private/LatencyInjector.h
@@ -2,6 +2,9 @@
 
 #pragma once
 
+#include "HAL/Platform.h"
+#include "Misc/Build.h"
+
 namespace UE::IO::Private
 {
 
