[Chaos]
* Removed dll export that causes linux compile errors

diff --git a/Engine/Source/Runtime/Experimental/Chaos/Public/Chaos/TriangleMeshImplicitObject.h b/Engine/Source/Runtime/Experimental/Chaos/Public/Chaos/TriangleMeshImplicitObject.h
--- a/Engine/Source/Runtime/Experimental/Chaos/Public/Chaos/TriangleMeshImplicitObject.h
+++ b/Engine/Source/Runtime/Experimental/Chaos/Public/Chaos/TriangleMeshImplicitObject.h
@@ -47,7 +47,7 @@ namespace Chaos
 			Keep,
 		};
 		
-		CHAOS_API struct alignas(16) FChildData
+		struct alignas(16) FChildData
 		{
 			FChildData()
 				: ChildOrFaceIndex(INDEX_NONE)
