[Chaos]
* Re-added class dllexport to FPhysicsThreadContext and FChaosVDThreadContext.

This code relies on undefined behavior of how things are inlined. base class is a template with a static variable and by doing dll export on entire class all accessors knowing the subclass will end up calling a non-inlined version of the code even though the code is inlined

diff --git a/Engine/Source/Runtime/Experimental/Chaos/Public/ChaosVisualDebugger/ChaosVDContextProvider.h b/Engine/Source/Runtime/Experimental/Chaos/Public/ChaosVisualDebugger/ChaosVDContextProvider.h
--- a/Engine/Source/Runtime/Experimental/Chaos/Public/ChaosVisualDebugger/ChaosVDContextProvider.h
+++ b/Engine/Source/Runtime/Experimental/Chaos/Public/ChaosVisualDebugger/ChaosVDContextProvider.h
@@ -18,7 +18,7 @@ struct FChaosVDContext
 };
 
 /** Singleton class that manages the thread local storage used to store CVD Context data */
-class FChaosVDThreadContext : public TThreadSingleton<FChaosVDThreadContext>
+class CHAOS_API FChaosVDThreadContext : public TThreadSingleton<FChaosVDThreadContext>
 {
 public:
 
@@ -33,19 +33,19 @@ public:
 	/** Copies the Current CVD context data into the provided struct
 	 * @return true if the copy was successful
 	 */
-	CHAOS_API bool GetCurrentContext(FChaosVDContext& OutContext);
+	bool GetCurrentContext(FChaosVDContext& OutContext);
 	
 	/** Gets the current CVD context data -
 	 * Don't use of a function that will recursively push new context data as it might invalidate the pointer
 	 * @return Ptr to the Current CVD context data
 	 */
-	CHAOS_API const FChaosVDContext* GetCurrentContext();
+	const FChaosVDContext* GetCurrentContext();
 	
 	/** Pushed a new CVD Context data to the local cvd context stack */
-	CHAOS_API void PushContext(const FChaosVDContext& InContext);
+	void PushContext(const FChaosVDContext& InContext);
 	
 	/** Removes the CVD Context data at the top of the local cvd context stack */
-	CHAOS_API void PopContext();
+	void PopContext();
 
 protected:
 	TArray<FChaosVDContext, TInlineAllocator<16>> LocalContextStack;
diff --git a/Engine/Source/Runtime/Experimental/Chaos/Public/Framework/Threading.h b/Engine/Source/Runtime/Experimental/Chaos/Public/Framework/Threading.h
--- a/Engine/Source/Runtime/Experimental/Chaos/Public/Framework/Threading.h
+++ b/Engine/Source/Runtime/Experimental/Chaos/Public/Framework/Threading.h
@@ -149,7 +149,7 @@ namespace Chaos
 
 #if PHYSICS_THREAD_CONTEXT
 /** Debug helper to ensure threading mistakes are caught. Do not use for ship */
-class FPhysicsThreadContext : public TThreadSingleton<FPhysicsThreadContext>
+class CHAOS_API FPhysicsThreadContext : public TThreadSingleton<FPhysicsThreadContext>
 {
 public:
 	bool IsInPhysicsSimContext() const
