IAS: Fix parsing of command line parameters

#rnx

diff --git a/Engine/Source/Runtime/Experimental/IoStoreOnDemand/Private/IoStoreOnDemand.cpp b/Engine/Source/Runtime/Experimental/IoStoreOnDemand/Private/IoStoreOnDemand.cpp
--- a/Engine/Source/Runtime/Experimental/IoStoreOnDemand/Private/IoStoreOnDemand.cpp
+++ b/Engine/Source/Runtime/Experimental/IoStoreOnDemand/Private/IoStoreOnDemand.cpp
@@ -15,6 +15,7 @@
 #include "Serialization/CompactBinaryWriter.h"
 #include "Serialization/CompactBinarySerialization.h"
 #include "Serialization/LargeMemoryWriter.h"
+#include "String/LexFromString.h"
 
 DEFINE_LOG_CATEGORY(LogIoStoreOnDemand);
 
@@ -27,7 +28,7 @@ static int64 ParseSizeParam(FStringView Value)
 	Value = Value.TrimStartAndEnd();
 
 	int64 Size = -1;
-	LexFromString(Size, Value.GetData());
+	LexFromString(Size, Value);
 	if (Size >= 0)
 	{
 		if (Value.EndsWith(TEXT("GB"))) return Size << 30;
@@ -100,7 +101,7 @@ static FFileIoCacheConfig GetFileCacheConfig(const TCHAR* CommandLine)
 	};
 
 #define GetConfigInt(Name, Dest) \
-	do { GetConfigIntImpl(TEXT("FileCache.") Name, TEXT("Ias.FileCache.") Name, Dest); } while (false)
+	do { GetConfigIntImpl(TEXT("FileCache.") Name, TEXT("Ias.FileCache.") Name TEXT("="), Dest); } while (false)
 	GetConfigInt(TEXT("WritePeriodSeconds"),	Ret.WriteRate.Seconds);
 	GetConfigInt(TEXT("WriteOpsPerPeriod"),		Ret.WriteRate.Ops);
 	GetConfigInt(TEXT("WriteBytesPerPeriod"),	Ret.WriteRate.Allowance);
