Make geometry collection convert vertex colors to sRGB when converting from FLinearColor to the render buffer FColor for new geometry collections.

For old geometry collections, keep the old behavior so that vertex colors are not changed unexpectedly on existing assets.  Add a bool on the UGeometryCollection asset object to let user control whether the conversion is performed.

diff --git a/Engine/Source/Runtime/Core/Public/UObject/FortniteMainBranchObjectVersion.h b/Engine/Source/Runtime/Core/Public/UObject/FortniteMainBranchObjectVersion.h
--- a/Engine/Source/Runtime/Core/Public/UObject/FortniteMainBranchObjectVersion.h
+++ b/Engine/Source/Runtime/Core/Public/UObject/FortniteMainBranchObjectVersion.h
@@ -333,6 +333,9 @@ struct CORE_API FFortniteMainBranchObjectVersion
 
 		WaterBodyStaticMeshRename,
 
+		// Geometry Collection now by-default converts vertex colors to sRGB when creating render data
+		GeometryCollectionConvertVertexColorToSRGB,
+
 		// -----<new versions can be added above this line>-------------------------------------------------
 		VersionPlusOne,
 		LatestVersion = VersionPlusOne - 1
diff --git a/Engine/Source/Runtime/Experimental/GeometryCollectionEngine/Private/GeometryCollection/DerivedDataGeometryCollectionCooker.cpp b/Engine/Source/Runtime/Experimental/GeometryCollectionEngine/Private/GeometryCollection/DerivedDataGeometryCollectionCooker.cpp
--- a/Engine/Source/Runtime/Experimental/GeometryCollectionEngine/Private/GeometryCollection/DerivedDataGeometryCollectionCooker.cpp
+++ b/Engine/Source/Runtime/Experimental/GeometryCollectionEngine/Private/GeometryCollection/DerivedDataGeometryCollectionCooker.cpp
@@ -99,7 +99,7 @@ bool FDerivedDataGeometryCollectionRenderDataCooker::Build(TArray<uint8>& OutDat
 	{
 		Chaos::FErrorReporter ErrorReporter(GeometryCollection.GetName());
 
-		TUniquePtr<FGeometryCollectionRenderData> RenderData = FGeometryCollectionRenderData::Create(*Collection, GeometryCollection.EnableNanite, GeometryCollection.bUseFullPrecisionUVs);
+		TUniquePtr<FGeometryCollectionRenderData> RenderData = FGeometryCollectionRenderData::Create(*Collection, GeometryCollection.EnableNanite, GeometryCollection.bUseFullPrecisionUVs, GeometryCollection.bConvertVertexColorsToSRGB);
 		RenderData->Serialize(ChaosAr, GeometryCollection);
 
 		if (false && ErrorReporter.EncounteredAnyErrors())
diff --git a/Engine/Source/Runtime/Experimental/GeometryCollectionEngine/Private/GeometryCollection/GeometryCollectionObject.cpp b/Engine/Source/Runtime/Experimental/GeometryCollectionEngine/Private/GeometryCollection/GeometryCollectionObject.cpp
--- a/Engine/Source/Runtime/Experimental/GeometryCollectionEngine/Private/GeometryCollection/GeometryCollectionObject.cpp
+++ b/Engine/Source/Runtime/Experimental/GeometryCollectionEngine/Private/GeometryCollection/GeometryCollectionObject.cpp
@@ -10,6 +10,7 @@
 #include "Materials/Material.h"
 #include "UObject/DestructionObjectVersion.h"
 #include "UObject/UE5MainStreamObjectVersion.h"
+#include "UObject/FortniteMainBranchObjectVersion.h"
 #include "Serialization/ArchiveCountMem.h"
 #include "HAL/IConsoleManager.h"
 #include "Interfaces/ITargetPlatform.h"
@@ -774,6 +775,7 @@ void UGeometryCollection::Serialize(FArchive& Ar)
 	Ar.UsingCustomVersion(FUE5MainStreamObjectVersion::GUID);
 	Ar.UsingCustomVersion(FUE5ReleaseStreamObjectVersion::GUID);
 	Ar.UsingCustomVersion(FPhysicsObjectVersion::GUID);
+	Ar.UsingCustomVersion(FFortniteMainBranchObjectVersion::GUID);
 	
 	Chaos::FChaosArchive ChaosAr(Ar);
 
@@ -1034,6 +1036,13 @@ void UGeometryCollection::Serialize(FArchive& Ar)
 		EnsureDataIsCooked(true /*bInitResources*/, Ar.IsTransacting(), Ar.IsPersistent(), false /*bAllowCopyFromDDC*/);
 	}
 #endif
+
+	if (Ar.IsLoading() && Ar.CustomVer(FFortniteMainBranchObjectVersion::GUID) < FFortniteMainBranchObjectVersion::GeometryCollectionConvertVertexColorToSRGB)
+	{
+		// disable sRGB conversion for old assets to keep the default behavior from before this setting existed
+		// (new assets will default-enable the conversion, because that matches static meshes and is the more-expected behavior)
+		bConvertVertexColorsToSRGB = false;
+	}
 }
 
 const TCHAR* UGeometryCollection::GetSelectedMaterialPath()
@@ -1053,6 +1062,18 @@ void UGeometryCollection::SetEnableNanite(bool bValue)
 	}
 }
 
+void UGeometryCollection::SetConvertVertexColorsToSRGB(bool bValue)
+{
+	if (bConvertVertexColorsToSRGB != bValue)
+	{
+		bConvertVertexColorsToSRGB = bValue;
+
+#if WITH_EDITOR
+		RebuildRenderData();
+#endif
+	}
+}
+
 void UGeometryCollection::FillAutoInstanceMeshesInstancesIfNeeded()
 {
 	// make sure the instanced meshes have their instance count properly set 
@@ -1174,7 +1195,7 @@ void UGeometryCollection::RebuildRenderData()
 	if (RenderDataGuid != StateGuid)
 	{
 		ReleaseResources();
-		RenderData = FGeometryCollectionRenderData::Create(*GetGeometryCollection(), EnableNanite, bUseFullPrecisionUVs);
+		RenderData = FGeometryCollectionRenderData::Create(*GetGeometryCollection(), EnableNanite, bUseFullPrecisionUVs, bConvertVertexColorsToSRGB);
 		InitResources();
 		PropagateMarkDirtyToComponents();
 		RenderDataGuid = StateGuid;
@@ -1467,6 +1488,10 @@ void UGeometryCollection::PostEditChangeProperty(struct FPropertyChangedEvent& P
 			bDoInvalidateCollection = true;
 			bRebuildRenderData = true;
 		}
+		else if (PropertyChangedEvent.Property->GetFName() == GET_MEMBER_NAME_CHECKED(UGeometryCollection, bConvertVertexColorsToSRGB))
+		{
+			bRebuildRenderData = true;
+		}
 		else if (PropertyChangedEvent.Property->GetFName() == GET_MEMBER_NAME_CHECKED(UGeometryCollection, SizeSpecificData))
 		{
 			bDoInvalidateCollection = true;
diff --git a/Engine/Source/Runtime/Experimental/GeometryCollectionEngine/Private/GeometryCollection/GeometryCollectionRenderData.cpp b/Engine/Source/Runtime/Experimental/GeometryCollectionEngine/Private/GeometryCollection/GeometryCollectionRenderData.cpp
--- a/Engine/Source/Runtime/Experimental/GeometryCollectionEngine/Private/GeometryCollection/GeometryCollectionRenderData.cpp
+++ b/Engine/Source/Runtime/Experimental/GeometryCollectionEngine/Private/GeometryCollection/GeometryCollectionRenderData.cpp
@@ -156,7 +156,7 @@ struct FGeometryCollectionBuiltMeshData
 };
 
 /** Get the data needed to build render data from a Geometry Collection. */
-bool BuildMeshDataFromGeometryCollection(FGeometryCollection& InCollection, FGeometryCollectionBuiltMeshData& OutMeshData)
+bool BuildMeshDataFromGeometryCollection(FGeometryCollection& InCollection, FGeometryCollectionBuiltMeshData& OutMeshData, bool bConvertVertexColorsToSRGB)
 {
 	// Vertices Group
 	const TManagedArray<FVector3f>& VertexArray = InCollection.Vertex;
@@ -240,7 +240,7 @@ bool BuildMeshDataFromGeometryCollection(FGeometryCollection& InCollection, FGeo
 
 				if (bHasColors)
 				{
-					BuildVertexData.Color[DestVertexStart + VertexIndex] = ColorArray[VertexStart + VertexIndex].ToFColor(false /* sRGB */);
+					BuildVertexData.Color[DestVertexStart + VertexIndex] = ColorArray[VertexStart + VertexIndex].ToFColor(bConvertVertexColorsToSRGB /* sRGB */);
 				}
 
 				for (int32 TexCoord = 0; TexCoord < NumTexCoords; ++TexCoord)
@@ -685,14 +685,14 @@ void CreateNaniteData(FGeometryCollectionBuiltMeshData&& InMeshData, FGeometryCo
 	}
 }
 
-TUniquePtr<FGeometryCollectionRenderData> FGeometryCollectionRenderData::Create(FGeometryCollection& InCollection, bool bInEnableNanite, bool bInUseFullPrecisionUVs)
+TUniquePtr<FGeometryCollectionRenderData> FGeometryCollectionRenderData::Create(FGeometryCollection& InCollection, bool bInEnableNanite, bool bInUseFullPrecisionUVs, bool bConvertVertexColorsToSRGB)
 {
 	TRACE_CPUPROFILER_EVENT_SCOPE(FGeometryCollectionRenderData::Create);
 
 	TUniquePtr<FGeometryCollectionRenderData> RenderData = MakeUnique<FGeometryCollectionRenderData>();
 
 	FGeometryCollectionBuiltMeshData MeshBuildData;
-	if (BuildMeshDataFromGeometryCollection(InCollection, MeshBuildData))
+	if (BuildMeshDataFromGeometryCollection(InCollection, MeshBuildData, bConvertVertexColorsToSRGB))
 	{
 	{
 		if (bInEnableNanite)
diff --git a/Engine/Source/Runtime/Experimental/GeometryCollectionEngine/Private/GeometryCollection/GeometryCollectionRenderData.h b/Engine/Source/Runtime/Experimental/GeometryCollectionEngine/Private/GeometryCollection/GeometryCollectionRenderData.h
--- a/Engine/Source/Runtime/Experimental/GeometryCollectionEngine/Private/GeometryCollection/GeometryCollectionRenderData.h
+++ b/Engine/Source/Runtime/Experimental/GeometryCollectionEngine/Private/GeometryCollection/GeometryCollectionRenderData.h
@@ -116,7 +116,7 @@ public:
 	~FGeometryCollectionRenderData();
 
 #if WITH_EDITOR
-	static TUniquePtr<FGeometryCollectionRenderData> Create(FGeometryCollection& InCollection, bool bInEnableNanite, bool bInUseFullPrecisionUVs);
+	static TUniquePtr<FGeometryCollectionRenderData> Create(FGeometryCollection& InCollection, bool bInEnableNanite, bool bInUseFullPrecisionUVs, bool bConvertVertexColorsToSRGB);
 #endif
 
 	bool IsInitialized()
diff --git a/Engine/Source/Runtime/Experimental/GeometryCollectionEngine/Public/GeometryCollection/GeometryCollectionObject.h b/Engine/Source/Runtime/Experimental/GeometryCollectionEngine/Public/GeometryCollection/GeometryCollectionObject.h
--- a/Engine/Source/Runtime/Experimental/GeometryCollectionEngine/Public/GeometryCollection/GeometryCollectionObject.h
+++ b/Engine/Source/Runtime/Experimental/GeometryCollectionEngine/Public/GeometryCollection/GeometryCollectionObject.h
@@ -590,6 +590,15 @@ public:
 	UPROPERTY(EditAnywhere, BlueprintReadWrite, BlueprintSetter=SetEnableNanite, Category = "Nanite")
 	bool EnableNanite;
 
+	UFUNCTION(BlueprintCallable, Category = "Rendering")
+	void SetConvertVertexColorsToSRGB(bool bValue);
+
+	/**
+	 * Convert vertex colors to sRGB for rendering. Exposed to avoid changing vertex color rendering for legacy assets; should typically be true for new geometry collections.
+	 */
+	UPROPERTY(EditAnywhere, BlueprintReadWrite, BlueprintSetter=SetConvertVertexColorsToSRGB, Category = "Rendering")
+	bool bConvertVertexColorsToSRGB = true;
+
 #if WITH_EDITORONLY_DATA
 	/*
 	*  CollisionType defines how to initialize the rigid collision structures.
