Fix leaf convex hull generator clearing unrelated convex hulls

diff --git a/Engine/Source/Runtime/Experimental/Chaos/Private/GeometryCollection/GeometryCollectionConvexUtility.cpp b/Engine/Source/Runtime/Experimental/Chaos/Private/GeometryCollection/GeometryCollectionConvexUtility.cpp
--- a/Engine/Source/Runtime/Experimental/Chaos/Private/GeometryCollection/GeometryCollectionConvexUtility.cpp
+++ b/Engine/Source/Runtime/Experimental/Chaos/Private/GeometryCollection/GeometryCollectionConvexUtility.cpp
@@ -1573,10 +1573,6 @@ void FGeometryCollectionConvexUtility::GenerateLeafConvexHulls(FGeometryCollecti
 
 	CreateConvexHullAttributesIfNeeded(Collection);
 
-	TManagedArray<TUniquePtr<Chaos::FConvex>>& ConvexHullAttrib = Collection.ModifyAttribute<TUniquePtr<Chaos::FConvex>>("ConvexHull", "Convex");
-	TManagedArray<TSet<int32>>& TransformToConvexIndicesAttrib = Collection.ModifyAttribute<TSet<int32>>("TransformToConvexIndices", FTransformCollection::TransformGroup);
-	TransformToConvexIndicesAttrib.Fill(TSet<int32>());
-
 	auto ComputeTransformedHull = [](const Chaos::FConvex& HullIn, const FTransform& TransformIn) -> TUniquePtr<Chaos::FConvex>
 	{
 		FTransform3f Transform = (FTransform3f)TransformIn;
@@ -1606,6 +1602,55 @@ void FGeometryCollectionConvexUtility::GenerateLeafConvexHulls(FGeometryCollecti
 		return;
 	}
 
+	TManagedArray<TUniquePtr<Chaos::FConvex>>& ConvexHullAttrib = Collection.ModifyAttribute<TUniquePtr<Chaos::FConvex>>("ConvexHull", "Convex");
+	TManagedArray<TSet<int32>>& TransformToConvexIndicesAttrib = Collection.ModifyAttribute<TSet<int32>>("TransformToConvexIndices", FTransformCollection::TransformGroup);
+	
+	// Remove all convex hulls from bones that we have selected for re-compute
+	for (int32 TransformIdx : UseTransforms)
+	{
+		if (Collection.SimulationType[TransformIdx] != FGeometryCollection::ESimulationTypes::FST_Rigid)
+		{
+			continue;
+		}
+		TransformToConvexIndicesAttrib[TransformIdx].Empty();
+	}
+	// Remove all convex hulls that are no longer referenced by any transforms
+	{
+		TArray<bool> ConvexHullsUsed; ConvexHullsUsed.SetNumZeroed(ConvexHullAttrib.Num());
+		for (int32 TransformIdx = 0; TransformIdx < TransformToConvexIndicesAttrib.Num(); ++TransformIdx)
+		{
+			for (int32 ConvexIdx : TransformToConvexIndicesAttrib[TransformIdx])
+			{
+				ConvexHullsUsed[ConvexIdx] = true;
+			}
+		}
+		for (int32 ConvexIdx = 0; ConvexIdx < ConvexHullAttrib.Num(); ++ConvexIdx)
+		{
+			if (!ConvexHullsUsed[ConvexIdx])
+			{
+				ConvexHullAttrib[ConvexIdx].Reset();
+			}
+		}
+		RemoveEmptyConvexHulls(Collection);
+	}
+	
+	auto AddComputedHullsToCollection = [&Collection, &ConvexHullAttrib, &TransformToConvexIndicesAttrib](TArray<TUniquePtr<::Chaos::FConvex>>& Hulls, TArray<TSet<int32>>& TransformToHullsIndices)
+	{
+		int32 InitialNum = ConvexHullAttrib.Num();
+		Collection.Resize(InitialNum + Hulls.Num(), "Convex");
+		for (int32 Idx = 0; Idx < Hulls.Num(); ++Idx)
+		{
+			ConvexHullAttrib[InitialNum + Idx] = MoveTemp(Hulls[Idx]);
+		}
+		for (int32 TransformIdx = 0; TransformIdx < TransformToHullsIndices.Num(); ++TransformIdx)
+		{
+			for (int32 ConvexIdx : TransformToHullsIndices[TransformIdx])
+			{
+				TransformToConvexIndicesAttrib[TransformIdx].Add(ConvexIdx + InitialNum);
+			}
+		}
+	};
+
 	if (bUseExternal)
 	{
 		for (int32 SourceTransformIdx : UseTransforms)
@@ -1641,10 +1686,7 @@ void FGeometryCollectionConvexUtility::GenerateLeafConvexHulls(FGeometryCollecti
 				FinalHulls[HullIdx] = ComputeTransformedHull(*ExternalHulls[HullIdx].Convex, ExternalHulls[HullIdx].Transform);
 			});
 
-			Collection.EmptyGroup("Convex");
-			Collection.Resize(FinalHulls.Num(), "Convex");
-			ConvexHullAttrib = MoveTemp(FinalHulls);
-			TransformToConvexIndicesAttrib = MoveTemp(TransformToExternalHullsIndices);
+			AddComputedHullsToCollection(FinalHulls, TransformToExternalHullsIndices);
 		}
 	}
 	if (bUseGenerated)
@@ -1666,10 +1708,7 @@ void FGeometryCollectionConvexUtility::GenerateLeafConvexHulls(FGeometryCollecti
 
 		if (!bUseIntersect)
 		{
-			Collection.EmptyGroup("Convex");
-			Collection.Resize(ComputedHulls.Hulls.Num(), "Convex");
-			ConvexHullAttrib = MoveTemp(ComputedHulls.Hulls);
-			TransformToConvexIndicesAttrib = MoveTemp(ComputedHulls.TransformToHullsIndices);
+			AddComputedHullsToCollection(ComputedHulls.Hulls, ComputedHulls.TransformToHullsIndices);
 		}
 	}
 	if (bUseIntersect)
@@ -1680,6 +1719,7 @@ void FGeometryCollectionConvexUtility::GenerateLeafConvexHulls(FGeometryCollecti
 		// only on the geometry that is not removed by the intersection.
 
 		TArray<TUniquePtr<Chaos::FConvex>> FinalHulls;
+		int32 HullIndexOffset = ConvexHullAttrib.Num();
 		for (int32 SourceTransformIdx : UseTransforms)
 		{
 			if (Collection.SimulationType[SourceTransformIdx] != FGeometryCollection::ESimulationTypes::FST_Rigid)
@@ -1710,7 +1750,7 @@ void FGeometryCollectionConvexUtility::GenerateLeafConvexHulls(FGeometryCollecti
 					for (int32 ExtHull : TransformToExternalHullsIndices[SourceTransformIdx])
 					{
 						int32 HullIdx = FinalHulls.Add(ComputeTransformedHull(*ExternalHulls[ExtHull].Convex, ExternalHulls[ExtHull].Transform));
-						TransformToConvexIndicesAttrib[SourceTransformIdx].Add(HullIdx);
+						TransformToConvexIndicesAttrib[SourceTransformIdx].Add(HullIndexOffset + HullIdx);
 					}
 					continue;
 				}
@@ -1720,7 +1760,7 @@ void FGeometryCollectionConvexUtility::GenerateLeafConvexHulls(FGeometryCollecti
 			{
 				if (TransformToExternalHullsIndices.IsEmpty())
 				{
-					TransformToConvexIndicesAttrib[SourceTransformIdx].Add(FinalHulls.Add(MoveTemp(ComputedHulls.Hulls[GeoHullIdx])));
+					TransformToConvexIndicesAttrib[SourceTransformIdx].Add(HullIndexOffset + FinalHulls.Add(MoveTemp(ComputedHulls.Hulls[GeoHullIdx])));
 
 					continue;
 				}
@@ -1730,13 +1770,16 @@ void FGeometryCollectionConvexUtility::GenerateLeafConvexHulls(FGeometryCollecti
 					UE::GeometryCollectionConvexUtility::IntersectConvexHulls(FinalHulls[HullIdx].Get(),
 						ComputedHulls.Hulls[GeoHullIdx].Get(), 0.0f, ExternalHulls[ExtHull].Convex.Get(),
 						nullptr, &ExternalHulls[ExtHull].Transform, &ExternalHulls[ExtHull].Transform, Settings.SimplificationDistanceThreshold);
-					TransformToConvexIndicesAttrib[SourceTransformIdx].Add(HullIdx);
+					TransformToConvexIndicesAttrib[SourceTransformIdx].Add(HullIdx + HullIndexOffset);
 				}
 			}
 		}
-		Collection.EmptyGroup("Convex");
-		Collection.Resize(FinalHulls.Num(), "Convex");
-		ConvexHullAttrib = MoveTemp(FinalHulls);
+		// copy intersected hulls into the output hull attrib
+		Collection.Resize(HullIndexOffset + FinalHulls.Num(), "Convex");
+		for (int32 Idx = 0; Idx < FinalHulls.Num(); ++Idx)
+		{
+			ConvexHullAttrib[HullIndexOffset + Idx] = MoveTemp(FinalHulls[Idx]);
+		}
 	}
 
 	RemoveEmptyConvexHulls(Collection);
