Networked Physics: Fix ensure with duplicate GC particle caching in rewind data.

diff --git a/Engine/Source/Runtime/Experimental/Chaos/Private/RewindData.cpp b/Engine/Source/Runtime/Experimental/Chaos/Private/RewindData.cpp
--- a/Engine/Source/Runtime/Experimental/Chaos/Private/RewindData.cpp
+++ b/Engine/Source/Runtime/Experimental/Chaos/Private/RewindData.cpp
@@ -932,28 +932,32 @@ void FRewindData::PushPTDirtyData(TPBDRigidParticleHandle<FReal,3>& Handle,const
 	FDirtyParticleInfo& Info = FindOrAddDirtyObj(Handle);
 	FGeometryParticleStateBase& Latest = Info.AddFrame(CurFrame);
 
-
 	const FFrameAndPhase FrameAndPhase{ CurFrame, FFrameAndPhase::PostCallbacks };
-
-	ensure(!bRecordingHistory || Latest.IsCleanExcludingDynamics(FrameAndPhase));	//PostCallbacks should be clean before we write sim results
-
-	if(bRecordingHistory || Latest.ParticlePositionRotation.IsClean(FrameAndPhase))
+	
+	if (bRecordingHistory || Latest.ParticlePositionRotation.IsClean(FrameAndPhase))
 	{
-		Latest.ParticlePositionRotation.WriteAccessMonotonic(FrameAndPhase, PropertiesPool).CopyFrom(Handle);
+		if (FParticlePositionRotation* PreXR = Latest.ParticlePositionRotation.WriteAccessNonDecreasing(FrameAndPhase, PropertiesPool))
+		{
+			PreXR->CopyFrom(Handle);
+		}
 	}
 
-	if(bRecordingHistory || Latest.Velocities.IsClean(FrameAndPhase))
+	if (bRecordingHistory || Latest.Velocities.IsClean(FrameAndPhase))
 	{
-		FParticleVelocities& PreVelocities = Latest.Velocities.WriteAccessMonotonic(FrameAndPhase, PropertiesPool);
-		PreVelocities.SetV(Handle.PreV());
-		PreVelocities.SetW(Handle.PreW());
+		if (FParticleVelocities* PreVelocities = Latest.Velocities.WriteAccessNonDecreasing(FrameAndPhase, PropertiesPool))
+		{
+			PreVelocities->SetV(Handle.PreV());
+			PreVelocities->SetW(Handle.PreW());
+		}
 	}
 	
-	if(bRecordingHistory || Latest.DynamicsMisc.IsClean(FrameAndPhase))
+	if (bRecordingHistory || Latest.DynamicsMisc.IsClean(FrameAndPhase))
 	{
-		FParticleDynamicMisc& PreDynamicMisc = Latest.DynamicsMisc.WriteAccessMonotonic(FrameAndPhase, PropertiesPool);
-		PreDynamicMisc.CopyFrom(Handle);	//everything is immutable except object state
-		PreDynamicMisc.SetObjectState(Handle.PreObjectState());
+		if (FParticleDynamicMisc* PreDynamicMisc = Latest.DynamicsMisc.WriteAccessNonDecreasing(FrameAndPhase, PropertiesPool))
+		{
+			PreDynamicMisc->CopyFrom(Handle);	//everything is immutable except object state
+			PreDynamicMisc->SetObjectState(Handle.PreObjectState());
+		}
 	}
 }
 
