chaos fields should not enable all particles it makes dynamic - will cause geometry collection particles to collide against their root particle

diff --git a/Engine/Source/Runtime/Experimental/Chaos/Public/PhysicsProxy/FieldSystemProxyHelper.h b/Engine/Source/Runtime/Experimental/Chaos/Public/PhysicsProxy/FieldSystemProxyHelper.h
--- a/Engine/Source/Runtime/Experimental/Chaos/Public/PhysicsProxy/FieldSystemProxyHelper.h
+++ b/Engine/Source/Runtime/Experimental/Chaos/Public/PhysicsProxy/FieldSystemProxyHelper.h
@@ -400,10 +400,6 @@ namespace Chaos
 				if (bHasStateChanged)
 				{
 					UpdateSolverParticlesState(RigidSolver, EvaluatedSamples, ParticleHandles);
-					for (Chaos::FGeometryParticleHandle* ParticleHandle : ParticleHandles)
-					{
-						RigidSolver->GetEvolution()->EnableParticle(ParticleHandle);
-					}
 				}
 			}
 		}
