Chaos - fix CIS missing CHAOS_API on new cvars

diff --git a/Engine/Source/Runtime/Experimental/Chaos/Public/Chaos/PBDRigidsEvolution.h b/Engine/Source/Runtime/Experimental/Chaos/Public/Chaos/PBDRigidsEvolution.h
--- a/Engine/Source/Runtime/Experimental/Chaos/Public/Chaos/PBDRigidsEvolution.h
+++ b/Engine/Source/Runtime/Experimental/Chaos/Public/Chaos/PBDRigidsEvolution.h
@@ -68,7 +68,7 @@ extern CHAOS_API FBroadPhaseConfig BroadPhaseConfig;
 
 namespace CVars
 {
-	extern bool bDisallowSetKinematicTargetOnDynamics;
+	extern CHAOS_API bool bDisallowSetKinematicTargetOnDynamics;
 }
 
 namespace Collisions
diff --git a/Engine/Source/Runtime/Experimental/Chaos/Public/Chaos/PBDRigidsSOAs.h b/Engine/Source/Runtime/Experimental/Chaos/Public/Chaos/PBDRigidsSOAs.h
--- a/Engine/Source/Runtime/Experimental/Chaos/Public/Chaos/PBDRigidsSOAs.h
+++ b/Engine/Source/Runtime/Experimental/Chaos/Public/Chaos/PBDRigidsSOAs.h
@@ -9,7 +9,7 @@ namespace Chaos
 {
 namespace CVars
 {
-	extern bool bRemoveParticleFromMovingKinematicsOnDisable;
+	extern CHAOS_API bool bRemoveParticleFromMovingKinematicsOnDisable;
 }
 
 class IParticleUniqueIndices
