changing a particle to be dynamic on a chaos field should cause it to get added to the active array - hacky way by enabling the particle if it's not already disabled (this way it works for GCs too as it won't try to wake up all the enabled particles that still have a parent)

diff --git a/Engine/Source/Runtime/Experimental/Chaos/Public/PhysicsProxy/FieldSystemProxyHelper.h b/Engine/Source/Runtime/Experimental/Chaos/Public/PhysicsProxy/FieldSystemProxyHelper.h
--- a/Engine/Source/Runtime/Experimental/Chaos/Public/PhysicsProxy/FieldSystemProxyHelper.h
+++ b/Engine/Source/Runtime/Experimental/Chaos/Public/PhysicsProxy/FieldSystemProxyHelper.h
@@ -201,6 +201,9 @@ namespace Chaos
 			if (!RigidHandle->Disabled())
 			{
 				RigidSolver->GetEvolution()->GetIslandManager().AddParticle(RigidHandle);
+
+				// This is mainly just to refresh the views. Note that we do not want to enable disabled particles.
+				RigidSolver->GetEvolution()->GetParticles().EnableParticle(RigidHandle);
 			}
 		}
 	}
