Add a new hull merge node for reducing the hulls used on convex decompositions
+ minor cleanup/fixes of hull merging algorithms

diff --git a/Engine/Plugins/Experimental/Fracture/Source/FractureEngine/Private/FractureEngineConvex.cpp b/Engine/Plugins/Experimental/Fracture/Source/FractureEngine/Private/FractureEngineConvex.cpp
--- a/Engine/Plugins/Experimental/Fracture/Source/FractureEngine/Private/FractureEngineConvex.cpp
+++ b/Engine/Plugins/Experimental/Fracture/Source/FractureEngine/Private/FractureEngineConvex.cpp
@@ -373,7 +373,7 @@ namespace UE::FractureEngine::Convex
 		return true;
 	}
 
-	bool ComputeConvexHullsNegativeSpace(FManagedArrayCollection& Collection, UE::Geometry::FSphereCovering& OutNegativeSpace, const UE::Geometry::FNegativeSpaceSampleSettings& Settings, bool bRestrictToSelection, const TArrayView<const int32> TransformSelection)
+	bool ComputeConvexHullsNegativeSpace(FManagedArrayCollection& Collection, UE::Geometry::FSphereCovering& OutNegativeSpace, const UE::Geometry::FNegativeSpaceSampleSettings& Settings, bool bRestrictToSelection, const TArrayView<const int32> TransformSelection, bool bFromRigidTransforms)
 	{
 		if (!FGeometryCollectionConvexUtility::HasConvexHullData(&Collection))
 		{
@@ -387,15 +387,18 @@ namespace UE::FractureEngine::Convex
 
 		GeometryCollection::Facades::FCollectionTransformFacade TransformFacade(Collection);
 		GeometryCollection::Facades::FCollectionTransformSelectionFacade SelectionFacade(Collection);
-		TArray<int> RigidSelection;
+		TArray<int> UseSelection;
 		if (!bRestrictToSelection)
 		{
-			RigidSelection = SelectionFacade.SelectLeaf();
+			UseSelection = bFromRigidTransforms ? SelectionFacade.SelectLeaf() : SelectionFacade.SelectAll();
 		}
 		else
 		{
-			RigidSelection.Append(TransformSelection);
-			SelectionFacade.ConvertSelectionToRigidNodes(RigidSelection);
+			UseSelection.Append(TransformSelection);
+			if (bFromRigidTransforms)
+			{
+				SelectionFacade.ConvertSelectionToRigidNodes(UseSelection);
+			}
 		}
 		
 		TArray<FTransform> GlobalTransformArray = TransformFacade.ComputeCollectionSpaceTransforms();
@@ -417,7 +420,7 @@ namespace UE::FractureEngine::Convex
 		};
 
 		bool bNoFailures = true;
-		for (int32 BoneIdx : RigidSelection)
+		for (int32 BoneIdx : UseSelection)
 		{
 			bool bSuccess = ProcessBone(BoneIdx);
 			bNoFailures = bNoFailures && bSuccess;
diff --git a/Engine/Plugins/Experimental/Fracture/Source/FractureEngine/Public/FractureEngineConvex.h b/Engine/Plugins/Experimental/Fracture/Source/FractureEngine/Public/FractureEngineConvex.h
--- a/Engine/Plugins/Experimental/Fracture/Source/FractureEngine/Public/FractureEngineConvex.h
+++ b/Engine/Plugins/Experimental/Fracture/Source/FractureEngine/Public/FractureEngineConvex.h
@@ -58,7 +58,7 @@ namespace UE::FractureEngine::Convex
 	 * Compute negative space for the convex hulls on a geometry collection
 	 * @return true if any negative space was found
 	 */
-	bool FRACTUREENGINE_API ComputeConvexHullsNegativeSpace(FManagedArrayCollection& Collection, UE::Geometry::FSphereCovering& OutNegativeSpace, const UE::Geometry::FNegativeSpaceSampleSettings& Settings, bool bRestrictToSelection = false, const TArrayView<const int32> TransformSelection = TArrayView<const int32>());
+	bool FRACTUREENGINE_API ComputeConvexHullsNegativeSpace(FManagedArrayCollection& Collection, UE::Geometry::FSphereCovering& OutNegativeSpace, const UE::Geometry::FNegativeSpaceSampleSettings& Settings, bool bRestrictToSelection = false, const TArrayView<const int32> TransformSelection = TArrayView<const int32>(), bool bFromRigidTransforms = true);
 
 }
 
diff --git a/Engine/Plugins/Experimental/GeometryCollectionPlugin/Source/GeometryCollectionNodes/Private/Dataflow/GeometryCollectionUtilityNodes.cpp b/Engine/Plugins/Experimental/GeometryCollectionPlugin/Source/GeometryCollectionNodes/Private/Dataflow/GeometryCollectionUtilityNodes.cpp
--- a/Engine/Plugins/Experimental/GeometryCollectionPlugin/Source/GeometryCollectionNodes/Private/Dataflow/GeometryCollectionUtilityNodes.cpp
+++ b/Engine/Plugins/Experimental/GeometryCollectionPlugin/Source/GeometryCollectionNodes/Private/Dataflow/GeometryCollectionUtilityNodes.cpp
@@ -27,6 +27,7 @@ namespace Dataflow
 		DATAFLOW_NODE_REGISTER_CREATION_FACTORY(FCreateNonOverlappingConvexHullsDataflowNode);
 		DATAFLOW_NODE_REGISTER_CREATION_FACTORY(FGenerateClusterConvexHullsFromLeafHullsDataflowNode);
 		DATAFLOW_NODE_REGISTER_CREATION_FACTORY(FGenerateClusterConvexHullsFromChildrenHullsDataflowNode);
+		DATAFLOW_NODE_REGISTER_CREATION_FACTORY(FMergeConvexHullsDataflowNode);
 		DATAFLOW_NODE_REGISTER_CREATION_FACTORY(FUpdateVolumeAttributesDataflowNode);
 		DATAFLOW_NODE_REGISTER_CREATION_FACTORY(FGetConvexHullVolumeDataflowNode);
 	}
@@ -232,6 +233,7 @@ void FGenerateClusterConvexHullsFromLeafHullsDataflowNode::Evaluate(Dataflow::FC
 				NegativeSpaceSettings.MinRadius = GetValue(Context, &MinRadius);
 				NegativeSpaceSettings.ReduceRadiusMargin = GetValue(Context, &NegativeSpaceTolerance);
 				NegativeSpaceSettings.MinSpacing = GetValue(Context, &MinSampleSpacing);
+				NegativeSpaceSettings.Sanitize();
 				bHasNegativeSpace = UE::FractureEngine::Convex::ComputeConvexHullsNegativeSpace(*GeomCollection, NegativeSpace, NegativeSpaceSettings, bHasSelectionFilter, SelectionArray);
 			}
 
@@ -314,6 +316,7 @@ void FGenerateClusterConvexHullsFromChildrenHullsDataflowNode::Evaluate(Dataflow
 				NegativeSpaceSettings.MinRadius = GetValue(Context, &MinRadius);
 				NegativeSpaceSettings.ReduceRadiusMargin = GetValue(Context, &NegativeSpaceTolerance);
 				NegativeSpaceSettings.MinSpacing = GetValue(Context, &MinSampleSpacing);
+				NegativeSpaceSettings.Sanitize();
 				bHasNegativeSpace = UE::FractureEngine::Convex::ComputeConvexHullsNegativeSpace(*GeomCollection, NegativeSpace, NegativeSpaceSettings, bHasSelectionFilter, SelectionArray);
 			}
 
@@ -353,6 +356,68 @@ void FGenerateClusterConvexHullsFromChildrenHullsDataflowNode::Evaluate(Dataflow
 	}
 }
 
+FMergeConvexHullsDataflowNode::FMergeConvexHullsDataflowNode(const Dataflow::FNodeParameters& InParam, FGuid InGuid)
+	: FDataflowNode(InParam, InGuid)
+{
+	RegisterInputConnection(&Collection);
+	RegisterInputConnection(&MaxConvexCount);
+	RegisterInputConnection(&ErrorTolerance);
+	RegisterInputConnection(&OptionalSelectionFilter);
+	RegisterInputConnection(&bProtectNegativeSpace);
+	RegisterInputConnection(&TargetNumSamples);
+	RegisterInputConnection(&MinSampleSpacing);
+	RegisterInputConnection(&NegativeSpaceTolerance);
+	RegisterInputConnection(&MinRadius);
+
+	RegisterOutputConnection(&Collection);
+	RegisterOutputConnection(&SphereCovering);
+}
+
+void FMergeConvexHullsDataflowNode::Evaluate(Dataflow::FContext& Context, const FDataflowOutput* Out) const
+{
+	if (Out->IsA(&Collection) || Out->IsA(&SphereCovering))
+	{
+		FManagedArrayCollection InCollection = GetValue<FManagedArrayCollection>(Context, &Collection);
+		FDataflowSphereCovering Spheres;
+
+		TArray<int32> SelectionArray;
+		bool bHasSelectionFilter = IsConnected(&OptionalSelectionFilter);
+		if (bHasSelectionFilter)
+		{
+			const FDataflowTransformSelection& InOptionalSelectionFilter = GetValue<FDataflowTransformSelection>(Context, &OptionalSelectionFilter);
+			SelectionArray = InOptionalSelectionFilter.AsArray();
+		}
+
+		bool bHasNegativeSpace = false;
+		UE::Geometry::FSphereCovering NegativeSpace;
+		if (GetValue(Context, &bProtectNegativeSpace))
+		{
+			UE::Geometry::FNegativeSpaceSampleSettings NegativeSpaceSettings;
+			NegativeSpaceSettings.TargetNumSamples = GetValue(Context, &TargetNumSamples);
+			NegativeSpaceSettings.MinRadius = GetValue(Context, &MinRadius);
+			NegativeSpaceSettings.ReduceRadiusMargin = GetValue(Context, &NegativeSpaceTolerance);
+			NegativeSpaceSettings.MinSpacing = GetValue(Context, &MinSampleSpacing);
+			NegativeSpaceSettings.Sanitize();
+			bHasNegativeSpace = UE::FractureEngine::Convex::ComputeConvexHullsNegativeSpace(InCollection, NegativeSpace, NegativeSpaceSettings, bHasSelectionFilter, SelectionArray, false);
+		}
+
+		const int32 InMaxConvexCount = GetValue(Context, &MaxConvexCount);
+		const double InErrorToleranceInCm = GetValue(Context, &ErrorTolerance);
+		FGeometryCollectionConvexUtility::FMergeConvexHullSettings HullMergeSettings;
+		HullMergeSettings.EmptySpace = bHasNegativeSpace ? &NegativeSpace : nullptr;
+		HullMergeSettings.ErrorToleranceInCm = InErrorToleranceInCm;
+		HullMergeSettings.MaxConvexCount = InMaxConvexCount;
+
+		FGeometryCollectionConvexUtility::MergeHullsOnTransforms(InCollection, HullMergeSettings, bHasSelectionFilter, SelectionArray);
+
+		SetValue<FManagedArrayCollection>(Context, MoveTemp(InCollection), &Collection);
+
+		// Move the negative space to the output container at the end, after it is no longer needed
+		Spheres.Spheres = MoveTemp(NegativeSpace);
+		SetValue(Context, MoveTemp(Spheres), &SphereCovering);
+	}
+}
+
 
 FUpdateVolumeAttributesDataflowNode::FUpdateVolumeAttributesDataflowNode(const Dataflow::FNodeParameters& InParam, FGuid InGuid)
 	: FDataflowNode(InParam, InGuid)
diff --git a/Engine/Plugins/Experimental/GeometryCollectionPlugin/Source/GeometryCollectionNodes/Public/Dataflow/GeometryCollectionUtilityNodes.h b/Engine/Plugins/Experimental/GeometryCollectionPlugin/Source/GeometryCollectionNodes/Public/Dataflow/GeometryCollectionUtilityNodes.h
--- a/Engine/Plugins/Experimental/GeometryCollectionPlugin/Source/GeometryCollectionNodes/Public/Dataflow/GeometryCollectionUtilityNodes.h
+++ b/Engine/Plugins/Experimental/GeometryCollectionPlugin/Source/GeometryCollectionNodes/Public/Dataflow/GeometryCollectionUtilityNodes.h
@@ -365,6 +365,65 @@ public:
 
 };
 
+
+/** Merge convex hulls on transforms with multiple hulls */
+USTRUCT(meta = (DataflowGeometryCollection))
+struct FMergeConvexHullsDataflowNode : public FDataflowNode
+{
+	GENERATED_USTRUCT_BODY()
+	DATAFLOW_NODE_DEFINE_INTERNAL(FMergeConvexHullsDataflowNode, "MergeConvexHulls", "GeometryCollection|Utilities", "")
+
+public:
+	UPROPERTY(meta = (DataflowInput, DataflowOutput))
+	FManagedArrayCollection Collection;
+
+	// A representation of the negative space protected by the 'protect negative space' option. If negative space is not protected, this will contain zero spheres.
+	UPROPERTY(meta = (DataflowOutput))
+	FDataflowSphereCovering SphereCovering;
+
+	/** Maximum number of convex to generate per transform. Ignored if < 0. */
+	UPROPERTY(EditAnywhere, Category = "Convex", meta = (DataflowInput))
+	int32 MaxConvexCount = -1;
+
+	/**
+	* Error tolerance to use to decide to merge leaf convex together.
+	* This is in centimeters and represents the side of a cube, the volume of which will be used as threshold
+	* to know if the volume of the generated convex is too large compared to the sum of the volume of the leaf convex
+	*/
+	UPROPERTY(EditAnywhere, Category = "Convex", meta = (DataflowInput, ClampMin = "0", UIMax = "100.", Units = cm))
+	double ErrorTolerance = 0.0;
+
+	/** Optional transform selection to compute cluster hulls on -- if not provided, all cluster hulls will be computed. */
+	UPROPERTY(meta = (DataflowInput, DataflowIntrinsic))
+	FDataflowTransformSelection OptionalSelectionFilter;
+
+	/** Whether to use a sphere cover to define negative space that should not be covered by convex hulls */
+	UPROPERTY(EditAnywhere, Category = NegativeSpace, meta = (DataflowInput))
+	bool bProtectNegativeSpace = false;
+
+	/** Approximate number of spheres to consider when covering negative space */
+	UPROPERTY(EditAnywhere, Category = NegativeSpace, meta = (DataflowInput, ClampMin = 1, EditCondition = "bProtectNegativeSpace", EditConditionHides))
+	int32 TargetNumSamples = 50;
+
+	/** Minimum desired spacing between spheres; if > 0, will attempt not to place sphere centers closer than this */
+	UPROPERTY(EditAnywhere, Category = NegativeSpace, meta = (DataflowInput, ClampMin = 0, Units = cm, EditCondition = "bProtectNegativeSpace", EditConditionHides))
+	double MinSampleSpacing = 1.0;
+
+	/** Amount of space to leave between convex hulls and protected negative space */
+	UPROPERTY(EditAnywhere, Category = NegativeSpace, meta = (DataflowInput, ClampMin = .01, UIMin = .1, Units = cm, EditCondition = "bProtectNegativeSpace", EditConditionHides))
+	double NegativeSpaceTolerance = 2.0;
+
+	/** Spheres smaller than this are not included in the negative space */
+	UPROPERTY(EditAnywhere, Category = NegativeSpace, meta = (DataflowInput, ClampMin = 0, Units = cm, EditCondition = "bProtectNegativeSpace", EditConditionHides))
+	double MinRadius = 10.0;
+
+	FMergeConvexHullsDataflowNode(const Dataflow::FNodeParameters& InParam, FGuid InGuid = FGuid::NewGuid());
+
+	virtual void Evaluate(Dataflow::FContext& Context, const FDataflowOutput* Out) const override;
+
+};
+
+
 /**
  *
  * Update the Volume and Size attributes on the target Collection (and add them if they were not present)
diff --git a/Engine/Source/Runtime/Experimental/Chaos/Private/GeometryCollection/GeometryCollectionConvexUtility.cpp b/Engine/Source/Runtime/Experimental/Chaos/Private/GeometryCollection/GeometryCollectionConvexUtility.cpp
--- a/Engine/Source/Runtime/Experimental/Chaos/Private/GeometryCollection/GeometryCollectionConvexUtility.cpp
+++ b/Engine/Source/Runtime/Experimental/Chaos/Private/GeometryCollection/GeometryCollectionConvexUtility.cpp
@@ -1823,6 +1823,10 @@ void FGeometryCollectionConvexUtility::GenerateClusterConvexHullsFromLeafOrChild
 	ProximityUtility.RequireProximity();
 	const TManagedArray<TSet<int32>>& Proximity = Collection.GetAttribute<TSet<int32>>("Proximity", FGeometryCollection::GeometryGroup);
 
+	// Whether to consider directly merging hulls from the same child/leaf transform; without this, these hulls can still be merged if they are also merged to a neighboring hull
+	// TODO: decide whether to always-enable this, or expose this option to the user, or leave it disabled
+	constexpr bool bConsiderIntraTransformConvexMerges = false;
+
 	const TArray<FTransform> GlobalTransforms = TransformFacade.ComputeCollectionSpaceTransforms();
 	const TManagedArrayAccessor<int32> SimulationTypeAttribute(Collection, FGeometryCollection::SimulationTypeAttribute, FTransformCollection::TransformGroup);
 	const TArray<int32> DepthFirstTransformIndices = HierarchyFacade.GetTransformArrayInDepthFirstOrder();
@@ -1882,8 +1886,6 @@ void FGeometryCollectionConvexUtility::GenerateClusterConvexHullsFromLeafOrChild
 				FGeometryCollectionClusteringUtility::GetLeafBones(&Collection, TransformIndex, true, SourceTransformIndices);
 			}
 			
-
-			UE::Geometry::FDynamicMesh3 CombinedGeometry;
 			if (SourceTransformIndices.Num() > 0)
 			{
 				const FTransform& ParentTransform = GlobalTransforms[TransformIndex];
@@ -1947,6 +1949,7 @@ void FGeometryCollectionConvexUtility::GenerateClusterConvexHullsFromLeafOrChild
 						const int32 TransformIndexA = SourceTransformIndices[IndexA];
 						if (const TArray<int32>* HullIndicesA = TransformToHullIdx.Find(TransformIndexA))
 						{
+							// Consider merges between convex hulls from separate transforms
 							for (int32 IndexB = IndexA + 1; IndexB < SourceTransformIndices.Num(); IndexB++)
 							{
 								const int32 TransformIndexB = SourceTransformIndices[IndexB];
@@ -1961,6 +1964,19 @@ void FGeometryCollectionConvexUtility::GenerateClusterConvexHullsFromLeafOrChild
 									}
 								}
 							}
+							// Consider merges between convex hulls within a transform
+							if (bConsiderIntraTransformConvexMerges)
+							{
+								for (int32 HullSubIdx = 0; HullSubIdx + 1 < HullIndicesA->Num(); ++HullSubIdx)
+								{
+									const int32 HullA = (*HullIndicesA)[HullSubIdx];
+									for (int32 ToHull = HullSubIdx + 1; ToHull < HullIndicesA->Num(); ++ToHull)
+									{
+										const int32 HullB = (*HullIndicesA)[ToHull];
+										HullProximity.Emplace(HullA, HullB);
+									}
+								}
+							}
 						}
 					}
 				}
@@ -1972,10 +1988,11 @@ void FGeometryCollectionConvexUtility::GenerateClusterConvexHullsFromLeafOrChild
 						const int32 GeoIdx = Collection.TransformToGeometryIndex[SourceTransformIndex];
 						if (GeoIdx > INDEX_NONE)
 						{
+							// Consider merges between convex hulls from separate transforms
 							for (int32 NeighborGeoIndex : Proximity[GeoIdx])
 							{
-								const int32 NeightborTransformIdx = Collection.TransformIndex[NeighborGeoIndex];
-								if (const TArray<int32>* NeighborHullIndices = TransformToHullIdx.Find(NeightborTransformIdx))
+								const int32 NeighborTransformIdx = Collection.TransformIndex[NeighborGeoIndex];
+								if (const TArray<int32>* NeighborHullIndices = TransformToHullIdx.Find(NeighborTransformIdx))
 								{
 									for (int32 SourceHullIdx : SourceHullIndices)
 									{
@@ -1986,6 +2003,19 @@ void FGeometryCollectionConvexUtility::GenerateClusterConvexHullsFromLeafOrChild
 									}
 								}
 							}
+							// Consider merges between convex hulls within a transform
+							if (bConsiderIntraTransformConvexMerges)
+							{
+								for (int32 HullSubIdx = 0; HullSubIdx + 1 < SourceHullIndices.Num(); ++HullSubIdx)
+								{
+									const int32 HullA = SourceHullIndices[HullSubIdx];
+									for (int32 ToHull = HullSubIdx + 1; ToHull < SourceHullIndices.Num(); ++ToHull)
+									{
+										const int32 HullB = SourceHullIndices[ToHull];
+										HullProximity.Emplace(HullA, HullB);
+									}
+								}
+							}
 						}
 					}
 				}
@@ -2037,6 +2067,113 @@ void FGeometryCollectionConvexUtility::GenerateClusterConvexHullsFromLeafOrChild
 
 }
 
+void FGeometryCollectionConvexUtility::MergeHullsOnTransforms(FManagedArrayCollection& Collection, const FMergeConvexHullSettings& Settings, bool bRestrictToSelection, const TArrayView<const int32> OptionalTransformSelection)
+{
+	static FName ConvexGroupName("Convex");
+
+	TOptional<FGeometryCollectionConvexUtility::FGeometryCollectionConvexData> ConvexData = GetConvexHullDataIfPresent(&Collection);
+	if (!ConvexData)
+	{
+		return;
+	}
+
+	TManagedArray<TSet<int32>>& TransformToConvexIndices = ConvexData->TransformToConvexIndices;
+	TManagedArray<TUniquePtr<Chaos::FConvex>>& ConvexHull = ConvexData->ConvexHull;
+
+	GeometryCollection::Facades::FCollectionTransformFacade TransformFacade(Collection);
+
+	const TArray<FTransform> GlobalTransforms = TransformFacade.ComputeCollectionSpaceTransforms();
+	TArray<int32> TransformsToProcess;
+	if (bRestrictToSelection)
+	{
+		TransformsToProcess.Append(OptionalTransformSelection);
+	}
+	else // process all
+	{
+		TransformsToProcess.Reserve(TransformToConvexIndices.Num());
+		for (int32 Idx = 0; Idx < TransformToConvexIndices.Num(); ++Idx)
+		{
+			TransformsToProcess.Add(Idx);
+		}
+	}
+
+	for (int32 TransformIndex : TransformsToProcess)
+	{
+		const int32 InitialNumConvex = TransformToConvexIndices[TransformIndex].Num();
+		if (InitialNumConvex <= 1)
+		{
+			continue;
+		}
+
+		TArray<Chaos::FConvex*> Hulls;
+		for (int32 ConvexIdx : TransformToConvexIndices[TransformIndex])
+		{
+			Hulls.Add(ConvexHull[ConvexIdx].Get());
+		}
+		TArray<TPair<int32, int32>> HullProximity;
+		for (int32 ConvexA = 0; ConvexA < InitialNumConvex; ++ConvexA)
+		{
+			for (int32 ConvexB = ConvexA + 1; ConvexB < InitialNumConvex; ++ConvexB)
+			{
+				HullProximity.Emplace(ConvexA, ConvexB);
+			}
+		}
+
+		// try to merge the leaf convex into a simpler set of convex
+		const auto GetHullVolume = [&Hulls](int32 Idx) { return (double)Hulls[Idx]->GetVolume(); };
+		const auto GetHullNumVertices = [&Hulls](int32 Idx) { return Hulls[Idx]->NumVertices(); };
+		const auto GetHullVertex = [&Hulls](int32 Hull, int32 V) -> FVector3d {
+			const FVector3d LocalVertex(Hulls[Hull]->GetVertex(V));
+			return LocalVertex;
+		};
+		UE::Geometry::FConvexDecomposition3 ConvexDecomposition;
+		ConvexDecomposition.InitializeFromHulls(Hulls.Num(), GetHullVolume, GetHullNumVertices, GetHullVertex, HullProximity);
+		ConvexDecomposition.MergeBest(Settings.MaxConvexCount, Settings.ErrorToleranceInCm, 0, true /*bAllowCompact*/, false /*bRequireHullTriangles*/, Settings.MaxConvexCount /*MaxHulls*/, Settings.EmptySpace, &GlobalTransforms[TransformIndex]);
+
+		if (!ensure(ConvexDecomposition.Decomposition.Num() > 0))
+		{
+			// Empty convex decomposition cannot be used / should not happen
+			continue;
+		}
+		if (InitialNumConvex <= ConvexDecomposition.Decomposition.Num())
+		{
+			// no need to update hulls if the merge didn't change the number of hulls
+			continue;
+		}
+
+		// reset existing hulls for this transform index
+		for (int32 ConvexIndex : TransformToConvexIndices[TransformIndex])
+		{
+			ConvexHull[ConvexIndex] = nullptr;
+		}
+		TransformToConvexIndices[TransformIndex].Reset();
+
+		// add new computed ones
+		for (int32 DecompIndex = 0; DecompIndex < ConvexDecomposition.Decomposition.Num(); DecompIndex++)
+		{
+			// Make Implicit Convex of a decomposed element
+			TArray<Chaos::FConvex::FVec3Type> Particles;
+			const TArray<FVector> Points = ConvexDecomposition.GetVertices<double>(DecompIndex);
+			Particles.SetNum(Points.Num());
+			for (int32 PointIndex = 0; PointIndex < Points.Num(); PointIndex++)
+			{
+				Particles[PointIndex] = Points[PointIndex];
+			}
+			TUniquePtr<Chaos::FConvex> ImplicitConvex = MakeUnique<Chaos::FConvex>(MoveTemp(Particles), 0.0f);
+
+			// Add the the element to the union
+			const int32 NewConvexIndex = Collection.AddElements(1, ConvexGroupName);
+			ConvexHull[NewConvexIndex] = MoveTemp(ImplicitConvex);
+			TransformToConvexIndices[TransformIndex].Add(NewConvexIndex);
+		}
+	}
+
+	// remove empty or null convex hulls
+	RemoveEmptyConvexHulls(Collection);
+
+	checkSlow(FGeometryCollectionConvexUtility::ValidateConvexData(&Collection));
+}
+
 bool FGeometryCollectionConvexUtility::ValidateConvexData(const FGeometryCollection* GeometryCollection)
 {
 	if (!GeometryCollection->HasAttribute("ConvexHull", "Convex") || !GeometryCollection->HasAttribute("TransformToConvexIndices", FTransformCollection::TransformGroup))
diff --git a/Engine/Source/Runtime/Experimental/Chaos/Public/GeometryCollection/GeometryCollectionConvexUtility.h b/Engine/Source/Runtime/Experimental/Chaos/Public/GeometryCollection/GeometryCollectionConvexUtility.h
--- a/Engine/Source/Runtime/Experimental/Chaos/Public/GeometryCollection/GeometryCollectionConvexUtility.h
+++ b/Engine/Source/Runtime/Experimental/Chaos/Public/GeometryCollection/GeometryCollectionConvexUtility.h
@@ -132,6 +132,15 @@ public:
 	static CHAOS_API void GenerateClusterConvexHullsFromLeafHulls(FGeometryCollection& Collection, const FClusterConvexHullSettings& Settings, const TArrayView<const int32> OptionalTransformSubset);
 	static CHAOS_API void GenerateClusterConvexHullsFromLeafHulls(FGeometryCollection& Collection, const FClusterConvexHullSettings& Settings);
 
+	struct FMergeConvexHullSettings
+	{
+		int32 MaxConvexCount = -1;
+		double ErrorToleranceInCm = 0.0;
+		UE::Geometry::FSphereCovering* EmptySpace = nullptr;
+	};
+	// Merge convex hulls that are currently on each (selected) transform. If convex hulls are not present, does nothing.
+	static void MergeHullsOnTransforms(FManagedArrayCollection& Collection, const FMergeConvexHullSettings& Settings, bool bRestrictToSelection, const TArrayView<const int32> OptionalTransformSelection);
+
 	// Additional settings for filtering when the EGenerateConvexMethod::IntersectExternalWithComputed is applied
 	struct FIntersectionFilters
 	{
diff --git a/Engine/Source/Runtime/GeometryCore/Private/CompGeom/ConvexDecomposition3.cpp b/Engine/Source/Runtime/GeometryCore/Private/CompGeom/ConvexDecomposition3.cpp
--- a/Engine/Source/Runtime/GeometryCore/Private/CompGeom/ConvexDecomposition3.cpp
+++ b/Engine/Source/Runtime/GeometryCore/Private/CompGeom/ConvexDecomposition3.cpp
@@ -1329,7 +1329,7 @@ void FConvexDecomposition3::InitializeFromHulls(int32 NumHulls, TFunctionRef<dou
 int32 FConvexDecomposition3::MergeBest(int32 InTargetNumParts, double MaxErrorTolerance, double MinThicknessToleranceWorldSpace, bool bAllowCompact, bool bRequireHullTriangles, int32 MaxOutputHulls,
 	const FSphereCovering* OptionalNegativeSpace, const FTransform* OptionalTransformIntoNegativeSpace)
 {
-	int32 TargetNumParts = InTargetNumParts;
+	int32 TargetNumParts = FMath::Max(1, InTargetNumParts);
 	const bool bHasValidMaxHulls = MaxOutputHulls > 0;
 	if (bHasValidMaxHulls)
 	{
diff --git a/Engine/Source/Runtime/GeometryCore/Public/CompGeom/ConvexDecomposition3.h b/Engine/Source/Runtime/GeometryCore/Public/CompGeom/ConvexDecomposition3.h
--- a/Engine/Source/Runtime/GeometryCore/Public/CompGeom/ConvexDecomposition3.h
+++ b/Engine/Source/Runtime/GeometryCore/Public/CompGeom/ConvexDecomposition3.h
@@ -60,6 +60,15 @@ struct FNegativeSpaceSampleSettings
 
 	// Ignore spheres with smaller radius than this
 	double MinRadius = 10.0;
+
+	// Make sure the settings values are in valid ranges
+	void Sanitize()
+	{
+		TargetNumSamples = FMath::Max(1, TargetNumSamples);
+		MinSpacing = FMath::Max(0.0, MinSpacing);
+		ReduceRadiusMargin = FMath::Max(0.0, ReduceRadiusMargin);
+		MinRadius = FMath::Max(0.0, MinRadius);
+	}
 };
 
 // Define a volume with a set of spheres
