fix mac / debug build error for Merge Convex Hulls

diff --git a/Engine/Source/Runtime/Experimental/Chaos/Private/GeometryCollection/GeometryCollectionConvexUtility.cpp b/Engine/Source/Runtime/Experimental/Chaos/Private/GeometryCollection/GeometryCollectionConvexUtility.cpp
--- a/Engine/Source/Runtime/Experimental/Chaos/Private/GeometryCollection/GeometryCollectionConvexUtility.cpp
+++ b/Engine/Source/Runtime/Experimental/Chaos/Private/GeometryCollection/GeometryCollectionConvexUtility.cpp
@@ -2170,8 +2170,6 @@ void FGeometryCollectionConvexUtility::MergeHullsOnTransforms(FManagedArrayColle
 
 	// remove empty or null convex hulls
 	RemoveEmptyConvexHulls(Collection);
-
-	checkSlow(FGeometryCollectionConvexUtility::ValidateConvexData(&Collection));
 }
 
 bool FGeometryCollectionConvexUtility::ValidateConvexData(const FGeometryCollection* GeometryCollection)
