Fixed an issue where FRewindData::RewindToFrame() would keep dirtying particles/joints every rewind. If the game was constantly rewinding, this would cause the size of the FRewindData::DirtyParticles array to get too big

diff --git a/Engine/Source/Runtime/Experimental/Chaos/Private/RewindData.cpp b/Engine/Source/Runtime/Experimental/Chaos/Private/RewindData.cpp
--- a/Engine/Source/Runtime/Experimental/Chaos/Private/RewindData.cpp
+++ b/Engine/Source/Runtime/Experimental/Chaos/Private/RewindData.cpp
@@ -310,7 +310,7 @@ bool FRewindData::RewindToFrame(int32 Frame)
 		//rewind is about to start, all particles should be in sync at this point
 		ensure(PTParticle->SyncState() == ESyncState::InSync);
 		
-		FGeometryParticleStateBase& History = DirtyParticleInfo.AddFrame(CurFrame);	//non-const in case we need to record what's at head for a rewind (CurFrame has already been increased to the next frame)
+		FGeometryParticleStateBase& History = DirtyParticleInfo.GetHistory(); //non-const in case we need to record what's at head for a rewind (CurFrame has already been increased to the next frame)
 
 		History.CachePreCorrectionState(*PTParticle);
 
@@ -349,16 +349,15 @@ bool FRewindData::RewindToFrame(int32 Frame)
 		}
 	}
 
+#if !UE_BUILD_SHIPPING
+	// For now, just ensure that the joints are InSync
 	for(FDirtyJointInfo& DirtyJointInfo : DirtyJoints)
 	{
-		FPBDJointConstraintHandle* Joint = DirtyJointInfo.GetObjectPtr();
-
+		const FPBDJointConstraintHandle* Joint = DirtyJointInfo.GetObjectPtr();
 		//rewind is about to start, all particles should be in sync at this point
 		ensure(Joint->SyncState() == ESyncState::InSync);
-
-		FJointStateBase& History = DirtyJointInfo.AddFrame(CurFrame);	//non-const in case we need to record what's at head for a rewind (CurFrame has already been increased to the next frame)
-
 	}
+#endif
 
 	CurFrame = Frame;
 	bNeedsSave = false;
