Add scopes to async physics tick in various components

diff --git a/Engine/Source/Runtime/Engine/Private/PhysicsEngine/AsyncPhysicsInputComponent.cpp b/Engine/Source/Runtime/Engine/Private/PhysicsEngine/AsyncPhysicsInputComponent.cpp
--- a/Engine/Source/Runtime/Engine/Private/PhysicsEngine/AsyncPhysicsInputComponent.cpp
+++ b/Engine/Source/Runtime/Engine/Private/PhysicsEngine/AsyncPhysicsInputComponent.cpp
@@ -137,6 +137,8 @@ void UAsyncPhysicsInputComponent::SetDataClass(TSubclassOf<UAsyncPhysicsData> In
 
 void UAsyncPhysicsInputComponent::AsyncPhysicsTickComponent(float DeltaTime, float SimTime)
 {
+	QUICK_SCOPE_CYCLE_COUNTER(AsyncPhysicsInputComponent_AsyncPhysicsTick);
+
 	Super::AsyncPhysicsTickComponent(DeltaTime, SimTime);
 	ensureMsgf(DataClass != nullptr, TEXT("You must call SetDataClass after creating the component"));
 
diff --git a/Engine/Source/Runtime/Engine/Private/PhysicsEngine/NetworkPhysicsComponent.cpp b/Engine/Source/Runtime/Engine/Private/PhysicsEngine/NetworkPhysicsComponent.cpp
--- a/Engine/Source/Runtime/Engine/Private/PhysicsEngine/NetworkPhysicsComponent.cpp
+++ b/Engine/Source/Runtime/Engine/Private/PhysicsEngine/NetworkPhysicsComponent.cpp
@@ -410,6 +410,8 @@ void UNetworkPhysicsComponent::UpdatePackageMap()
 
 void UNetworkPhysicsComponent::AsyncPhysicsTickComponent(float DeltaTime, float SimTime)
 {
+	QUICK_SCOPE_CYCLE_COUNTER(NetworkPhysicsComponent_AsyncPhysicsTick);
+
 	Super ::AsyncPhysicsTickComponent(DeltaTime, SimTime);
 #if DEBUG_NETWORK_PHYSICS
 	if(HasServerWorld() && !HasLocalController() && InputsHistory)
diff --git a/Engine/Source/Runtime/Experimental/GeometryCollectionEngine/Private/GeometryCollection/GeometryCollectionComponent.cpp b/Engine/Source/Runtime/Experimental/GeometryCollectionEngine/Private/GeometryCollection/GeometryCollectionComponent.cpp
--- a/Engine/Source/Runtime/Experimental/GeometryCollectionEngine/Private/GeometryCollection/GeometryCollectionComponent.cpp
+++ b/Engine/Source/Runtime/Experimental/GeometryCollectionEngine/Private/GeometryCollection/GeometryCollectionComponent.cpp
@@ -2662,6 +2662,8 @@ void UGeometryCollectionComponent::TickComponent(float DeltaTime, enum ELevelTic
 
 void UGeometryCollectionComponent::AsyncPhysicsTickComponent(float DeltaTime, float SimTime)
 {
+	QUICK_SCOPE_CYCLE_COUNTER(GeometryCollectionComponent_AsyncPhysicsTick);
+
 	if (!PhysicsProxy || !PhysicsProxy->IsInitializedOnPhysicsThread())
 	{
 		return;
