Moved IAS' statistics collector into its own unit

#rnx

diff --git a/Engine/Source/Runtime/Experimental/IoStoreOnDemand/Private/OnDemandIoDispatcherBackend.cpp b/Engine/Source/Runtime/Experimental/IoStoreOnDemand/Private/OnDemandIoDispatcherBackend.cpp
--- a/Engine/Source/Runtime/Experimental/IoStoreOnDemand/Private/OnDemandIoDispatcherBackend.cpp
+++ b/Engine/Source/Runtime/Experimental/IoStoreOnDemand/Private/OnDemandIoDispatcherBackend.cpp
@@ -1,6 +1,7 @@
 // Copyright Epic Games, Inc. All Rights Reserved.
 
 #include "OnDemandIoDispatcherBackend.h"
+#include "Statistics.h"
 
 #include "CancellationToken.h"
 #include "Containers/StringView.h"
@@ -25,8 +26,6 @@
 #include "Misc/ScopeLock.h"
 #include "Misc/ScopeRWLock.h"
 #include "Modules/ModuleManager.h"
-#include "ProfilingDebugging/CountersTrace.h"
-#include "ProfilingDebugging/CsvProfiler.h"
 #include "Serialization/CompactBinary.h"
 #include "Serialization/CompactBinarySerialization.h"
 #include "Serialization/JsonReader.h"
@@ -55,236 +54,6 @@ namespace UE::IO::Private
 
 using namespace UE::Tasks;
 
-///////////////////////////////////////////////////////////////////////////////
-// TRACE STATS
-// iorequest stats
-TRACE_DECLARE_INT_COUNTER(IoRequestsMade, TEXT("Ias/IoRequestsMade"));
-TRACE_DECLARE_INT_COUNTER(IoRequestsCompleted, TEXT("Ias/IoRequestsCompleted"));
-TRACE_DECLARE_MEMORY_COUNTER(IoRequestsCompletedSize, TEXT("Ias/Size/IoRequestsCompletedSize"));
-TRACE_DECLARE_INT_COUNTER(IoRequestsCancelled, TEXT("Ias/IoRequestsCancelled"));
-TRACE_DECLARE_INT_COUNTER(IoRequestsFailed, TEXT("Ias/IoRequestsFailed"));
-// chunkrequest stats
-TRACE_DECLARE_INT_COUNTER(ReadRequestsCreated, TEXT("Ias/ReadRequestsCreated"));
-TRACE_DECLARE_INT_COUNTER(ReadRequestsRemoved, TEXT("Ias/ReadRequestsRemoved"));
-// cache stats
-TRACE_DECLARE_INT_COUNTER(CacheHits, TEXT("Ias/CacheHits"));
-TRACE_DECLARE_MEMORY_COUNTER(CacheHitsSize, TEXT("Ias/Size/CacheHitsSize"));
-TRACE_DECLARE_INT_COUNTER(CachePuts, TEXT("Ias/CachePuts"));
-TRACE_DECLARE_MEMORY_COUNTER(CachePutsSize, TEXT("Ias/Size/CachePutsSize"));
-TRACE_DECLARE_INT_COUNTER(CacheRejects, TEXT("Ias/CacheRejects"));
-TRACE_DECLARE_MEMORY_COUNTER(CacheRejectsSize, TEXT("Ias/Size/CacheRejectsSize"));
-// http stats
-TRACE_DECLARE_INT_COUNTER(HttpRequestsCompleted, TEXT("Ias/HttpRequestsCompleted"));
-TRACE_DECLARE_INT_COUNTER(HttpRequestsFailed, TEXT("Ias/HttpRequestsFailed"));
-TRACE_DECLARE_INT_COUNTER(HttpRequestsPending, TEXT("Ias/HttpRequestsPending"));
-TRACE_DECLARE_INT_COUNTER(HttpRequestsInflight, TEXT("Ias/HttpRequestsInflight"));
-TRACE_DECLARE_MEMORY_COUNTER(HttpRequestsCompletedSize, TEXT("Ias/Size/HttpRequestsCompletedSize"));
-///////////////////////////////////////////////////////////////////////////////
-// CSV STATS
-CSV_DEFINE_CATEGORY(Ias, true);
-// iorequest stats
-CSV_DEFINE_STAT(Ias, FrameIoRequestsMade);
-CSV_DEFINE_STAT(Ias, FrameIoRequestsCompleted);
-CSV_DEFINE_STAT(Ias, FrameIoRequestsCompletedSize);
-CSV_DEFINE_STAT(Ias, FrameIoRequestsCancelled);
-CSV_DEFINE_STAT(Ias, FrameIoRequestsFailed);
-// chunkrequest stats
-CSV_DEFINE_STAT(Ias, FrameReadRequestsCreated);
-CSV_DEFINE_STAT(Ias, FrameReadRequestsRemoved);
-// cache stats
-CSV_DEFINE_STAT(Ias, FrameCacheHits);
-CSV_DEFINE_STAT(Ias, FrameCacheHitsSize);
-CSV_DEFINE_STAT(Ias, FrameCachePuts);
-CSV_DEFINE_STAT(Ias, FrameCachePutsSize);
-CSV_DEFINE_STAT(Ias, FrameCacheRejects);
-CSV_DEFINE_STAT(Ias, FrameCacheRejectsSize);
-// http stats
-CSV_DEFINE_STAT(Ias, FrameHttpRequestsCompleted);
-CSV_DEFINE_STAT(Ias, FrameHttpRequestsFailed);
-CSV_DEFINE_STAT(Ias, FrameHttpRequestsPending);
-CSV_DEFINE_STAT(Ias, FrameHttpRequestsInflight);
-CSV_DEFINE_STAT(Ias, FrameHttpRequestsCompletedSize);
-
-#if (COUNTERSTRACE_ENABLED || CSV_PROFILER)
-class FOnDemandIoBackendStats
-{
-public:
-	void OnIoRequestEnqueue(FIoRequestImpl* Request);
-	void OnIoRequestComplete(FIoRequestImpl* Request);
-	void OnIoRequestCancel(FIoRequestImpl* Request);
-	void OnIoRequestFail(FIoRequestImpl* Request);
-	void OnChunkRequestCreate();
-	void OnChunkRequestRelease();
-	void OnCacheHit(uint64 InSize);
-	void OnCachePut(uint64 InSize);
-	void OnCacheReject(uint64 InSize);
-	void OnHttpRequestEnqueue();
-	void OnHttpRequestDequeue();
-	void OnHttpRequestComplete(uint64 InSize);
-	void OnHttpRequestFail();
-
-private:
-	int32 IoRequestsMade = 0;
-	std::atomic<int32> IoRequestsCompleted = 0;
-	int32 IoRequestsCancelled = 0;
-	std::atomic<int32> IoRequestsFailed = 0;
-	int32 ReadRequestsCreated = 0;
-	std::atomic<int32> ReadRequestsRemoved = 0;
-	std::atomic<int32> CacheHits = 0;
-	std::atomic<int32> CachePuts = 0;
-	std::atomic<int32> CacheRejects = 0;
-	int32 HttpRequestsCompleted = 0;
-	int32 HttpRequestsFailed = 0;
-	std::atomic<int32> HttpRequestsPending = 0;
-	int32 HttpRequestsInflight = 0;
-
-	std::atomic<uint64> IoRequestsCompletedSize = 0;
-	std::atomic<uint64> CacheHitsSize = 0;
-	std::atomic<uint64> CachePutsSize = 0;
-	std::atomic<uint64> CacheRejectsSize = 0;
-	uint64 HttpRequestsCompletedSize = 0;
-
-	float BytesToApproxMB(uint64 Bytes)
-	{
-		return float(double(Bytes) / 1024.0 / 1024.0);
-	}
-
-	float BytesToApproxKB(uint64 Bytes)
-	{
-		return float(double(Bytes) / 1024.0);
-	}
-
-};
-void FOnDemandIoBackendStats::OnIoRequestEnqueue(FIoRequestImpl* Request)
-{
-	++IoRequestsMade;
-	TRACE_COUNTER_SET(IoRequestsMade, IoRequestsMade);
-	CSV_CUSTOM_STAT_DEFINED(FrameIoRequestsMade, IoRequestsMade, ECsvCustomStatOp::Set);
-}
-void FOnDemandIoBackendStats::OnIoRequestComplete(FIoRequestImpl* Request)
-{
-	int32 Count = IoRequestsCompleted.fetch_add(1, std::memory_order_relaxed);
-	TRACE_COUNTER_SET(IoRequestsCompleted, Count);
-	CSV_CUSTOM_STAT_DEFINED(FrameIoRequestsCompleted, Count, ECsvCustomStatOp::Set);
-
-	uint64 Size = IoRequestsCompletedSize.fetch_add(Request->GetBuffer().DataSize());
-	TRACE_COUNTER_SET(IoRequestsCompletedSize, Size);
-	CSV_CUSTOM_STAT_DEFINED(FrameIoRequestsCompletedSize, BytesToApproxKB(Size), ECsvCustomStatOp::Set);
-}
-void FOnDemandIoBackendStats::OnIoRequestCancel(FIoRequestImpl* Request)
-{
-	++IoRequestsCancelled;
-	TRACE_COUNTER_SET(IoRequestsCancelled, IoRequestsCancelled);
-	CSV_CUSTOM_STAT_DEFINED(FrameIoRequestsCancelled, IoRequestsCancelled, ECsvCustomStatOp::Set);
-}
-void FOnDemandIoBackendStats::OnIoRequestFail(FIoRequestImpl* Request)
-{
-	int32 Count = IoRequestsFailed.fetch_add(1, std::memory_order_relaxed);
-	TRACE_COUNTER_SET(IoRequestsFailed, Count);
-	CSV_CUSTOM_STAT_DEFINED(FrameIoRequestsFailed, Count, ECsvCustomStatOp::Set);
-}
-void FOnDemandIoBackendStats::OnChunkRequestCreate()
-{
-	++ReadRequestsCreated;
-	TRACE_COUNTER_SET(ReadRequestsCreated, ReadRequestsCreated);
-	CSV_CUSTOM_STAT_DEFINED(FrameReadRequestsCreated, ReadRequestsCreated, ECsvCustomStatOp::Set);
-}
-void FOnDemandIoBackendStats::OnChunkRequestRelease()
-{
-	int32 Count = ReadRequestsRemoved.fetch_add(1, std::memory_order_relaxed);
-	TRACE_COUNTER_SET(ReadRequestsRemoved, Count);
-	CSV_CUSTOM_STAT_DEFINED(FrameReadRequestsRemoved, Count, ECsvCustomStatOp::Set);
-}
-void FOnDemandIoBackendStats::OnCacheHit(uint64 InSize)
-{
-	int32 Count = CacheHits.fetch_add(1, std::memory_order_relaxed);
-	TRACE_COUNTER_SET(CacheHits, Count);
-	CSV_CUSTOM_STAT_DEFINED(FrameCacheHits, Count, ECsvCustomStatOp::Set);
-
-	uint64 Size = CacheHitsSize.fetch_add(InSize, std::memory_order_relaxed);
-	TRACE_COUNTER_SET(CacheHitsSize, Size);
-	CSV_CUSTOM_STAT_DEFINED(FrameCacheHitsSize, BytesToApproxKB(Size), ECsvCustomStatOp::Set);
-}
-void FOnDemandIoBackendStats::OnCachePut(uint64 InSize)
-{
-	int32 Count = CachePuts.fetch_add(1, std::memory_order_relaxed);
-	TRACE_COUNTER_SET(CachePuts, Count);
-	CSV_CUSTOM_STAT_DEFINED(FrameCachePuts, Count, ECsvCustomStatOp::Set);
-
-	uint64 Size = CachePutsSize.fetch_add(InSize, std::memory_order_relaxed);
-	TRACE_COUNTER_SET(CachePutsSize, Size);
-	CSV_CUSTOM_STAT_DEFINED(FrameCachePutsSize, BytesToApproxKB(Size), ECsvCustomStatOp::Set);
-}
-void FOnDemandIoBackendStats::OnCacheReject(uint64 InSize)
-{
-	int32 Count = CacheRejects.fetch_add(1, std::memory_order_relaxed);
-	TRACE_COUNTER_SET(CacheRejects, Count);
-	CSV_CUSTOM_STAT_DEFINED(FrameCacheRejects, Count, ECsvCustomStatOp::Set);
-
-	uint64 Size = CacheRejectsSize.fetch_add(InSize, std::memory_order_relaxed);
-	TRACE_COUNTER_SET(CacheRejectsSize, Size);
-	CSV_CUSTOM_STAT_DEFINED(FrameCacheRejectsSize, BytesToApproxKB(Size), ECsvCustomStatOp::Set);
-}
-void FOnDemandIoBackendStats::OnHttpRequestEnqueue()
-{
-	int32 Pending = HttpRequestsPending.fetch_add(1, std::memory_order_relaxed);
-	TRACE_COUNTER_SET(HttpRequestsPending, Pending);
-	CSV_CUSTOM_STAT_DEFINED(FrameHttpRequestsPending, Pending, ECsvCustomStatOp::Set);
-}
-void FOnDemandIoBackendStats::OnHttpRequestDequeue()
-{
-	int32 Pending = HttpRequestsPending.fetch_sub(1, std::memory_order_relaxed);
-	TRACE_COUNTER_SET(HttpRequestsPending, Pending);
-	CSV_CUSTOM_STAT_DEFINED(FrameHttpRequestsPending, Pending, ECsvCustomStatOp::Set);
-
-	int32 Inflight = ++HttpRequestsInflight;
-	TRACE_COUNTER_SET(HttpRequestsInflight, Inflight);
-	CSV_CUSTOM_STAT_DEFINED(FrameHttpRequestsPending, Inflight, ECsvCustomStatOp::Set);
-}
-void FOnDemandIoBackendStats::OnHttpRequestComplete(uint64 InSize)
-{
-	int32 Inflight = --HttpRequestsInflight;
-	TRACE_COUNTER_SET(HttpRequestsInflight, Inflight);
-	CSV_CUSTOM_STAT_DEFINED(FrameHttpRequestsPending, Inflight, ECsvCustomStatOp::Set);
-
-	++HttpRequestsCompleted;
-	TRACE_COUNTER_SET(HttpRequestsCompleted, HttpRequestsCompleted);
-	CSV_CUSTOM_STAT_DEFINED(FrameHttpRequestsCompleted, HttpRequestsCompleted, ECsvCustomStatOp::Set);
-
-	HttpRequestsCompletedSize += InSize;
-	TRACE_COUNTER_SET(HttpRequestsCompletedSize, HttpRequestsCompletedSize);
-	CSV_CUSTOM_STAT_DEFINED(FrameHttpRequestsCompletedSize, BytesToApproxKB(HttpRequestsCompletedSize), ECsvCustomStatOp::Set);
-}
-void FOnDemandIoBackendStats::OnHttpRequestFail()
-{
-	int32 Inflight = --HttpRequestsInflight;
-	TRACE_COUNTER_SET(HttpRequestsInflight, Inflight);
-	CSV_CUSTOM_STAT_DEFINED(FrameHttpRequestsPending, Inflight, ECsvCustomStatOp::Set);
-
-	++HttpRequestsFailed;
-	TRACE_COUNTER_SET(HttpRequestsFailed, HttpRequestsFailed);
-	CSV_CUSTOM_STAT_DEFINED(FrameHttpRequestsFailed, HttpRequestsFailed, ECsvCustomStatOp::Set);
-
-}
-#else
-class FOnDemandIoBackendStats
-{
-public:
-	void OnIoRequestEnqueue(FIoRequestImpl* Request) {}
-	void OnIoRequestComplete(FIoRequestImpl* Request) {}
-	void OnIoRequestCancel(FIoRequestImpl* Request) {}
-	void OnIoRequestFail(FIoRequestImpl* Request) {}
-	void OnChunkRequestCreate() {}
-	void OnChunkRequestRelease() {}
-	void OnCacheHit(uint64 InSize) {}
-	void OnCachePut(uint64 InSize) {}
-	void OnCacheReject(uint64 InSize) {}
-	void OnHttpRequestEnqueue() {}
-	void OnHttpRequestDequeue() {}
-	void OnHttpRequestComplete(uint64 InSize) {}
-	void OnHttpRequestFail() {}
-};
-#endif
 ///////////////////////////////////////////////////////////////////////////////
 FIoHash GetChunkKey(const FIoHash& ChunkHash, const FIoOffsetAndLength& Range)
 {
diff --git a/Engine/Source/Runtime/Experimental/IoStoreOnDemand/Private/Statistics.cpp b/Engine/Source/Runtime/Experimental/IoStoreOnDemand/Private/Statistics.cpp
new file mode 100644
--- /dev/null
+++ b/Engine/Source/Runtime/Experimental/IoStoreOnDemand/Private/Statistics.cpp
@@ -0,0 +1,176 @@
+// Copyright Epic Games, Inc. All Rights Reserved.
+
+#include "Statistics.h"
+
+#if IAS_WITH_STATISTICS
+
+#include "IO/IoDispatcherBackend.h"
+
+namespace UE::IO::Private
+{
+
+///////////////////////////////////////////////////////////////////////////////
+// TRACE STATS
+// iorequest stats
+TRACE_DECLARE_INT_COUNTER(IoRequestsMade, TEXT("Ias/IoRequestsMade"));
+TRACE_DECLARE_INT_COUNTER(IoRequestsCompleted, TEXT("Ias/IoRequestsCompleted"));
+TRACE_DECLARE_MEMORY_COUNTER(IoRequestsCompletedSize, TEXT("Ias/Size/IoRequestsCompletedSize"));
+TRACE_DECLARE_INT_COUNTER(IoRequestsCancelled, TEXT("Ias/IoRequestsCancelled"));
+TRACE_DECLARE_INT_COUNTER(IoRequestsFailed, TEXT("Ias/IoRequestsFailed"));
+// chunkrequest stats
+TRACE_DECLARE_INT_COUNTER(ReadRequestsCreated, TEXT("Ias/ReadRequestsCreated"));
+TRACE_DECLARE_INT_COUNTER(ReadRequestsRemoved, TEXT("Ias/ReadRequestsRemoved"));
+// cache stats
+TRACE_DECLARE_INT_COUNTER(CacheHits, TEXT("Ias/CacheHits"));
+TRACE_DECLARE_MEMORY_COUNTER(CacheHitsSize, TEXT("Ias/Size/CacheHitsSize"));
+TRACE_DECLARE_INT_COUNTER(CachePuts, TEXT("Ias/CachePuts"));
+TRACE_DECLARE_MEMORY_COUNTER(CachePutsSize, TEXT("Ias/Size/CachePutsSize"));
+TRACE_DECLARE_INT_COUNTER(CacheRejects, TEXT("Ias/CacheRejects"));
+TRACE_DECLARE_MEMORY_COUNTER(CacheRejectsSize, TEXT("Ias/Size/CacheRejectsSize"));
+// http stats
+TRACE_DECLARE_INT_COUNTER(HttpRequestsCompleted, TEXT("Ias/HttpRequestsCompleted"));
+TRACE_DECLARE_INT_COUNTER(HttpRequestsFailed, TEXT("Ias/HttpRequestsFailed"));
+TRACE_DECLARE_INT_COUNTER(HttpRequestsPending, TEXT("Ias/HttpRequestsPending"));
+TRACE_DECLARE_INT_COUNTER(HttpRequestsInflight, TEXT("Ias/HttpRequestsInflight"));
+TRACE_DECLARE_MEMORY_COUNTER(HttpRequestsCompletedSize, TEXT("Ias/Size/HttpRequestsCompletedSize"));
+///////////////////////////////////////////////////////////////////////////////
+// CSV STATS
+CSV_DEFINE_CATEGORY(Ias, true);
+// iorequest stats
+CSV_DEFINE_STAT(Ias, FrameIoRequestsMade);
+CSV_DEFINE_STAT(Ias, FrameIoRequestsCompleted);
+CSV_DEFINE_STAT(Ias, FrameIoRequestsCompletedSize);
+CSV_DEFINE_STAT(Ias, FrameIoRequestsCancelled);
+CSV_DEFINE_STAT(Ias, FrameIoRequestsFailed);
+// chunkrequest stats
+CSV_DEFINE_STAT(Ias, FrameReadRequestsCreated);
+CSV_DEFINE_STAT(Ias, FrameReadRequestsRemoved);
+// cache stats
+CSV_DEFINE_STAT(Ias, FrameCacheHits);
+CSV_DEFINE_STAT(Ias, FrameCacheHitsSize);
+CSV_DEFINE_STAT(Ias, FrameCachePuts);
+CSV_DEFINE_STAT(Ias, FrameCachePutsSize);
+CSV_DEFINE_STAT(Ias, FrameCacheRejects);
+CSV_DEFINE_STAT(Ias, FrameCacheRejectsSize);
+// http stats
+CSV_DEFINE_STAT(Ias, FrameHttpRequestsCompleted);
+CSV_DEFINE_STAT(Ias, FrameHttpRequestsFailed);
+CSV_DEFINE_STAT(Ias, FrameHttpRequestsPending);
+CSV_DEFINE_STAT(Ias, FrameHttpRequestsInflight);
+CSV_DEFINE_STAT(Ias, FrameHttpRequestsCompletedSize);
+
+void FOnDemandIoBackendStats::OnIoRequestEnqueue(FIoRequestImpl* Request)
+{
+	++IoRequestsMade;
+	TRACE_COUNTER_SET(IoRequestsMade, IoRequestsMade);
+	CSV_CUSTOM_STAT_DEFINED(FrameIoRequestsMade, IoRequestsMade, ECsvCustomStatOp::Set);
+}
+void FOnDemandIoBackendStats::OnIoRequestComplete(FIoRequestImpl* Request)
+{
+	int32 Count = IoRequestsCompleted.fetch_add(1, std::memory_order_relaxed);
+	TRACE_COUNTER_SET(IoRequestsCompleted, Count);
+	CSV_CUSTOM_STAT_DEFINED(FrameIoRequestsCompleted, Count, ECsvCustomStatOp::Set);
+
+	uint64 Size = IoRequestsCompletedSize.fetch_add(Request->GetBuffer().DataSize());
+	TRACE_COUNTER_SET(IoRequestsCompletedSize, Size);
+	CSV_CUSTOM_STAT_DEFINED(FrameIoRequestsCompletedSize, BytesToApproxKB(Size), ECsvCustomStatOp::Set);
+}
+void FOnDemandIoBackendStats::OnIoRequestCancel(FIoRequestImpl* Request)
+{
+	++IoRequestsCancelled;
+	TRACE_COUNTER_SET(IoRequestsCancelled, IoRequestsCancelled);
+	CSV_CUSTOM_STAT_DEFINED(FrameIoRequestsCancelled, IoRequestsCancelled, ECsvCustomStatOp::Set);
+}
+void FOnDemandIoBackendStats::OnIoRequestFail(FIoRequestImpl* Request)
+{
+	int32 Count = IoRequestsFailed.fetch_add(1, std::memory_order_relaxed);
+	TRACE_COUNTER_SET(IoRequestsFailed, Count);
+	CSV_CUSTOM_STAT_DEFINED(FrameIoRequestsFailed, Count, ECsvCustomStatOp::Set);
+}
+void FOnDemandIoBackendStats::OnChunkRequestCreate()
+{
+	++ReadRequestsCreated;
+	TRACE_COUNTER_SET(ReadRequestsCreated, ReadRequestsCreated);
+	CSV_CUSTOM_STAT_DEFINED(FrameReadRequestsCreated, ReadRequestsCreated, ECsvCustomStatOp::Set);
+}
+void FOnDemandIoBackendStats::OnChunkRequestRelease()
+{
+	int32 Count = ReadRequestsRemoved.fetch_add(1, std::memory_order_relaxed);
+	TRACE_COUNTER_SET(ReadRequestsRemoved, Count);
+	CSV_CUSTOM_STAT_DEFINED(FrameReadRequestsRemoved, Count, ECsvCustomStatOp::Set);
+}
+void FOnDemandIoBackendStats::OnCacheHit(uint64 InSize)
+{
+	int32 Count = CacheHits.fetch_add(1, std::memory_order_relaxed);
+	TRACE_COUNTER_SET(CacheHits, Count);
+	CSV_CUSTOM_STAT_DEFINED(FrameCacheHits, Count, ECsvCustomStatOp::Set);
+
+	uint64 Size = CacheHitsSize.fetch_add(InSize, std::memory_order_relaxed);
+	TRACE_COUNTER_SET(CacheHitsSize, Size);
+	CSV_CUSTOM_STAT_DEFINED(FrameCacheHitsSize, BytesToApproxKB(Size), ECsvCustomStatOp::Set);
+}
+void FOnDemandIoBackendStats::OnCachePut(uint64 InSize)
+{
+	int32 Count = CachePuts.fetch_add(1, std::memory_order_relaxed);
+	TRACE_COUNTER_SET(CachePuts, Count);
+	CSV_CUSTOM_STAT_DEFINED(FrameCachePuts, Count, ECsvCustomStatOp::Set);
+
+	uint64 Size = CachePutsSize.fetch_add(InSize, std::memory_order_relaxed);
+	TRACE_COUNTER_SET(CachePutsSize, Size);
+	CSV_CUSTOM_STAT_DEFINED(FrameCachePutsSize, BytesToApproxKB(Size), ECsvCustomStatOp::Set);
+}
+void FOnDemandIoBackendStats::OnCacheReject(uint64 InSize)
+{
+	int32 Count = CacheRejects.fetch_add(1, std::memory_order_relaxed);
+	TRACE_COUNTER_SET(CacheRejects, Count);
+	CSV_CUSTOM_STAT_DEFINED(FrameCacheRejects, Count, ECsvCustomStatOp::Set);
+
+	uint64 Size = CacheRejectsSize.fetch_add(InSize, std::memory_order_relaxed);
+	TRACE_COUNTER_SET(CacheRejectsSize, Size);
+	CSV_CUSTOM_STAT_DEFINED(FrameCacheRejectsSize, BytesToApproxKB(Size), ECsvCustomStatOp::Set);
+}
+void FOnDemandIoBackendStats::OnHttpRequestEnqueue()
+{
+	int32 Pending = HttpRequestsPending.fetch_add(1, std::memory_order_relaxed);
+	TRACE_COUNTER_SET(HttpRequestsPending, Pending);
+	CSV_CUSTOM_STAT_DEFINED(FrameHttpRequestsPending, Pending, ECsvCustomStatOp::Set);
+}
+void FOnDemandIoBackendStats::OnHttpRequestDequeue()
+{
+	int32 Pending = HttpRequestsPending.fetch_sub(1, std::memory_order_relaxed);
+	TRACE_COUNTER_SET(HttpRequestsPending, Pending);
+	CSV_CUSTOM_STAT_DEFINED(FrameHttpRequestsPending, Pending, ECsvCustomStatOp::Set);
+
+	int32 Inflight = ++HttpRequestsInflight;
+	TRACE_COUNTER_SET(HttpRequestsInflight, Inflight);
+	CSV_CUSTOM_STAT_DEFINED(FrameHttpRequestsPending, Inflight, ECsvCustomStatOp::Set);
+}
+void FOnDemandIoBackendStats::OnHttpRequestComplete(uint64 InSize)
+{
+	int32 Inflight = --HttpRequestsInflight;
+	TRACE_COUNTER_SET(HttpRequestsInflight, Inflight);
+	CSV_CUSTOM_STAT_DEFINED(FrameHttpRequestsPending, Inflight, ECsvCustomStatOp::Set);
+
+	++HttpRequestsCompleted;
+	TRACE_COUNTER_SET(HttpRequestsCompleted, HttpRequestsCompleted);
+	CSV_CUSTOM_STAT_DEFINED(FrameHttpRequestsCompleted, HttpRequestsCompleted, ECsvCustomStatOp::Set);
+
+	HttpRequestsCompletedSize += InSize;
+	TRACE_COUNTER_SET(HttpRequestsCompletedSize, HttpRequestsCompletedSize);
+	CSV_CUSTOM_STAT_DEFINED(FrameHttpRequestsCompletedSize, BytesToApproxKB(HttpRequestsCompletedSize), ECsvCustomStatOp::Set);
+}
+void FOnDemandIoBackendStats::OnHttpRequestFail()
+{
+	int32 Inflight = --HttpRequestsInflight;
+	TRACE_COUNTER_SET(HttpRequestsInflight, Inflight);
+	CSV_CUSTOM_STAT_DEFINED(FrameHttpRequestsPending, Inflight, ECsvCustomStatOp::Set);
+
+	++HttpRequestsFailed;
+	TRACE_COUNTER_SET(HttpRequestsFailed, HttpRequestsFailed);
+	CSV_CUSTOM_STAT_DEFINED(FrameHttpRequestsFailed, HttpRequestsFailed, ECsvCustomStatOp::Set);
+
+}
+
+} // namespace UE::IO::Private
+
+#endif // IAS_WITH_STATISTICS
diff --git a/Engine/Source/Runtime/Experimental/IoStoreOnDemand/Private/Statistics.h b/Engine/Source/Runtime/Experimental/IoStoreOnDemand/Private/Statistics.h
new file mode 100644
--- /dev/null
+++ b/Engine/Source/Runtime/Experimental/IoStoreOnDemand/Private/Statistics.h
@@ -0,0 +1,92 @@
+// Copyright Epic Games, Inc. All Rights Reserved.
+
+#pragma once
+
+#include "ProfilingDebugging/CountersTrace.h"
+#include "ProfilingDebugging/CsvProfiler.h"
+
+class FIoRequestImpl;
+
+#if COUNTERSTRACE_ENABLED || CSV_PROFILER
+#	define IAS_WITH_STATISTICS 1
+#else
+#	define IAS_WITH_STATISTICS 0
+#endif
+
+namespace UE::IO::Private
+{
+
+#if IAS_WITH_STATISTICS
+
+class FOnDemandIoBackendStats
+{
+public:
+	void OnIoRequestEnqueue(FIoRequestImpl* Request);
+	void OnIoRequestComplete(FIoRequestImpl* Request);
+	void OnIoRequestCancel(FIoRequestImpl* Request);
+	void OnIoRequestFail(FIoRequestImpl* Request);
+	void OnChunkRequestCreate();
+	void OnChunkRequestRelease();
+	void OnCacheHit(uint64 InSize);
+	void OnCachePut(uint64 InSize);
+	void OnCacheReject(uint64 InSize);
+	void OnHttpRequestEnqueue();
+	void OnHttpRequestDequeue();
+	void OnHttpRequestComplete(uint64 InSize);
+	void OnHttpRequestFail();
+
+private:
+	int32 IoRequestsMade = 0;
+	std::atomic<int32> IoRequestsCompleted = 0;
+	int32 IoRequestsCancelled = 0;
+	std::atomic<int32> IoRequestsFailed = 0;
+	int32 ReadRequestsCreated = 0;
+	std::atomic<int32> ReadRequestsRemoved = 0;
+	std::atomic<int32> CacheHits = 0;
+	std::atomic<int32> CachePuts = 0;
+	std::atomic<int32> CacheRejects = 0;
+	int32 HttpRequestsCompleted = 0;
+	int32 HttpRequestsFailed = 0;
+	std::atomic<int32> HttpRequestsPending = 0;
+	int32 HttpRequestsInflight = 0;
+
+	std::atomic<uint64> IoRequestsCompletedSize = 0;
+	std::atomic<uint64> CacheHitsSize = 0;
+	std::atomic<uint64> CachePutsSize = 0;
+	std::atomic<uint64> CacheRejectsSize = 0;
+	uint64 HttpRequestsCompletedSize = 0;
+
+	float BytesToApproxMB(uint64 Bytes)
+	{
+		return float(double(Bytes) / 1024.0 / 1024.0);
+	}
+
+	float BytesToApproxKB(uint64 Bytes)
+	{
+		return float(double(Bytes) / 1024.0);
+	}
+};
+
+#else // IAS_WITH_STATISTICS
+
+class FOnDemandIoBackendStats
+{
+public:
+	void OnIoRequestEnqueue(FIoRequestImpl* Request) {}
+	void OnIoRequestComplete(FIoRequestImpl* Request) {}
+	void OnIoRequestCancel(FIoRequestImpl* Request) {}
+	void OnIoRequestFail(FIoRequestImpl* Request) {}
+	void OnChunkRequestCreate() {}
+	void OnChunkRequestRelease() {}
+	void OnCacheHit(uint64 InSize) {}
+	void OnCachePut(uint64 InSize) {}
+	void OnCacheReject(uint64 InSize) {}
+	void OnHttpRequestEnqueue() {}
+	void OnHttpRequestDequeue() {}
+	void OnHttpRequestComplete(uint64 InSize) {}
+	void OnHttpRequestFail() {}
+};
+
+#endif // IAS_WITH_STATISTICS
+
+} // namespace UE::IO::Private
