particles that get added into a cluster union must get a temporary ChildToParent set immediately prior to the deferred cluster union properties update

diff --git a/Engine/Source/Runtime/Experimental/Chaos/Private/Chaos/ClusterUnionManager.cpp b/Engine/Source/Runtime/Experimental/Chaos/Private/Chaos/ClusterUnionManager.cpp
--- a/Engine/Source/Runtime/Experimental/Chaos/Private/Chaos/ClusterUnionManager.cpp
+++ b/Engine/Source/Runtime/Experimental/Chaos/Private/Chaos/ClusterUnionManager.cpp
@@ -460,6 +460,15 @@ namespace Chaos
 				Properties.bIsAuxiliaryParticle = false;
 				Cluster->ChildProperties.Add(Particle, Properties);
 			}
+
+			// Need to set a temporary ChildToParent on the particle. If this ChildToParent doesn't exist, and the particle goes through the removal process BEFORE
+			// the deferred cluster union update properties is called, the position of the particle will get reset to the position of the cluster union.
+			if (FPBDRigidClusteredParticleHandle* ClusterParticle = Particle->CastToClustered())
+			{
+				const FRigidTransform3 ClusterWorldTM(Cluster->InternalCluster->X(), Cluster->InternalCluster->R());
+				const FRigidTransform3 Frame = GetParticleRigidFrameInClusterUnion(Particle, ClusterWorldTM);
+				ClusterParticle->SetChildToParent(Frame);
+			}
 		}
 
 		MClustering.AddParticlesToCluster(Cluster->InternalCluster, FinalParticlesToAdd, ChildToParentMap);
