Make 'select by size' dataflow node give the option to choose between using the 'relative size' or 'size' values

diff --git a/Engine/Plugins/Experimental/GeometryCollectionPlugin/Source/GeometryCollectionNodes/Private/Dataflow/GeometryCollectionSelectionNodes.cpp b/Engine/Plugins/Experimental/GeometryCollectionPlugin/Source/GeometryCollectionNodes/Private/Dataflow/GeometryCollectionSelectionNodes.cpp
--- a/Engine/Plugins/Experimental/GeometryCollectionPlugin/Source/GeometryCollectionNodes/Private/Dataflow/GeometryCollectionSelectionNodes.cpp
+++ b/Engine/Plugins/Experimental/GeometryCollectionPlugin/Source/GeometryCollectionNodes/Private/Dataflow/GeometryCollectionSelectionNodes.cpp
@@ -622,7 +622,7 @@ void FCollectionTransformSelectionBySizeDataflowNode::Evaluate(Dataflow::FContex
 		bool bInsideRange = RangeSetting == ERangeSettingEnum::Dataflow_RangeSetting_InsideRange;
 
 		GeometryCollection::Facades::FCollectionTransformSelectionFacade TransformSelectionFacade(InCollection);
-		const TArray<int32>& SelectionArr = TransformSelectionFacade.SelectBySize(InSizeMin, InSizeMax, bInclusive, bInsideRange);
+		const TArray<int32>& SelectionArr = TransformSelectionFacade.SelectBySize(InSizeMin, InSizeMax, bInclusive, bInsideRange, bUseRelativeSize);
 
 		FDataflowTransformSelection NewTransformSelection;
 		NewTransformSelection.Initialize(InCollection.NumElements(FGeometryCollection::TransformGroup), false);
diff --git a/Engine/Plugins/Experimental/GeometryCollectionPlugin/Source/GeometryCollectionNodes/Public/Dataflow/GeometryCollectionSelectionNodes.h b/Engine/Plugins/Experimental/GeometryCollectionPlugin/Source/GeometryCollectionNodes/Public/Dataflow/GeometryCollectionSelectionNodes.h
--- a/Engine/Plugins/Experimental/GeometryCollectionPlugin/Source/GeometryCollectionNodes/Public/Dataflow/GeometryCollectionSelectionNodes.h
+++ b/Engine/Plugins/Experimental/GeometryCollectionPlugin/Source/GeometryCollectionNodes/Public/Dataflow/GeometryCollectionSelectionNodes.h
@@ -686,6 +686,10 @@ public:
 	UPROPERTY(EditAnywhere, Category = "Size")
 	bool bInclusive = true;
 
+	/** Whether to use the 'Relative Size' -- i.e., the Size / Largest Bone Size. Otherwise, Size is the cube root of Volume. */
+	UPROPERTY(EditAnywhere, Category = "Size")
+	bool bUseRelativeSize = true;
+
 	/** Array of the selected bone indicies */
 	UPROPERTY(meta = (DataflowOutput, DisplayName = "TransformSelection"))
 	FDataflowTransformSelection TransformSelection;
diff --git a/Engine/Source/Runtime/Experimental/Chaos/Private/GeometryCollection/Facades/CollectionTransformSelectionFacade.cpp b/Engine/Source/Runtime/Experimental/Chaos/Private/GeometryCollection/Facades/CollectionTransformSelectionFacade.cpp
--- a/Engine/Source/Runtime/Experimental/Chaos/Private/GeometryCollection/Facades/CollectionTransformSelectionFacade.cpp
+++ b/Engine/Source/Runtime/Experimental/Chaos/Private/GeometryCollection/Facades/CollectionTransformSelectionFacade.cpp
@@ -628,8 +628,16 @@ namespace GeometryCollection::Facades
 		InOutSelection.SetNum(NewNumElements);
 	}
 
-	TArray<int32> FCollectionTransformSelectionFacade::SelectBySize(float SizeMin, float SizeMax, bool bInclusive, bool bInsideRange) const
+	TArray<int32> FCollectionTransformSelectionFacade::SelectBySize(float SizeMin, float SizeMax, bool bInclusive, bool bInsideRange, bool bUseRelativeSize) const
 	{
+		// if not using relative size, convert sizes to volumes and select by volume
+		if (!bUseRelativeSize)
+		{
+			float VolumeMin = SizeMin * SizeMin * SizeMin;
+			float VolumeMax = SizeMax * SizeMax * SizeMax;
+			return SelectByVolume(VolumeMin, VolumeMax, bInclusive, bInsideRange);
+		}
+
 		TArray<int32> OutSelection;
 
 		// TODO: (See also SelectByVolume) We should add a method to get the volumes without modifying the collection, and then remove this full geometrycollection copy
diff --git a/Engine/Source/Runtime/Experimental/Chaos/Public/GeometryCollection/Facades/CollectionTransformSelectionFacade.h b/Engine/Source/Runtime/Experimental/Chaos/Public/GeometryCollection/Facades/CollectionTransformSelectionFacade.h
--- a/Engine/Source/Runtime/Experimental/Chaos/Public/GeometryCollection/Facades/CollectionTransformSelectionFacade.h
+++ b/Engine/Source/Runtime/Experimental/Chaos/Public/GeometryCollection/Facades/CollectionTransformSelectionFacade.h
@@ -119,8 +119,8 @@ namespace GeometryCollection::Facades
 		/**  */
 		static CHAOS_API void SelectByPercentage(TArray<int32>& InOutSelection, int32 Percentage, bool Deterministic, float RandomSeed);
 
-		/**  */
-		CHAOS_API TArray<int32> SelectBySize(float SizeMin, float SizeMax, bool bInclusive, bool bInsideRange) const;
+		/** Select transforms by their Size or Relative Size */
+		CHAOS_API TArray<int32> SelectBySize(float SizeMin, float SizeMax, bool bInclusive, bool bInsideRange, bool bUseRelativeSize = true) const;
 
 		/**  */
 		CHAOS_API TArray<int32> SelectByVolume(float VolumeMin, float VolumeMax, bool bInclusive, bool bInsideRange) const;
