Chaos - New and improved implicit casting which supports casting to intermediate classes (e.g., to FImplicitObjectUnion when you have an FImplicitObjectUnionClustered)

diff --git a/Engine/Source/Programs/HeadlessChaos/Private/HeadlessChaosTestImplicits.cpp b/Engine/Source/Programs/HeadlessChaos/Private/HeadlessChaosTestImplicits.cpp
--- a/Engine/Source/Programs/HeadlessChaos/Private/HeadlessChaosTestImplicits.cpp
+++ b/Engine/Source/Programs/HeadlessChaos/Private/HeadlessChaosTestImplicits.cpp
@@ -2148,4 +2148,42 @@ namespace ChaosTest {
 
 	}
 
+	GTEST_TEST(ImplicitTests, TestImplicitCasts)
+	{
+		TUniquePtr<FImplicitObject> Sphere = MakeUnique<FImplicitSphere3>(FVec3(0), 100.0);
+		const FImplicitObject* ConstSphere = Sphere.Get();
+
+		TArray<TUniquePtr<FImplicitObject>> UnionObjects;
+		UnionObjects.Emplace(MakeUnique<FImplicitSphere3>(FVec3(0), 100.0));
+
+		TArray<TUniquePtr<FImplicitObject>> UnionClusteredObjects;
+		UnionClusteredObjects.Emplace(MakeUnique<FImplicitSphere3>(FVec3(0), 100.0));
+
+		TUniquePtr<FImplicitObject> Union = MakeUnique<FImplicitObjectUnion>(MoveTemp(UnionObjects));
+		TUniquePtr<FImplicitObject> UnionClustered = MakeUnique<FImplicitObjectUnionClustered>(MoveTemp(UnionClusteredObjects));
+
+		// We can cast an implicit object to its exact type (const and non-const)
+		FImplicitSphere3* SphereCast = Sphere->AsA<FImplicitSphere3>();
+		EXPECT_NE(SphereCast, nullptr);
+
+		FImplicitSphere3* SphereCastChecked = Sphere->AsAChecked<FImplicitSphere3>();
+		EXPECT_NE(SphereCastChecked, nullptr);
+
+		const FImplicitSphere3* ConstSphereCast = ConstSphere->AsA<FImplicitSphere3>();
+		EXPECT_NE(ConstSphereCast, nullptr);
+
+		const FImplicitSphere3* ConstSphereCastChecked = ConstSphere->AsAChecked<FImplicitSphere3>();
+		EXPECT_NE(ConstSphereCastChecked, nullptr);
+
+		FImplicitObjectUnion* UnionCast = Union->AsA<FImplicitObjectUnion>();
+		EXPECT_NE(UnionCast, nullptr);
+
+		FImplicitObjectUnionClustered* UnionClusteredCast = UnionClustered->AsA<FImplicitObjectUnionClustered>();
+		EXPECT_NE(UnionClusteredCast, nullptr);
+
+		// We can cast a Clustered Union to a Union
+		FImplicitObjectUnion* UnionClusteredUpCast = UnionClustered->AsA<FImplicitObjectUnion>();
+		EXPECT_NE(UnionClusteredUpCast, nullptr);
+
+	}
 }
diff --git a/Engine/Source/Runtime/Experimental/Chaos/Private/Chaos/ChaosDebugDraw.cpp b/Engine/Source/Runtime/Experimental/Chaos/Private/Chaos/ChaosDebugDraw.cpp
--- a/Engine/Source/Runtime/Experimental/Chaos/Private/Chaos/ChaosDebugDraw.cpp
+++ b/Engine/Source/Runtime/Experimental/Chaos/Private/Chaos/ChaosDebugDraw.cpp
@@ -826,12 +826,7 @@ namespace Chaos
 
 		void DrawParticleBVHImpl(const FRigidTransform3& SpaceTransform, const FGeometryParticleHandle* Particle, const FColor& InColor, const FChaosDebugDrawSettings& Settings)
 		{
-			const FImplicitObjectUnion* Union = Particle->Geometry()->GetObject<FImplicitObjectUnion>();
-			if (Union == nullptr)
-			{
-				Union = Particle->Geometry()->GetObject<FImplicitObjectUnionClustered>();
-			}
-			if (Union != nullptr)
+			if (const FImplicitObjectUnion* Union = Particle->Geometry()->template AsA<FImplicitObjectUnion>())
 			{
 				if (Union->GetBVH() != nullptr)
 				{
diff --git a/Engine/Source/Runtime/Experimental/Chaos/Private/Chaos/Collision/ParticlePairMidPhase.cpp b/Engine/Source/Runtime/Experimental/Chaos/Private/Chaos/Collision/ParticlePairMidPhase.cpp
--- a/Engine/Source/Runtime/Experimental/Chaos/Private/Chaos/Collision/ParticlePairMidPhase.cpp
+++ b/Engine/Source/Runtime/Experimental/Chaos/Private/Chaos/Collision/ParticlePairMidPhase.cpp
@@ -95,14 +95,7 @@ namespace Chaos
 			// All implicit excepts unions have 1 leaf
 			NumImplicits = 1;
 
-			// @todo(chaos): we need better downcasting in FImplicitObject
-			const FImplicitObjectUnion* Union = Implicit->template GetObject<FImplicitObjectUnion>();
-			if (Union == nullptr)
-			{
-				Union = Implicit->template GetObject<FImplicitObjectUnionClustered>();
-			}
-
-			if (Union != nullptr)
+			if (const FImplicitObjectUnion* Union = Implicit->template AsA<FImplicitObjectUnion>())
 			{
 				NumImplicits = Union->GetNumLeafObjects();
 			}
@@ -1080,11 +1073,11 @@ namespace Chaos
 		const Private::FImplicitBVH* BVH1 = nullptr;
 		if (CVars::bChaosUnionBVHEnabled)
 		{
-			if (const FImplicitObjectUnion* Union0 = Implicit0->template GetObject<FImplicitObjectUnion>())
+			if (const FImplicitObjectUnion* Union0 = Implicit0->template AsA<FImplicitObjectUnion>())
 			{
 				BVH0 = Union0->GetBVH();
 			}
-			if (const FImplicitObjectUnion* Union1 = Implicit1->template GetObject<FImplicitObjectUnion>())
+			if (const FImplicitObjectUnion* Union1 = Implicit1->template AsA<FImplicitObjectUnion>())
 			{
 				BVH1 = Union1->GetBVH();
 			}
diff --git a/Engine/Source/Runtime/Experimental/Chaos/Private/Chaos/ImplicitObject.cpp b/Engine/Source/Runtime/Experimental/Chaos/Private/Chaos/ImplicitObject.cpp
--- a/Engine/Source/Runtime/Experimental/Chaos/Private/Chaos/ImplicitObject.cpp
+++ b/Engine/Source/Runtime/Experimental/Chaos/Private/Chaos/ImplicitObject.cpp
@@ -89,7 +89,7 @@ namespace Chaos
 
 	bool FImplicitObject::IsUnderlyingUnion() const
 	{
-		return Type == ImplicitObjectType::Union;
+		return (Type == ImplicitObjectType::Union) || (Type == ImplicitObjectType::UnionClustered);
 	}
 
 	FReal FImplicitObject::SignedDistance(const FVec3& x) const
diff --git a/Engine/Source/Runtime/Experimental/Chaos/Public/Chaos/ImplicitObject.h b/Engine/Source/Runtime/Experimental/Chaos/Public/Chaos/ImplicitObject.h
--- a/Engine/Source/Runtime/Experimental/Chaos/Public/Chaos/ImplicitObject.h
+++ b/Engine/Source/Runtime/Experimental/Chaos/Public/Chaos/ImplicitObject.h
@@ -71,6 +71,19 @@ using FImplicitHierarchyVisitor = TFunctionRef<void(const FImplicitObject* Impli
 using FImplicitHierarchyVisitorBool = TFunctionRef<bool(const FImplicitObject* Implicit, const FRigidTransform3& Transform, const int32 RootObjectIndex, const int32 ObjectIndex, const int32 LeafObjectIndex)>;
 
 
+// Specialized for derived classes so that we can downcast to non-leaf types in the class hierarchy
+// @see TImplicitTypeInfo<FImplicitObjectUnion>
+template<typename T>
+struct TImplicitTypeInfo
+{
+	// Return true if implicits of type InType can be cast to T.
+	// I.e., is T a base class (or the class) of InType.
+	static bool IsBaseOf(const EImplicitObjectType InType)
+	{
+		return (T::StaticType() == InType);
+	}
+};
+
 /*
  * Base class for implicit collision geometry such as spheres, capsules, boxes, etc.
  * 
@@ -111,6 +124,49 @@ public:
 	FImplicitObject(FImplicitObject&&) = delete;
 	CHAOS_API virtual ~FImplicitObject();
 
+	// Can this object be cast to type T_DERIVED?
+	template<typename TargetType>
+	bool IsA() const
+	{
+		return TImplicitTypeInfo<TargetType>::IsBaseOf(GetType());
+	}
+
+	// Dynamic cast to type T_DERIVED. Returns null if T_DERIVED is not a valid cast for this object.
+	template<typename TargetType>
+	const TargetType* AsA() const
+	{
+		if (IsA<TargetType>())
+		{
+			return static_cast<const TargetType*>(this);
+		}
+		return nullptr;
+	}
+
+	template<typename TargetType>
+	TargetType* AsA()
+	{
+		if (IsA<TargetType>())
+		{
+			return static_cast<TargetType*>(this);
+		}
+		return nullptr;
+	}
+
+	template<typename TargetType>
+	const TargetType* AsAChecked() const
+	{
+		check(IsA<TargetType>());
+		return static_cast<const TargetType*>(this);
+	}
+
+	template<typename TargetType>
+	TargetType* AsAChecked()
+	{
+		check(IsA<TargetType>());
+		return static_cast<TargetType*>(this);
+	}
+
+
 	template<class T_DERIVED>
 	T_DERIVED* GetObject()
 	{
diff --git a/Engine/Source/Runtime/Experimental/Chaos/Public/Chaos/ImplicitObjectUnion.h b/Engine/Source/Runtime/Experimental/Chaos/Public/Chaos/ImplicitObjectUnion.h
--- a/Engine/Source/Runtime/Experimental/Chaos/Public/Chaos/ImplicitObjectUnion.h
+++ b/Engine/Source/Runtime/Experimental/Chaos/Public/Chaos/ImplicitObjectUnion.h
@@ -359,4 +359,18 @@ private:
 	TArray<FPBDRigidParticleHandle*> MOriginalParticleLookupHack;	
 	TMap<const FImplicitObject*,FPBDRigidParticleHandle*> MCollisionParticleLookupHack;	//temp hack for finding collision particles
 };
+
+template<>
+struct TImplicitTypeInfo<FImplicitObjectUnion>
+{
+	// @todo(chaos): this is a bit topsy-turvy because the base class needs to know about all derived classes. 
+	// Ideally we would have a function like GetBaseType(InType) and implement TImplicitTypeInfo<FImplicitObjectUnionClustered> 
+	// but then we'd need a runtime of type->basetype for all valid types (and then there's the bitmask complication)
+	static bool IsBaseOf(const EImplicitObjectType InType)
+	{
+		return (InType == FImplicitObjectUnion::StaticType()) || TImplicitTypeInfo<FImplicitObjectUnionClustered>::IsBaseOf(InType);
+	}
+};
+
+
 }
