Removed IAS-stats' dependency on an request object

#rnx

diff --git a/Engine/Source/Runtime/Experimental/IoStoreOnDemand/Private/OnDemandIoDispatcherBackend.cpp b/Engine/Source/Runtime/Experimental/IoStoreOnDemand/Private/OnDemandIoDispatcherBackend.cpp
--- a/Engine/Source/Runtime/Experimental/IoStoreOnDemand/Private/OnDemandIoDispatcherBackend.cpp
+++ b/Engine/Source/Runtime/Experimental/IoStoreOnDemand/Private/OnDemandIoDispatcherBackend.cpp
@@ -974,12 +974,12 @@ void FOnDemandIoBackend::CompleteRequest(FChunkRequest* ChunkRequest)
 		
 		if (!bDecoded)
 		{
-			Stats.OnIoRequestFail(Request);
+			Stats.OnIoRequestFail();
 			bCanCache = false;
 			Request->SetFailed();
 		}
 
-		Stats.OnIoRequestComplete(Request);
+		Stats.OnIoRequestComplete(Request->GetBuffer().GetSize());
 		CompletedRequests.Enqueue(Request);
 		BackendContext->WakeUpDispatcherThreadDelegate.Execute();
 	}
@@ -1014,7 +1014,7 @@ bool FOnDemandIoBackend::Resolve(FIoRequestImpl* Request)
 		return false;
 	}
 
-	Stats.OnIoRequestEnqueue(Request);
+	Stats.OnIoRequestEnqueue();
 	FChunkRequestParams RequestParams = FChunkRequestParams::Create(Request, ChunkInfo);
 	FChunkRequest* ChunkRequest = ChunkRequests.Create(Request, RequestParams);
 
diff --git a/Engine/Source/Runtime/Experimental/IoStoreOnDemand/Private/Statistics.cpp b/Engine/Source/Runtime/Experimental/IoStoreOnDemand/Private/Statistics.cpp
--- a/Engine/Source/Runtime/Experimental/IoStoreOnDemand/Private/Statistics.cpp
+++ b/Engine/Source/Runtime/Experimental/IoStoreOnDemand/Private/Statistics.cpp
@@ -4,8 +4,6 @@
 
 #if IAS_WITH_STATISTICS
 
-#include "IO/IoDispatcherBackend.h"
-
 namespace UE::IO::Private
 {
 
@@ -59,29 +57,29 @@ CSV_DEFINE_STAT(Ias, FrameHttpRequestsPending);
 CSV_DEFINE_STAT(Ias, FrameHttpRequestsInflight);
 CSV_DEFINE_STAT(Ias, FrameHttpRequestsCompletedSize);
 
-void FOnDemandIoBackendStats::OnIoRequestEnqueue(FIoRequestImpl* Request)
+void FOnDemandIoBackendStats::OnIoRequestEnqueue()
 {
 	++IoRequestsMade;
 	TRACE_COUNTER_SET(IoRequestsMade, IoRequestsMade);
 	CSV_CUSTOM_STAT_DEFINED(FrameIoRequestsMade, IoRequestsMade, ECsvCustomStatOp::Set);
 }
-void FOnDemandIoBackendStats::OnIoRequestComplete(FIoRequestImpl* Request)
+void FOnDemandIoBackendStats::OnIoRequestComplete(uint64 RequestSize)
 {
 	int32 Count = IoRequestsCompleted.fetch_add(1, std::memory_order_relaxed);
 	TRACE_COUNTER_SET(IoRequestsCompleted, Count);
 	CSV_CUSTOM_STAT_DEFINED(FrameIoRequestsCompleted, Count, ECsvCustomStatOp::Set);
 
-	uint64 Size = IoRequestsCompletedSize.fetch_add(Request->GetBuffer().DataSize());
+	uint64 Size = IoRequestsCompletedSize.fetch_add(RequestSize);
 	TRACE_COUNTER_SET(IoRequestsCompletedSize, Size);
 	CSV_CUSTOM_STAT_DEFINED(FrameIoRequestsCompletedSize, BytesToApproxKB(Size), ECsvCustomStatOp::Set);
 }
-void FOnDemandIoBackendStats::OnIoRequestCancel(FIoRequestImpl* Request)
+void FOnDemandIoBackendStats::OnIoRequestCancel()
 {
 	++IoRequestsCancelled;
 	TRACE_COUNTER_SET(IoRequestsCancelled, IoRequestsCancelled);
 	CSV_CUSTOM_STAT_DEFINED(FrameIoRequestsCancelled, IoRequestsCancelled, ECsvCustomStatOp::Set);
 }
-void FOnDemandIoBackendStats::OnIoRequestFail(FIoRequestImpl* Request)
+void FOnDemandIoBackendStats::OnIoRequestFail()
 {
 	int32 Count = IoRequestsFailed.fetch_add(1, std::memory_order_relaxed);
 	TRACE_COUNTER_SET(IoRequestsFailed, Count);
diff --git a/Engine/Source/Runtime/Experimental/IoStoreOnDemand/Private/Statistics.h b/Engine/Source/Runtime/Experimental/IoStoreOnDemand/Private/Statistics.h
--- a/Engine/Source/Runtime/Experimental/IoStoreOnDemand/Private/Statistics.h
+++ b/Engine/Source/Runtime/Experimental/IoStoreOnDemand/Private/Statistics.h
@@ -5,7 +5,7 @@
 #include "ProfilingDebugging/CountersTrace.h"
 #include "ProfilingDebugging/CsvProfiler.h"
 
-class FIoRequestImpl;
+#include <atomic>
 
 #if COUNTERSTRACE_ENABLED || CSV_PROFILER
 #	define IAS_WITH_STATISTICS 1
@@ -21,10 +21,10 @@ namespace UE::IO::Private
 class FOnDemandIoBackendStats
 {
 public:
-	void OnIoRequestEnqueue(FIoRequestImpl* Request);
-	void OnIoRequestComplete(FIoRequestImpl* Request);
-	void OnIoRequestCancel(FIoRequestImpl* Request);
-	void OnIoRequestFail(FIoRequestImpl* Request);
+	void OnIoRequestEnqueue();
+	void OnIoRequestComplete(uint64 RequestSize);
+	void OnIoRequestCancel();
+	void OnIoRequestFail();
 	void OnChunkRequestCreate();
 	void OnChunkRequestRelease();
 	void OnCacheHit(uint64 InSize);
@@ -72,10 +72,10 @@ private:
 class FOnDemandIoBackendStats
 {
 public:
-	void OnIoRequestEnqueue(FIoRequestImpl* Request) {}
-	void OnIoRequestComplete(FIoRequestImpl* Request) {}
-	void OnIoRequestCancel(FIoRequestImpl* Request) {}
-	void OnIoRequestFail(FIoRequestImpl* Request) {}
+	void OnIoRequestEnqueue() {}
+	void OnIoRequestComplete(uint64 RequestSize) {}
+	void OnIoRequestCancel() {}
+	void OnIoRequestFail() {}
 	void OnChunkRequestCreate() {}
 	void OnChunkRequestRelease() {}
 	void OnCacheHit(uint64 InSize) {}
