Moved some localised utilities from .h to .cpp

#rnx

diff --git a/Engine/Source/Runtime/Experimental/IoStoreOnDemand/Private/Statistics.cpp b/Engine/Source/Runtime/Experimental/IoStoreOnDemand/Private/Statistics.cpp
--- a/Engine/Source/Runtime/Experimental/IoStoreOnDemand/Private/Statistics.cpp
+++ b/Engine/Source/Runtime/Experimental/IoStoreOnDemand/Private/Statistics.cpp
@@ -7,7 +7,11 @@
 namespace UE::IO::Private
 {
 
-///////////////////////////////////////////////////////////////////////////////
+////////////////////////////////////////////////////////////////////////////////
+static int32 BytesToApproxMB(uint64 Bytes) { return int32(Bytes >> 20); }
+static int32 BytesToApproxKB(uint64 Bytes) { return int32(Bytes >> 10); }
+
+////////////////////////////////////////////////////////////////////////////////
 // TRACE STATS
 // iorequest stats
 TRACE_DECLARE_INT_COUNTER(IoRequestsMade, TEXT("Ias/IoRequestsMade"));
diff --git a/Engine/Source/Runtime/Experimental/IoStoreOnDemand/Private/Statistics.h b/Engine/Source/Runtime/Experimental/IoStoreOnDemand/Private/Statistics.h
--- a/Engine/Source/Runtime/Experimental/IoStoreOnDemand/Private/Statistics.h
+++ b/Engine/Source/Runtime/Experimental/IoStoreOnDemand/Private/Statistics.h
@@ -55,16 +55,6 @@ private:
 	std::atomic<uint64> CachePutsSize = 0;
 	std::atomic<uint64> CacheRejectsSize = 0;
 	uint64 HttpRequestsCompletedSize = 0;
-
-	float BytesToApproxMB(uint64 Bytes)
-	{
-		return float(double(Bytes) / 1024.0 / 1024.0);
-	}
-
-	float BytesToApproxKB(uint64 Bytes)
-	{
-		return float(double(Bytes) / 1024.0);
-	}
 };
 
 #else // IAS_WITH_STATISTICS
