FileCache -> FileIoCache

#rnx
#rnx

diff --git a/Engine/Source/Runtime/Experimental/IoStoreOnDemand/Private/FileCache.cpp b/Engine/Source/Runtime/Experimental/IoStoreOnDemand/Private/FileCache.cpp
deleted file mode 100644
diff --git a/Engine/Source/Runtime/Experimental/IoStoreOnDemand/Private/FileCache.h b/Engine/Source/Runtime/Experimental/IoStoreOnDemand/Private/FileCache.h
deleted file mode 100644
diff --git a/Engine/Source/Runtime/Experimental/IoStoreOnDemand/Private/FileIoCache.cpp b/Engine/Source/Runtime/Experimental/IoStoreOnDemand/Private/FileIoCache.cpp
new file mode 100644
--- /dev/null
+++ b/Engine/Source/Runtime/Experimental/IoStoreOnDemand/Private/FileIoCache.cpp
@@ -0,0 +1,934 @@
+// Copyright Epic Games, Inc. All Rights Reserved.
+
+// {{{1 misc
+
+#include "FileIoCache.h"
+
+#include "Containers/IntrusiveDoubleLinkedList.h"
+#include "GenericPlatform/GenericPlatformFile.h"
+#include "HAL/Event.h"
+#include "HAL/FileManager.h"
+#include "HAL/Platform.h"
+#include "HAL/PlatformFile.h"
+#include "HAL/Runnable.h"
+#include "HAL/RunnableThread.h"
+#include "IO/IoDispatcher.h"
+#include "IO/IoHash.h"
+#include "Misc/ScopeLock.h"
+#include "Misc/Paths.h"
+#include "Misc/CommandLine.h"
+#include "ProfilingDebugging/CountersTrace.h"
+#include "ProfilingDebugging/CpuProfilerTrace.h"
+#include "Templates/UniquePtr.h"
+
+#include <atomic>
+
+DEFINE_LOG_CATEGORY(LogIoCache);
+
+// }}}
+
+namespace UE::IO::Private
+{
+
+using namespace UE::Tasks;
+
+// {{{1 stats
+
+////////////////////////////////////////////////////////////////////////////////
+#define BUILD_DATUM_X(x) \
+	x(CachedBytes) \
+	x(ErrorCount) \
+	x(GetCount) \
+	x(PendingBytes) \
+	x(PutCount) \
+	x(PutExistingCount) \
+	x(PutRejectCount) \
+	x(ReadBytes)
+
+#if !COUNTERSTRACE_ENABLED
+namespace FCountersTrace
+{
+	struct FCounterInt
+	{
+		FCounterInt(...) {}
+		void Set(int64 i) { V = i; }
+		void Add(int64 d) { V += d; }
+		int64 Get() const { return V;}
+		std::atomic<int64> V = 0;
+	};
+	using FCounterAtomicInt = FCounterInt;
+}
+#endif 
+
+class FStatistics
+{
+public:
+	enum class EDatum : int32
+	{
+#define DATUM_X(x) x,
+		BUILD_DATUM_X(DATUM_X)
+#undef DATUM_X
+		_Count,
+	};
+
+	void Add(EDatum Which, int64 Delta);
+	void Set(EDatum Which, int64 Value);
+	int64 Get(EDatum Which) const;
+
+private:
+	struct FDatum
+	{
+		FDatum(const TCHAR* Name);
+		FCountersTrace::FCounterAtomicInt Value;
+	};
+
+	FDatum Datums[int(EDatum::_Count)] = {
+#define DATUM_X(x) FDatum{ TEXT("Ias/FileCache/") TEXT(PREPROCESSOR_TO_STRING(x)) },
+		BUILD_DATUM_X(DATUM_X)
+#undef DATUM_X
+	};
+};
+
+static FStatistics GStatistics;
+
+FStatistics::FDatum::FDatum(const TCHAR* Name)
+: Value(TraceCounterNameType_Static, Name, TraceCounterDisplayHint_None)
+{
+	FCountersTrace::FCounterAtomicInt(TEXT(""), TraceCounterDisplayHint_None).Get();
+	FCountersTrace::FCounterInt(TEXT(""), TraceCounterDisplayHint_None).Get();
+}
+
+void FStatistics::Add(EDatum Which, int64 Delta)
+{
+	Datums[int(Which)].Value.Add(Delta);
+}
+
+void FStatistics::Set(EDatum Which, int64 Value)
+{
+	Datums[int(Which)].Value.Set(Value);
+}
+
+int64 FStatistics::Get(EDatum Which) const
+{
+	return Datums[int(Which)].Value.Get();
+}
+
+
+
+// {{{1 toc-file
+
+////////////////////////////////////////////////////////////////////////////////
+class FCacheFileToc
+{
+	struct FHeader
+	{
+		static constexpr uint32 ExpectedMagic = 0x2e696f; // .io
+
+		uint32 Magic = 0;
+		uint32 EntryCount = 0;
+		uint64 CursorPos = 0;
+	};
+
+public:
+	struct FTocEntry
+	{
+		FIoHash Key;
+		FIoHash Hash;
+		uint64 SerialOffset = 0;
+		uint64 SerialSize = 0;
+
+		friend FArchive& operator<<(FArchive& Ar, FTocEntry& Entry)
+		{
+			Ar << Entry.Key;
+			Ar << Entry.Hash;
+			Ar << Entry.SerialOffset;
+			Ar << Entry.SerialSize;
+
+			return Ar;
+		}
+	};
+
+	FCacheFileToc() = default;
+
+	void AddEntry(const FIoHash& Key, const FIoHash& Hash, uint64 SerialOffset, uint64 SerialSize);
+	FIoStatus Load(const FString& FilePath, uint64& OutCursorPos);
+	FIoStatus Save(const FString& FilePath, const uint64 CursorPos);
+
+	TConstArrayView<FTocEntry> GetEntries() const
+	{ 
+		return TocEntries;
+	}
+
+private:
+	TArray<FTocEntry> TocEntries;
+};
+
+void FCacheFileToc::AddEntry(const FIoHash& Key, const FIoHash& Hash, uint64 SerialOffset, uint64 SerialSize)
+{
+	TocEntries.Add(FTocEntry
+	{
+		Key,
+		Hash,
+		SerialOffset,
+		SerialSize
+	});
+}
+
+FIoStatus FCacheFileToc::Load(const FString& FilePath, uint64& OutCursorPos)
+{
+	IFileManager& FileMgr = IFileManager::Get();
+	TUniquePtr<FArchive> Ar(FileMgr.CreateFileReader(*FilePath));
+	
+	OutCursorPos = ~uint64(0);
+
+	if (!Ar.IsValid() || Ar->IsError())
+	{
+		return FIoStatus(EIoErrorCode::FileNotOpen);
+	}
+
+	FHeader Header;
+	Ar->Serialize(&Header, sizeof(FHeader));
+
+	if (Header.Magic != FHeader::ExpectedMagic)
+	{
+		return FIoStatus(EIoErrorCode::CorruptToc);
+	}
+
+	OutCursorPos = Header.CursorPos;
+	TocEntries.Empty();
+	TocEntries.Reserve(Header.EntryCount);
+	*Ar << TocEntries;
+
+	return FIoStatus(EIoErrorCode::Ok);
+}
+
+FIoStatus FCacheFileToc::Save(const FString& FilePath, const uint64 CursorPos)
+{
+	IFileManager& FileMgr = IFileManager::Get();
+	TUniquePtr<FArchive> Ar(FileMgr.CreateFileWriter(*FilePath));
+
+	if (!Ar.IsValid() || Ar->IsError())
+	{
+		return FIoStatus(EIoErrorCode::FileNotOpen);
+	}
+
+	FHeader Header;
+	Header.Magic = FHeader::ExpectedMagic;
+	Header.EntryCount = TocEntries.Num();
+	Header.CursorPos = CursorPos;
+
+	Ar->Serialize(&Header, sizeof(FHeader));
+	*Ar << TocEntries;
+
+	return FIoStatus(EIoErrorCode::Ok);
+}
+
+
+
+// {{{1 cache-map
+
+////////////////////////////////////////////////////////////////////////////////
+enum class ECacheEntryState : uint8
+{
+	None,
+	Pending,
+	Writing,
+	Persisted
+};
+
+////////////////////////////////////////////////////////////////////////////////
+struct FCacheEntry
+	: public TIntrusiveDoubleLinkedListNode<FCacheEntry>
+{
+	FIoHash Key;
+	FIoHash Hash;
+	uint64 SerialOffset = 0;
+	uint64 SerialSize = 0;
+	FIoBuffer Data;
+	ECacheEntryState State = ECacheEntryState::None;
+};
+
+using FCacheEntryList = TIntrusiveDoubleLinkedList<FCacheEntry>;
+
+////////////////////////////////////////////////////////////////////////////////
+class FCacheMap
+{
+public:
+	void SetCacheLimits(uint64 MaxPendingBytes, uint64 MaxPersistedBytes);
+	void Reset();
+	bool Contains(FIoHash Key) const;
+	bool Get(const FIoHash& Key, FCacheEntry& OutEntry) const;
+	bool InsertPending(FIoHash Key, FIoBuffer& Data, bool& bAdded);
+	int32 RemovePending(FCacheEntryList& OutPending, uint32 MaxSize);
+	void InsertPersisted(FCacheEntryList&& InPersisted, const uint64 CursorPos);
+	void RemovePersisted(const uint64 RequiredSize);
+	uint64 GetPendingBytes() const { return TotalPendingBytes; }
+	FIoStatus Load(const FString& FilePath, uint64& OutCursorPos);
+	FIoStatus Save(const FString& FilePath, const uint64 CursorPos);
+
+private:
+	FCacheEntryList Pending;
+	FCacheEntryList Persisted;
+	TMap<FIoHash, TUniquePtr<FCacheEntry>> Lookup;
+	mutable FCriticalSection Cs;
+	std::atomic_uint64_t TotalPendingBytes = 0;
+	std::atomic_uint64_t TotalPersistedBytes = 0;
+	uint64 MaxPersistedBytes = 0;
+	uint64 MaxPendingBytes = 0;
+};
+
+void FCacheMap::SetCacheLimits(uint64 InMaxPendingBytes, uint64 InMaxPersistedBytes)
+{
+	MaxPendingBytes = InMaxPendingBytes;
+	MaxPersistedBytes = InMaxPersistedBytes;
+}
+
+void FCacheMap::Reset()
+{
+	Pending.Reset();
+	Persisted.Reset();
+	Lookup.Reset();
+	TotalPendingBytes = 0;
+	TotalPersistedBytes = 0;
+}
+
+bool FCacheMap::Contains(FIoHash Key) const
+{
+	FScopeLock _(&Cs);
+	return Lookup.Contains(Key);
+}
+
+bool FCacheMap::Get(const FIoHash& Key, FCacheEntry& OutEntry) const
+{
+	FScopeLock _(&Cs);
+
+	if (const TUniquePtr<FCacheEntry>* Entry = Lookup.Find(Key))
+	{
+		OutEntry = *(Entry->Get());
+		return true;
+	}
+
+	return false;
+}
+
+bool FCacheMap::InsertPending(FIoHash Key, FIoBuffer& Data, bool& bAdded)
+{
+	check(Data.GetSize() > 0);
+
+	FScopeLock _(&Cs);
+
+	bAdded = false;
+	if (TotalPendingBytes + Data.GetSize() > MaxPendingBytes)
+	{
+		GStatistics.Add(FStatistics::EDatum::PutRejectCount, 1);
+		return false;
+	}
+
+	if (Lookup.Contains(Key))
+	{
+		GStatistics.Add(FStatistics::EDatum::PutExistingCount, 1);
+		return true;
+	}
+
+	TUniquePtr<FCacheEntry>& Entry = Lookup.FindOrAdd(Key);
+	Entry = MakeUnique<FCacheEntry>();
+	Entry->Key = Key;
+	Entry->Data = FIoBuffer(Data);
+	Entry->State = ECacheEntryState::Pending;
+
+	Pending.AddTail(Entry.Get());
+
+	TotalPendingBytes += Data.GetSize();
+
+	GStatistics.Add(FStatistics::EDatum::PendingBytes, Data.GetSize());
+	GStatistics.Add(FStatistics::EDatum::PutCount, 1);
+
+	return bAdded = true;
+}
+
+int32 FCacheMap::RemovePending(FCacheEntryList& OutPending, uint32 MaxSize)
+{
+	FScopeLock _(&Cs);
+
+	if (Pending.IsEmpty())
+	{
+		return 0;
+	}
+
+	uint32 ReturnSize = 0;
+	for (FCacheEntryList::TIterator Iter = Pending.begin(); Iter != Pending.end();)
+	{
+		FCacheEntry& Entry = *Iter;
+		++Iter;
+
+		if (MaxSize < ReturnSize + Entry.Data.GetSize())
+		{
+			break;
+		}
+
+		ReturnSize += Entry.Data.GetSize();
+
+		Pending.Remove(&Entry);
+		OutPending.AddTail(&Entry);
+	}
+
+	TotalPendingBytes -= ReturnSize;
+	GStatistics.Set(FStatistics::EDatum::PendingBytes, TotalPendingBytes);
+
+	return ReturnSize;
+}
+
+void FCacheMap::InsertPersisted(FCacheEntryList&& InPersisted, const uint64 CursorPos)
+{
+	check(InPersisted.GetTail());
+	FCacheEntry& Tail = *InPersisted.GetTail();
+	const uint64 ExpectedCursosPos = (Tail.SerialOffset + Tail.SerialSize) % MaxPersistedBytes;
+	check(ExpectedCursosPos == CursorPos);
+
+	FScopeLock _(&Cs);
+
+	uint64 PersistedBytes = 0;
+	for (FCacheEntry& Entry : InPersisted)
+	{
+		check(Entry.SerialSize > 0);
+		Entry.State = ECacheEntryState::Persisted;
+		Entry.Data = FIoBuffer();
+		PersistedBytes += Entry.SerialSize;
+	}
+
+	Persisted.AddTail(MoveTemp(InPersisted));
+	TotalPersistedBytes += PersistedBytes;
+	GStatistics.Add(FStatistics::EDatum::CachedBytes, PersistedBytes);
+}
+
+void FCacheMap::RemovePersisted(const uint64 RequiredSize)
+{
+	FScopeLock _(&Cs);
+
+	uint64 RemovedBytes = 0;
+	for (;;)
+	{
+		if ((TotalPersistedBytes - RemovedBytes) + RequiredSize < MaxPersistedBytes)
+		{
+			break;
+		}
+
+		FCacheEntry* Entry = Persisted.PopHead();
+		check(Entry);
+		const uint64 EntrySize = Entry->SerialSize;
+		const FIoHash Key = Entry->Key;
+		Lookup.Remove(Key);
+		RemovedBytes += EntrySize;
+	}
+
+	TotalPersistedBytes -= RemovedBytes;
+	GStatistics.Add(FStatistics::EDatum::CachedBytes, -int64(RemovedBytes));
+}
+
+FIoStatus FCacheMap::Load(const FString& FilePath, uint64& OutCursorPos)
+{
+	FCacheFileToc CacheFileToc;
+	if (FIoStatus Status = CacheFileToc.Load(FilePath, OutCursorPos); !Status.IsOk())
+	{
+		return Status;
+	}
+
+	TConstArrayView<FCacheFileToc::FTocEntry> TocEntries = CacheFileToc.GetEntries();
+	for (const FCacheFileToc::FTocEntry& Entry : TocEntries)
+	{
+		check(!Lookup.Contains(Entry.Key));
+
+		TUniquePtr<FCacheEntry>& CacheEntry = Lookup.FindOrAdd(Entry.Key);
+		CacheEntry = MakeUnique<FCacheEntry>();
+
+		CacheEntry->Key = Entry.Key;
+		CacheEntry->Hash = Entry.Hash;
+		CacheEntry->SerialOffset = Entry.SerialOffset;
+		CacheEntry->SerialSize = Entry.SerialSize;
+		CacheEntry->State = ECacheEntryState::Persisted;
+		Persisted.AddTail(CacheEntry.Get());
+		TotalPersistedBytes += Entry.SerialSize;
+	}
+
+	if (FCacheEntry* Tail = Persisted.GetTail())
+	{
+		const uint64 ExpectedCursosPos = (Tail->SerialOffset + Tail->SerialSize) % MaxPersistedBytes;
+		check(ExpectedCursosPos == OutCursorPos);
+	}
+
+	GStatistics.Set(FStatistics::EDatum::CachedBytes, TotalPersistedBytes);
+
+	return FIoStatus(EIoErrorCode::Ok);
+}
+
+FIoStatus FCacheMap::Save(const FString& FilePath, const uint64 CursorPos)
+{
+	FCacheFileToc CacheFileToc;
+
+	if (FCacheEntry* Tail = Persisted.GetTail())
+	{
+		const uint64 ExpectedCursosPos = (Tail->SerialOffset + Tail->SerialSize) % MaxPersistedBytes;
+		check(ExpectedCursosPos == CursorPos);
+	}
+
+	for (FCacheEntry& Entry : Persisted)
+	{
+		CacheFileToc.AddEntry(Entry.Key, Entry.Hash, Entry.SerialOffset, Entry.SerialSize);
+	}
+
+	return CacheFileToc.Save(FilePath, CursorPos);
+}
+
+
+
+// {{{1 governor
+
+////////////////////////////////////////////////////////////////////////////////
+class FGovernorExternal
+{
+public:
+				FGovernorExternal();
+	void		Set(...) {}
+	uint32		TickAllowance();
+	void		Return(uint32) {}
+
+private:
+	int64		CycleThreshold;
+	int64		CycleLast;
+};
+
+////////////////////////////////////////////////////////////////////////////////
+FGovernorExternal::FGovernorExternal()
+{
+	CycleThreshold = int64(1.0 / FPlatformTime::GetSecondsPerCycle());
+	CycleThreshold = (CycleThreshold * 3) / 4;
+	CycleLast = FPlatformTime::Cycles() - CycleThreshold;
+}
+
+////////////////////////////////////////////////////////////////////////////////
+uint32 FGovernorExternal::TickAllowance()
+{
+	// We'll only check the platform layer for our allowance every now and again
+	int64 Cycle = FPlatformTime::Cycles64();
+	if (Cycle - CycleLast < CycleThreshold)
+	{
+		return 0;
+	}
+
+	CycleLast = Cycle;
+
+	IPlatformFile& Ipf = IPlatformFile::GetPlatformPhysical();
+	return Ipf.GetAllowedBytesToWriteThrottledStorage();
+}
+
+
+
+////////////////////////////////////////////////////////////////////////////////
+class FGovernorInternal
+{
+public:
+	void		Set(uint32 InAllowance, uint32 InOps, uint32 Seconds);
+	uint32		TickAllowance();
+	void		Return(uint32 LeftOver);
+
+private:
+	uint32		TickAllowance(int64 Cycles);
+	int64		CycleThreshold;
+	int64		CyclePrev;
+	uint32		Allowance;
+	uint32		RunOff;
+};
+
+////////////////////////////////////////////////////////////////////////////////
+void FGovernorInternal::Set(uint32 InAllowance, uint32 Ops, uint32 Seconds)
+{
+	CycleThreshold = int64(1.0 / FPlatformTime::GetSecondsPerCycle());
+
+	CycleThreshold = (CycleThreshold * Seconds) / Ops;
+	Allowance = InAllowance / Ops;
+
+	CyclePrev = FPlatformTime::Cycles() - CycleThreshold;
+}
+
+////////////////////////////////////////////////////////////////////////////////
+uint32 FGovernorInternal::TickAllowance()
+{
+	int64 Cycle = FPlatformTime::Cycles64();
+	uint32 Ret = TickAllowance(Cycle) + RunOff;
+	RunOff = 0;
+	return Ret;
+}
+
+////////////////////////////////////////////////////////////////////////////////
+uint32 FGovernorInternal::TickAllowance(int64 Cycle)
+{
+	int64 CycleDelta = FMath::Max(0ll, Cycle - CyclePrev);
+	if (CycleDelta < CycleThreshold)
+	{
+		return 0;
+	}
+
+	// A crude guard against runaway.
+	for (; CycleDelta > CycleThreshold; CycleDelta -= CycleThreshold);
+	CyclePrev = Cycle + CycleDelta;
+
+	return Allowance;
+}
+
+////////////////////////////////////////////////////////////////////////////////
+void FGovernorInternal::Return(uint32 LeftOver)
+{
+	// Roughly keep to some arbitrary limit so as to avoid excessive growth
+	uint32 OnePointFive = Allowance + (Allowance >> 1);
+	RunOff = FMath::Min(OnePointFive, LeftOver);
+}
+
+
+
+////////////////////////////////////////////////////////////////////////////////
+#if !defined(UE_USE_PLATFORM_GOVERNOR)
+#	define UE_USE_PLATFORM_GOVERNOR 1
+#endif
+
+#if UE_USE_PLATFORM_GOVERNOR
+	using FGovernor = FGovernorExternal;
+#else
+	using FGovernor = FGovernorInternal;
+#endif
+
+
+
+// {{{1 file-impl
+
+////////////////////////////////////////////////////////////////////////////////
+class FFileIoCache final
+	: public FRunnable
+	, public IIoCache
+{
+public:
+	explicit FFileIoCache(const FFileIoCacheConfig& Config);
+	virtual ~FFileIoCache();
+
+	virtual bool ContainsChunk(const FIoHash& Key) const override;
+	virtual TTask<TIoStatusOr<FIoBuffer>> Get(const FIoHash& Key, const FIoReadOptions& Options, const FIoCancellationToken* CancellationToken) override;
+	virtual FIoStatus Put(const FIoHash& Key, FIoBuffer& Data) override;
+
+	// Runnable
+	virtual bool Init() override
+	{ 
+		return true;
+	}
+
+	virtual void Stop() override
+	{
+		bStopRequested = true;
+		TickWriterEvent->Trigger();
+	}
+
+private:
+	virtual uint32 Run() override;
+	void Initialize();
+	void Shutdown();
+	void FileWriterThreadInner();
+
+	FFileIoCacheConfig CacheConfig;
+	FCacheMap CacheMap;
+	TUniquePtr<FRunnableThread> WriterThread;
+	FEventRef TickWriterEvent;
+	FString CacheFilePath;
+	uint64 WriteCursorPos = 0;
+	std::atomic_bool bStopRequested{false};
+	FGovernor Governor;
+};
+
+FFileIoCache::FFileIoCache(const FFileIoCacheConfig& Config)
+	: CacheConfig(Config)
+{
+	Governor.Set(Config.WriteRate.Allowance, Config.WriteRate.Ops, Config.WriteRate.Seconds);
+
+	CacheMap.SetCacheLimits(Config.MemoryQuota, Config.DiskQuota);
+
+	Initialize();
+}
+
+FFileIoCache::~FFileIoCache()
+{
+	Shutdown();
+}
+
+bool FFileIoCache::ContainsChunk(const FIoHash& Key) const
+{
+	return CacheMap.Contains(Key);
+}
+
+TTask<TIoStatusOr<FIoBuffer>> FFileIoCache::Get(const FIoHash& Key, const FIoReadOptions& Options, const FIoCancellationToken* CancellationToken)
+{
+	return Launch(UE_SOURCE_LOCATION,
+		[this, Key, Options, CancellationToken]()
+		{
+			FCacheEntry Entry;
+			if (!CacheMap.Get(Key, Entry))
+			{
+				return TIoStatusOr<FIoBuffer>(FIoStatus(EIoErrorCode::Unknown));
+			}
+
+			if (CancellationToken && CancellationToken->IsCancelled())
+			{
+				return TIoStatusOr<FIoBuffer>(FIoStatus(EIoErrorCode::Cancelled));
+			}
+
+			if (Entry.Data.GetSize() > 0)
+			{
+				TRACE_CPUPROFILER_EVENT_SCOPE(FFileIoCache::ReadPendingEntry);
+
+				const uint64 ReadOffset = Options.GetOffset();
+				const uint64 ReadSize = FMath::Min(Options.GetSize(), Entry.Data.GetSize());
+				FIoBuffer Buffer = Options.GetTargetVa() ? FIoBuffer(FIoBuffer::Wrap, Options.GetTargetVa(), ReadSize) : FIoBuffer(ReadSize);
+				Buffer.GetMutableView().CopyFrom(Entry.Data.GetView().RightChop(ReadOffset));
+				GStatistics.Add(FStatistics::EDatum::GetCount, 1);
+
+				return TIoStatusOr<FIoBuffer>(Buffer);
+			}
+			else
+			{
+				TRACE_CPUPROFILER_EVENT_SCOPE(FFileIoCache::ReadPersistedEntry);
+
+				check(Entry.SerialSize > 0);
+				check(Entry.Hash != FIoHash::Zero);
+
+				const uint64 ReadSize = FMath::Min(Options.GetSize(), Entry.SerialSize);
+				const uint64 ReadOffset = Entry.SerialOffset + Options.GetOffset();
+				FIoBuffer Buffer = Options.GetTargetVa() ? FIoBuffer(FIoBuffer::Wrap, Options.GetTargetVa(), ReadSize) : FIoBuffer(ReadSize);
+
+				IPlatformFile& Ipf = IPlatformFile::GetPlatformPhysical();
+				if (TUniquePtr<IFileHandle> FileHandle(Ipf.OpenRead(*CacheFilePath, true)); FileHandle.IsValid())
+				{
+					UE_LOG(LogIoCache, VeryVerbose, TEXT("Read chunk, Key='%s', Hash='%s', File='%s', Offset='%llu', Size='%llu'"),
+						*LexToString(Entry.Key), *LexToString(Entry.Hash), *CacheFilePath, Entry.SerialOffset, Entry.SerialSize);
+					
+					FileHandle->Seek(int64(ReadOffset));
+					FileHandle->Read(Buffer.GetData(), ReadSize);
+
+					const FIoHash& ExpectedHash = Entry.Hash;
+					const FIoHash Hash = FIoHash::HashBuffer(Buffer.GetView());
+					if (Hash == ExpectedHash) 
+					{
+						GStatistics.Add(FStatistics::EDatum::GetCount, 1);
+						GStatistics.Add(FStatistics::EDatum::ReadBytes, ReadSize);
+						return TIoStatusOr<FIoBuffer>(Buffer);
+					}
+
+					GStatistics.Add(FStatistics::EDatum::ErrorCount, 1);
+					UE_LOG(LogIoCache, Verbose, TEXT("Read chunk failed, hash mismatch, Key='%s', Hash='%s', ExpectedHash='%s', File='%s', Offset='%llu', Size='%llu'"),
+						*LexToString(Entry.Key), *LexToString(Hash), *LexToString(ExpectedHash), *CacheFilePath, ReadOffset, ReadSize);
+
+					return TIoStatusOr<FIoBuffer>(FIoStatus(EIoErrorCode::ReadError));
+				}
+				else
+				{
+					GStatistics.Add(FStatistics::EDatum::ErrorCount, 1);
+					UE_LOG(LogIoCache, Warning, TEXT("Read chunk failed, unable to open cache file '%s' for reading"), *CacheFilePath);
+					return TIoStatusOr<FIoBuffer>(FIoStatus(EIoErrorCode::ReadError));
+				}
+			}
+		});
+}
+	
+FIoStatus FFileIoCache::Put(const FIoHash& Key, FIoBuffer& Data)
+{
+	bool bAdded = false;
+	if (CacheMap.InsertPending(Key, Data, bAdded))
+	{
+		if (bAdded)
+		{
+			TickWriterEvent->Trigger();
+		}
+		return FIoStatus::Ok;
+	}
+
+	return FIoStatus(EIoErrorCode::Unknown);
+}
+
+void FFileIoCache::Initialize()
+{
+	UE_LOG(LogIoCache, Log,
+		TEXT("Initializing file I/O cache, disk size %lluB, memory size %lluB"), CacheConfig.DiskQuota, CacheConfig.MemoryQuota);
+
+	WriterThread.Reset(FRunnableThread::Create(this, TEXT("Ias.FileCache"), 0, TPri_BelowNormal));
+
+	const FString CacheDir = FPaths::ProjectPersistentDownloadDir() / TEXT("IoCache");
+	const FString CacheTocPath = CacheDir / TEXT("cache.utoc");
+	CacheFilePath = CacheDir / TEXT("cache.ucas");
+	WriteCursorPos = 0;
+
+	IFileManager& FileMgr = IFileManager::Get();
+
+	if (CacheConfig.DropCache)
+	{
+		FileMgr.Delete(*CacheTocPath);
+		FileMgr.Delete(*CacheFilePath);
+		return;
+	}
+	
+	FIoStatus CacheStatus(EIoErrorCode::Unknown);
+	if (FileMgr.FileExists(*CacheTocPath))
+	{
+		if (FParse::Param(FCommandLine::Get(), TEXT("ClearIoCache")))
+		{
+			UE_LOG(LogIoCache, Log, TEXT("Deleting cache file '%s'"), *CacheFilePath);
+			FileMgr.Delete(*CacheFilePath);
+		}
+		else
+		{
+			CacheStatus = CacheMap.Load(CacheTocPath, WriteCursorPos);
+			if (CacheStatus.IsOk())
+			{
+				check(WriteCursorPos != ~uint64(0));
+				
+				UE_LOG(LogIoCache, Log, TEXT("Loaded TOC '%s'"), *CacheTocPath);
+				if (FileMgr.FileExists(*CacheFilePath))
+				{
+					//TODO: Integrity check?
+				}
+				else
+				{
+					UE_LOG(LogIoCache, Warning, TEXT("Failed to open cache file ''"), *CacheFilePath);
+					CacheStatus = FIoStatus(EIoErrorCode::FileNotOpen);
+				}
+			}
+			else
+			{
+				UE_LOG(LogIoCache, Warning, TEXT("Failed to load TOC '%s'"), *CacheTocPath);
+			}
+		}
+	}
+
+	if (!CacheStatus.IsOk())
+	{
+		WriteCursorPos = 0;
+		CacheMap.Reset();
+		FileMgr.Delete(*CacheFilePath);
+
+		if (!FileMgr.DirectoryExists(*CacheDir))
+		{
+			FileMgr.MakeDirectory(*CacheDir, true);
+		}
+	}
+}
+
+void FFileIoCache::Shutdown()
+{
+	if (bStopRequested)
+	{
+		return;
+	}
+
+	bStopRequested = true;
+	TickWriterEvent->Trigger();
+	WriterThread->Kill();
+
+	const FString CacheTocPath = FPaths::ProjectPersistentDownloadDir() / TEXT("IoCache") / TEXT("cache.utoc");
+	UE_LOG(LogIoCache, Log, TEXT("Saving TOC '%s'"), *CacheTocPath);
+	CacheMap.Save(CacheTocPath, WriteCursorPos);
+
+	CacheMap.Reset();
+}
+
+uint32 FFileIoCache::Run()
+{
+	while (!bStopRequested)
+	{
+		FileWriterThreadInner();
+		TickWriterEvent->Wait(10);
+	}
+
+	return 0;
+}
+
+void FFileIoCache::FileWriterThreadInner()
+{
+	static uint32 _TickCount = 0;
+	++_TickCount;
+
+	// Update the rate limiting and see how much we are allowed to write
+	uint32 WriteAllowance = Governor.TickAllowance();
+	if (!WriteAllowance)
+	{
+		return;
+	}
+
+	// Collect pending entries up to the write allowance
+	FCacheEntryList Entries;
+	int32 PendingSize = CacheMap.RemovePending(Entries, WriteAllowance);
+	if (PendingSize <= 0)
+	{
+		return;
+	}
+
+	if (WriteAllowance -= PendingSize)
+	{
+		Governor.Return(WriteAllowance);
+	}
+
+	// Wrap the write cursor if we'd end up writing over the end
+	if (WriteCursorPos + PendingSize > CacheConfig.DiskQuota)
+	{
+		WriteCursorPos = 0;
+	}
+
+	// Copy data to be cached into a buffer
+	auto* Buffer = (uint8*)FMemory::Malloc(PendingSize);
+	auto* Cursor = (uint8*)Buffer;
+	for (FCacheEntry& Entry : Entries)
+	{
+		const FMemoryView& View = Entry.Data.GetView();
+
+		Entry.SerialOffset = WriteCursorPos + uint32(ptrdiff_t(Cursor - Buffer));
+		Entry.SerialSize = View.GetSize();
+		Entry.Hash = FIoHash::HashBuffer(View);
+		Entry.State = ECacheEntryState::Writing;
+
+		::memcpy(Cursor, View.GetData(), View.GetSize());
+		Cursor += View.GetSize();
+	}
+
+	// Drop buffers.
+	for (FCacheEntry& Entry : Entries)
+	{
+		FIoBuffer& Data = Entry.Data;
+		Data = FIoBuffer();
+	}
+
+	// Open cache file and write the block to it.
+	UE_LOG(LogIoCache, VeryVerbose, TEXT("Write; cursor:%llu size:%d"), WriteCursorPos, PendingSize);
+	TRACE_CPUPROFILER_EVENT_SCOPE(FFileIoCache::WriteCacheEntry);
+
+	IPlatformFile& Ipf = IPlatformFile::GetPlatformPhysical();
+	TUniquePtr<IFileHandle> FileHandle(Ipf.OpenWrite(*CacheFilePath, true, true));
+	if (!FileHandle.IsValid())
+	{
+		UE_LOG(LogIoCache, Warning, TEXT("Write chunks failed, unable to open file '%s' for writing"), *CacheFilePath);
+		return;
+	}
+	FileHandle->Seek(WriteCursorPos);
+	FileHandle->Write(Buffer, PendingSize);
+	FileHandle->Flush();
+	WriteCursorPos += PendingSize;
+
+	CacheMap.InsertPersisted(MoveTemp(Entries), WriteCursorPos);
+
+	CacheMap.RemovePersisted(PendingSize);
+
+	Entries.Reset();
+}
+
+// }}}
+
+} // namespace UE::IO::Private
+
+TUniquePtr<IIoCache> MakeFileIoCache(const FFileIoCacheConfig& Config)
+{
+	return MakeUnique<UE::IO::Private::FFileIoCache>(Config);
+}
+
+// vim: foldlevel=1
diff --git a/Engine/Source/Runtime/Experimental/IoStoreOnDemand/Private/FileIoCache.h b/Engine/Source/Runtime/Experimental/IoStoreOnDemand/Private/FileIoCache.h
new file mode 100644
--- /dev/null
+++ b/Engine/Source/Runtime/Experimental/IoStoreOnDemand/Private/FileIoCache.h
@@ -0,0 +1,51 @@
+// Copyright Epic Games, Inc. All Rights Reserved.
+
+#pragma once
+
+#include "CancellationToken.h"
+#include "IO/IoStatus.h"
+#include "Memory/MemoryFwd.h"
+#include "Tasks/Task.h"
+#include "Templates/UniquePtr.h"
+
+class FIoBuffer;
+class FIoReadOptions;
+struct FIoHash;
+
+DECLARE_LOG_CATEGORY_EXTERN(LogIoCache, Log, All);
+
+/** Cache for binary blobs with a 20 byte cache key. */
+class IIoCache
+{
+public:
+	virtual ~IIoCache() = default;
+
+	/** Returns whether the specified cache key is present in the cache. */
+	virtual bool ContainsChunk(const FIoHash& Key) const = 0;
+
+	/** Get the chunk associated with the specified cache key. */
+	virtual UE::Tasks::TTask<TIoStatusOr<FIoBuffer>> Get(
+		const FIoHash& Key,
+		const FIoReadOptions& Options,
+		const FIoCancellationToken* CancellationToken) = 0;
+
+	/** Insert a new chunk into the cache. */
+	virtual FIoStatus Put(const FIoHash& Key, FIoBuffer& Data) = 0;
+};
+
+struct FFileIoCacheConfig
+{
+	struct FRate
+	{
+		uint32 Allowance = ~0u;
+		uint32 Ops = 256;
+		uint32 Seconds = 1;
+	};
+
+	uint64 DiskQuota = 1ull << 30;
+	uint32 MemoryQuota = 3 << 20;
+	FRate WriteRate;
+	bool DropCache = false;
+};
+
+TUniquePtr<IIoCache> MakeFileIoCache(const FFileIoCacheConfig& Config);
diff --git a/Engine/Source/Runtime/Experimental/IoStoreOnDemand/Private/IoStoreOnDemand.cpp b/Engine/Source/Runtime/Experimental/IoStoreOnDemand/Private/IoStoreOnDemand.cpp
--- a/Engine/Source/Runtime/Experimental/IoStoreOnDemand/Private/IoStoreOnDemand.cpp
+++ b/Engine/Source/Runtime/Experimental/IoStoreOnDemand/Private/IoStoreOnDemand.cpp
@@ -5,7 +5,7 @@
 #include "EncryptionKeyManager.h"
 #include "LatencyInjector.h"
 
-#include "FileCache.h"
+#include "FileIoCache.h"
 #include "HAL/PlatformMisc.h"
 #include "HAL/FileManager.h"
 #include "Misc/Base64.h"
@@ -73,7 +73,7 @@ static bool ParseEncryptionKeyParam(const FString& Param, FGuid& OutKeyGuid, FAE
 }
 
 ////////////////////////////////////////////////////////////////////////////////
-static FFileIoCacheConfig GetFileCacheConfig(const TCHAR* CommandLine)
+static FFileIoCacheConfig GetFileIoCacheConfig(const TCHAR* CommandLine)
 {
 	FFileIoCacheConfig Ret;
 
@@ -449,7 +449,7 @@ void FIoStoreOnDemandModule::StartupModule()
 	Endpoint.EndpointType = UE::EOnDemandEndpointType::CDN;
 
 	TSharedPtr<IIoCache> Cache;
-	if (FFileIoCacheConfig Config = GetFileCacheConfig(CommandLine); Config.DiskQuota > 0)
+	if (FFileIoCacheConfig Config = GetFileIoCacheConfig(CommandLine); Config.DiskQuota > 0)
 	{
 		Cache = MakeShareable(MakeFileIoCache(Config).Release());
 	}
diff --git a/Engine/Source/Runtime/Experimental/IoStoreOnDemand/Private/OnDemandIoDispatcherBackend.cpp b/Engine/Source/Runtime/Experimental/IoStoreOnDemand/Private/OnDemandIoDispatcherBackend.cpp
--- a/Engine/Source/Runtime/Experimental/IoStoreOnDemand/Private/OnDemandIoDispatcherBackend.cpp
+++ b/Engine/Source/Runtime/Experimental/IoStoreOnDemand/Private/OnDemandIoDispatcherBackend.cpp
@@ -7,7 +7,7 @@
 #include "Containers/StringView.h"
 #include "CoreHttp/Client.h"
 #include "EncryptionKeyManager.h"
-#include "FileCache.h"
+#include "FileIoCache.h"
 #include "HAL/Event.h"
 #include "HAL/Platform.h"
 #include "HAL/PlatformTime.h"
