Change IAS' log category to LogIas.

#rnx

diff --git a/Engine/Source/Runtime/Experimental/IoStoreOnDemand/Private/IoStoreOnDemand.cpp b/Engine/Source/Runtime/Experimental/IoStoreOnDemand/Private/IoStoreOnDemand.cpp
--- a/Engine/Source/Runtime/Experimental/IoStoreOnDemand/Private/IoStoreOnDemand.cpp
+++ b/Engine/Source/Runtime/Experimental/IoStoreOnDemand/Private/IoStoreOnDemand.cpp
@@ -17,7 +17,7 @@
 #include "Serialization/LargeMemoryWriter.h"
 #include "String/LexFromString.h"
 
-DEFINE_LOG_CATEGORY(LogIoStoreOnDemand);
+DEFINE_LOG_CATEGORY(LogIas);
 
 namespace UE::IO::Private
 {
