Renamed IAS' log category to LogIas

#rnx

diff --git a/Engine/Source/Runtime/Experimental/IoStoreOnDemand/Private/OnDemandIoDispatcherBackend.cpp b/Engine/Source/Runtime/Experimental/IoStoreOnDemand/Private/OnDemandIoDispatcherBackend.cpp
--- a/Engine/Source/Runtime/Experimental/IoStoreOnDemand/Private/OnDemandIoDispatcherBackend.cpp
+++ b/Engine/Source/Runtime/Experimental/IoStoreOnDemand/Private/OnDemandIoDispatcherBackend.cpp
@@ -168,7 +168,7 @@ void FDistributionEndpoints::IssueRequests()
 			}
 
 			FResolveRequest& ResolveRequest = *Kv.Value.Get();
-			UE_LOG(LogIoStoreOnDemand, Log, TEXT("Resolving '%s' (#%d/%d)"), *ResolveRequest.DistributionUrl, ResolveRequest.RetryCount + 1, MaxAttempts);
+			UE_LOG(LogIas, Log, TEXT("Resolving '%s' (#%d/%d)"), *ResolveRequest.DistributionUrl, ResolveRequest.RetryCount + 1, MaxAttempts);
 
 			FHttpRequestPtr HttpRequest = HttpModule.Get().CreateRequest();
 			HttpRequest->SetTimeout(3.0f);
@@ -355,14 +355,14 @@ void FHttpClient::Issue(UE::HTTP::FRequest&& Request, FIoReadCallback&& Callback
 				const bool bSuccessful = StatusCode > 199 && StatusCode < 300;
 				if (const FIoBuffer& Content = Status.GetContent(); bSuccessful && Content.GetSize() > 0)
 				{
-					UE_LOG(LogIoStoreOnDemand, VeryVerbose, TEXT("%s"),
+					UE_LOG(LogIas, VeryVerbose, TEXT("%s"),
 						*WriteToString<256>(TEXT("HTTP GET - "), Url, TEXT(" ("), StatusCode, TEXT(" "), Duration, TEXT("ms "), Content.GetSize(), TEXT(" Bytes)")));
 
 					Callback(Content);
 				}
 				else
 				{
-					UE_LOG(LogIoStoreOnDemand, VeryVerbose, TEXT("%s"),
+					UE_LOG(LogIas, VeryVerbose, TEXT("%s"),
 						*WriteToString<256>(TEXT("HTTP GET - "), Url, TEXT(" ("), StatusCode, TEXT(" "), Duration, TEXT("ms)")));
 					
 					Callback(FIoStatus(EIoErrorCode::NotFound, TEXTVIEW("Invalid Content")));
@@ -371,7 +371,7 @@ void FHttpClient::Issue(UE::HTTP::FRequest&& Request, FIoReadCallback&& Callback
 			else if (FTicketStatus::EId::Error == Status.GetId())
 			{
 				const uint64 Duration = (uint64)FPlatformTime::ToMilliseconds64(FPlatformTime::Cycles64() - StartTime);
-				UE_LOG(LogIoStoreOnDemand, VeryVerbose, TEXT("%s"),
+				UE_LOG(LogIas, VeryVerbose, TEXT("%s"),
 					*WriteToString<256>(TEXT("HTTP GET - "), Url, TEXT(" ("), StatusCode, TEXT(" "), Duration, TEXTVIEW("ms)")));
 
 				Callback(FIoStatus(EIoErrorCode::ReadError, TEXTVIEW("HTTP Error")));
@@ -468,7 +468,7 @@ FOnDemandIoStore::~FOnDemandIoStore()
 void FOnDemandIoStore::AddToc(const FOnDemandEndpoint& Ep, FOnDemandToc&& Toc)
 {
 	check(Ep.IsValid());
-	UE_LOG(LogIoStoreOnDemand, Log, TEXT("Adding TOC '%s/%s'"), *Ep.ServiceUrl, *Ep.TocPath);
+	UE_LOG(LogIas, Log, TEXT("Adding TOC '%s/%s'"), *Ep.ServiceUrl, *Ep.TocPath);
 
 	TRACE_CPUPROFILER_EVENT_SCOPE(FOnDemandIoBackend::AddToc);
 
@@ -555,7 +555,7 @@ void FOnDemandIoStore::AddDeferredContainers()
 		FContainer* Container = *It;
 		if (Container->EncryptionKeyGuid.IsEmpty())
 		{
-			UE_LOG(LogIoStoreOnDemand, Log, TEXT("Mounting container '%s'"), *Container->Name);
+			UE_LOG(LogIas, Log, TEXT("Mounting container '%s'"), *Container->Name);
 			RegisteredContainers.Add(Container);
 			It.RemoveCurrent();
 		}
@@ -565,14 +565,14 @@ void FOnDemandIoStore::AddDeferredContainers()
 			ensure(FGuid::Parse(Container->EncryptionKeyGuid, KeyGuid));
 			if (const FAES::FAESKey* Key = FEncryptionKeyManager::Get().GetKey(KeyGuid))
 			{
-				UE_LOG(LogIoStoreOnDemand, Log, TEXT("Mounting container '%s'"), *Container->Name);
+				UE_LOG(LogIas, Log, TEXT("Mounting container '%s'"), *Container->Name);
 				Container->EncryptionKey = *Key;
 				RegisteredContainers.Add(Container);
 				It.RemoveCurrent();
 			}
 			else
 			{
-				UE_LOG(LogIoStoreOnDemand, Log, TEXT("Defeering container '%s', encryption key '%s' not available"), *Container->Name, *Container->EncryptionKeyGuid);
+				UE_LOG(LogIas, Log, TEXT("Defeering container '%s', encryption key '%s' not available"), *Container->Name, *Container->EncryptionKeyGuid);
 			}
 		}
 	}
@@ -812,7 +812,7 @@ class FOnDemandIoBackend final
 			FScopeLock _(&Mutex);
 
 			FBackendData& BackendData = FBackendData::Get(Request);
-			UE_LOG(LogIoStoreOnDemand, VeryVerbose, TEXT("%s"),
+			UE_LOG(LogIas, VeryVerbose, TEXT("%s"),
 				*WriteToString<256>(TEXT("Cancelling I/O request ChunkId='"), LexToString(Request->ChunkId), TEXT("' ChunkKey='"), BackendData.ChunkKey, TEXT("'")));
 
 			if (FChunkRequest** InflightRequest = Inflight.Find(BackendData.ChunkKey))
@@ -917,7 +917,7 @@ FOnDemandIoBackend::~FOnDemandIoBackend()
 
 void FOnDemandIoBackend::Initialize(TSharedRef<const FIoDispatcherBackendContext> Context)
 {
-	UE_LOG(LogIoStoreOnDemand, Log, TEXT("Initializing on demand I/O dispatcher backend"));
+	UE_LOG(LogIas, Log, TEXT("Initializing on demand I/O dispatcher backend"));
 	BackendContext = Context;
 	DistributionEndpoints.ResolveDeferredEndpoints();
 }
@@ -929,7 +929,7 @@ void FOnDemandIoBackend::Shutdown()
 		return;
 	}
 
-	UE_LOG(LogIoStoreOnDemand, Log, TEXT("Shutting down on demand I/O dispatcher backend"));
+	UE_LOG(LogIas, Log, TEXT("Shutting down on demand I/O dispatcher backend"));
 
 	bStopRequested = true;
 	TickBackendEvent->Trigger();
@@ -1154,7 +1154,7 @@ void FOnDemandIoBackend::Mount(const FOnDemandEndpoint& Endpoint)
 
 	if ((Endpoint.DistributionUrl.IsEmpty() && Endpoint.ServiceUrl.IsEmpty()) || Endpoint.TocPath.IsEmpty())
 	{
-		UE_LOG(LogIoStoreOnDemand, Error, TEXT("Trying to mount an invalid on demand endpoint"));
+		UE_LOG(LogIas, Error, TEXT("Trying to mount an invalid on demand endpoint"));
 		return;
 	}
 
@@ -1171,19 +1171,19 @@ void FOnDemandIoBackend::Mount(const FOnDemandEndpoint& Endpoint)
 			{
 				if (FIoStatus Status = MountDeferredEndpoints(DistributionUrl, SerivceUrls); !Status.IsOk())
 				{
-					UE_LOG(LogIoStoreOnDemand, Error, TEXT("Failed to add endpoint(s), reason '%s'"), *Status.ToString());
+					UE_LOG(LogIas, Error, TEXT("Failed to add endpoint(s), reason '%s'"), *Status.ToString());
 				}
 			});
 		}
 		else if (FIoStatus Status = AddToc(Endpoint); !Status.IsOk())
 		{
-			UE_LOG(LogIoStoreOnDemand, Error, TEXT("Failed to add TOC '%s/%s', reason '%s'"),
+			UE_LOG(LogIas, Error, TEXT("Failed to add TOC '%s/%s', reason '%s'"),
 				*Endpoint.ServiceUrl, *Endpoint.TocPath, *Status.ToString());
 		}
 	}
 	else
 	{
-		UE_LOG(LogIoStoreOnDemand, Log, TEXT("Mounting ZEN endpoint, Url='%s'"), *Endpoint.ServiceUrl);
+		UE_LOG(LogIas, Log, TEXT("Mounting ZEN endpoint, Url='%s'"), *Endpoint.ServiceUrl);
 	}
 }
 
@@ -1196,7 +1196,7 @@ TIoStatusOr<FOnDemandToc> FOnDemandIoBackend::GetToc(FHttpClient& HttpClient, co
 
 	for (int32 Attempt = 0, MaxAttempts = GIoDispatcherMaxHttpRetryCount; Attempt <= MaxAttempts; ++Attempt)
 	{
-		UE_LOG(LogIoStoreOnDemand, Log, TEXT("Fetching TOC '%s/%s' (#%d/%d)"), *HttpClient.ServiceUrl(), *TocPath, Attempt + 1, MaxAttempts);
+		UE_LOG(LogIas, Log, TEXT("Fetching TOC '%s/%s' (#%d/%d)"), *HttpClient.ServiceUrl(), *TocPath, Attempt + 1, MaxAttempts);
 		
 		TIoStatusOr<FOnDemandToc> Toc;
 		HttpClient.Get(Url.ToView(), [&Toc](TIoStatusOr<FIoBuffer> Response)
@@ -1211,12 +1211,12 @@ TIoStatusOr<FOnDemandToc> FOnDemandIoBackend::GetToc(FHttpClient& HttpClient, co
 				}
 				else
 				{
-					UE_LOG(LogIoStoreOnDemand, Error, TEXT("Failed loading on demand TOC from compact binary"));
+					UE_LOG(LogIas, Error, TEXT("Failed loading on demand TOC from compact binary"));
 				}
 			}
 			else
 			{
-				UE_LOG(LogIoStoreOnDemand, Error, TEXT("Failed fetching TOC, reason '%s'"), *Response.Status().ToString());
+				UE_LOG(LogIas, Error, TEXT("Failed fetching TOC, reason '%s'"), *Response.Status().ToString());
 			}
 		});
 
@@ -1253,7 +1253,7 @@ FIoStatus FOnDemandIoBackend::AddToc(const FOnDemandEndpoint& Endpoint)
 	IoStore->AddToc(Endpoint, Toc.ConsumeValueOrDie());
 
 	UE_CLOG(Endpoint.ServiceUrl !=  HttpClient->ServiceUrl(),
-		LogIoStoreOnDemand, Fatal, TEXT("Fetching on demand content from multiple endpoints are currently not supported"));
+		LogIas, Fatal, TEXT("Fetching on demand content from multiple endpoints are currently not supported"));
 
 	return FIoStatus::Ok;
 }
diff --git a/Engine/Source/Runtime/Experimental/IoStoreOnDemand/Public/IO/IoStoreOnDemand.h b/Engine/Source/Runtime/Experimental/IoStoreOnDemand/Public/IO/IoStoreOnDemand.h
--- a/Engine/Source/Runtime/Experimental/IoStoreOnDemand/Public/IO/IoStoreOnDemand.h
+++ b/Engine/Source/Runtime/Experimental/IoStoreOnDemand/Public/IO/IoStoreOnDemand.h
@@ -18,7 +18,7 @@ class IIoStoreWriter;
 struct FIoContainerSettings;
 struct FIoStoreWriterSettings;
 
-UE_API DECLARE_LOG_CATEGORY_EXTERN(LogIoStoreOnDemand, VeryVerbose, All);
+UE_API DECLARE_LOG_CATEGORY_EXTERN(LogIas, VeryVerbose, All);
 
 namespace UE
 {
