geometry collection PullFromPhysicsState needs to check both NextData and PrevData for changes to a particle's non-interpolatable state

diff --git a/Engine/Source/Runtime/Experimental/Chaos/Private/PhysicsProxy/GeometryCollectionPhysicsProxy.cpp b/Engine/Source/Runtime/Experimental/Chaos/Private/PhysicsProxy/GeometryCollectionPhysicsProxy.cpp
--- a/Engine/Source/Runtime/Experimental/Chaos/Private/PhysicsProxy/GeometryCollectionPhysicsProxy.cpp
+++ b/Engine/Source/Runtime/Experimental/Chaos/Private/PhysicsProxy/GeometryCollectionPhysicsProxy.cpp
@@ -1,6 +1,5 @@
 // Copyright Epic Games, Inc. All Rights Reserved.
 
-
 #include "PhysicsProxy/GeometryCollectionPhysicsProxy.h"
 #include "PhysicsProxy/FieldSystemProxyHelper.h"
 
@@ -59,6 +58,14 @@
 #define GC_PHYSICSPROXY_CHECK_FOR_NAN(Vec)
 #endif
 
+namespace
+{
+	const FName LinearVelocityAttributeName = "LinearVelocity";
+	const FName AngularVelocitiesAttributeName = "AngularVelocity";
+	const FName InternalClusterParentTypeArrayAttributeName = "InternalClusterParentTypeArray";
+	const FName AnimateTransformAttributeName = "AnimateTransformAttribute";
+}
+
 namespace Chaos{
 	extern int32 AccelerationStructureSplitStaticAndDynamic;
 	extern int32 AccelerationStructureIsolateQueryOnlyObjects;
@@ -3374,128 +3381,147 @@ private:
 
 };
 
-// Called from FPhysScene_ChaosInterface::SyncBodies(), NOT the solver.
-bool FGeometryCollectionPhysicsProxy::PullFromPhysicsState(const Chaos::FDirtyGeometryCollectionData& PullData, const int32 SolverSyncTimestamp, const Chaos::FDirtyGeometryCollectionData* NextPullData, const Chaos::FRealSingle* Alpha)
+bool FGeometryCollectionPhysicsProxy::PullNonInterpolatableDataFromSinglePhysicsState(const Chaos::FDirtyGeometryCollectionData& BufferData, bool bForcePullXRVW, const TBitArray<>* Seen)
 {
-	if(IsObjectDeleting) return false;
-
-	/**
-	 * CONTEXT: GAMETHREAD (Read Locked)
-	 * Perform a similar operation to Sync, but take the data from a gamethread-safe buffer. This will be called
-	 * from the game thread when it cannot sync to the physics thread. The simulation is very likely to be running
-	 * when this happens so never read any physics thread data here!
-	 *
-	 * Note: A read lock will have been acquired for this - so the physics thread won't force a buffer flip while this
-	 * sync is ongoing
-	 *
-	 */
-
-	// earlier exit
-	const bool HasResults = (PullData.Results().GetNumEntries() > 0);
-	const bool HasNextResults = (NextPullData && NextPullData->Results().GetNumEntries() > 0);
-	if (!HasResults && !HasNextResults)
+	const FGeometryCollectionResults& CurrentResults = BufferData.Results();
+	const bool bHasResults = CurrentResults.GetNumEntries() > 0;
+	if (!bHasResults)
 	{
 		return false;
 	}
 
-	const FGeometryCollectionResults& CurrentResults = NextPullData ? NextPullData->Results() : PullData.Results();
-
-	const int32 NumTransforms = GameThreadCollection.Transform.Num();
-	const bool bNeedInterpolation = (NextPullData != nullptr);
-
 	bool bIsCollectionDirty = false;
 
-	if (NumTransforms > 0)
+	TManagedArray<FVector3f>* LinearVelocities = GameThreadCollection.FindAttributeTyped<FVector3f>(LinearVelocityAttributeName, FTransformCollection::TransformGroup);
+	TManagedArray<FVector3f>* AngularVelocities = GameThreadCollection.FindAttributeTyped<FVector3f>(AngularVelocitiesAttributeName, FTransformCollection::TransformGroup);
+	TManagedArray<uint8>* InternalClusterParentTypeArray = GameThreadCollection.FindAttributeTyped<uint8>(InternalClusterParentTypeArrayAttributeName, FTransformCollection::TransformGroup);
+	const TManagedArray<bool>* AnimationsActive = GameThreadCollection.FindAttribute<bool>(AnimateTransformAttributeName, FGeometryCollection::TransformGroup);
+
+	// first step : process the non values that do not need to be interpolated 
+	for (int32 EntryIndex = 0; EntryIndex < CurrentResults.GetNumEntries(); EntryIndex++)
 	{
-		static const FName LinearVelocityAttributeName = "LinearVelocity";
-		static const FName AngularVelocitiesAttributeName = "AngularVelocity";
-		static const FName InternalClusterParentTypeArrayAttributeName = "InternalClusterParentTypeArray";
-		static const FName AnimateTransformAttributeName = "AnimateTransformAttribute";
+		const FGeometryCollectionResults::FStateData& StateData = CurrentResults.GetState(EntryIndex);
+		const int32 TransformGroupIndex = StateData.TransformIndex;
 
-		TManagedArray<FVector3f>* LinearVelocities = GameThreadCollection.FindAttributeTyped<FVector3f>(LinearVelocityAttributeName, FTransformCollection::TransformGroup);
-		TManagedArray<FVector3f>* AngularVelocities = GameThreadCollection.FindAttributeTyped<FVector3f>(AngularVelocitiesAttributeName, FTransformCollection::TransformGroup);
-		TManagedArray<uint8>* InternalClusterParentTypeArray = GameThreadCollection.FindAttributeTyped<uint8>(InternalClusterParentTypeArrayAttributeName, FTransformCollection::TransformGroup);
-		const TManagedArray<bool>* AnimationsActive = GameThreadCollection.FindAttribute<bool>(AnimateTransformAttributeName, FGeometryCollection::TransformGroup);
+		if (Seen && (*Seen)[TransformGroupIndex])
+		{
+			continue;
+		}
 
-		// first step : process the non values that do not need to be interpolated 
-		for (int32 EntryIndex = 0; EntryIndex < CurrentResults.GetNumEntries(); EntryIndex++)
+		if (GTParticles[TransformGroupIndex] == nullptr)
 		{
-			const FGeometryCollectionResults::FStateData& StateData = CurrentResults.GetState(EntryIndex);
-			const int32 TransformGroupIndex = StateData.TransformIndex;
+			continue;
+		}
 
-			if (GTParticles[TransformGroupIndex] == nullptr)
-			{
-				continue;
-			}
-			FParticle& GTParticle = *GTParticles[TransformGroupIndex];
+		FParticle& GTParticle = *GTParticles[TransformGroupIndex];
 
-			const bool bIsActive = !StateData.State.DisabledState;
-			if (UpdateValue(GameThreadCollection.Active[TransformGroupIndex], bIsActive))
-			{
-				GTParticle.SetDisabled(!bIsActive);
-				bIsCollectionDirty = true;
-			}
-			
+		const bool bIsActive = !StateData.State.DisabledState;
+		if (UpdateValue(GameThreadCollection.Active[TransformGroupIndex], bIsActive))
+		{
+			GTParticle.SetDisabled(!bIsActive);
+			bIsCollectionDirty = true;
+		}
+
+		if (UpdateValue(GameThreadCollection.DynamicState[TransformGroupIndex], StateData.State.DynamicState))
+		{
+			GTParticle.SetObjectState(static_cast<Chaos::EObjectStateType>(StateData.State.DynamicState));
+			bIsCollectionDirty = true;
+		}
+
+		if (UpdateValue(GameThreadCollection.Parent[TransformGroupIndex], StateData.ParentTransformIndex))
+		{
+			bIsCollectionDirty = true;
+		}
 
-			if (UpdateValue(GameThreadCollection.DynamicState[TransformGroupIndex], StateData.State.DynamicState))
+		if (InternalClusterParentTypeArray)
+		{
+			Chaos::EInternalClusterType ParentType = Chaos::EInternalClusterType::None;
+			if (StateData.State.HasInternalClusterParent != 0)
 			{
-				GTParticle.SetObjectState(static_cast<Chaos::EObjectStateType>(StateData.State.DynamicState));
-				bIsCollectionDirty = true;
+				ParentType = (StateData.State.DynamicInternalClusterParent != 0) ? Chaos::EInternalClusterType::Dynamic : Chaos::EInternalClusterType::KinematicOrStatic;
 			}
-
-			if (UpdateValue(GameThreadCollection.Parent[TransformGroupIndex], StateData.ParentTransformIndex))
+			const uint8 ParentTypeUInt8 = static_cast<uint8>(ParentType);
+			if (UpdateValue((*InternalClusterParentTypeArray)[TransformGroupIndex], ParentTypeUInt8))
 			{
 				bIsCollectionDirty = true;
 			}
+		}
 
-			if (InternalClusterParentTypeArray)
+		// if interpolation is off , we need to apply XR, VW and transforms
+		if (bForcePullXRVW)
+		{
+			const bool bAnimatingWhileDisabled = AnimationsActive ? (*AnimationsActive)[TransformGroupIndex] : false;
+			if (bIsActive || bAnimatingWhileDisabled)
 			{
-				Chaos::EInternalClusterType ParentType = Chaos::EInternalClusterType::None;
-				if (StateData.State.HasInternalClusterParent != 0)
-				{
-					ParentType = (StateData.State.DynamicInternalClusterParent != 0)? Chaos::EInternalClusterType::Dynamic: Chaos::EInternalClusterType::KinematicOrStatic;
-				}
-				const uint8 ParentTypeUInt8 = static_cast<uint8>(ParentType);
-				if (UpdateValue((*InternalClusterParentTypeArray)[TransformGroupIndex], ParentTypeUInt8))
+				const FGeometryCollectionResults::FPositionData& PositionData = CurrentResults.GetPositions(EntryIndex);
+
+				const bool XRModified = UpdateGTParticleXR(GTParticle, PositionData.ParticleX, PositionData.ParticleR);
+				bIsCollectionDirty |= XRModified;
+
+				if (LinearVelocities && AngularVelocities)
 				{
-					bIsCollectionDirty = true;
+					const FGeometryCollectionResults::FVelocityData& VelocityData = CurrentResults.GetVelocities(EntryIndex);
+					bIsCollectionDirty |= UpdateValue((*LinearVelocities)[TransformGroupIndex], VelocityData.ParticleV);
+					bIsCollectionDirty |= UpdateValue((*AngularVelocities)[TransformGroupIndex], VelocityData.ParticleW);
 				}
+
+				const FTransform& NewTransform = CurrentResults.GetTransform(EntryIndex);
+				bIsCollectionDirty |= UpdateTransform(GameThreadCollection.Transform[TransformGroupIndex], NewTransform);
 			}
+		}
 
-			// if interpolation is off , we need to apply XR, VW and transforms
-			if (bNeedInterpolation == false)
-			{
-				const bool bAnimatingWhileDisabled = AnimationsActive ? (*AnimationsActive)[TransformGroupIndex] : false;
-				if (bIsActive || bAnimatingWhileDisabled)
-				{
-					const FGeometryCollectionResults::FPositionData& PositionData = CurrentResults.GetPositions(EntryIndex);
+		// internal cluster index map update
+		if (StateData.InternalClusterUniqueIdx > INDEX_NONE)
+		{
+			GTParticlesToInternalClusterUniqueIdx.Add(GTParticles[TransformGroupIndex].Get(), StateData.InternalClusterUniqueIdx);
+			InternalClusterUniqueIdxToChildrenTransformIndices.FindOrAdd(StateData.InternalClusterUniqueIdx).Add(TransformGroupIndex);
+		}
+	}
 
-					const bool XRModified = UpdateGTParticleXR(GTParticle, PositionData.ParticleX, PositionData.ParticleR);
-					bIsCollectionDirty |= XRModified;
+	return bIsCollectionDirty;
+}
 
-					if (LinearVelocities && AngularVelocities)
-					{
-						const FGeometryCollectionResults::FVelocityData& VelocityData = CurrentResults.GetVelocities(EntryIndex);
-						bIsCollectionDirty |= UpdateValue((*LinearVelocities)[TransformGroupIndex], VelocityData.ParticleV);
-						bIsCollectionDirty |= UpdateValue((*AngularVelocities)[TransformGroupIndex], VelocityData.ParticleW);
-					}
+// Called from FPhysScene_ChaosInterface::SyncBodies(), NOT the solver.
+bool FGeometryCollectionPhysicsProxy::PullFromPhysicsState(const Chaos::FDirtyGeometryCollectionData& PullData, const int32 SolverSyncTimestamp, const Chaos::FDirtyGeometryCollectionData* NextPullData, const Chaos::FRealSingle* Alpha)
+{
+	if(IsObjectDeleting) return false;
 
-					const FTransform& NewTransform = CurrentResults.GetTransform(EntryIndex);
-					bIsCollectionDirty |= UpdateTransform(GameThreadCollection.Transform[TransformGroupIndex], NewTransform);
-				}
-			}
-		}
+	/**
+	 * CONTEXT: GAMETHREAD (Read Locked)
+	 * Perform a similar operation to Sync, but take the data from a gamethread-safe buffer. This will be called
+	 * from the game thread when it cannot sync to the physics thread. The simulation is very likely to be running
+	 * when this happens so never read any physics thread data here!
+	 *
+	 * Note: A read lock will have been acquired for this - so the physics thread won't force a buffer flip while this
+	 * sync is ongoing
+	 *
+	 */
 
-#if WITH_EDITORONLY_DATA
-		if (FDamageCollector* Collector = FRuntimeDataCollector::GetInstance().Find(CollectorGuid))
+	// earlier exit
+	const bool bHasResults = (PullData.Results().GetNumEntries() > 0);
+	const bool bHasNextResults = (NextPullData && NextPullData->Results().GetNumEntries() > 0);
+
+	if (!bHasResults && !bHasNextResults)
+	{
+		return false;
+	}
+
+	const int32 NumTransforms = GameThreadCollection.Transform.Num();
+	const bool bNeedInterpolation = (NextPullData != nullptr);
+
+	bool bIsCollectionDirty = false;
+
+	if (NumTransforms > 0)
+	{
+		GTParticlesToInternalClusterUniqueIdx.Reset();
+		InternalClusterUniqueIdxToChildrenTransformIndices.Reset();
+
+		// first: non-interpolatable data (everything besides XRVW OR if no next data exists, everything is non-interpolatable).
+		if (NextPullData)
 		{
-			for (int32 TransformGroupIndex = 0; TransformGroupIndex < NumTransforms; ++TransformGroupIndex)
-			{
-				const FGeometryCollectionResults::FDamageData& DamageData = CurrentResults.GetDamages(TransformGroupIndex);
-				Collector->SampleDamage(TransformGroupIndex, DamageData.Damage, DamageData.DamageThreshold);
-			}
+			bIsCollectionDirty |= PullNonInterpolatableDataFromSinglePhysicsState(*NextPullData, false, nullptr);
 		}
-#endif
+		bIsCollectionDirty |= PullNonInterpolatableDataFromSinglePhysicsState(PullData, !bNeedInterpolation, NextPullData ? &NextPullData->Results().GetModifiedTransformIndices() : nullptr);
 
 		// second : interpolate-able ones
 		if (bNeedInterpolation)
@@ -3503,6 +3529,10 @@ bool FGeometryCollectionPhysicsProxy::PullFromPhysicsState(const Chaos::FDirtyGe
 			const FGeometryCollectionResults& PrevResults = PullData.Results();
 			const FGeometryCollectionResults& NextResults = NextPullData->Results();
 
+			TManagedArray<FVector3f>* LinearVelocities = GameThreadCollection.FindAttributeTyped<FVector3f>(LinearVelocityAttributeName, FTransformCollection::TransformGroup);
+			TManagedArray<FVector3f>* AngularVelocities = GameThreadCollection.FindAttributeTyped<FVector3f>(AngularVelocitiesAttributeName, FTransformCollection::TransformGroup);
+			const TManagedArray<bool>* AnimationsActive = GameThreadCollection.FindAttribute<bool>(AnimateTransformAttributeName, FGeometryCollection::TransformGroup);
+
 			// for that case we cannot just go through the list of entries since Results and NextResults may have different number of entries that don't always match
 			// so we need to go through the transform indices and find the matching entries on both side 
 			// this is helped by the FResultAccessor structure
@@ -3548,23 +3578,19 @@ bool FGeometryCollectionPhysicsProxy::PullFromPhysicsState(const Chaos::FDirtyGe
 			}
 		} 
 
-		// internal cluster index map update
-		GTParticlesToInternalClusterUniqueIdx.Reset();
-		InternalClusterUniqueIdxToChildrenTransformIndices.Reset();
-		for (int32 TransformGroupIndex = 0; TransformGroupIndex < NumTransforms; ++TransformGroupIndex)
+
+#if WITH_EDITORONLY_DATA
+		// Damage is collected in full every time so no need to go through PullNonInterpolatableDataFromSinglePhysicsState.
+		if (FDamageCollector* Collector = FRuntimeDataCollector::GetInstance().Find(CollectorGuid))
 		{
-			const FGeometryCollectionResults::FEntryIndex EntryIndex = CurrentResults.GetEntryIndexByTransformIndex(TransformGroupIndex);
-			if (EntryIndex != INDEX_NONE)
+			const FGeometryCollectionResults& CurrentResults = NextPullData ? NextPullData->Results() : PullData.Results();
+			for (int32 TransformGroupIndex = 0; TransformGroupIndex < NumTransforms; ++TransformGroupIndex)
 			{
-				const FGeometryCollectionResults::FStateData& StateData = CurrentResults.GetState(EntryIndex);
-
-				if (StateData.InternalClusterUniqueIdx > INDEX_NONE)
-				{
-					GTParticlesToInternalClusterUniqueIdx.Add(GTParticles[TransformGroupIndex].Get(), StateData.InternalClusterUniqueIdx);
-					InternalClusterUniqueIdxToChildrenTransformIndices.FindOrAdd(StateData.InternalClusterUniqueIdx).Add(TransformGroupIndex);
-				}
+				const FGeometryCollectionResults::FDamageData& DamageData = CurrentResults.GetDamages(TransformGroupIndex);
+				Collector->SampleDamage(TransformGroupIndex, DamageData.Damage, DamageData.DamageThreshold);
 			}
 		}
+#endif
 
 		// if physics world transform was dirtied this frame we need to force an update to make sure transforms of 
 		// broken sleeping particles on kinemically driven GCs remained updated and the object feel grounded to the world
diff --git a/Engine/Source/Runtime/Experimental/Chaos/Public/GeometryCollectionProxyData.h b/Engine/Source/Runtime/Experimental/Chaos/Public/GeometryCollectionProxyData.h
--- a/Engine/Source/Runtime/Experimental/Chaos/Public/GeometryCollectionProxyData.h
+++ b/Engine/Source/Runtime/Experimental/Chaos/Public/GeometryCollectionProxyData.h
@@ -431,6 +431,11 @@ public:
 		Transforms[EntryIndex] = Transform;
 	}
 
+	inline const TBitArray<>& GetModifiedTransformIndices() const
+	{
+		return ModifiedTransformIndices;
+	}
+
 private:
 	Chaos::FReal SolverDt;
 
diff --git a/Engine/Source/Runtime/Experimental/Chaos/Public/PhysicsProxy/GeometryCollectionPhysicsProxy.h b/Engine/Source/Runtime/Experimental/Chaos/Public/PhysicsProxy/GeometryCollectionPhysicsProxy.h
--- a/Engine/Source/Runtime/Experimental/Chaos/Public/PhysicsProxy/GeometryCollectionPhysicsProxy.h
+++ b/Engine/Source/Runtime/Experimental/Chaos/Public/PhysicsProxy/GeometryCollectionPhysicsProxy.h
@@ -517,6 +517,18 @@ protected:
 	CHAOS_API Chaos::FPBDRigidClusteredParticleHandle* FindClusteredParticleHandleByItemIndex_Internal(FGeometryCollectionItemIndex ItemIndex) const;
 private:
 
+	/**
+	 * Since geometry collections only buffer data that has changed, when PullFromPhysicsState is given both PrevData and NextData it must
+	 * examine *both* PrevData and NextData for data about a particle (since that particle's data coudl be in PrevData and not NextData).
+	 * PullNonInterpolatableDataFromSinglePhysicsState will do the work to pull the non-interpolatable data (which could include X/R/V/W in certain scenarios)
+	 * to the game thread.
+	 * 
+	 * @param BufferData The buffered physics data to pull data from (could be PrevData or NextData).
+	 * @param bForcePullXRVW Whether or not to pull interpolatable data as non-interpolatable data (i.e. X/R/V/W). This happens only when NextData doesn't exist.
+	 * @param Seen The bit array of the previously handled FDirtyGeometryCollectionData (should be null in the case of handling NextData, and should be the TBitArray on NextData when handling PrevData).
+	 * @return A boolean indicating whether or not a change to the GT state was detected.
+	 */
+	bool PullNonInterpolatableDataFromSinglePhysicsState(const Chaos::FDirtyGeometryCollectionData& BufferData, bool bForcePullXRVW, const TBitArray<>* Seen);
 
 	/* set to true once InitializeBodiesPT has been called*/
 	bool bIsInitializedOnPhysicsThread = false;
