Chaos - friction stiffness tuning to reduce jitter in hard-to-solve or unsolvable contact configurations.
- Add cvar to control friction stiffness in the velocity phase (used by quadratic shapes) for testing if we run more than one velocity iteration
- Remove some unused cvars.

diff --git a/Engine/Source/Runtime/Experimental/Chaos/Private/Chaos/Collision/PBDCollisionSolver.cpp b/Engine/Source/Runtime/Experimental/Chaos/Private/Chaos/Collision/PBDCollisionSolver.cpp
--- a/Engine/Source/Runtime/Experimental/Chaos/Private/Chaos/Collision/PBDCollisionSolver.cpp
+++ b/Engine/Source/Runtime/Experimental/Chaos/Private/Chaos/Collision/PBDCollisionSolver.cpp
@@ -25,30 +25,17 @@ namespace Chaos
 {
 	namespace CVars
 	{
-		extern int32 Chaos_Collision_UseShockPropagation;
-
-		//
-		// Solver Settings
-		//
-
-		bool bChaos_PBDCollisionSolver_ISPC = false;	// NOTE: WORK IN PROGRESS. Not ready to be enabled.
-		FAutoConsoleVariableRef CVarChaos_PBDCollisionSolver_ISPC(TEXT("p.Chaos.PBDCollisionSolver.ISPC"), bChaos_PBDCollisionSolver_ISPC, TEXT("Use ISPC collision solver (WIP)"));
-
 		//
 		// Position Solver Settings
 		//
 
 		bool bChaos_PBDCollisionSolver_Position_SolveEnabled = true;
 		float Chaos_PBDCollisionSolver_Position_MinInvMassScale = 0.77f;
-		float Chaos_PBDCollisionSolver_Position_StaticFrictionStiffness = 1.0f;
-		float Chaos_PBDCollisionSolver_Position_PositionSolverTolerance = 0.001f;		// cms
-		float Chaos_PBDCollisionSolver_Position_RotationSolverTolerance = 0.001f;		// rads
+		float Chaos_PBDCollisionSolver_Position_StaticFrictionStiffness = 0.5f;
 
 		FAutoConsoleVariableRef CVarChaos_PBDCollisionSolver_Position_SolveEnabled(TEXT("p.Chaos.PBDCollisionSolver.Position.SolveEnabled"), bChaos_PBDCollisionSolver_Position_SolveEnabled, TEXT(""));
 		FAutoConsoleVariableRef CVarChaos_PBDCollisionSolver_Position_MinInvMassScale(TEXT("p.Chaos.PBDCollisionSolver.Position.MinInvMassScale"), Chaos_PBDCollisionSolver_Position_MinInvMassScale, TEXT(""));
 		FAutoConsoleVariableRef CVarChaos_PBDCollisionSolver_Position_StaticFrictionStiffness(TEXT("p.Chaos.PBDCollisionSolver.Position.StaticFriction.Stiffness"), Chaos_PBDCollisionSolver_Position_StaticFrictionStiffness, TEXT(""));
-		FAutoConsoleVariableRef CVarChaos_PBDCollisionSolver_Position_PositionSolverTolerance(TEXT("p.Chaos.PBDCollisionSolver.Position.PositionTolerance"), Chaos_PBDCollisionSolver_Position_PositionSolverTolerance, TEXT(""));
-		FAutoConsoleVariableRef CVarChaos_PBDCollisionSolver_Position_RotationSolverTolerance(TEXT("p.Chaos.PBDCollisionSolver.Position.RotationTolerance"), Chaos_PBDCollisionSolver_Position_RotationSolverTolerance, TEXT(""));
 
 		//
 		// Velocity Solver Settings
@@ -58,11 +45,13 @@ namespace Chaos
 		// If Chaos_PBDCollisionSolver_Velocity_MinInvMassScale is the same as Chaos_PBDCollisionSolver_Position_MinInvMassScale and all velocity iterations have shockpropagation, we avoid recalculating constraint-space mass
 		float Chaos_PBDCollisionSolver_Velocity_MinInvMassScale = Chaos_PBDCollisionSolver_Position_MinInvMassScale;
 		bool bChaos_PBDCollisionSolver_Velocity_FrictionEnabled = true;
+		float Chaos_PBDCollisionSolver_Velocity_StaticFrictionStiffness = 1.0f;
 		bool bChaos_PBDCollisionSolver_Velocity_AveragePointEnabled = true;
 
 		FAutoConsoleVariableRef CVarChaos_PBDCollisionSolver_Velocity_SolveEnabled(TEXT("p.Chaos.PBDCollisionSolver.Velocity.SolveEnabled"), bChaos_PBDCollisionSolver_Velocity_SolveEnabled, TEXT(""));
 		FAutoConsoleVariableRef CVarChaos_PBDCollisionSolver_Velocity_MinInvMassScale(TEXT("p.Chaos.PBDCollisionSolver.Velocity.MinInvMassScale"), Chaos_PBDCollisionSolver_Velocity_MinInvMassScale, TEXT(""));
 		FAutoConsoleVariableRef CVarChaos_PBDCollisionSolver_Velocity_FrictionEnabled(TEXT("p.Chaos.PBDCollisionSolver.Velocity.FrictionEnabled"), bChaos_PBDCollisionSolver_Velocity_FrictionEnabled, TEXT(""));
+		FAutoConsoleVariableRef CVarChaos_PBDCollisionSolver_Velocity_StaticFrictionStiffness(TEXT("p.Chaos.PBDCollisionSolver.Velocity.StaticFriction.Stiffness"), Chaos_PBDCollisionSolver_Velocity_StaticFrictionStiffness, TEXT(""));
 		FAutoConsoleVariableRef CVarChaos_PBDCollisionSolver_Velocity_AveragePointEnabled(TEXT("p.Chaos.PBDCollisionSolver.Velocity.AveragePointEnabled"), bChaos_PBDCollisionSolver_Velocity_AveragePointEnabled, TEXT(""));
 	}
 	using namespace CVars;
diff --git a/Engine/Source/Runtime/Experimental/Chaos/Public/Chaos/Collision/PBDCollisionSolver.h b/Engine/Source/Runtime/Experimental/Chaos/Public/Chaos/Collision/PBDCollisionSolver.h
--- a/Engine/Source/Runtime/Experimental/Chaos/Public/Chaos/Collision/PBDCollisionSolver.h
+++ b/Engine/Source/Runtime/Experimental/Chaos/Public/Chaos/Collision/PBDCollisionSolver.h
@@ -21,6 +21,7 @@ namespace Chaos
 		extern bool bChaos_PBDCollisionSolver_Velocity_AveragePointEnabled;
 		extern bool bChaos_PBDCollisionSolver_Velocity_FrictionEnabled;
 		extern float Chaos_PBDCollisionSolver_Position_StaticFrictionStiffness;
+		extern float Chaos_PBDCollisionSolver_Velocity_StaticFrictionStiffness;
 	}
 
 	namespace Private
@@ -584,8 +585,6 @@ namespace Chaos
 			FPBDCollisionSolverManifoldPoint& ManifoldPoint)
 		{
 			FSolverReal ImpulseNormal = -Stiffness * ManifoldPoint.ContactMassNormal * ContactVelocityDeltaNormal;
-			FSolverReal ImpulseTangentU = -Stiffness * ManifoldPoint.ContactMassTangentU * ContactVelocityDeltaTangent0;
-			FSolverReal ImpulseTangentV = -Stiffness * ManifoldPoint.ContactMassTangentV * ContactVelocityDeltaTangent1;
 
 			// Clamp the total impulse to be positive along the normal. We can apply a net negative impulse, 
 			// but only to correct the velocity that was added by pushout (in which case MinImpulseNormal will be negative).
@@ -596,9 +595,19 @@ namespace Chaos
 				ImpulseNormal = MinImpulseNormal - ManifoldPoint.NetImpulseNormal;
 			}
 
+			ManifoldPoint.NetImpulseNormal += ImpulseNormal;
+
+			FSolverVec3 Impulse = ImpulseNormal * ManifoldPoint.ContactNormal;
+
 			// Clamp the tangential impulses to the friction cone
+			FSolverReal ImpulseTangentU = 0;
+			FSolverReal ImpulseTangentV = 0;
 			if ((DynamicFriction > 0) && (Dt > 0))
 			{
+				const FSolverReal FrictionStiffness = Stiffness * CVars::Chaos_PBDCollisionSolver_Velocity_StaticFrictionStiffness;
+				ImpulseTangentU = -FrictionStiffness * ManifoldPoint.ContactMassTangentU * ContactVelocityDeltaTangent0;
+				ImpulseTangentV = -FrictionStiffness * ManifoldPoint.ContactMassTangentV * ContactVelocityDeltaTangent1;
+
 				const FSolverReal MaxImpulseTangent = FMath::Max(FSolverReal(0), DynamicFriction * (ManifoldPoint.NetImpulseNormal + ImpulseNormal + ManifoldPoint.NetPushOutNormal / Dt));
 				const FSolverReal MaxImpulseTangentSq = FMath::Square(MaxImpulseTangent);
 				const FSolverReal ImpulseTangentSq = FMath::Square(ImpulseTangentU) + FMath::Square(ImpulseTangentV);
@@ -608,14 +617,14 @@ namespace Chaos
 					ImpulseTangentU *= ImpulseTangentScale;
 					ImpulseTangentV *= ImpulseTangentScale;
 				}
-			}
 
-			ManifoldPoint.NetImpulseNormal += ImpulseNormal;
-			ManifoldPoint.NetImpulseTangentU += ImpulseTangentU;
-			ManifoldPoint.NetImpulseTangentV += ImpulseTangentV;
+				ManifoldPoint.NetImpulseTangentU += ImpulseTangentU;
+				ManifoldPoint.NetImpulseTangentV += ImpulseTangentV;
+
+				Impulse += ImpulseTangentU * ManifoldPoint.ContactTangentU + ImpulseTangentV * ManifoldPoint.ContactTangentV;
+			}
 
 			// Apply the velocity deltas from the impulse
-			const FSolverVec3 Impulse = ImpulseNormal * ManifoldPoint.ContactNormal + ImpulseTangentU * ManifoldPoint.ContactTangentU + ImpulseTangentV * ManifoldPoint.ContactTangentV;
 			if (IsDynamic(0))
 			{
 				const FSolverVec3 DV0 = State.InvMs[0] * Impulse;
