Attempt to make the Static Mesh -> Geometry Collection conversion more robust to missing normal/tangent information by forcing a recompute if either attribute is *entirely* invalid, even if the recompute flag is not set.

diff --git a/Engine/Source/Runtime/Experimental/GeometryCollectionEngine/Private/GeometryCollection/GeometryCollectionEngineConversion.cpp b/Engine/Source/Runtime/Experimental/GeometryCollectionEngine/Private/GeometryCollection/GeometryCollectionEngineConversion.cpp
--- a/Engine/Source/Runtime/Experimental/GeometryCollectionEngine/Private/GeometryCollection/GeometryCollectionEngineConversion.cpp
+++ b/Engine/Source/Runtime/Experimental/GeometryCollectionEngine/Private/GeometryCollection/GeometryCollectionEngineConversion.cpp
@@ -382,38 +382,66 @@ void FGeometryCollectionEngineConversion::AppendMeshDescription(
 
 namespace
 {
-	// Note: This is the same as the corresponding ModelingComponents AssetUtils function, in the MeshModelingToolset plugin,
-	// but we cannot use a plugin function from here, so have our own copy
-	void InitializeAutoGeneratedAttributes(FMeshDescription& Mesh, const FMeshBuildSettings* BuildSettings)
+	// Based on FStaticMeshOperations::AreNormalsAndTangentsValid but tests if *any* are valid instead of if *all* are valid to avoid forcing a recomputes for any small degen tri
+	void HasValidNormalsAndTangents(const FMeshDescription& MeshDescription, bool& bHasValidNormals, bool& bHasValidTangents)
+	{
+		bHasValidNormals = false;
+		bHasValidTangents = false;
+
+		FStaticMeshConstAttributes Attributes(MeshDescription);
+		TArrayView<const FVector3f> VertexInstanceNormals = Attributes.GetVertexInstanceNormals().GetRawArray();
+		TArrayView<const FVector3f> VertexInstanceTangents = Attributes.GetVertexInstanceTangents().GetRawArray();
+
+		for (const FVertexInstanceID VertexInstanceID : MeshDescription.VertexInstances().GetElementIDs())
+		{
+			bHasValidNormals |= (!VertexInstanceNormals[VertexInstanceID].IsNearlyZero() && !VertexInstanceNormals[VertexInstanceID].ContainsNaN());
+			bHasValidTangents |= (!VertexInstanceTangents[VertexInstanceID].IsNearlyZero() && !VertexInstanceTangents[VertexInstanceID].ContainsNaN());
+			if (bHasValidNormals && bHasValidTangents)
+			{
+				break;
+			}
+		}
+	}
+
+	// Note: This is similar to the InitializeAutoGeneratedAttributes ModelingComponents AssetUtils function, in the MeshModelingToolset plugin,
+	// but we cannot use a plugin function from here, so have our own version. This version forces a recompute if Normals or Tangents are fully invalid, 
+	// even if the corresponding Recompute flag is not set.
+	void InitializeNormalsAndTangentsIfNeededOrRequested(FMeshDescription& Mesh, const FMeshBuildSettings* BuildSettings)
 	{
 		check(BuildSettings);
 
-		// if the build settings don't indicate calculated normals/tangents, we will assume they are good
-		if (BuildSettings->bRecomputeNormals || BuildSettings->bRecomputeTangents)
+		// Recompute according to build settings
+		bool bShouldRecomputeNormals = BuildSettings->bRecomputeNormals;
+		bool bShouldRecomputeTangents = BuildSettings->bRecomputeTangents;
+		// Also recompute if normals or tangents are fully invalid
+		// Note: We don't force a recompute if only some elements are invalid as these may just be a degenerate element on a mesh where the normals/tangents should otherwise be preserved.
+		if (!BuildSettings->bRecomputeNormals || !BuildSettings->bRecomputeTangents)
 		{
-			// check if any are invalid
-			bool bHasInvalidNormals, bHasInvalidTangents;
-			FStaticMeshOperations::AreNormalsAndTangentsValid(Mesh, bHasInvalidNormals, bHasInvalidTangents);
+			bool bHasValidNormals = false, bHasValidTangents = false;
+			HasValidNormalsAndTangents(Mesh, bHasValidNormals, bHasValidTangents);
+			bShouldRecomputeNormals |= !bHasValidNormals;
+			bShouldRecomputeTangents |= !bHasValidTangents;
+		}
+		
 
-			// if neither are invalid we are not going to recompute
-			if (bHasInvalidNormals || bHasInvalidTangents)
+		// run recompute function if either normals or tangents need recompute
+		if (bShouldRecomputeNormals || bShouldRecomputeTangents)
+		{
+			FStaticMeshAttributes Attributes(Mesh);
+			if (!Attributes.GetTriangleNormals().IsValid() || !Attributes.GetTriangleTangents().IsValid())
 			{
-				FStaticMeshAttributes Attributes(Mesh);
-				if (!Attributes.GetTriangleNormals().IsValid() || !Attributes.GetTriangleTangents().IsValid())
-				{
-					// If these attributes don't exist, create them and compute their values for each triangle
-					FStaticMeshOperations::ComputeTriangleTangentsAndNormals(Mesh);
-				}
+				// If these attributes don't exist, create them and compute their values for each triangle
+				FStaticMeshOperations::ComputeTriangleTangentsAndNormals(Mesh);
+			}
 
-				EComputeNTBsFlags ComputeNTBsOptions = EComputeNTBsFlags::BlendOverlappingNormals;
-				ComputeNTBsOptions |= BuildSettings->bRecomputeNormals ? EComputeNTBsFlags::Normals : EComputeNTBsFlags::None;
-				ComputeNTBsOptions |= BuildSettings->bRecomputeTangents ? EComputeNTBsFlags::Tangents : EComputeNTBsFlags::None;
-				ComputeNTBsOptions |= BuildSettings->bUseMikkTSpace ? EComputeNTBsFlags::UseMikkTSpace : EComputeNTBsFlags::None;
-				ComputeNTBsOptions |= BuildSettings->bComputeWeightedNormals ? EComputeNTBsFlags::WeightedNTBs : EComputeNTBsFlags::None;
-				ComputeNTBsOptions |= BuildSettings->bRemoveDegenerates ? EComputeNTBsFlags::IgnoreDegenerateTriangles : EComputeNTBsFlags::None;
+			EComputeNTBsFlags ComputeNTBsOptions = EComputeNTBsFlags::BlendOverlappingNormals;
+			ComputeNTBsOptions |= (bShouldRecomputeNormals) ? EComputeNTBsFlags::Normals : EComputeNTBsFlags::None;
+			ComputeNTBsOptions |= (bShouldRecomputeTangents) ? EComputeNTBsFlags::Tangents : EComputeNTBsFlags::None;
+			ComputeNTBsOptions |= BuildSettings->bUseMikkTSpace ? EComputeNTBsFlags::UseMikkTSpace : EComputeNTBsFlags::None;
+			ComputeNTBsOptions |= BuildSettings->bComputeWeightedNormals ? EComputeNTBsFlags::WeightedNTBs : EComputeNTBsFlags::None;
+			ComputeNTBsOptions |= BuildSettings->bRemoveDegenerates ? EComputeNTBsFlags::IgnoreDegenerateTriangles : EComputeNTBsFlags::None;
 
-				FStaticMeshOperations::ComputeTangentsAndNormals(Mesh, ComputeNTBsOptions);
-			}
+			FStaticMeshOperations::ComputeTangentsAndNormals(Mesh, ComputeNTBsOptions);
 		}
 	}
 }
@@ -440,7 +468,7 @@ FMeshDescription* FGeometryCollectionEngineConversion::GetMaxResMeshDescriptionW
 		SourceModel = &StaticMesh->GetSourceModel(0);
 	}
 
-	InitializeAutoGeneratedAttributes(*MeshDescription, &SourceModel->BuildSettings);
+	InitializeNormalsAndTangentsIfNeededOrRequested(*MeshDescription, &SourceModel->BuildSettings);
 #endif //WITH_EDITORONLY_DATA
 	return MeshDescription;
 }
