Include error reason in IAS' log output.

#rnx

diff --git a/Engine/Source/Runtime/Experimental/IoStoreOnDemand/Private/OnDemandIoDispatcherBackend.cpp b/Engine/Source/Runtime/Experimental/IoStoreOnDemand/Private/OnDemandIoDispatcherBackend.cpp
--- a/Engine/Source/Runtime/Experimental/IoStoreOnDemand/Private/OnDemandIoDispatcherBackend.cpp
+++ b/Engine/Source/Runtime/Experimental/IoStoreOnDemand/Private/OnDemandIoDispatcherBackend.cpp
@@ -372,7 +372,7 @@ void FHttpClient::Issue(UE::HTTP::FRequest&& Request, FIoReadCallback&& Callback
 			{
 				const uint64 Duration = (uint64)FPlatformTime::ToMilliseconds64(FPlatformTime::Cycles64() - StartTime);
 				UE_LOG(LogIas, VeryVerbose, TEXT("%s"),
-					*WriteToString<256>(TEXT("HTTP GET - "), Url, TEXT(" ("), StatusCode, TEXT(" "), Duration, TEXTVIEW("ms)")));
+					*WriteToString<256>(TEXT("HTTP GET - "), Url, TEXT(" ("), StatusCode, TEXT(" "), Duration, TEXTVIEW("ms) : "), Status.GetErrorReason()));
 
 				Callback(FIoStatus(EIoErrorCode::ReadError, TEXTVIEW("HTTP Error")));
 			}
