Control poll behaviour a little more finely

#rnx

diff --git a/Engine/Source/Runtime/Experimental/IoStoreOnDemand/Private/CoreHttp/Client.h b/Engine/Source/Runtime/Experimental/IoStoreOnDemand/Private/CoreHttp/Client.h
--- a/Engine/Source/Runtime/Experimental/IoStoreOnDemand/Private/CoreHttp/Client.h
+++ b/Engine/Source/Runtime/Experimental/IoStoreOnDemand/Private/CoreHttp/Client.h
@@ -170,7 +170,7 @@ public:
 
 							FEventLoop() = default;
 							~FEventLoop();
-	uint32					Tick();
+	uint32					Tick(uint32 PollTimeoutMs=0);
 	bool					IsIdle() const;
 	void					Cancel(FTicket Ticket);
 	[[nodiscard]] FRequest	Request(FAnsiStringView Method, FAnsiStringView Url, const FRequestParams* Params=nullptr);
diff --git a/Engine/Source/Runtime/Experimental/IoStoreOnDemand/Private/CoreHttp/Http.cpp b/Engine/Source/Runtime/Experimental/IoStoreOnDemand/Private/CoreHttp/Http.cpp
--- a/Engine/Source/Runtime/Experimental/IoStoreOnDemand/Private/CoreHttp/Http.cpp
+++ b/Engine/Source/Runtime/Experimental/IoStoreOnDemand/Private/CoreHttp/Http.cpp
@@ -1114,7 +1114,7 @@ static int32 Select(SocketType Socket, int32 Mode, int32 TimeoutMs=0)
 class FEventLoopInternal
 {
 public:
-	static bool		DoIsReady(FActivity* Activity);
+	static bool		DoIsReady(FActivity* Activity, uint32 TimeoutMs=0);
 	static int32	DoResolve(FActivity* Activity);
 	static int32	DoConnect(FActivity* Activity);
 	static int32	DoSend(FActivity* Activity);
@@ -1128,7 +1128,7 @@ public:
 
 
 ////////////////////////////////////////////////////////////////////////////////
-bool FEventLoopInternal::DoIsReady(FActivity* Activity)
+bool FEventLoopInternal::DoIsReady(FActivity* Activity, uint32 TimeoutMs)
 {
 	TRACE_CPUPROFILER_EVENT_SCOPE(CoreHttp::DoIsReady);
 
@@ -1856,7 +1856,7 @@ void FEventLoop::Cancel(FTicket Ticket)
 }
 
 ////////////////////////////////////////////////////////////////////////////////
-uint32 FEventLoop::Tick()
+uint32 FEventLoop::Tick(uint32 PollTimeoutMs)
 {
 	TRACE_CPUPROFILER_EVENT_SCOPE(CoreHttp::Tick);
 
@@ -1892,7 +1892,7 @@ uint32 FEventLoop::Tick()
 		++BusyCount;
 
 		int32 Result = 1;
-		if (!FEventLoopInternal::DoIsReady(Activity))
+		if (!FEventLoopInternal::DoIsReady(Activity, PollTimeoutMs))
 		{
 			continue;
 		}
