setting per-particle collision filter data on the geometry collection needs to sync that data to the physics thread

diff --git a/Engine/Source/Runtime/Experimental/Chaos/Private/PhysicsProxy/GeometryCollectionPhysicsProxy.cpp b/Engine/Source/Runtime/Experimental/Chaos/Private/PhysicsProxy/GeometryCollectionPhysicsProxy.cpp
--- a/Engine/Source/Runtime/Experimental/Chaos/Private/PhysicsProxy/GeometryCollectionPhysicsProxy.cpp
+++ b/Engine/Source/Runtime/Experimental/Chaos/Private/PhysicsProxy/GeometryCollectionPhysicsProxy.cpp
@@ -16,6 +16,7 @@
 #include "ChaosSolversModule.h"
 #include "Chaos/PBDCollisionConstraintsUtil.h"
 #include "Chaos/PerParticleGravity.h"
+#include "Chaos/PhysicsObjectInternalInterface.h"
 #include "Chaos/ImplicitObject.h"
 #include "Chaos/ImplicitObjectScaled.h"
 #include "Chaos/ImplicitObjectUnion.h"
@@ -3638,6 +3639,37 @@ void FGeometryCollectionPhysicsProxy::UpdateFilterData_External(const FCollision
 	SetProxyDirty_External();
 }
 
+DECLARE_CYCLE_STAT(TEXT("FGeometryCollectionPhysicsProxy::UpdatePerParticleFilterData_External"), STAT_GeometryCollectionPhysicsProxyUpdatePerParticleFilterData, STATGROUP_Chaos);
+void FGeometryCollectionPhysicsProxy::UpdatePerParticleFilterData_External(const TArray<FParticleCollisionFilterData>& PerParticleData)
+{
+	SCOPE_CYCLE_COUNTER(STAT_GeometryCollectionPhysicsProxyUpdatePerParticleFilterData);
+	check(IsInGameThread());
+
+	// TODO: Go through the PushToPhysicsState flow instead?
+	// TODO: Need to figure out how the per-particle collision filter data works with the global one. The per-particle one should probably replace the global one entirely...
+	if (Chaos::FPhysicsSolver* RBDSolver = GetSolver<Chaos::FPhysicsSolver>())
+	{
+		RBDSolver->EnqueueCommandImmediate([this, PerParticleData]()
+		{
+			Chaos::FWritePhysicsObjectInterface_Internal Interface= Chaos::FPhysicsObjectInternalInterface::GetWrite();
+
+			int32 ParticleIndex = 0;
+			for (const FParticleCollisionFilterData& Data : PerParticleData)
+			{
+				Chaos::FPhysicsObjectHandle Object = PhysicsObjects[ParticleIndex].Get();
+				if (Data.bIsValid && Object)
+				{
+					TArrayView<Chaos::FPhysicsObjectHandle> ParticleView{ &Object, 1 };
+					Interface.UpdateShapeCollisionFlags(ParticleView, Data.bSimEnabled, Data.bQueryEnabled);
+					Interface.UpdateShapeFilterData(ParticleView, Data.QueryFilter, Data.SimFilter);
+				}
+
+				++ParticleIndex;
+			}
+		});
+	}
+}
+
 //==============================================================================
 // STATIC SETUP FUNCTIONS
 //==============================================================================
diff --git a/Engine/Source/Runtime/Experimental/Chaos/Public/PhysicsProxy/GeometryCollectionPhysicsProxy.h b/Engine/Source/Runtime/Experimental/Chaos/Public/PhysicsProxy/GeometryCollectionPhysicsProxy.h
--- a/Engine/Source/Runtime/Experimental/Chaos/Public/PhysicsProxy/GeometryCollectionPhysicsProxy.h
+++ b/Engine/Source/Runtime/Experimental/Chaos/Public/PhysicsProxy/GeometryCollectionPhysicsProxy.h
@@ -443,6 +443,17 @@ public:
 	EReplicationMode GetReplicationMode() const { return ReplicationMode; }
 
 	CHAOS_API void UpdateFilterData_External(const FCollisionFilterData& NewSimFilter, const FCollisionFilterData& NewQueryFilter);
+
+	struct FParticleCollisionFilterData
+	{
+		bool bIsValid = false;
+		bool bQueryEnabled = false;
+		bool bSimEnabled = false;
+		FCollisionFilterData QueryFilter;
+		FCollisionFilterData SimFilter;
+	};
+
+	CHAOS_API void UpdatePerParticleFilterData_External(const TArray<FParticleCollisionFilterData>& Data);
 	
 	/** 
 	 * Traverses the parents of TransformGroupIdx counting number of levels,
diff --git a/Engine/Source/Runtime/Experimental/GeometryCollectionEngine/Private/GeometryCollection/GeometryCollectionComponent.cpp b/Engine/Source/Runtime/Experimental/GeometryCollectionEngine/Private/GeometryCollection/GeometryCollectionComponent.cpp
--- a/Engine/Source/Runtime/Experimental/GeometryCollectionEngine/Private/GeometryCollection/GeometryCollectionComponent.cpp
+++ b/Engine/Source/Runtime/Experimental/GeometryCollectionEngine/Private/GeometryCollection/GeometryCollectionComponent.cpp
@@ -3383,26 +3383,20 @@ void UGeometryCollectionComponent::LoadCollisionProfiles()
 
 	// Cache the FCollisionResponseTemplate as well as the query/sim collision filter data for a given collision profile name
 	// so we don't have to recreate it every time.
-	
-	struct CollisionProfileDataCache
-	{
-		FCollisionResponseTemplate Template;
-		FCollisionFilterData QueryFilter;
-		FCollisionFilterData SimFilter;
-	};
-
-	TMap<FName, CollisionProfileDataCache> CachedData;
+	using FCollisionProfileDataCache = FGeometryCollectionPhysicsProxy::FParticleCollisionFilterData;
+	TMap<FName, FCollisionProfileDataCache> CachedData;
 
 	// Returns nullptr if we can't create or get the data.
-	auto CreateOrGetCollisionProfileData = [this, &CachedData](const FName& ProfileName) -> const CollisionProfileDataCache*
+	auto CreateOrGetCollisionProfileData = [this, &CachedData](const FName& ProfileName) -> const FCollisionProfileDataCache*
 	{
-		if (const CollisionProfileDataCache* Data = CachedData.Find(ProfileName))
+		if (const FCollisionProfileDataCache* Data = CachedData.Find(ProfileName))
 		{
 			return Data;
 		}
 
-		CollisionProfileDataCache Cache;
-		if (ProfileName == NAME_None || !UCollisionProfile::Get()->GetProfileTemplate(ProfileName, Cache.Template))
+		FCollisionProfileDataCache Cache;
+		FCollisionResponseTemplate Template;
+		if (ProfileName == NAME_None || !UCollisionProfile::Get()->GetProfileTemplate(ProfileName, Template))
 		{
 			return nullptr;
 		}
@@ -3413,22 +3407,28 @@ void UGeometryCollectionComponent::LoadCollisionProfiles()
 
 		Cache.QueryFilter = InitialQueryFilter;
 		Cache.SimFilter = InitialSimFilter;
-		CreateShapeFilterData(Cache.Template.ObjectType, BodyInstance.GetMaskFilter(), ActorID, Cache.Template.ResponseToChannels, CompID, INDEX_NONE, Cache.QueryFilter, Cache.SimFilter, BodyInstance.bUseCCD, bNotifyCollisions || bNotifyGlobalCollisions, false, false);
+		CreateShapeFilterData(Template.ObjectType, BodyInstance.GetMaskFilter(), ActorID, Template.ResponseToChannels, CompID, INDEX_NONE, Cache.QueryFilter, Cache.SimFilter, BodyInstance.bUseCCD, bNotifyCollisions || bNotifyGlobalCollisions, false, false);
 
 		// Maintain parity with the rest of the geometry collection filters.
 		Cache.QueryFilter.Word3 |= (EPDF_SimpleCollision | EPDF_ComplexCollision);
 		Cache.SimFilter.Word3 |= (EPDF_SimpleCollision | EPDF_ComplexCollision);
 		
+		Cache.bQueryEnabled = CollisionEnabledHasQuery(Template.CollisionEnabled);
+		Cache.bSimEnabled = CollisionEnabledHasPhysics(Template.CollisionEnabled);
+		Cache.bIsValid = true;
 		return &CachedData.Add(ProfileName, Cache);
 	};
 
-	const CollisionProfileDataCache* AbandonedData = (bEnableAbandonAfterLevel && GetIsReplicated()) ? CreateOrGetCollisionProfileData(AbandonedCollisionProfileName) : nullptr;
+	const FCollisionProfileDataCache* AbandonedData = (bEnableAbandonAfterLevel && GetIsReplicated()) ? CreateOrGetCollisionProfileData(AbandonedCollisionProfileName) : nullptr;
 
 	Chaos::Facades::FCollectionHierarchyFacade HierarchyFacade(*RestCollection->GetGeometryCollection());
 	// Use GetAllPhysicsObjectIncludingNulls instead of GetAllPhysicsObjects if you need to use Level data or any data from HierarchyFacade
 	TArray<Chaos::FPhysicsObjectHandle> PhysicsObjects = PhysicsProxy->GetAllPhysicsObjectIncludingNulls();
 	FLockedWritePhysicsObjectExternalInterface Interface = FPhysicsObjectExternalInterface::LockWrite(PhysicsObjects);
 
+	TArray<FCollisionProfileDataCache> PerParticleData;
+	PerParticleData.Reserve(PhysicsObjects.Num());
+
 	for (int32 ParticleIndex = 0; ParticleIndex < PhysicsObjects.Num(); ++ParticleIndex)
 	{
 		if (PhysicsObjects[ParticleIndex] != nullptr)
@@ -3436,7 +3436,7 @@ void UGeometryCollectionComponent::LoadCollisionProfiles()
 			const int32 Level = HierarchyFacade.GetInitialLevel(ParticleIndex);
 
 			TArrayView<Chaos::FPhysicsObjectHandle> ParticleView{ &PhysicsObjects[ParticleIndex], 1 };
-			const CollisionProfileDataCache* Data = nullptr;
+			const FCollisionProfileDataCache* Data = nullptr;
 			if (!CollisionProfilePerParticle.IsEmpty() && CollisionProfilePerParticle[ParticleIndex] != NAME_None)
 			{
 				Data = CreateOrGetCollisionProfileData(CollisionProfilePerParticle[ParticleIndex]);
@@ -3452,11 +3452,18 @@ void UGeometryCollectionComponent::LoadCollisionProfiles()
 
 			if (Data)
 			{
-				Interface->UpdateShapeCollisionFlags(ParticleView, CollisionEnabledHasPhysics(Data->Template.CollisionEnabled), CollisionEnabledHasQuery(Data->Template.CollisionEnabled));
+				PerParticleData.Add(*Data);
+				Interface->UpdateShapeCollisionFlags(ParticleView, Data->bSimEnabled, Data->bQueryEnabled);
 				Interface->UpdateShapeFilterData(ParticleView, Data->QueryFilter, Data->SimFilter);
 			}
+			else
+			{
+				PerParticleData.Add({});
+			}
 		}
 	}
+
+	PhysicsProxy->UpdatePerParticleFilterData_External(PerParticleData);
 }
 
 #if WITH_EDITORONLY_DATA
