Chaos Cloth: fix accessing out of bounds when only some maps exist on constraints with multiple maps (this is all newer constraints like XPBD and updated BendingConstraints).

diff --git a/Engine/Source/Runtime/Experimental/Chaos/Private/Chaos/PBDBendingConstraints.cpp b/Engine/Source/Runtime/Experimental/Chaos/Private/Chaos/PBDBendingConstraints.cpp
--- a/Engine/Source/Runtime/Experimental/Chaos/Private/Chaos/PBDBendingConstraints.cpp
+++ b/Engine/Source/Runtime/Experimental/Chaos/Private/Chaos/PBDBendingConstraints.cpp
@@ -192,10 +192,10 @@ void FPBDBendingConstraints::Apply(FSolverParticles& Particles, const FSolverRea
 						&RestAngles.GetData()[ColorStart],
 						&IsBuckled.GetData()[ColorStart],
 						StiffnessHasWeightMap,
-						&Stiffness.GetIndices().GetData()[ColorStart],
+						StiffnessHasWeightMap ? &Stiffness.GetIndices().GetData()[ColorStart] : nullptr,
 						&Stiffness.GetTable().GetData()[0],
 						BucklingStiffnessHasWeightMap,
-						&BucklingStiffness.GetIndices().GetData()[ColorStart],
+						BucklingStiffnessHasWeightMap ? &BucklingStiffness.GetIndices().GetData()[ColorStart] : nullptr,
 						&BucklingStiffness.GetTable().GetData()[0],
 						ColorSize);
 				}
diff --git a/Engine/Source/Runtime/Experimental/Chaos/Private/Chaos/PBDBendingConstraints.ispc b/Engine/Source/Runtime/Experimental/Chaos/Private/Chaos/PBDBendingConstraints.ispc
--- a/Engine/Source/Runtime/Experimental/Chaos/Private/Chaos/PBDBendingConstraints.ispc
+++ b/Engine/Source/Runtime/Experimental/Chaos/Private/Chaos/PBDBendingConstraints.ispc
@@ -134,8 +134,8 @@ export void ApplyBendingConstraintsWithMaps(uniform FVector4f PandInvM[],
 		const varying int32 i2 = Constraint.V[1];
 		const varying int32 i3 = Constraint.V[2];
 		const varying int32 i4 = Constraint.V[3];
-		const varying int8 StiffnessIndex = select(StiffnessHasMap, StiffnessIndices[i], 0);
-		const varying int8 BucklingStiffnessIndex = select(BucklingStiffnessHasMap, BucklingStiffnessIndices[i], 0);
+		const varying int8 StiffnessIndex = StiffnessHasMap ? StiffnessIndices[i] : 0;
+		const varying int8 BucklingStiffnessIndex = BucklingStiffnessHasMap ? BucklingStiffnessIndices[i] : 0;
 
 		const varying FVector4f PandInvM2 = VectorGather(&PandInvM[i2]);
 		const varying FVector4f PandInvM3 = VectorGather(&PandInvM[i3]);
diff --git a/Engine/Source/Runtime/Experimental/Chaos/Private/Chaos/XPBDAnisotropicBendingConstraints.cpp b/Engine/Source/Runtime/Experimental/Chaos/Private/Chaos/XPBDAnisotropicBendingConstraints.cpp
--- a/Engine/Source/Runtime/Experimental/Chaos/Private/Chaos/XPBDAnisotropicBendingConstraints.cpp
+++ b/Engine/Source/Runtime/Experimental/Chaos/Private/Chaos/XPBDAnisotropicBendingConstraints.cpp
@@ -591,25 +591,25 @@ void FXPBDAnisotropicBendingConstraints::Apply(FSolverParticles& Particles, cons
 							Dt,
 							MinStiffness,
 							StiffnessHasWeightMap,
-							&Stiffness.GetIndices().GetData()[ColorStart],
+							StiffnessHasWeightMap ? &Stiffness.GetIndices().GetData()[ColorStart] : nullptr,
 							&Stiffness.GetTable().GetData()[0],
 							StiffnessWeftHasWeightMap,
-							&StiffnessWeft.GetIndices().GetData()[ColorStart],
+							StiffnessWeftHasWeightMap ? &StiffnessWeft.GetIndices().GetData()[ColorStart] : nullptr,
 							&StiffnessWeft.GetTable().GetData()[0],
 							StiffnessBiasHasWeightMap,
-							&StiffnessBias.GetIndices().GetData()[ColorStart],
+							StiffnessBiasHasWeightMap ? &StiffnessBias.GetIndices().GetData()[ColorStart] : nullptr,
 							&StiffnessBias.GetTable().GetData()[0],
 							BucklingStiffnessHasWeightMap,
-							&BucklingStiffness.GetIndices().GetData()[ColorStart],
+							BucklingStiffnessHasWeightMap ? &BucklingStiffness.GetIndices().GetData()[ColorStart] : nullptr,
 							&BucklingStiffness.GetTable().GetData()[0],
 							BucklingStiffnessWeftHasWeightMap,
-							&BucklingStiffnessWeft.GetIndices().GetData()[ColorStart],
+							BucklingStiffnessWeftHasWeightMap ? &BucklingStiffnessWeft.GetIndices().GetData()[ColorStart] : nullptr,
 							&BucklingStiffnessWeft.GetTable().GetData()[0],
 							BucklingStiffnessBiasHasWeightMap,
-							&BucklingStiffnessBias.GetIndices().GetData()[ColorStart],
+							BucklingStiffnessBiasHasWeightMap ? &BucklingStiffnessBias.GetIndices().GetData()[ColorStart] : nullptr,
 							&BucklingStiffnessBias.GetTable().GetData()[0],
 							DampingHasWeightMap,
-							&DampingRatio.GetIndices().GetData()[ColorStart],
+							DampingHasWeightMap ? &DampingRatio.GetIndices().GetData()[ColorStart] : nullptr,
 							&DampingRatio.GetTable().GetData()[0],
 							ColorSize);
 					}
@@ -630,22 +630,22 @@ void FXPBDAnisotropicBendingConstraints::Apply(FSolverParticles& Particles, cons
 							Dt,
 							MinStiffness,
 							StiffnessHasWeightMap,
-							&Stiffness.GetIndices().GetData()[ColorStart],
+							StiffnessHasWeightMap ? &Stiffness.GetIndices().GetData()[ColorStart] : nullptr,
 							&Stiffness.GetTable().GetData()[0],
 							StiffnessWeftHasWeightMap,
-							&StiffnessWeft.GetIndices().GetData()[ColorStart],
+							StiffnessWeftHasWeightMap ? &StiffnessWeft.GetIndices().GetData()[ColorStart] : nullptr,
 							&StiffnessWeft.GetTable().GetData()[0],
 							StiffnessBiasHasWeightMap,
-							&StiffnessBias.GetIndices().GetData()[ColorStart],
+							StiffnessBiasHasWeightMap ? &StiffnessBias.GetIndices().GetData()[ColorStart] : nullptr,
 							&StiffnessBias.GetTable().GetData()[0],
 							BucklingStiffnessHasWeightMap,
-							&BucklingStiffness.GetIndices().GetData()[ColorStart],
+							BucklingStiffnessHasWeightMap ? &BucklingStiffness.GetIndices().GetData()[ColorStart] : nullptr,
 							&BucklingStiffness.GetTable().GetData()[0],
 							BucklingStiffnessWeftHasWeightMap,
-							&BucklingStiffnessWeft.GetIndices().GetData()[ColorStart],
+							BucklingStiffnessWeftHasWeightMap ? &BucklingStiffnessWeft.GetIndices().GetData()[ColorStart] : nullptr,
 							&BucklingStiffnessWeft.GetTable().GetData()[0],
 							BucklingStiffnessBiasHasWeightMap,
-							&BucklingStiffnessBias.GetIndices().GetData()[ColorStart],
+							BucklingStiffnessBiasHasWeightMap ? &BucklingStiffnessBias.GetIndices().GetData()[ColorStart] : nullptr,
 							&BucklingStiffnessBias.GetTable().GetData()[0],
 							ColorSize);
 					}
diff --git a/Engine/Source/Runtime/Experimental/Chaos/Private/Chaos/XPBDAnisotropicBendingConstraints.ispc b/Engine/Source/Runtime/Experimental/Chaos/Private/Chaos/XPBDAnisotropicBendingConstraints.ispc
--- a/Engine/Source/Runtime/Experimental/Chaos/Private/Chaos/XPBDAnisotropicBendingConstraints.ispc
+++ b/Engine/Source/Runtime/Experimental/Chaos/Private/Chaos/XPBDAnisotropicBendingConstraints.ispc
@@ -241,12 +241,12 @@ export void ApplyXPBDAnisotropicBendingConstraintsWithMaps(uniform FVector4f Pan
 		const varying int32 i2 = Constraint.V[1];
 		const varying int32 i3 = Constraint.V[2];
 		const varying int32 i4 = Constraint.V[3];
-		const varying int8 StiffnessIndex = select(StiffnessHasMap, StiffnessIndices[i], 0);
-		const varying int8 StiffnessWeftIndex = select(StiffnessWeftHasMap, StiffnessWeftIndices[i], 0);
-		const varying int8 StiffnessBiasIndex = select(StiffnessBiasHasMap, StiffnessBiasIndices[i], 0);
-		const varying int8 BucklingStiffnessIndex = select(BucklingStiffnessHasMap, BucklingStiffnessIndices[i], 0);
-		const varying int8 BucklingStiffnessWeftIndex = select(BucklingStiffnessWeftHasMap, BucklingStiffnessWeftIndices[i], 0);
-		const varying int8 BucklingStiffnessBiasIndex = select(BucklingStiffnessBiasHasMap, BucklingStiffnessBiasIndices[i], 0);
+		const varying int8 StiffnessIndex = StiffnessHasMap ? StiffnessIndices[i] : 0;
+		const varying int8 StiffnessWeftIndex = StiffnessWeftHasMap ? StiffnessWeftIndices[i] : 0;
+		const varying int8 StiffnessBiasIndex = StiffnessBiasHasMap ? StiffnessBiasIndices[i] : 0;
+		const varying int8 BucklingStiffnessIndex = BucklingStiffnessHasMap ? BucklingStiffnessIndices[i] : 0;
+		const varying int8 BucklingStiffnessWeftIndex = BucklingStiffnessWeftHasMap ? BucklingStiffnessWeftIndices[i] : 0;
+		const varying int8 BucklingStiffnessBiasIndex = BucklingStiffnessBiasHasMap ? BucklingStiffnessBiasIndices[i] : 0;
 
 		const varying FVector4f PandInvM2 = VectorGather(&PandInvM[i2]);
 		const varying FVector4f PandInvM3 = VectorGather(&PandInvM[i3]);
@@ -355,13 +355,13 @@ export void ApplyXPBDAnisotropicBendingConstraintsWithDampingAndMaps(uniform FVe
 		const varying int32 i2 = Constraint.V[1];
 		const varying int32 i3 = Constraint.V[2];
 		const varying int32 i4 = Constraint.V[3];
-		const varying int8 StiffnessIndex = select(StiffnessHasMap, StiffnessIndices[i], 0);
-		const varying int8 StiffnessWeftIndex = select(StiffnessWeftHasMap, StiffnessWeftIndices[i], 0);
-		const varying int8 StiffnessBiasIndex = select(StiffnessBiasHasMap, StiffnessBiasIndices[i], 0);
-		const varying int8 BucklingStiffnessIndex = select(BucklingStiffnessHasMap, BucklingStiffnessIndices[i], 0);
-		const varying int8 BucklingStiffnessWeftIndex = select(BucklingStiffnessWeftHasMap, BucklingStiffnessWeftIndices[i], 0);
-		const varying int8 BucklingStiffnessBiasIndex = select(BucklingStiffnessBiasHasMap, BucklingStiffnessBiasIndices[i], 0);
-		const varying int8 DampingIndex = select(DampingHasMap, DampingIndices[i], 0);
+		const varying int8 StiffnessIndex = StiffnessHasMap ? StiffnessIndices[i] : 0;
+		const varying int8 StiffnessWeftIndex = StiffnessWeftHasMap ? StiffnessWeftIndices[i] : 0;
+		const varying int8 StiffnessBiasIndex = StiffnessBiasHasMap ? StiffnessBiasIndices[i] : 0;
+		const varying int8 BucklingStiffnessIndex = BucklingStiffnessHasMap ? BucklingStiffnessIndices[i] : 0;
+		const varying int8 BucklingStiffnessWeftIndex = BucklingStiffnessWeftHasMap ? BucklingStiffnessWeftIndices[i] : 0;
+		const varying int8 BucklingStiffnessBiasIndex = BucklingStiffnessBiasHasMap ? BucklingStiffnessBiasIndices[i] : 0;
+		const varying int8 DampingIndex = DampingHasMap ? DampingIndices[i] : 0;
 
 		const varying FVector4f PandInvM2 = VectorGather(&PandInvM[i2]);
 		const varying FVector4f PandInvM3 = VectorGather(&PandInvM[i3]);
diff --git a/Engine/Source/Runtime/Experimental/Chaos/Private/Chaos/XPBDBendingConstraints.cpp b/Engine/Source/Runtime/Experimental/Chaos/Private/Chaos/XPBDBendingConstraints.cpp
--- a/Engine/Source/Runtime/Experimental/Chaos/Private/Chaos/XPBDBendingConstraints.cpp
+++ b/Engine/Source/Runtime/Experimental/Chaos/Private/Chaos/XPBDBendingConstraints.cpp
@@ -282,13 +282,13 @@ void FXPBDBendingConstraints::Apply(FSolverParticles& Particles, const FSolverRe
 							Dt,
 							MinStiffness,
 							StiffnessHasWeightMap,
-							&Stiffness.GetIndices().GetData()[ColorStart],
+							StiffnessHasWeightMap ? &Stiffness.GetIndices().GetData()[ColorStart] : nullptr,
 							&Stiffness.GetTable().GetData()[0],
 							BucklingStiffnessHasWeightMap,
-							&BucklingStiffness.GetIndices().GetData()[ColorStart],
+							BucklingStiffnessHasWeightMap ? &BucklingStiffness.GetIndices().GetData()[ColorStart] : nullptr,
 							&BucklingStiffness.GetTable().GetData()[0],
 							DampingHasWeightMap,
-							&DampingRatio.GetIndices().GetData()[ColorStart],
+							DampingHasWeightMap ? &DampingRatio.GetIndices().GetData()[ColorStart] : nullptr,
 							&DampingRatio.GetTable().GetData()[0],
 							ColorSize);
 					}
@@ -308,10 +308,10 @@ void FXPBDBendingConstraints::Apply(FSolverParticles& Particles, const FSolverRe
 							Dt,
 							MinStiffness,
 							StiffnessHasWeightMap,
-							&Stiffness.GetIndices().GetData()[ColorStart],
+							StiffnessHasWeightMap ? &Stiffness.GetIndices().GetData()[ColorStart] : nullptr,
 							&Stiffness.GetTable().GetData()[0],
 							BucklingStiffnessHasWeightMap,
-							&BucklingStiffness.GetIndices().GetData()[ColorStart],
+							BucklingStiffnessHasWeightMap ? &BucklingStiffness.GetIndices().GetData()[ColorStart] : nullptr,
 							&BucklingStiffness.GetTable().GetData()[0],
 							ColorSize);
 					}
diff --git a/Engine/Source/Runtime/Experimental/Chaos/Private/Chaos/XPBDBendingConstraints.ispc b/Engine/Source/Runtime/Experimental/Chaos/Private/Chaos/XPBDBendingConstraints.ispc
--- a/Engine/Source/Runtime/Experimental/Chaos/Private/Chaos/XPBDBendingConstraints.ispc
+++ b/Engine/Source/Runtime/Experimental/Chaos/Private/Chaos/XPBDBendingConstraints.ispc
@@ -224,8 +224,8 @@ export void ApplyXPBDBendingConstraintsWithMaps(uniform FVector4f PandInvM[],
 		const varying int32 i2 = Constraint.V[1];
 		const varying int32 i3 = Constraint.V[2];
 		const varying int32 i4 = Constraint.V[3];
-		const varying int8 StiffnessIndex = select(StiffnessHasMap, StiffnessIndices[i], 0);
-		const varying int8 BucklingStiffnessIndex = select(BucklingStiffnessHasMap, BucklingStiffnessIndices[i], 0);
+		const varying int8 StiffnessIndex = StiffnessHasMap ? StiffnessIndices[i] : 0;
+		const varying int8 BucklingStiffnessIndex = BucklingStiffnessHasMap ? BucklingStiffnessIndices[i] : 0;
 
 		const varying FVector4f PandInvM2 = VectorGather(&PandInvM[i2]);
 		const varying FVector4f PandInvM3 = VectorGather(&PandInvM[i3]);
@@ -307,9 +307,9 @@ export void ApplyXPBDBendingConstraintsWithDampingAndMaps(uniform FVector4f Pand
 		const varying int32 i2 = Constraint.V[1];
 		const varying int32 i3 = Constraint.V[2];
 		const varying int32 i4 = Constraint.V[3];
-		const varying int8 StiffnessIndex = select(StiffnessHasMap, StiffnessIndices[i], 0);
-		const varying int8 BucklingStiffnessIndex = select(BucklingStiffnessHasMap, BucklingStiffnessIndices[i], 0);
-		const varying int8 DampingIndex = select(DampingHasMap, DampingIndices[i], 0);
+		const varying int8 StiffnessIndex = StiffnessHasMap ? StiffnessIndices[i] : 0;
+		const varying int8 BucklingStiffnessIndex = BucklingStiffnessHasMap ? BucklingStiffnessIndices[i] : 0;
+		const varying int8 DampingIndex = DampingHasMap ? DampingIndices[i] : 0;
 
 		const varying FVector4f PandInvM2 = VectorGather(&PandInvM[i2]);
 		const varying FVector4f PandInvM3 = VectorGather(&PandInvM[i3]);
diff --git a/Engine/Source/Runtime/Experimental/Chaos/Private/Chaos/XPBDSpringConstraints.cpp b/Engine/Source/Runtime/Experimental/Chaos/Private/Chaos/XPBDSpringConstraints.cpp
--- a/Engine/Source/Runtime/Experimental/Chaos/Private/Chaos/XPBDSpringConstraints.cpp
+++ b/Engine/Source/Runtime/Experimental/Chaos/Private/Chaos/XPBDSpringConstraints.cpp
@@ -181,10 +181,10 @@ void FXPBDSpringConstraints::Apply(FSolverParticles& Particles, const FSolverRea
 							Dt,
 							MinStiffness,
 							StiffnessHasWeightMap,
-							&Stiffness.GetIndices().GetData()[ColorStart],
+							StiffnessHasWeightMap ? &Stiffness.GetIndices().GetData()[ColorStart] : nullptr,
 							&Stiffness.GetTable().GetData()[0],
 							DampingHasWeightMap,
-							&DampingRatio.GetIndices().GetData()[ColorStart],
+							DampingHasWeightMap ? &DampingRatio.GetIndices().GetData()[ColorStart] : nullptr,
 							&DampingRatio.GetTable().GetData()[0],
 							ColorSize);
 					}
diff --git a/Engine/Source/Runtime/Experimental/Chaos/Private/Chaos/XPBDSpringConstraints.ispc b/Engine/Source/Runtime/Experimental/Chaos/Private/Chaos/XPBDSpringConstraints.ispc
--- a/Engine/Source/Runtime/Experimental/Chaos/Private/Chaos/XPBDSpringConstraints.ispc
+++ b/Engine/Source/Runtime/Experimental/Chaos/Private/Chaos/XPBDSpringConstraints.ispc
@@ -208,8 +208,8 @@ export void ApplyXPBDSpringConstraintsWithDampingAndWeightMaps(uniform FVector4f
 		const varying FIntVector2 Constraint = VectorLoad(&Constraints[extract(i,0)]);
 		const varying int32 i1 = Constraint.V[0];
 		const varying int32 i2 = Constraint.V[1];
-		const varying int8 StiffnessIndex = select(StiffnessHasMap, StiffnessIndices[i], 0);
-		const varying int8 DampingIndex = select(DampingHasMap, DampingIndices[i], 0);
+		const varying int8 StiffnessIndex = StiffnessHasMap ? StiffnessIndices[i] : 0;
+		const varying int8 DampingIndex = DampingHasMap ? DampingIndices[i] : 0;
 
 		const varying FVector4f PandInvM1 = VectorGather(&PandInvM[i1]);
 		const varying FVector4f PandInvM2 = VectorGather(&PandInvM[i2]);
diff --git a/Engine/Source/Runtime/Experimental/Chaos/Private/Chaos/XPBDStretchBiasElementConstraints.cpp b/Engine/Source/Runtime/Experimental/Chaos/Private/Chaos/XPBDStretchBiasElementConstraints.cpp
--- a/Engine/Source/Runtime/Experimental/Chaos/Private/Chaos/XPBDStretchBiasElementConstraints.cpp
+++ b/Engine/Source/Runtime/Experimental/Chaos/Private/Chaos/XPBDStretchBiasElementConstraints.cpp
@@ -538,22 +538,22 @@ void FXPBDStretchBiasElementConstraints::Apply(FSolverParticles& Particles, cons
 							Dt,
 							MinStiffness,
 							StiffnessWarpHasWeightMap,
-							&StiffnessWarp.GetIndices().GetData()[ColorStart],
+							StiffnessWarpHasWeightMap ? &StiffnessWarp.GetIndices().GetData()[ColorStart] : nullptr,
 							&StiffnessWarp.GetTable().GetData()[0],
 							StiffnessWeftHasWeightMap,
-							&StiffnessWeft.GetIndices().GetData()[ColorStart],
+							StiffnessWeftHasWeightMap ? &StiffnessWeft.GetIndices().GetData()[ColorStart] : nullptr,
 							&StiffnessWeft.GetTable().GetData()[0],
 							StiffnessBiasHasWeightMap,
-							&StiffnessBias.GetIndices().GetData()[ColorStart],
+							StiffnessBiasHasWeightMap ? &StiffnessBias.GetIndices().GetData()[ColorStart] : nullptr,
 							&StiffnessBias.GetTable().GetData()[0],
 							DampingHasWeightMap,
-							&DampingRatio.GetIndices().GetData()[ColorStart],
+							DampingHasWeightMap ? &DampingRatio.GetIndices().GetData()[ColorStart] : nullptr,
 							&DampingRatio.GetTable().GetData()[0],
 							WarpScaleHasWeightMap,
-							&WarpScale.GetIndices().GetData()[ColorStart],
+							WarpScaleHasWeightMap ? &WarpScale.GetIndices().GetData()[ColorStart] : nullptr,
 							&WarpScale.GetTable().GetData()[0],
 							WeftScaleHasWeightMap,
-							&WeftScale.GetIndices().GetData()[ColorStart],
+							WeftScaleHasWeightMap ? &WeftScale.GetIndices().GetData()[ColorStart] : nullptr,
 							&WeftScale.GetTable().GetData()[0],
 							ColorSize);
 					}
@@ -575,19 +575,19 @@ void FXPBDStretchBiasElementConstraints::Apply(FSolverParticles& Particles, cons
 							Dt,
 							MinStiffness,
 							StiffnessWarpHasWeightMap,
-							&StiffnessWarp.GetIndices().GetData()[ColorStart],
+							StiffnessWarpHasWeightMap ? &StiffnessWarp.GetIndices().GetData()[ColorStart] : nullptr,
 							&StiffnessWarp.GetTable().GetData()[0],
 							StiffnessWeftHasWeightMap,
-							&StiffnessWeft.GetIndices().GetData()[ColorStart],
+							StiffnessWeftHasWeightMap ? &StiffnessWeft.GetIndices().GetData()[ColorStart] : nullptr,
 							&StiffnessWeft.GetTable().GetData()[0],
 							StiffnessBiasHasWeightMap,
-							&StiffnessBias.GetIndices().GetData()[ColorStart],
+							StiffnessBiasHasWeightMap ? &StiffnessBias.GetIndices().GetData()[ColorStart] : nullptr,
 							&StiffnessBias.GetTable().GetData()[0],
 							WarpScaleHasWeightMap,
-							&WarpScale.GetIndices().GetData()[ColorStart],
+							WarpScaleHasWeightMap ? &WarpScale.GetIndices().GetData()[ColorStart] : nullptr,
 							&WarpScale.GetTable().GetData()[0],
 							WeftScaleHasWeightMap,
-							&WeftScale.GetIndices().GetData()[ColorStart],
+							WeftScaleHasWeightMap ? &WeftScale.GetIndices().GetData()[ColorStart] : nullptr,
 							&WeftScale.GetTable().GetData()[0],
 							ColorSize);
 					}
diff --git a/Engine/Source/Runtime/Experimental/Chaos/Private/Chaos/XPBDStretchBiasElementConstraints.ispc b/Engine/Source/Runtime/Experimental/Chaos/Private/Chaos/XPBDStretchBiasElementConstraints.ispc
--- a/Engine/Source/Runtime/Experimental/Chaos/Private/Chaos/XPBDStretchBiasElementConstraints.ispc
+++ b/Engine/Source/Runtime/Experimental/Chaos/Private/Chaos/XPBDStretchBiasElementConstraints.ispc
@@ -259,11 +259,11 @@ export void ApplyXPBDStretchBiasConstraintsWithMaps(uniform FVector4f PandInvM[]
 		const varying int32 i0 = Constraint.V[0];
 		const varying int32 i1 = Constraint.V[1];
 		const varying int32 i2 = Constraint.V[2];
-		const varying int8 StiffnessWarpIndex = select(StiffnessWarpHasMap, StiffnessWarpIndices[i], 0);
-		const varying int8 StiffnessWeftIndex = select(StiffnessWeftHasMap, StiffnessWeftIndices[i], 0);
-		const varying int8 StiffnessBiasIndex = select(StiffnessBiasHasMap, StiffnessBiasIndices[i], 0);
-		const varying int8 WarpScaleIndex = select(WarpScaleHasMap, WarpScaleIndices[i], 0);
-		const varying int8 WeftScaleIndex = select(WeftScaleHasMap, WeftScaleIndices[i], 0);
+		const varying int8 StiffnessWarpIndex = StiffnessWarpHasMap ? StiffnessWarpIndices[i] : 0;
+		const varying int8 StiffnessWeftIndex = StiffnessWeftHasMap ? StiffnessWeftIndices[i] : 0;
+		const varying int8 StiffnessBiasIndex = StiffnessBiasHasMap ? StiffnessBiasIndices[i] : 0;
+		const varying int8 WarpScaleIndex = WarpScaleHasMap ? WarpScaleIndices[i] : 0;
+		const varying int8 WeftScaleIndex = WeftScaleHasMap ? WeftScaleIndices[i] : 0;
 		
 		const varying FVector4f PandInvM0 = VectorGather(&PandInvM[i0]);
 		const varying FVector4f PandInvM1 = VectorGather(&PandInvM[i1]);
@@ -364,12 +364,12 @@ export void ApplyXPBDStretchBiasConstraintsWithDampingAndMaps(uniform FVector4f
 		const varying int32 i0 = Constraint.V[0];
 		const varying int32 i1 = Constraint.V[1];
 		const varying int32 i2 = Constraint.V[2];
-		const varying int8 StiffnessWarpIndex = select(StiffnessWarpHasMap, StiffnessWarpIndices[i], 0);
-		const varying int8 StiffnessWeftIndex = select(StiffnessWeftHasMap, StiffnessWeftIndices[i], 0);
-		const varying int8 StiffnessBiasIndex = select(StiffnessBiasHasMap, StiffnessBiasIndices[i], 0);
-		const varying int8 WarpScaleIndex = select(WarpScaleHasMap, WarpScaleIndices[i], 0);
-		const varying int8 WeftScaleIndex = select(WeftScaleHasMap, WeftScaleIndices[i], 0);
-		const varying int8 DampingIndex = select(DampingHasMap, DampingIndices[i], 0);
+		const varying int8 StiffnessWarpIndex = StiffnessWarpHasMap ? StiffnessWarpIndices[i] : 0;
+		const varying int8 StiffnessWeftIndex = StiffnessWeftHasMap ? StiffnessWeftIndices[i] : 0;
+		const varying int8 StiffnessBiasIndex = StiffnessBiasHasMap ? StiffnessBiasIndices[i] : 0;
+		const varying int8 WarpScaleIndex = WarpScaleHasMap ? WarpScaleIndices[i] : 0;
+		const varying int8 WeftScaleIndex = WeftScaleHasMap ? WeftScaleIndices[i] : 0;
+		const varying int8 DampingIndex = DampingHasMap ? DampingIndices[i] : 0;
 		
 		const varying FVector4f PandInvM0 = VectorGather(&PandInvM[i0]);
 		const varying FVector4f PandInvM1 = VectorGather(&PandInvM[i1]);
