Quick fix in hopes of preventing crashes when handle ref cache goes bad.

#rnx

diff --git a/Engine/Source/Runtime/Experimental/Iris/Core/Private/Iris/ReplicationSystem/ObjectReferenceCache.cpp b/Engine/Source/Runtime/Experimental/Iris/Core/Private/Iris/ReplicationSystem/ObjectReferenceCache.cpp
--- a/Engine/Source/Runtime/Experimental/Iris/Core/Private/Iris/ReplicationSystem/ObjectReferenceCache.cpp
+++ b/Engine/Source/Runtime/Experimental/Iris/Core/Private/Iris/ReplicationSystem/ObjectReferenceCache.cpp
@@ -403,7 +403,17 @@ FNetRefHandle FObjectReferenceCache::GetObjectReferenceHandleFromObject(const UO
 	if (const FNetRefHandle* Reference = ObjectToNetReferenceHandle.Find(Object))
 	{
 		// Verify that this is the same object
-		const FCachedNetObjectReference& CachedObject = ReferenceHandleToCachedReference.FindChecked(*Reference);
+
+		const FCachedNetObjectReference* CachedObjectPtr = ReferenceHandleToCachedReference.Find(*Reference);
+		if (CachedObjectPtr == nullptr)
+		{
+			ensureAlwaysMsgf(CachedObjectPtr != nullptr, TEXT("Mismatch between ReferenceCache maps. Object %s (0x%p) has handle (%s) but no cache."), *GetNameSafe(Object), Object, *Reference->ToString());
+			const_cast<TMap<const UObject*, FNetRefHandle>&>(ObjectToNetReferenceHandle).Remove(Object);
+			UE_LOG_REFERENCECACHE(Warning, TEXT("ObjectReferenceCache::GetObjectReferenceHandleFromObject removed %s due to mismatched object %s (0x%p) "), *Reference->ToString(), *GetNameSafe(Object), Object);
+			return FNetRefHandle();
+		}
+
+		const FCachedNetObjectReference& CachedObject = *CachedObjectPtr;
 
 		const UObject* ExistingObject = CachedObject.Object.Get();
 		if (ExistingObject == Object)
