In the Geometry Collection Component the async tick had to block Game Thread and Physics Thread to apply the replicate data. Now all the replication code is copied and executed on the Physics thread. This allow not to block both thread at the same time.

diff --git a/Engine/Source/Runtime/Experimental/Chaos/Public/Chaos/SimCallbackObject.h b/Engine/Source/Runtime/Experimental/Chaos/Public/Chaos/SimCallbackObject.h
--- a/Engine/Source/Runtime/Experimental/Chaos/Public/Chaos/SimCallbackObject.h
+++ b/Engine/Source/Runtime/Experimental/Chaos/Public/Chaos/SimCallbackObject.h
@@ -337,6 +337,14 @@ class FSimCallbackCommandObject : public ISimCallbackObject
 public:
 	FSimCallbackCommandObject(TUniqueFunction<void()>&& InFunc)
 		: Func(MoveTemp(InFunc))
+		, Func2(nullptr)
+		, bFuncHasTimeParameters(false)
+	{}
+
+	FSimCallbackCommandObject(TUniqueFunction<void(FReal DeltaTime, FReal SimTime)>&& InFunc)
+		: Func(nullptr)
+		, Func2(MoveTemp(InFunc))
+		, bFuncHasTimeParameters(true)
 	{}
 
 	virtual void FreeOutputData_External(FSimCallbackOutput* Output)
@@ -368,10 +376,22 @@ private:
 
 	virtual void OnPreSimulate_Internal() override
 	{
-		Func();
+		if (!bFuncHasTimeParameters)
+		{ 
+			Func();
+		}
+		else
+		{
+			Func2(GetDeltaTime_Internal(), GetSimTime_Internal());
+		}
+		
 	}
 
+	// Those two function could be in an union 
 	TUniqueFunction<void()> Func;
+	TUniqueFunction<void(FReal deltaTime, FReal SimTime)> Func2;
+
+	bool bFuncHasTimeParameters;
 
 	friend struct FSimCallbackInput;
 };
diff --git a/Engine/Source/Runtime/Experimental/GeometryCollectionEngine/Private/GeometryCollection/GeometryCollectionComponent.cpp b/Engine/Source/Runtime/Experimental/GeometryCollectionEngine/Private/GeometryCollection/GeometryCollectionComponent.cpp
--- a/Engine/Source/Runtime/Experimental/GeometryCollectionEngine/Private/GeometryCollection/GeometryCollectionComponent.cpp
+++ b/Engine/Source/Runtime/Experimental/GeometryCollectionEngine/Private/GeometryCollection/GeometryCollectionComponent.cpp
@@ -1707,17 +1707,38 @@ void ActivateClusters(Chaos::FRigidClustering& Clustering, Chaos::TPBDRigidClust
 	Clustering.DeactivateClusterParticle(Cluster);
 }
 
-void UGeometryCollectionComponent::ResetRepData()
+void UGeometryCollectionComponent::ResetRepDataCommon()
 {
-	ClustersToRep.Reset();
-	RepData.Reset();
 	OneOffActivatedProcessed = 0;
 	VersionProcessed = INDEX_NONE;
 	LastHardsnapTimeInMs = 0;
 }
 
+void UGeometryCollectionComponent::ResetRepData()
+{
+	ClustersToRep.Reset();
+	RepData.Reset();
+
+	if (GetNetMode() == ENetMode::NM_Client)
+	{
+		// Reset on the Physics Thread no to read/write on the client from different thread
+		if (Chaos::FPhysicsSolver* CurrSolver = GetSolver(*this))
+		{
+			CurrSolver->EnqueueCommandImmediate([this]()
+				{
+					ResetRepDataCommon();
+				});
+		}
+	}
+	else
+	{
+		ResetRepDataCommon();
+	}
+}
+
 void UGeometryCollectionComponent::UpdateRepData()
 {
+	check(GetNetMode() != ENetMode::NM_Client);
 	using namespace Chaos;
 	if(!bEnableReplication)
 	{
@@ -1995,200 +2016,151 @@ void UGeometryCollectionComponent::ProcessRepData()
 	ProcessRepData(0.f, 0.f);
 }
 
-bool UGeometryCollectionComponent::ProcessRepData(const float DeltaTime, const float SimTime)
-{
-	using namespace Chaos;
-
-	if (!PhysicsProxy || !PhysicsProxy->IsInitializedOnPhysicsThread() || PhysicsProxy->GetReplicationMode() != FGeometryCollectionPhysicsProxy::EReplicationMode::Client)
+namespace {
+	bool ProcessRepDataCommon(const float DeltaTime, const float SimTime, const FGeometryCollectionRepData& RepData, FGeometryCollectionPhysicsProxy* PhysicsProxy, int32& VersionProcessed, int32& OneOffActivatedProcessed, double& LastHardsnapTimeInMs, float ReceivedTime, bool bReplicateMovement)
 	{
-		return false;
-	}
+		using namespace Chaos;
 
-	// Track the sim time that this rep data was received on.
-	if (!RepData.RepDataReceivedTime.IsSet())
-	{
-		RepData.RepDataReceivedTime = SimTime;
-	}
-
-	// How far we must extrapolate from when we received the data
-	const float RepExtrapTime = FMath::Max(0.f, SimTime - *RepData.RepDataReceivedTime);
+		if (!PhysicsProxy || !PhysicsProxy->IsInitializedOnPhysicsThread() || PhysicsProxy->GetReplicationMode() != FGeometryCollectionPhysicsProxy::EReplicationMode::Client)
+		{
+			return false;
+		}
 
-	// If we've extrapolated past a threshold, then stop tracking
-	// the last received rep data
-	if (RepExtrapTime > GeometryCollectionRepMaxExtrapolationTime)
-	{
-		return false;
-	}
+		// How far we must extrapolate from when we received the data
+		const float RepExtrapTime = FMath::Max(0.f, SimTime - ReceivedTime);
 
-	// Create a little little function for applying a lambda to each
-	// corresponding pair of replicated and local clusters.
-	const auto ForEachClusterPair = [this](const auto& Lambda)
-	{
-		for (const FGeometryCollectionClusterRep& RepCluster : RepData.Clusters)
+		// If we've extrapolated past a threshold, then stop tracking
+		// the last received rep data
+		if (RepExtrapTime > GeometryCollectionRepMaxExtrapolationTime)
 		{
-			if (Chaos::FPBDRigidParticleHandle* Cluster = PhysicsProxy->GetParticles()[RepCluster.ClusterIdx])
-			{
-				if (RepCluster.ClusterState.IsInternalCluster())
-				{
-					// internal cluster do not have an index so we rep data send one of the children's
-					// let's find the parent
-					Cluster = Cluster->CastToClustered()->Parent();
-				}
-
-				if (Cluster && Cluster->Disabled() == false)
-				{
-					Lambda(RepCluster, *Cluster);
-				}
-			}
+			return false;
 		}
-	};
 
-#if ENABLE_DRAW_DEBUG
-	if (bGeometryCollectionDebugDrawRep)
-	{
-		ForEachClusterPair([RepExtrapTime](const FGeometryCollectionClusterRep& RepCluster, Chaos::FPBDRigidParticleHandle& Cluster)
+		// Create a little little function for applying a lambda to each
+		// corresponding pair of replicated and local clusters.
+		const auto ForEachClusterPair = [&RepData, &PhysicsProxy](const auto& Lambda)
+		{
+			for (const FGeometryCollectionClusterRep& RepCluster : RepData.Clusters)
 			{
-				// Don't bother debug drawing if the delta is too small
-				if ((Cluster.X() - RepCluster.Position).SizeSquared() < .1f)
+				if (Chaos::FPBDRigidParticleHandle* Cluster = PhysicsProxy->GetParticles()[RepCluster.ClusterIdx])
 				{
-					FVector Axis;
-					float Angle;
-					(RepCluster.Rotation.Inverse() * Cluster.R()).ToAxisAndAngle(Axis, Angle);
-					if (FMath::Abs(Angle) < .1f)
+					if (RepCluster.ClusterState.IsInternalCluster())
 					{
-						return;
+						// internal cluster do not have an index so we rep data send one of the children's
+						// let's find the parent
+						Cluster = Cluster->CastToClustered()->Parent();
 					}
-				}
 
-				Chaos::FDebugDrawQueue& DrawQueue = Chaos::FDebugDrawQueue::GetInstance();
-				DrawQueue.DrawDebugCoordinateSystem(Cluster.X(), FRotator(Cluster.R()), 100.f, false, -1, -1, 1.f);
-				DrawQueue.DrawDebugBox(Cluster.X() + Cluster.LocalBounds().Center(), Cluster.LocalBounds().Extents(), Cluster.R(), FColor::White, false, -1, -1, 1.f);
-				DrawQueue.DrawDebugBox(RepCluster.Position + Cluster.LocalBounds().Center(), Cluster.LocalBounds().Extents(), RepCluster.Rotation, FColor::Green, false, -1, -1, 1.f);
-
-				if (bGeometryCollectionRepUseClusterVelocityMatch)
-				{
-					const FVector RepVel = RepCluster.LinearVelocity;
-					const FVector RepAngVel = RepCluster.AngularVelocity;
-					const FVector RepExtrapPos = RepCluster.Position + (RepVel * RepExtrapTime);
-					const Chaos::FRotation3 RepExtrapAng = Chaos::FRotation3::IntegrateRotationWithAngularVelocity(RepCluster.Rotation, RepAngVel, RepExtrapTime);
-					DrawQueue.DrawDebugCoordinateSystem(RepExtrapPos, FRotator(RepExtrapAng), 100.f, false, -1, -1, 1.f);
-					DrawQueue.DrawDebugDirectionalArrow(Cluster.X(), RepExtrapPos, 10.f, FColor::White, false, -1, -1, 1.f);
-					DrawQueue.DrawDebugBox(RepExtrapPos + Cluster.LocalBounds().Center(), Cluster.LocalBounds().Extents(), RepExtrapAng, FColor::Orange, false, -1, -1, 1.f);
-				}
-				else
-				{
-					DrawQueue.DrawDebugCoordinateSystem(RepCluster.Position, FRotator(RepCluster.Rotation), 100.f, false, -1, -1, 1.f);
+					if (Cluster && Cluster->Disabled() == false)
+					{
+						Lambda(RepCluster, *Cluster);
+					}
 				}
-			});
-	}
-#endif
+			}
+		};
 
-	// If not doing velocity match, don't bother processing the same version twice.
-	// Do this one after the debug draw so that we can still easily see the diff
-	// between the position and the target position.
-	if (VersionProcessed == RepData.Version && !bGeometryCollectionRepUseClusterVelocityMatch)
-	{
-		return false;
-	}
 
-	// Update the anchored state of the root particle
-	const int32 InitialRootIndex = PhysicsProxy->GetSimParameters().InitialRootIndex;
-	const TArray<FPBDRigidClusteredParticleHandle*>& ParticleHandles = PhysicsProxy->GetParticles();
-	if (FPBDRigidClusteredParticleHandle* RootHandle = ParticleHandles[InitialRootIndex])
-	{
-		if (RootHandle->Parent() || !RootHandle->Disabled())
+		// If not doing velocity match, don't bother processing the same version twice.
+		// Do this one after the debug draw so that we can still easily see the diff
+		// between the position and the target position.
+		if (VersionProcessed == RepData.Version && !bGeometryCollectionRepUseClusterVelocityMatch)
 		{
-			if (RepData.bIsRootAnchored != RootHandle->IsAnchored())
-			{
-				RootHandle->SetIsAnchored(RepData.bIsRootAnchored);
+			return false;
+		}
 
-				if (FPBDRigidsSolver* Solver = PhysicsProxy->GetSolver<Chaos::FPBDRigidsSolver>())
+		// Update the anchored state of the root particle
+		const int32 InitialRootIndex = PhysicsProxy->GetSimParameters().InitialRootIndex;
+		const TArray<FPBDRigidClusteredParticleHandle*>& ParticleHandles = PhysicsProxy->GetParticles();
+		if (FPBDRigidClusteredParticleHandle* RootHandle = ParticleHandles[InitialRootIndex])
+		{
+			if (RootHandle->Parent() || !RootHandle->Disabled())
+			{
+				if (RepData.bIsRootAnchored != RootHandle->IsAnchored())
 				{
-					if (FPBDRigidsEvolutionGBF* Evolution = Solver->GetEvolution())
+					RootHandle->SetIsAnchored(RepData.bIsRootAnchored);
+
+					if (FPBDRigidsSolver* Solver = PhysicsProxy->GetSolver<Chaos::FPBDRigidsSolver>())
 					{
-						FRigidClustering& RigidClustering = Solver->GetEvolution()->GetRigidClustering();
-						// NOTE: We have to start out with the particle as kinematic and let UpdateKinematicProperties correct it if this is wrong
-						Evolution->SetParticleObjectState(RootHandle, Chaos::EObjectStateType::Kinematic);
-						Chaos::UpdateKinematicProperties(RootHandle, RigidClustering.GetChildrenMap(), *Evolution);
+						if (FPBDRigidsEvolutionGBF* Evolution = Solver->GetEvolution())
+						{
+							FRigidClustering& RigidClustering = Solver->GetEvolution()->GetRigidClustering();
+							// NOTE: We have to start out with the particle as kinematic and let UpdateKinematicProperties correct it if this is wrong
+							Evolution->SetParticleObjectState(RootHandle, Chaos::EObjectStateType::Kinematic);
+							Chaos::UpdateKinematicProperties(RootHandle, RigidClustering.GetChildrenMap(), *Evolution);
+						}
 					}
 				}
 			}
 		}
-	}
-
-	AActor* Owner = GetOwner();
-	const bool bReplicateMovement = Owner ? GetOwner()->IsReplicatingMovement() : true;
-
-	bool bHardSnap = false;
-
-	if (bReplicateMovement)
-	{
-		const int64 CurrentTimeInMs = FPlatformTime::ToMilliseconds64(FPlatformTime::Cycles64());
-
-		// Always hard snap on the very first version received
-		if (VersionProcessed == 0)
-		{
-			bHardSnap = true;
-		}
-		else if (VersionProcessed < RepData.Version)
+		bool bHardSnap = false;
+		if (bReplicateMovement)
 		{
-			//TODO: this will not really work if a fracture happens and then immediately goes to sleep without updating client enough times
-			//A time method would work better here, but is limited to async mode. Maybe we can support both
-			bHardSnap = (RepData.Version - VersionProcessed) > GeometryCollectionHardMissingUpdatesSnapThreshold;
+			const int64 CurrentTimeInMs = FPlatformTime::ToMilliseconds64(FPlatformTime::Cycles64());
 
-			if (!bGeometryCollectionRepUseClusterVelocityMatch)
+			// Always hard snap on the very first version received
+			if (VersionProcessed == 0)
 			{
-				// When not doing velocity match for clusters, instead we do periodic hard snapping
-				bHardSnap |= (CurrentTimeInMs - LastHardsnapTimeInMs) > GeometryCollectionHardsnapThresholdMs;
+				bHardSnap = true;
 			}
-		}
-		else if (VersionProcessed > RepData.Version)
-		{
-			//rollover so just treat as hard snap - this case is extremely rare and a one off
-			bHardSnap = true;
-		}
+			else if (VersionProcessed < RepData.Version)
+			{
+				//TODO: this will not really work if a fracture happens and then immediately goes to sleep without updating client enough times
+				//A time method would work better here, but is limited to async mode. Maybe we can support both
+				bHardSnap = (RepData.Version - VersionProcessed) > GeometryCollectionHardMissingUpdatesSnapThreshold;
 
-		if (bHardSnap)
-		{
-			LastHardsnapTimeInMs = CurrentTimeInMs;
-		}
-	}
+				if (!bGeometryCollectionRepUseClusterVelocityMatch)
+				{
+					// When not doing velocity match for clusters, instead we do periodic hard snapping
+					bHardSnap |= (CurrentTimeInMs - LastHardsnapTimeInMs) > GeometryCollectionHardsnapThresholdMs;
+				}
+			}
+			else if (VersionProcessed > RepData.Version)
+			{
+				//rollover so just treat as hard snap - this case is extremely rare and a one off
+				bHardSnap = true;
+			}
 
-	FPBDRigidsSolver* Solver = PhysicsProxy->GetSolver<Chaos::FPBDRigidsSolver>();
-	FRigidClustering& RigidClustering = Solver->GetEvolution()->GetRigidClustering();
+			if (bHardSnap)
+			{
+				LastHardsnapTimeInMs = CurrentTimeInMs;
+			}
+		}
 
-	//First make sure all one off activations have been applied. This ensures our connectivity graph is the same and we have the same clusters as the server
-	for (; OneOffActivatedProcessed < RepData.OneOffActivated.Num(); ++OneOffActivatedProcessed)
-	{
-		const FGeometryCollectionActivatedCluster& ActivatedCluster = RepData.OneOffActivated[OneOffActivatedProcessed];
-		FPBDRigidParticleHandle* OneOff = PhysicsProxy->GetParticles()[ActivatedCluster.ActivatedIndex];
+		FPBDRigidsSolver* Solver = PhysicsProxy->GetSolver<Chaos::FPBDRigidsSolver>();
+		FRigidClustering& RigidClustering = Solver->GetEvolution()->GetRigidClustering();
 
-		if (ensure(OneOff))
+		//First make sure all one off activations have been applied. This ensures our connectivity graph is the same and we have the same clusters as the server
+		for (; OneOffActivatedProcessed < RepData.OneOffActivated.Num(); ++OneOffActivatedProcessed)
 		{
-			if (FPBDRigidClusteredParticleHandle* ClusterParticle = OneOff->CastToClustered())
+			const FGeometryCollectionActivatedCluster& ActivatedCluster = RepData.OneOffActivated[OneOffActivatedProcessed];
+			FPBDRigidParticleHandle* OneOff = PhysicsProxy->GetParticles()[ActivatedCluster.ActivatedIndex];
+
+			if (ensure(OneOff))
 			{
-				// Set initial velocities if not hard snapping
-				if (!bHardSnap)
+				if (FPBDRigidClusteredParticleHandle* ClusterParticle = OneOff->CastToClustered())
 				{
-					// TODO: we should get an update cluster position first so that when particles break off they get the right position 
-					// TODO: should we invalidate?
-					OneOff->SetV(ActivatedCluster.InitialLinearVelocity);
-					OneOff->SetW(ActivatedCluster.InitialAngularVelocity);
-				}
+					// Set initial velocities if not hard snapping
+					if (!bHardSnap)
+					{
+						// TODO: we should get an update cluster position first so that when particles break off they get the right position 
+						// TODO: should we invalidate?
+						OneOff->SetV(ActivatedCluster.InitialLinearVelocity);
+						OneOff->SetW(ActivatedCluster.InitialAngularVelocity);
+					}
 
-				// OneOffActivated is not garanteed to be sorted by level, so we may need to release particles before their parent get released
-				// so we make sure the full parent chain is properly released as well 
-				RigidClustering.ForceReleaseChildParticleAndParents(ClusterParticle, /* bTriggerBreakEvents */true);
+					// OneOffActivated is not guaranteed to be sorted by level, so we may need to release particles before their parent get released
+					// so we make sure the full parent chain is properly released as well 
+					RigidClustering.ForceReleaseChildParticleAndParents(ClusterParticle, /* bTriggerBreakEvents */true);
+				}
 			}
 		}
-	}
 
-	// Keep track of whether we did some "work" on this frame so we can turn off the async tick after
-	// multiple frames of not doing anything.
-	bool bProcessed = false;
-	if (bReplicateMovement)
-	{
-		ForEachClusterPair([Solver, DeltaTime, RepExtrapTime, bHardSnap, &bProcessed](const FGeometryCollectionClusterRep& RepCluster, Chaos::FPBDRigidParticleHandle& Cluster)
+		// Keep track of whether we did some "work" on this frame so we can turn off the async tick after
+		// multiple frames of not doing anything.
+		bool bProcessed = false;
+		if (bReplicateMovement)
+		{
+			ForEachClusterPair([Solver, DeltaTime, RepExtrapTime, bHardSnap, &bProcessed](const FGeometryCollectionClusterRep& RepCluster, Chaos::FPBDRigidParticleHandle& Cluster)
 			{
 				bool bWake = false;
 
@@ -2245,10 +2217,135 @@ bool UGeometryCollectionComponent::ProcessRepData(const float DeltaTime, const f
 					Solver->GetEvolution()->SetParticleObjectState(&Cluster, Chaos::EObjectStateType::Dynamic);
 				}
 			});
+		}
+
+		VersionProcessed = RepData.Version;
+		return bProcessed;
+	}
+}
+
+bool UGeometryCollectionComponent::ProcessRepData(const float DeltaTime, const float SimTime)
+{
+
+	// Track the sim time that this rep data was received on.
+	if (!RepData.RepDataReceivedTime.IsSet())
+	{
+		RepData.RepDataReceivedTime = SimTime;
+	}
+
+	float ReceivedTime = *RepData.RepDataReceivedTime;
+
+	AActor* Owner = GetOwner();
+	const bool bReplicateMovement = Owner ? GetOwner()->IsReplicatingMovement() : true;
+
+	bool ReturnValue = ::ProcessRepDataCommon(DeltaTime, SimTime, RepData, PhysicsProxy, VersionProcessed, OneOffActivatedProcessed, LastHardsnapTimeInMs, ReceivedTime, bReplicateMovement);
+
+
+#if ENABLE_DRAW_DEBUG
+	
+	if (bGeometryCollectionDebugDrawRep)
+	{
+		// Track the sim time that this rep data was received on.
+		if (!RepData.RepDataReceivedTime.IsSet())
+		{
+			RepData.RepDataReceivedTime = SimTime;
+		}
+
+		// How far we must extrapolate from when we received the data
+		const float RepExtrapTime = FMath::Max(0.f, SimTime - *RepData.RepDataReceivedTime);
+
+		// Create a little little function for applying a lambda to each
+		// corresponding pair of replicated and local clusters.
+		const auto ForEachClusterPair = [this](const auto& Lambda)
+		{
+			for (const FGeometryCollectionClusterRep& RepCluster : RepData.Clusters)
+			{
+				if (Chaos::FPBDRigidParticleHandle* Cluster = PhysicsProxy->GetParticles()[RepCluster.ClusterIdx])
+				{
+					if (RepCluster.ClusterState.IsInternalCluster())
+					{
+						// internal cluster do not have an index so we rep data send one of the children's
+						// let's find the parent
+						Cluster = Cluster->CastToClustered()->Parent();
+					}
+
+					if (Cluster && Cluster->Disabled() == false)
+					{
+						Lambda(RepCluster, *Cluster);
+					}
+				}
+			}
+		};
+
+		ForEachClusterPair([RepExtrapTime](const FGeometryCollectionClusterRep& RepCluster, Chaos::FPBDRigidParticleHandle& Cluster)
+		{
+			// Don't bother debug drawing if the delta is too small
+			if ((Cluster.X() - RepCluster.Position).SizeSquared() < .1f)
+			{
+				FVector Axis;
+				float Angle;
+				(RepCluster.Rotation.Inverse() * Cluster.R()).ToAxisAndAngle(Axis, Angle);
+				if (FMath::Abs(Angle) < .1f)
+				{
+					return;
+				}
+			}
+
+			Chaos::FDebugDrawQueue& DrawQueue = Chaos::FDebugDrawQueue::GetInstance();
+			DrawQueue.DrawDebugCoordinateSystem(Cluster.X(), FRotator(Cluster.R()), 100.f, false, -1, -1, 1.f);
+			DrawQueue.DrawDebugBox(Cluster.X() + Cluster.LocalBounds().Center(), Cluster.LocalBounds().Extents(), Cluster.R(), FColor::White, false, -1, -1, 1.f);
+			DrawQueue.DrawDebugBox(RepCluster.Position + Cluster.LocalBounds().Center(), Cluster.LocalBounds().Extents(), RepCluster.Rotation, FColor::Green, false, -1, -1, 1.f);
+
+			if (bGeometryCollectionRepUseClusterVelocityMatch)
+			{
+				const FVector RepVel = RepCluster.LinearVelocity;
+				const FVector RepAngVel = RepCluster.AngularVelocity;
+				const FVector RepExtrapPos = RepCluster.Position + (RepVel * RepExtrapTime);
+				const Chaos::FRotation3 RepExtrapAng = Chaos::FRotation3::IntegrateRotationWithAngularVelocity(RepCluster.Rotation, RepAngVel, RepExtrapTime);
+				DrawQueue.DrawDebugCoordinateSystem(RepExtrapPos, FRotator(RepExtrapAng), 100.f, false, -1, -1, 1.f);
+				DrawQueue.DrawDebugDirectionalArrow(Cluster.X(), RepExtrapPos, 10.f, FColor::White, false, -1, -1, 1.f);
+				DrawQueue.DrawDebugBox(RepExtrapPos + Cluster.LocalBounds().Center(), Cluster.LocalBounds().Extents(), RepExtrapAng, FColor::Orange, false, -1, -1, 1.f);
+			}
+			else
+			{
+				DrawQueue.DrawDebugCoordinateSystem(RepCluster.Position, FRotator(RepCluster.Rotation), 100.f, false, -1, -1, 1.f);
+			}
+		});
+	}
+#endif
+
+	return ReturnValue;
+}
+
+void UGeometryCollectionComponent::ProcessRepDataOnPT()
+{
+	using namespace Chaos;
+
+	if (Chaos::FPhysicsSolver* CurrSolver = GetSolver(*this))
+	{
+		AActor* Owner = GetOwner();
+		const bool bReplicateMovement = Owner ? GetOwner()->IsReplicatingMovement() : true;
+
+		// Capture - VersionProcessed - LastHardsnapTimeInMs - OneOffActivatedProcessed
+		// which are only used in that function and in the reset function and both are executed on the Physics Thread
+		CurrSolver->EnqueueCommandImmediate([PhysicsProxy = PhysicsProxy, RepDataCopy = RepData, &VersionProcessed = VersionProcessed, &LastHardsnapTimeInMs = LastHardsnapTimeInMs, &OneOffActivatedProcessed = OneOffActivatedProcessed, bReplicateMovement](Chaos::FReal DeltaTime, Chaos::FReal SimTime)
+		{
+				float ReceivedTime;
+
+				// Track the sim time that this rep data was received on.
+				if (!RepDataCopy.RepDataReceivedTime.IsSet())
+				{
+					ReceivedTime = SimTime;
+				}
+				else
+				{
+					ReceivedTime = *RepDataCopy.RepDataReceivedTime;
+				}
+
+			::ProcessRepDataCommon(DeltaTime, SimTime, RepDataCopy, PhysicsProxy, VersionProcessed, OneOffActivatedProcessed, LastHardsnapTimeInMs, ReceivedTime, bReplicateMovement);
+		});
 	}
 
-	VersionProcessed = RepData.Version;
-	return bProcessed;
 }
 
 void UGeometryCollectionComponent::SetDynamicState(const Chaos::EObjectStateType& NewDynamicState)
@@ -2692,27 +2789,7 @@ void UGeometryCollectionComponent::AsyncPhysicsTickComponent(float DeltaTime, fl
 
 	Super::AsyncPhysicsTickComponent(DeltaTime, SimTime);
 
-	// using net mode for now as using local role seemed to cause other issues at initialization time
-	// we may nee dto to also use local role in the future if the authority is likely to change at runtime
-	if (GetNetMode() == ENetMode::NM_Client)
-	{
-		const int64 CurrentTimeInMs = FPlatformTime::ToMilliseconds64(FPlatformTime::Cycles64());
-		const bool bActive = ProcessRepData(DeltaTime, SimTime);
-		if (!bActive && LastAsyncPhysicsTickMs > 0 && CurrentTimeInMs - LastAsyncPhysicsTickMs > static_cast<int64>(MaxGeometryCollectionAsyncPhysicsTickIdleTimeMs))
-		{
-			DeferRemoveAsyncPhysicsTick();
-
-			// Reset LastAsyncPhysicsTickMs to 0 so that the next time when the async physics tick gets enabled, we
-			// allow the tick to run for at least MaxGeometryCollectionAsyncPhysicsTickIdleTimeMs. This handles the
-			// (hopefully unlikely) case where the async tick gets re-enabled but there's no work to do.
-			LastAsyncPhysicsTickMs = 0;
-		}
-		else
-		{
-			LastAsyncPhysicsTickMs = CurrentTimeInMs;
-		}
-	}
-	else
+	if (GetNetMode() != ENetMode::NM_Client)
 	{
 		UpdateRepData();
 	}
@@ -3330,7 +3407,10 @@ void UGeometryCollectionComponent::OnRep_RepData()
 {
 	// We have new data that was replicated! Turn on the async tick to process instead of just requesting a one-off
 	// since we may want to keep processing for extra time afterwards.
-	SetAsyncPhysicsTickEnabled(true);
+	check(IsInGameThread());
+    check(GetNetMode() == ENetMode::NM_Client);
+
+	ProcessRepDataOnPT();
 }
 
 void UGeometryCollectionComponent::SetAbandonedParticleCollisionProfileName(FName CollisionProfile)
diff --git a/Engine/Source/Runtime/Experimental/GeometryCollectionEngine/Public/GeometryCollection/GeometryCollectionComponent.h b/Engine/Source/Runtime/Experimental/GeometryCollectionEngine/Public/GeometryCollection/GeometryCollectionComponent.h
--- a/Engine/Source/Runtime/Experimental/GeometryCollectionEngine/Public/GeometryCollection/GeometryCollectionComponent.h
+++ b/Engine/Source/Runtime/Experimental/GeometryCollectionEngine/Public/GeometryCollection/GeometryCollectionComponent.h
@@ -1304,6 +1304,8 @@ protected:
 	int64 LastAsyncPhysicsTickMs = 0;
 
 private:
+	void ProcessRepDataOnPT();
+	void ResetRepDataCommon();
 
 	bool bRenderStateDirty;
 	bool bEnableBoneSelection;
