Update child XR positions in cluster union when new parent to child update comes through

diff --git a/Engine/Source/Runtime/Experimental/Chaos/Private/Chaos/ClusterUnionManager.cpp b/Engine/Source/Runtime/Experimental/Chaos/Private/Chaos/ClusterUnionManager.cpp
--- a/Engine/Source/Runtime/Experimental/Chaos/Private/Chaos/ClusterUnionManager.cpp
+++ b/Engine/Source/Runtime/Experimental/Chaos/Private/Chaos/ClusterUnionManager.cpp
@@ -889,7 +889,8 @@ namespace Chaos
 				{
 					if (const FClusterUnionChildToParentUpdate* Update = PendingChildToParentUpdates.Find(ChildHandle))
 					{
-						ChildHandle->SetChildToParent(Update->ChildToParent);
+						const FRigidTransform3 ChildToParent = Update->ChildToParent;
+						ChildHandle->SetChildToParent(ChildToParent);
 
 						if (!Update->bLock && !ChildHandle->IsChildToParentLocked())
 						{
@@ -898,6 +899,14 @@ namespace Chaos
 						ChildHandle->SetChildToParentLocked(true);
 						PendingChildToParentUpdates.Remove(ChildHandle);
 
+						// Update the child's world transform to be consistent with its ChildToParent transform
+						const FPBDRigidClusteredParticleHandle* ParentHandle = ClusterUnion->InternalCluster;
+						const FRigidTransform3 ParticleToWorld = ChildToParent * FRigidTransform3(ParentHandle->X(), ParentHandle->R());
+						ChildHandle->SetX(ParticleToWorld.GetTranslation());
+						ChildHandle->SetP(ParticleToWorld.GetTranslation());
+						ChildHandle->SetR(ParticleToWorld.GetRotation());
+						ChildHandle->SetQ(ParticleToWorld.GetRotation());
+
 						// A child to parent update needs to remove *and* add to the connectivity graph (in that order) since
 						// the child to parent update might move the node so far away as to make the old connectivity edges incorrect.
 						ClusterUnion->PendingConnectivityOperations.Add({ Particle, EClusterUnionConnectivityOperation::Remove });
