Chaos - add cvar to tune capsule-convex point merge distance tolerance

diff --git a/Engine/Source/Runtime/Experimental/Chaos/Private/Chaos/CollisionOneShotManifolds.cpp b/Engine/Source/Runtime/Experimental/Chaos/Private/Chaos/CollisionOneShotManifolds.cpp
--- a/Engine/Source/Runtime/Experimental/Chaos/Private/Chaos/CollisionOneShotManifolds.cpp
+++ b/Engine/Source/Runtime/Experimental/Chaos/Private/Chaos/CollisionOneShotManifolds.cpp
@@ -36,6 +36,10 @@ namespace Chaos
 	FRealSingle Chaos_Collision_Manifold_CapsuleRadialContactFraction = 0.25f;
 	FAutoConsoleVariableRef CVarChaos_Manifold_CapsuleRadialContactFraction(TEXT("p.Chaos.Collision.Manifold.CapsuleRadialContactFraction"), Chaos_Collision_Manifold_CapsuleRadialContactFraction, TEXT(""));
 
+	// When a capsule-convex contact has points closer together than this fraction of radius, ignore one of the contacts
+	FRealSingle Chaos_Collision_Manifold_CapsuleMinContactDistanceFraction = 0.1f;
+	FAutoConsoleVariableRef CVarChaos_Manifold_CapsuleMinContactDistanceFraction(TEXT("p.Chaos.Collision.Manifold.CapsuleMinContactDistanceFraction"), Chaos_Collision_Manifold_CapsuleMinContactDistanceFraction, TEXT(""));
+
 	FRealSingle Chaos_Collision_Manifold_PlaneContactNormalEpsilon = 0.001f;
 	FAutoConsoleVariableRef CVarChaos_Manifold_PlaneContactNormalEpsilon(TEXT("p.Chaos.Collision.Manifold.PlaneContactNormalEpsilon"), Chaos_Collision_Manifold_PlaneContactNormalEpsilon, TEXT("Normal tolerance used to distinguish face contacts from edge-edge contacts"));
 
@@ -1162,7 +1166,8 @@ namespace Chaos
 						MaxContactPointCount,
 						PlaneTolerance);
 
-					const FReal PointDistanceToleranceSq = FMath::Square(FReal(0.1) * CapsuleRadius);
+					const FReal PointDistanceToleranceFraction = Chaos_Collision_Manifold_CapsuleMinContactDistanceFraction;
+					const FReal PointDistanceToleranceSq = FMath::Square(PointDistanceToleranceFraction * CapsuleRadius);
 
 					// Add the vertices if not clipped away
 					for (int32 ContactPointIndex = 0; ContactPointIndex < ContactPointCount; ++ContactPointIndex)
