Wait for the last physics steps, but don't wait for the physics steps that were launched during the current game tick
#prefligh 648cbb966e2511694daf1998

diff --git a/Engine/Source/Runtime/Experimental/Chaos/Private/Chaos/Framework/PhysicsSolverBase.cpp b/Engine/Source/Runtime/Experimental/Chaos/Private/Chaos/Framework/PhysicsSolverBase.cpp
--- a/Engine/Source/Runtime/Experimental/Chaos/Private/Chaos/Framework/PhysicsSolverBase.cpp
+++ b/Engine/Source/Runtime/Experimental/Chaos/Private/Chaos/Framework/PhysicsSolverBase.cpp
@@ -182,7 +182,7 @@ namespace Chaos
 	// 0 blocks on any physics steps generated from past GT Frames, and blocks on none of the tasks from current frame.
 	// 1 blocks on everything except the single most recent task (including tasks from current frame)
 	// 1 should gurantee we will always have a future output for interpolation from 2 frames in the past
-	int32 AsyncPhysicsBlockMode = 1;
+	int32 AsyncPhysicsBlockMode = 0;
 	FAutoConsoleVariableRef CVarAsyncPhysicsBlockMode(TEXT("p.AsyncPhysicsBlockMode"), AsyncPhysicsBlockMode, TEXT("Setting to 0 blocks on any physics steps generated from past GT Frames, and blocks on none of the tasks from current frame."
 		" 1 blocks on everything except the single most recent task (including tasks from current frame). 1 should gurantee we will always have a future output for interpolation from 2 frames in the past."));
 
