Iris:
* Added alot of verbose logs to FObjectReferenceCache to help track down rare crash

#rnx

diff --git a/Engine/Source/Runtime/Experimental/Iris/Core/Private/Iris/ReplicationSystem/ObjectReferenceCache.cpp b/Engine/Source/Runtime/Experimental/Iris/Core/Private/Iris/ReplicationSystem/ObjectReferenceCache.cpp
--- a/Engine/Source/Runtime/Experimental/Iris/Core/Private/Iris/ReplicationSystem/ObjectReferenceCache.cpp
+++ b/Engine/Source/Runtime/Experimental/Iris/Core/Private/Iris/ReplicationSystem/ObjectReferenceCache.cpp
@@ -25,14 +25,14 @@
 
 DEFINE_LOG_CATEGORY_STATIC(LogIrisReferences, Log, All);
 
-#if (UE_BUILD_SHIPPING || UE_BUILD_TEST)
-#	define UE_NET_ENABLE_REFERENCECACHE_LOG 0
-#	define UE_NET_VALIDATE_REFERENCECACHE 0
-#else
-#	define UE_NET_ENABLE_REFERENCECACHE_LOG 1
-#	define UE_NET_VALIDATE_REFERENCECACHE 1
+#ifndef UE_NET_ENABLE_REFERENCECACHE_LOG
+#	define UE_NET_ENABLE_REFERENCECACHE_LOG !(UE_BUILD_SHIPPING || UE_BUILD_TEST)
 #endif 
 
+#ifndef UE_NET_VALIDATE_REFERENCECACHE
+#	define UE_NET_VALIDATE_REFERENCECACHE !(UE_BUILD_SHIPPING || UE_BUILD_TEST)
+#endif
+
 #if UE_NET_ENABLE_REFERENCECACHE_LOG
 #	define UE_LOG_REFERENCECACHE(Verbosity, Format, ...)  UE_LOG(LogIrisReferences, Verbosity, Format, ##__VA_ARGS__)
 #else
@@ -242,26 +242,35 @@ bool FObjectReferenceCache::CreateObjectReferenceInternal(const UObject* Object,
 
 		if (CachedObject.Object.Get() == Object)
 		{
-			UE_LOG_REFERENCECACHE(VeryVerbose, TEXT("ObjectReferenceCache::CreateObjectReferenceHandle Found existing %s for Object %s, OuterNetRefHandle: %s"), *CachedObject.NetRefHandle.ToString(), ToCStr(Object->GetPathName()), *CachedObject.OuterNetRefHandle.ToString());
+			UE_LOG_REFERENCECACHE(VeryVerbose, TEXT("ObjectReferenceCache::CreateObjectReferenceHandle Found existing %s for ObjectPath %s Object: %s (0x%p), OuterNetRefHandle: %s"), 
+				*CachedObject.NetRefHandle.ToString(), ToCStr(Object->GetPathName()), *GetNameSafe(Object), Object, *CachedObject.OuterNetRefHandle.ToString());
 			OutReference = MakeNetObjectReference(CachedObject);			
 			return true;
 		}
 		else
 		{
-			UE_LOG_REFERENCECACHE(Verbose, TEXT("Detected stale cached Object removing %s from ObjectToNetReferenceHandle lookup"), *CachedObject.NetRefHandle.ToString());
+			UE_LOG_REFERENCECACHE(Verbose, TEXT("ObjectReferenceCache::CreateObjectReferenceHandle Removed %s from ObjectToNetReferenceHandle due to stale cache. Cached Object %s (0x%p).  New object %s (0x%p)"), 
+				*CachedObject.NetRefHandle.ToString(), *GetNameSafe(CachedObject.ObjectKey), CachedObject.ObjectKey, *GetNameSafe(Object), Object);
 
 			ObjectToNetReferenceHandle.Remove(CachedObject.ObjectKey);
 			ObjectToNetReferenceHandle.Remove(Object);
 
 			if (RefHandle.IsStatic())
 			{
+				UE_LOG_REFERENCECACHE(Verbose, TEXT("ObjectReferenceCache::CreateObjectReferenceHandle clearing stale cache for static handle %s. Cached Object %s (0x%p).  New object %s (0x%p)"),
+					*CachedObject.NetRefHandle.ToString(), *GetNameSafe(CachedObject.ObjectKey), CachedObject.ObjectKey, *GetNameSafe(Object), Object);
+
 				// Note we only cleanse the object reference
 				// we still keep the cachedObject data around in order to be able to serialize static destruction infos
 				CachedObject.ObjectKey = nullptr;
 				CachedObject.Object = nullptr;
+
 			}
 			else
 			{
+				UE_LOG_REFERENCECACHE(Verbose, TEXT("ObjectReferenceCache::CreateObjectReferenceHandle removing cache for handle %s. Cached Object %s (0x%p). New object %s (0x%p)"),
+					*CachedObject.NetRefHandle.ToString(), *GetNameSafe(CachedObject.ObjectKey), CachedObject.ObjectKey, *GetNameSafe(Object), Object);
+
 				ReferenceHandleToCachedReference.Remove(RefHandle);
 			}
 		}
@@ -372,7 +381,8 @@ bool FObjectReferenceCache::CreateObjectReferenceInternal(const UObject* Object,
 	}
 #endif
 	
-	UE_LOG_REFERENCECACHE(Verbose, TEXT("ObjectReferenceCache::CreateObjectReferenceHandle %s for Object %s, OuterNetRefHandle: %s"), *CachedObject.NetRefHandle.ToString(), ToCStr(Object->GetPathName()), *CachedObject.OuterNetRefHandle.ToString());
+	UE_LOG_REFERENCECACHE(Verbose, TEXT("ObjectReferenceCache::CreateObjectReferenceHandle Adding %s to ObjectToNetReferenceHandle and Cache for ObjectPath %s, Object (0x%p), OuterNetRefHandle: %s"), 
+		*CachedObject.NetRefHandle.ToString(), ToCStr(Object->GetPathName()), *GetNameSafe(Object), Object, *CachedObject.OuterNetRefHandle.ToString());
 
 	// Create reference
 	OutReference = MakeNetObjectReference(CachedObject);
@@ -409,7 +419,7 @@ FNetRefHandle FObjectReferenceCache::GetObjectReferenceHandleFromObject(const UO
 		{
 			ensureAlwaysMsgf(CachedObjectPtr != nullptr, TEXT("Mismatch between ReferenceCache maps. Object %s (0x%p) has handle (%s) but no cache."), *GetNameSafe(Object), Object, *Reference->ToString());
 			const_cast<TMap<const UObject*, FNetRefHandle>&>(ObjectToNetReferenceHandle).Remove(Object);
-			UE_LOG_REFERENCECACHE(Warning, TEXT("ObjectReferenceCache::GetObjectReferenceHandleFromObject removed %s due to mismatched object %s (0x%p) "), *Reference->ToString(), *GetNameSafe(Object), Object);
+			UE_LOG_REFERENCECACHE(Warning, TEXT("ObjectReferenceCache::GetObjectReferenceHandleFromObject removed %s from ObjectToNetReferenceHandle due to mismatch with cache. Object %s (0x%p) "), *Reference->ToString(), *GetNameSafe(Object), Object);
 			return FNetRefHandle();
 		}
 
@@ -424,7 +434,8 @@ FNetRefHandle FObjectReferenceCache::GetObjectReferenceHandleFromObject(const UO
 		}
 		else
 		{
-			UE_LOG_REFERENCECACHE(Verbose, TEXT("ObjectReferenceCache::Found %s cached RefHandle %s for object %s"), CachedObject.Object.IsStale() ? TEXT("Stale") : TEXT("UnResolved"), *CachedObject.NetRefHandle.ToString(), ToCStr(Object->GetName()));
+			UE_LOG_REFERENCECACHE(Verbose, TEXT("ObjectReferenceCache::Found %s cached RefHandle %s for object %s (0x%p)"), 
+				CachedObject.Object.IsStale() ? TEXT("Stale") : TEXT("UnResolved"), *CachedObject.NetRefHandle.ToString(), *GetNameSafe(Object), Object);
 		}
 	}
 
@@ -440,7 +451,7 @@ void FObjectReferenceCache::AddRemoteReference(FNetRefHandle RefHandle, const UO
 
 	if (ensure(RefHandle.IsDynamic()))
 	{
-		UE_LOG_REFERENCECACHE(Verbose, TEXT("AddRemoteReference: %s, Object: %s Pointer: %p" ), *RefHandle.ToString(), Object ? *Object->GetName() : TEXT("NULL"), Object);
+		UE_LOG_REFERENCECACHE(Verbose, TEXT("ObjectReferenceCache::AddRemoteReference: Adding %s to ObjectToNetReferenceHandle, Object: %s (0x%p)" ), *RefHandle.ToString(), *GetNameSafe(Object), Object);
 
 	#if UE_NET_VALIDATE_REFERENCECACHE
 		// Dynamic object might have been referenced before it gets replicated, then it will then already be registered (with a path relative to its outer)
@@ -449,7 +460,8 @@ void FObjectReferenceCache::AddRemoteReference(FNetRefHandle RefHandle, const UO
 			FCachedNetObjectReference& ExistingCachedObject = ReferenceHandleToCachedReference.FindChecked(*ExistingHandle);
 			if (ExistingCachedObject.NetRefHandle != RefHandle)
 			{
-				UE_LOG_REFERENCECACHE_WARNING(TEXT("AddRemoteReference: Detected conflicting Ref %s"), *ExistingCachedObject.NetRefHandle.ToString());
+				UE_LOG_REFERENCECACHE_WARNING(TEXT("ObjectReferenceCache::AddRemoteReference: Detected conflicting Ref %s: Received: %s, Object: %s (0x%p)"), *ExistingCachedObject.NetRefHandle.ToString(),
+					*RefHandle.ToString(), *GetNameSafe(Object), Object);
 			}
 		}
 	#endif
@@ -457,9 +469,12 @@ void FObjectReferenceCache::AddRemoteReference(FNetRefHandle RefHandle, const UO
 		// The object could already be known as it has been exported as a reference prior to being replicated to this connection
 		if (FCachedNetObjectReference* CacheObjectPtr = ReferenceHandleToCachedReference.Find(RefHandle))
 		{
+			UE_LOG_REFERENCECACHE(Verbose, TEXT("ObjectReferenceCache::AddRemoteReference: %s, Object: %s (0x%p).  Patching up existing cache that held: 0x%p"), *RefHandle.ToString(), *GetNameSafe(Object), Object, CacheObjectPtr->ObjectKey);
+
 			// Patch-up object pointers
 			CacheObjectPtr->Object = MakeWeakObjectPtr(const_cast<UObject*>(Object));
 			CacheObjectPtr->ObjectKey = Object;
+
 		}
 		else
 		{
@@ -478,6 +493,8 @@ void FObjectReferenceCache::AddRemoteReference(FNetRefHandle RefHandle, const UO
 			CachedObject.bIsPending = false;
 			CachedObject.bIgnoreWhenMissing = false;
 
+			UE_LOG_REFERENCECACHE(Verbose, TEXT("ObjectReferenceCache::AddRemoteReference: Adding %s to Cache Object: %s (0x%p).  Adding new cache"), *RefHandle.ToString(), *GetNameSafe(Object), Object);
+
 			ReferenceHandleToCachedReference.Add(RefHandle, CachedObject);
 		}
 
@@ -490,7 +507,7 @@ void FObjectReferenceCache::RemoveReference(FNetRefHandle RefHandle, const UObje
 	// Cleanup cached data for dynamic instance
 	if (ensure(RefHandle.IsDynamic()))
 	{
-		UE_LOG_REFERENCECACHE(Verbose, TEXT("ObjectReferenceCache::RemoveReference %s, Object: %s Pointer: %p"), *RefHandle.ToString(), Object ? *Object->GetName() : TEXT("NULL"), Object);
+		UE_LOG_REFERENCECACHE(Verbose, TEXT("ObjectReferenceCache::RemoveReference: %s, Object: %s (0x%p)"), *RefHandle.ToString(), *GetNameSafe(Object), Object);
 
 		// We create an iterator to avoid having to find the key twice
 		TMap<FNetRefHandle, FCachedNetObjectReference>::TKeyIterator It = ReferenceHandleToCachedReference.CreateKeyIterator(RefHandle);
@@ -519,7 +536,7 @@ void FObjectReferenceCache::RemoveReference(FNetRefHandle RefHandle, const UObje
 					if (ReferenceHandleToCachedReference.RemoveAndCopyValue(SubRefHandle, RemovedSubObjectRef))
 					{
 						const UObject* SubRefObject = RemovedSubObjectRef.Object.Get();
-						UE_LOG_REFERENCECACHE(Verbose, TEXT("ObjectReferenceCache::RemoveReference %s, Referenced dynamic %s, Object: %s Pointer: %p"), *SubRefHandle.ToString(), *RefHandle.ToString(), SubRefObject ? *SubRefObject->GetName() : TEXT("NULL"), RemovedSubObjectRef.ObjectKey);
+						UE_LOG_REFERENCECACHE(Verbose, TEXT("ObjectReferenceCache::RemoveReference: %s removed from ObjectToNetReferenceHandle. Referenced dynamic %s, Object: %s Pointer: %p"), *SubRefHandle.ToString(), *RefHandle.ToString(), SubRefObject ? *SubRefObject->GetName() : TEXT("NULL"), RemovedSubObjectRef.ObjectKey);
 
 						ObjectToNetReferenceHandle.Remove(RemovedSubObjectRef.ObjectKey);
 					}
@@ -532,11 +549,14 @@ void FObjectReferenceCache::RemoveReference(FNetRefHandle RefHandle, const UObje
 			// Remove if RelativePath is invalid
 			if (!CachedItemToRemove.RelativePath.IsValid())
 			{
+				UE_LOG_REFERENCECACHE(Verbose, TEXT("ObjectReferenceCache::RemoveReference: %s, Object: %s (0x%p). Destroying cache since Path is invalid (%s)"),
+					*RefHandle.ToString(), *GetNameSafe(Object), Object, *CachedItemToRemove.RelativePath.ToString());
 				It.RemoveCurrent();
 			}
 			else
 			{
 				// Clear out object reference since it will be invalid
+				UE_LOG_REFERENCECACHE(Verbose, TEXT("ObjectReferenceCache::RemoveReference: %s, Object: %s (0x%p). Clearing cache reference to (0x%p)"), *RefHandle.ToString(), *GetNameSafe(Object), Object, CachedItemToRemove.ObjectKey);
 				CachedItemToRemove.ObjectKey = nullptr;
 				CachedItemToRemove.Object.Reset();
 			}
@@ -551,15 +571,23 @@ void FObjectReferenceCache::RemoveReference(FNetRefHandle RefHandle, const UObje
 				const FNetRefHandle ExistingRefHandle = ObjectIt.Value();
 				if (ExistingRefHandle == RefHandle)
 				{
+					UE_LOG_REFERENCECACHE(Verbose, TEXT("ObjectReferenceCache::RemoveReference: %s removed from ObjectToNetReferenceHandle Object: %s (0x%p)."), 
+						*RefHandle.ToString(), *GetNameSafe(Object), Object);
 					ObjectIt.RemoveCurrent();
 				}
 				else
 				{
 					// If we have an existing handle associated with the Object that differs from the one stored in the cache it means that the client might have destroyed the instance without notifying the replication system
 					// and reused the instance, if that is the case we just leave the association as it is
-					UE_LOG_REFERENCECACHE(Verbose, TEXT("FObjectReferenceCache::RemoveReference ExistingReference %s != RefHandle being removed: %s "), *(ExistingRefHandle.ToString()), *RefHandle.ToString());
+					UE_LOG_REFERENCECACHE(Verbose, TEXT("ObjectReferenceCache::RemoveReference: ExistingReference %s != RefHandle being removed: %s, Object: %s (0x%p) "), 
+						*(ExistingRefHandle.ToString()), *RefHandle.ToString(), *GetNameSafe(Object), Object);
 				}
 			}
+			else
+			{
+				UE_LOG_REFERENCECACHE(Warning, TEXT("ObjectReferenceCache::RemoveReference: %s, Object: %s (0x%p). Skipped removing since ObjectToNetReferenceHandle had no entry"), 
+					*RefHandle.ToString(), *GetNameSafe(Object), Object);
+			}
 		}
 	}
 }
@@ -737,7 +765,7 @@ UObject* FObjectReferenceCache::ResolveObjectReferenceHandleInternal(FNetRefHand
 					CacheObjectPtr->bIsPending = true;
 
 					StartAsyncLoadingPackage(*CacheObjectPtr, ObjectPathName, RefHandle, false);
-					UE_LOG_REFERENCECACHE(Log, TEXT("GetObjectFromRefHandle: Async loading package. Path: %s, %s"), *ObjectPath, *RefHandle.ToString());
+					UE_LOG_REFERENCECACHE(VeryVerbose, TEXT("GetObjectFromRefHandle: Async loading package. Path: %s, %s"), *ObjectPath, *RefHandle.ToString());
 				}
 				else
 				{
@@ -920,9 +948,12 @@ UObject* FObjectReferenceCache::ResolveObjectReferenceHandleInternal(FNetRefHand
 	{
 		ObjectToNetReferenceHandle.Add(Object, CompleteRefHandle);
 		UE_NET_TRACE_NETHANDLE_CREATED(CompleteRefHandle, CreatePersistentNetDebugName(ResolvedToken), 0, IsAuthority() ? 1U : 0U);
+
+		UE_LOG_REFERENCECACHE(Verbose, TEXT("ObjectReferenceCache::ResolveObjectReferenceHandleInternal Adding %s Object %s (0x%p) to ObjectToNetReferenceHandle"),
+			*CompleteRefHandle.ToString(), *GetNameSafe(Object), Object);
 	}
 	
-	UE_LOG_REFERENCECACHE(Verbose, TEXT("ObjectReferenceCache::ResolveObjectRefererenceHandle %s for Object: %s Pointer: %p"), *RefHandle.ToString(), ToCStr(Object->GetPathName()), Object);
+	UE_LOG_REFERENCECACHE(VeryVerbose, TEXT("ObjectReferenceCache::ResolveObjectRefererenceHandle %s for ObjectPath: %s Object %s (0x%p)"), *RefHandle.ToString(), ToCStr(Object->GetPathName()), *GetNameSafe(Object), Object);
 
 	// Update our QueuedObjectReference if one exists.
 	UpdateTrackedQueuedBatchObjectReference(CompleteRefHandle, Object);
@@ -1082,7 +1113,8 @@ FNetObjectReference FObjectReferenceCache::GetOrCreateObjectReference(const FStr
 
 		ReferenceHandleToCachedReference.Add(RefHandle, CachedObject);
 	
-		UE_LOG_REFERENCECACHE(Verbose, TEXT("ObjectReferenceCache::GetOrCreateObjectReference(FromPath) %s for Object %s, %s"), *CachedObject.NetRefHandle.ToString(), ToCStr(ObjectPath), *CachedObject.OuterNetRefHandle.ToString());
+		UE_LOG_REFERENCECACHE(Verbose, TEXT("ObjectReferenceCache::GetOrCreateObjectReference(FromPath) Added cache for %s for Object: %s, OuterHandle %s OuterObject: %s (0x%p)"), 
+			*CachedObject.NetRefHandle.ToString(), ToCStr(ObjectPath), *CachedObject.OuterNetRefHandle.ToString(), *GetNameSafe(Outer), Outer);
 	}
 
 	OutReference = FNetObjectReference(RefHandle);
@@ -1284,6 +1316,9 @@ void FObjectReferenceCache::ReadFullReferenceInternal(FNetSerializationContext&
 		CachedObject.bIgnoreWhenMissing = bNoLoad;
 
 		ReferenceHandleToCachedReference.Add(NetRefHandle, CachedObject);
+
+		UE_LOG_REFERENCECACHE(Verbose, TEXT("ObjectReferenceCache::ReadFullReferenceInternal Added cache for %s for Object: %s, OuterHandle %s"),
+			*CachedObject.NetRefHandle.ToString(), *RelativePath.ToString(), *OuterRef.GetRefHandle().ToString());
 	}
 
 	OutObjectRef = FNetObjectReference(NetRefHandle);
@@ -1403,7 +1438,7 @@ void FObjectReferenceCache::WriteReference(FNetSerializationContext& Context, FN
 		return;
 	}
 
-	UE_LOG_REFERENCECACHE(Verbose, TEXT("ObjectReferenceCache::WriteReference Auth: %u %s"), IsAuthority() ? 1U : 0U, *RefHandle.ToString());
+	UE_LOG_REFERENCECACHE(VeryVerbose, TEXT("ObjectReferenceCache::WriteReference Auth: %u %s"), IsAuthority() ? 1U : 0U, *RefHandle.ToString());
 
 	checkSlow(Ref.CanBeExported() || RefHandle.IsDynamic());
 
@@ -1449,7 +1484,7 @@ void FObjectReferenceCache::ReadReference(FNetSerializationContext& Context, FNe
 	// The handle must exist in the cache or it must be a dynamic handle that will be explicitly registered
 	if (ensureAlwaysMsgf(NetRefHandle.IsDynamic() || ReferenceHandleToCachedReference.Find(NetRefHandle), TEXT("FObjectReferenceCache::ReadReference: Handle %s must already be known or be dynamic"), *NetRefHandle.ToString()))
 	{
-		UE_LOG_REFERENCECACHE(Verbose, TEXT("ObjectReferenceCache::ReadReference Auth: %u %s"), IsAuthority() ? 1U : 0U, *NetRefHandle.ToString());
+		UE_LOG_REFERENCECACHE(VeryVerbose, TEXT("ObjectReferenceCache::ReadReference Auth: %u %s"), IsAuthority() ? 1U : 0U, *NetRefHandle.ToString());
 		OutRef = FNetObjectReference(NetRefHandle);
 	}
 	else
@@ -1691,12 +1726,12 @@ void FObjectReferenceCache::ValidateAsyncLoadingPackage(FCachedNetObjectReferenc
 
 	if (PendingLoadRequest.NetRefHandles.Last() != RefHandle)
 	{
-		UE_LOG_REFERENCECACHE(Log, TEXT("ValidateAsyncLoadingPackage: Already async loading package with a different NetRefHandle. Path: %s, original RefHandle: %s, new RefHandle: %s"),
+		UE_LOG_REFERENCECACHE(VeryVerbose, TEXT("ValidateAsyncLoadingPackage: Already async loading package with a different NetRefHandle. Path: %s, original RefHandle: %s, new RefHandle: %s"),
 			*PackagePath.ToString(), *PendingLoadRequest.NetRefHandles.Last().ToString(), *RefHandle.ToString());
 	}
 	else
 	{
-		UE_LOG_REFERENCECACHE(Log, TEXT("ValidateAsyncLoadingPackage: Already async loading package. Path: %s, RefHandle: %s"), *PackagePath.ToString(), *RefHandle.ToString());
+		UE_LOG_REFERENCECACHE(VeryVerbose, TEXT("ValidateAsyncLoadingPackage: Already async loading package. Path: %s, RefHandle: %s"), *PackagePath.ToString(), *RefHandle.ToString());
 	}
 	
 //#if CSV_PROFILER
