[ChaosVD] Properly playback delta frames when manually scrubbing a solver timeline.

Frames are recorded as deltas most of the time, full captures are done every 10 seconds (or another value configured via cvar), therefore if you jump to any point of the timeline randomly, the visualization will get out of sync.

This change makes sure that if you are jumping to a frame that is more that one frame away, it will play all the frames in between staring from the closest keyframe

This ensures the visualization in the tool is always in sync with was recorded.

In addition, this CL fixes an issue where the frames recorded at the beginning are not flag as keyframe when they are.

diff --git a/Engine/Plugins/Experimental/ChaosVD/Source/ChaosVD/Private/ChaosVDPlaybackController.cpp b/Engine/Plugins/Experimental/ChaosVD/Source/ChaosVD/Private/ChaosVDPlaybackController.cpp
--- a/Engine/Plugins/Experimental/ChaosVD/Source/ChaosVD/Private/ChaosVDPlaybackController.cpp
+++ b/Engine/Plugins/Experimental/ChaosVD/Source/ChaosVD/Private/ChaosVDPlaybackController.cpp
@@ -99,6 +99,27 @@ void FChaosVDPlaybackController::UnloadCurrentRecording(EChaosVDUnloadRecordingF
 	}
 }
 
+void FChaosVDPlaybackController::PlayFromClosestKeyFrame(const int32 InTrackID, const int32 FrameNumber, FChaosVDScene& InSceneToControl)
+{
+	const int32 KeyFrameNumber = LoadedRecording->FindFirstSolverKeyFrameNumberFromFrame(InTrackID, FrameNumber);
+
+	if (!ensure(KeyFrameNumber >= 0))
+	{
+		return;
+	}
+
+	for (int32 CurrentFrameNumber = KeyFrameNumber; CurrentFrameNumber < FrameNumber; CurrentFrameNumber++)
+	{				
+		const FChaosVDSolverFrameData* SolverFrameData = LoadedRecording->GetSolverFrameData(InTrackID, CurrentFrameNumber);
+		const int32 LastStepNumber = SolverFrameData->SolverSteps.Num() - 1;
+
+		if (SolverFrameData && ensure(SolverFrameData->SolverSteps.IsValidIndex(LastStepNumber)))
+		{
+			InSceneToControl.UpdateFromRecordedStepData(InTrackID, SolverFrameData->DebugName, SolverFrameData->SolverSteps[LastStepNumber], *SolverFrameData);
+		}
+	}
+}
+
 void FChaosVDPlaybackController::GoToRecordedSolverStep(const int32 InTrackID, const int32 FrameNumber, const int32 Step, FGuid InstigatorID)
 {
 	if (const TSharedPtr<FChaosVDScene> SceneToControlSharedPtr = SceneToControl.Pin())
@@ -116,8 +137,6 @@ void FChaosVDPlaybackController::GoToRecordedSolverStep(const int32 InTrackID, c
 			// we will need to add a lock to the file. A feature that might need that is Recording Clips
 			TraceServices::FAnalysisSessionReadScope SessionReadScope(*TraceSession);
 
-			const FChaosVDSolverFrameData* SolverFrameData = LoadedRecording->GetSolverFrameData(InTrackID, FrameNumber);
-
 			TSharedPtr<FChaosVDTrackInfo> CurrentTrackInfo;
 			if (TrackInfoByIDMap* TrackInfoByID = TrackInfoPerType.Find(EChaosVDTrackType::Solver))
 			{
@@ -132,6 +151,15 @@ void FChaosVDPlaybackController::GoToRecordedSolverStep(const int32 InTrackID, c
 				}
 			}
 
+			const int32 FrameDiff = FrameNumber - CurrentTrackInfo->CurrentFrame;
+			constexpr int32 FrameDriftTolerance = 1;
+			if (FMath::Abs(FrameDiff) > FrameDriftTolerance || CurrentTrackInfo->CurrentFrame == 0)
+			{
+				// As Frames are recorded as delta, we need to make sure of playing back all the deltas since the closest keyframe
+				PlayFromClosestKeyFrame(InTrackID, FrameNumber, *SceneToControlSharedPtr.Get());
+			}
+
+			const FChaosVDSolverFrameData* SolverFrameData = LoadedRecording->GetSolverFrameData(InTrackID, FrameNumber);
 			if (CurrentTrackInfo->LockedOnStep != INDEX_NONE)
 			{
 				// If this track is locked to a specific step, we need to play back the previous steps on the current frame, because not all steps capture the same data.
diff --git a/Engine/Plugins/Experimental/ChaosVD/Source/ChaosVD/Private/ChaosVDScene.cpp b/Engine/Plugins/Experimental/ChaosVD/Source/ChaosVD/Private/ChaosVDScene.cpp
--- a/Engine/Plugins/Experimental/ChaosVD/Source/ChaosVD/Private/ChaosVDScene.cpp
+++ b/Engine/Plugins/Experimental/ChaosVD/Source/ChaosVD/Private/ChaosVDScene.cpp
@@ -80,11 +80,15 @@ void FChaosVDScene::AddReferencedObjects(FReferenceCollector& Collector)
 void FChaosVDScene::UpdateFromRecordedStepData(const int32 SolverID, const FString& SolverName, const FChaosVDStepData& InRecordedStepData, const FChaosVDSolverFrameData& InFrameData)
 {
 	FChaosVDParticlesByIDMap& SolverParticlesByID = ParticlesBySolverID.FindChecked(SolverID);
+	
+	TSet<int32> ParticlesIDsInRecordedStepData;
+	ParticlesIDsInRecordedStepData.Reserve(InRecordedStepData.RecordedParticlesData.Num());
 
 	// Go over existing Particle VD Instances and update them or create them if needed 
 	for (const FChaosVDParticleDataWrapper& Particle : InRecordedStepData.RecordedParticlesData)
 	{
 		const int32 ParticleVDInstanceID = GetIDForRecordedParticleData(Particle);
+		ParticlesIDsInRecordedStepData.Add(ParticleVDInstanceID);
 
 		if (InRecordedStepData.ParticlesDestroyedIDs.Contains(ParticleVDInstanceID))
 		{
@@ -126,12 +130,13 @@ void FChaosVDScene::UpdateFromRecordedStepData(const int32 SolverID, const FStri
 
 	UpdateParticlesCollisionData(InRecordedStepData, SolverID);
 
-	// TODO: This will not work once the recording has what changes between solver steps
-	// So it needs to be update to remove particles based on a Particles Removed event
 	int32 AmountRemoved = 0;
 	for (FChaosVDParticlesByIDMap::TIterator RemoveIterator = SolverParticlesByID.CreateIterator(); RemoveIterator; ++RemoveIterator)
 	{
-		if (InRecordedStepData.ParticlesDestroyedIDs.Contains(RemoveIterator.Key()))
+		// If we are playing back a keyframe, the scene should only contain what it is in the recorded data
+		const bool bShouldDestroyParticleAnyway = InFrameData.bIsKeyFrame && !ParticlesIDsInRecordedStepData.Contains(RemoveIterator.Key());
+		
+		if (bShouldDestroyParticleAnyway || InFrameData.ParticlesDestroyedIDs.Contains(RemoveIterator.Key()))
 		{
 			if (AChaosVDParticleActor* ActorToRemove = RemoveIterator.Value())
 			{
diff --git a/Engine/Plugins/Experimental/ChaosVD/Source/ChaosVD/Private/Trace/ChaosVDTraceAnalyzer.cpp b/Engine/Plugins/Experimental/ChaosVD/Source/ChaosVD/Private/Trace/ChaosVDTraceAnalyzer.cpp
--- a/Engine/Plugins/Experimental/ChaosVD/Source/ChaosVD/Private/Trace/ChaosVDTraceAnalyzer.cpp
+++ b/Engine/Plugins/Experimental/ChaosVD/Source/ChaosVD/Private/Trace/ChaosVDTraceAnalyzer.cpp
@@ -117,9 +117,16 @@ bool FChaosVDTraceAnalyzer::OnEvent(uint16 RouteId, EStyle Style, const FOnEvent
 
 			if (FChaosVDSolverFrameData* FrameData = ChaosVDTraceProvider->GetLastSolverFrame(SolverID))
 			{
+				int32 ParticleDestroyedID = EventData.GetValue<int32>("ParticleID");
+
+				// We need to add all particles that were destroyed in any step of this frame to the Frame data structure
+				// So we can properly process these events when not all the steps are played back
+				// Either because of the lock sub-step feature or because we are manually scrubbing from frame to frame
+				FrameData->ParticlesDestroyedIDs.Add(ParticleDestroyedID);
+				
 				if (FrameData->SolverSteps.Num() > 0)
 				{
-					FrameData->SolverSteps.Last().ParticlesDestroyedIDs.Add(EventData.GetValue<int32>("ParticleID"));
+					FrameData->SolverSteps.Last().ParticlesDestroyedIDs.Add(ParticleDestroyedID);
 				}
 			}
 
diff --git a/Engine/Plugins/Experimental/ChaosVD/Source/ChaosVD/Public/ChaosVDPlaybackController.h b/Engine/Plugins/Experimental/ChaosVD/Source/ChaosVD/Public/ChaosVDPlaybackController.h
--- a/Engine/Plugins/Experimental/ChaosVD/Source/ChaosVD/Public/ChaosVDPlaybackController.h
+++ b/Engine/Plugins/Experimental/ChaosVD/Source/ChaosVD/Public/ChaosVDPlaybackController.h
@@ -209,6 +209,9 @@ protected:
 	/** Handles any data changes on the loaded recording - Usually called during Trace analysis */
 	void HandleCurrentRecordingUpdated();
 
+	/** Finds the closest Key frame to the provided frame number, and plays all the following frames until the specified frame number (no inclusive) */
+	void PlayFromClosestKeyFrame(int32 InTrackID, int32 FrameNumber, FChaosVDScene& InSceneToControl);
+
 	/** Map containing all track info, by track type*/
 	TMap<EChaosVDTrackType, TrackInfoByIDMap> TrackInfoPerType;
 
diff --git a/Engine/Source/Runtime/Experimental/Chaos/Private/ChaosVisualDebugger/ChaosVisualDebuggerTrace.cpp b/Engine/Source/Runtime/Experimental/Chaos/Private/ChaosVisualDebugger/ChaosVisualDebuggerTrace.cpp
--- a/Engine/Source/Runtime/Experimental/Chaos/Private/ChaosVisualDebugger/ChaosVisualDebuggerTrace.cpp
+++ b/Engine/Source/Runtime/Experimental/Chaos/Private/ChaosVisualDebugger/ChaosVisualDebuggerTrace.cpp
@@ -224,7 +224,7 @@ void FChaosVisualDebuggerTrace::TraceParticlesSoA(const Chaos::FPBDRigidsSOAs& P
 void FChaosVisualDebuggerTrace::SetupForFullCaptureIfNeeded(int32 SolverID, bool& bOutFullCaptureRequested)
 {
 	DeltaRecordingStatesLock.ReadLock();
-	bOutFullCaptureRequested = RequestedFullCaptureSolverIDs.Contains(SolverID);
+	bOutFullCaptureRequested = RequestedFullCaptureSolverIDs.Contains(SolverID) || !SolverIDsForDeltaRecording.Contains(SolverID);
 	DeltaRecordingStatesLock.ReadUnlock();
 
 	if (bOutFullCaptureRequested)
diff --git a/Engine/Source/Runtime/Experimental/ChaosVDData/Private/ChaosVDRecording.cpp b/Engine/Source/Runtime/Experimental/ChaosVDData/Private/ChaosVDRecording.cpp
--- a/Engine/Source/Runtime/Experimental/ChaosVDData/Private/ChaosVDRecording.cpp
+++ b/Engine/Source/Runtime/Experimental/ChaosVDData/Private/ChaosVDRecording.cpp
@@ -64,6 +64,44 @@ int32 FChaosVDRecording::GetLowestSolverFrameNumberAtCycle(int32 SolverID, uint6
 	return INDEX_NONE;
 }
 
+int32 FChaosVDRecording::FindFirstSolverKeyFrameNumberFromFrame(int32 SolverID, int32 StartFrameNumber)
+{
+	if (TArray<int32>* KeyFrameNumbersPtr = RecordedKeyFramesNumberPerSolver.Find(SolverID))
+	{
+		TArray<int32>& KeyFrameNumbers = *KeyFrameNumbersPtr;
+
+		const int32 IndexFound = Algo::LowerBound(KeyFrameNumbers, StartFrameNumber);
+		
+		// If StartFrameNumber is larger than the last keyframe recorded
+		// IndexFound will be outside of the array's bounds. In that case we want to use the last key frame available
+		if (IndexFound >= KeyFrameNumbers.Num())
+		{
+			return KeyFrameNumbers.Last();
+		}
+
+		if (KeyFrameNumbers.IsValidIndex(IndexFound))
+		{
+			// Frame numbers are not repeated, so the lower bound search should give us the
+			// index containing the provided "StartFrameNumber" if it was already a key frame
+			const int32 FoundKeyFrame = KeyFrameNumbers[IndexFound];
+			if (FoundKeyFrame == StartFrameNumber)
+			{
+				return FoundKeyFrame;
+			}
+
+			// If StartFrameNumber was not a keyframe, we will get the lowest index number containing a key frame number larger than "StartFrameNumber"
+			// in which case we want the previous one;
+			const int32 PrevKeyFrameIndex = IndexFound - 1;			
+			if (KeyFrameNumbers.IsValidIndex(PrevKeyFrameIndex))
+			{
+				return KeyFrameNumbers[PrevKeyFrameIndex];
+			}
+		}		
+	}
+
+	return INDEX_NONE;
+}
+
 int32 FChaosVDRecording::GetLowestSolverFrameNumberGameFrame(int32 SolverID, int32 GameFrame)
 {
 	if (!GameFrames.IsValidIndex(GameFrame))
@@ -96,17 +134,39 @@ int32 FChaosVDRecording::GetLowestGameFrameAtSolverFrameNumber(int32 SolverID, i
 	return INDEX_NONE;
 }
 
+void FChaosVDRecording::AddKeyFrameNumberForSolver(const int32 SolverID, int32 FrameNumber)
+{
+	if (TArray<int32>* KeyFrameNumber = RecordedKeyFramesNumberPerSolver.Find(SolverID))
+	{
+		KeyFrameNumber->Add(FrameNumber);
+	}
+	else
+	{
+		RecordedKeyFramesNumberPerSolver.Add(SolverID, { FrameNumber });
+	}
+}
+
 void FChaosVDRecording::AddFrameForSolver(const int32 SolverID, FChaosVDSolverFrameData&& InFrameData)
 {
+	int32 FrameNumber;
+	const bool bIsKeyFrame = InFrameData.bIsKeyFrame;
 	if (TArray<FChaosVDSolverFrameData>* SolverFrames = RecordedFramesDataPerSolver.Find(SolverID))
 	{
-		SolverFrames->Add(MoveTemp(InFrameData));
+		FrameNumber = SolverFrames->Num();
+
+		SolverFrames->Add(MoveTemp(InFrameData));	
 	}
 	else
 	{	
+		FrameNumber = 0;
 		RecordedFramesDataPerSolver.Add(SolverID, { MoveTemp(InFrameData) });
 	}
 
+	if (bIsKeyFrame)
+	{
+		AddKeyFrameNumberForSolver(SolverID, FrameNumber);
+	}
+
 	OnRecordingUpdated().Broadcast();
 }
 
diff --git a/Engine/Source/Runtime/Experimental/ChaosVDData/Public/ChaosVDRecording.h b/Engine/Source/Runtime/Experimental/ChaosVDData/Public/ChaosVDRecording.h
--- a/Engine/Source/Runtime/Experimental/ChaosVDData/Public/ChaosVDRecording.h
+++ b/Engine/Source/Runtime/Experimental/ChaosVDData/Public/ChaosVDRecording.h
@@ -37,6 +37,7 @@ struct CHAOSVDDATA_API FChaosVDSolverFrameData
 	Chaos::FRigidTransform3 SimulationTransform;
 	bool bIsKeyFrame = false;
 	FChaosVDStepsContainer SolverSteps;
+	TSet<int32> ParticlesDestroyedIDs;
 };
 
 struct FChaosVDGameFrameData
@@ -98,6 +99,8 @@ struct CHAOSVDDATA_API FChaosVDRecording
 	 * @return Found frame number. INDEX_NONE if no frame is found for the specified cycle
 	 */
 	int32 GetLowestSolverFrameNumberAtCycle(int32 SolverID, uint64 Cycle);
+
+	int32 FindFirstSolverKeyFrameNumberFromFrame(int32 SolverID, int32 StartFrameNumber);
 	
 	/**
 	 * Searches and returns the lowest frame number of a solver at the specified cycle
@@ -180,8 +183,14 @@ protected:
 	void AddImplicitObject(const uint32 ID, const Chaos::FImplicitObject* InImplicitObject);
 	
 	void AddImplicitObject_Internal(const uint32 ID, const TSharedPtr<const Chaos::FImplicitObject>& InImplicitObject);
+
+	/** Stores a frame number of a solver that is a Key Frame -
+	 * These are used when scrubbing to make sure the visualization is in sync with what was recorded
+	 */
+	void AddKeyFrameNumberForSolver(int32 SolverID, int32 FrameNumber);
 	
 	TMap<int32, TArray<FChaosVDSolverFrameData>> RecordedFramesDataPerSolver;
+	TMap<int32, TArray<int32>> RecordedKeyFramesNumberPerSolver;
 	TArray<FChaosVDGameFrameData> GameFrames;
 	FChaosVDRecordingUpdated RecordingUpdatedDelegate;
 	FChaosVDGeometryDataLoaded GeometryDataLoaded;
