Chaos - fix shock propagation what was accidentally disabled

diff --git a/Engine/Source/Runtime/Experimental/Chaos/Private/Chaos/Island/IslandManager.cpp b/Engine/Source/Runtime/Experimental/Chaos/Private/Chaos/Island/IslandManager.cpp
--- a/Engine/Source/Runtime/Experimental/Chaos/Private/Chaos/Island/IslandManager.cpp
+++ b/Engine/Source/Runtime/Experimental/Chaos/Private/Chaos/Island/IslandManager.cpp
@@ -452,6 +452,11 @@ namespace Chaos::Private
 		ApplyDeterminism();
 	}
 
+	void FPBDIslandManager::SetAssignLevels(const bool bInAssignLevels)
+	{
+		bAssignLevels = bInAssignLevels;
+	}
+
 	void FPBDIslandManager::AddConstraintContainer(const FPBDConstraintContainer& Container)
 	{
 		if (ConstraintContainers.Num() < Container.GetContainerId() + 1)
@@ -1732,15 +1737,13 @@ namespace Chaos::Private
 	{
 		SCOPE_CYCLE_COUNTER(STAT_IslandManager_AssignLevels);
 
-		// We only need to assign levels if shock propagation is enabled (but we always sort)
-		const bool bAssignLevels = ((CVars::ChaosSolverCollisionPositionShockPropagationIterations > 0) || (CVars::ChaosSolverCollisionVelocityShockPropagationIterations > 0));
-
 		// Could go wide here...
 		for (FPBDIsland* Island : Islands)
 		{
 			// Levels will only change when awake and when something has been added or removed
 			if (!Island->Flags.bIsSleeping && (!!Island->Flags.bItemsAdded || !!Island->Flags.bItemsRemoved))
 			{
+				// We only need to assign levels if shock propagation is enabled (but we always sort)
 				if (bAssignLevels)
 				{
 					AssignIslandLevels(Island);
diff --git a/Engine/Source/Runtime/Experimental/Chaos/Private/Chaos/PBDRigidsEvolutionGBF.cpp b/Engine/Source/Runtime/Experimental/Chaos/Private/Chaos/PBDRigidsEvolutionGBF.cpp
--- a/Engine/Source/Runtime/Experimental/Chaos/Private/Chaos/PBDRigidsEvolutionGBF.cpp
+++ b/Engine/Source/Runtime/Experimental/Chaos/Private/Chaos/PBDRigidsEvolutionGBF.cpp
@@ -599,6 +599,22 @@ void FPBDRigidsEvolutionGBF::SetIsDeterministic(const bool bInIsDeterministic)
 	IslandManager.SetIsDeterministic(bInIsDeterministic);
 }
 
+void FPBDRigidsEvolutionGBF::SetShockPropagationIterations(const int32 InPositionIts, const int32 InVelocityIts)
+{
+	// Negative inputs mean leave values as they are
+	if (InPositionIts >= 0)
+	{
+		CollisionConstraints.SetPositionShockPropagationIterations(InPositionIts);
+	}
+
+	if (InVelocityIts >= 0)
+	{
+		CollisionConstraints.SetVelocityShockPropagationIterations(InVelocityIts);
+	}
+
+	// If we have shock prop enabled, we require levels to be assigned to constraints and bodies
+	IslandManager.SetAssignLevels(CollisionConstraints.IsShockPropagationEnabled());
+}
 
 void FPBDRigidsEvolutionGBF::TestModeResetParticles()
 {
diff --git a/Engine/Source/Runtime/Experimental/Chaos/Private/PBDRigidsSolver.cpp b/Engine/Source/Runtime/Experimental/Chaos/Private/PBDRigidsSolver.cpp
--- a/Engine/Source/Runtime/Experimental/Chaos/Private/PBDRigidsSolver.cpp
+++ b/Engine/Source/Runtime/Experimental/Chaos/Private/PBDRigidsSolver.cpp
@@ -1156,13 +1156,8 @@ namespace Chaos
 			{
 				MEvolution->GetCollisionConstraints().SetVelocityFrictionIterations(ChaosSolverCollisionVelocityFrictionIterations);
 			}
-			if (ChaosSolverCollisionPositionShockPropagationIterations >= 0)
 			{
-				MEvolution->GetCollisionConstraints().SetPositionShockPropagationIterations(ChaosSolverCollisionPositionShockPropagationIterations);
-			}
-			if (ChaosSolverCollisionVelocityShockPropagationIterations >= 0)
-			{
-				MEvolution->GetCollisionConstraints().SetVelocityShockPropagationIterations(ChaosSolverCollisionVelocityShockPropagationIterations);
+				MEvolution->SetShockPropagationIterations(ChaosSolverCollisionPositionShockPropagationIterations, ChaosSolverCollisionVelocityShockPropagationIterations);
 			}
 			if (ChaosSolverPositionIterations >= 0)
 			{
diff --git a/Engine/Source/Runtime/Experimental/Chaos/Public/Chaos/Island/IslandManager.h b/Engine/Source/Runtime/Experimental/Chaos/Public/Chaos/Island/IslandManager.h
--- a/Engine/Source/Runtime/Experimental/Chaos/Public/Chaos/Island/IslandManager.h
+++ b/Engine/Source/Runtime/Experimental/Chaos/Public/Chaos/Island/IslandManager.h
@@ -421,6 +421,7 @@ namespace Chaos::Private
 		CHAOS_API void SetGravityForces(const FPerParticleGravity* InGravity);
 		CHAOS_API void SetDisableCounterThreshold(const int32 InDisableCounterThreshold);
 		CHAOS_API void SetIsDeterministic(const bool bInIsDeterministic);
+		CHAOS_API void SetAssignLevels(const bool bInAssignLevels);
 		CHAOS_API void UpdateParticles();
 		CHAOS_API void UpdateIslands();
 		CHAOS_API void UpdateSleep();
@@ -572,6 +573,9 @@ namespace Chaos::Private
 
 		// Whether we require determinism
 		bool bIsDeterministic = false;
+
+		// Whether we should assign levels (for shock propagation)
+		bool bAssignLevels = true;
 	};
 
 
diff --git a/Engine/Source/Runtime/Experimental/Chaos/Public/Chaos/PBDCollisionConstraints.h b/Engine/Source/Runtime/Experimental/Chaos/Public/Chaos/PBDCollisionConstraints.h
--- a/Engine/Source/Runtime/Experimental/Chaos/Public/Chaos/PBDCollisionConstraints.h
+++ b/Engine/Source/Runtime/Experimental/Chaos/Public/Chaos/PBDCollisionConstraints.h
@@ -260,6 +260,11 @@ public:
 		SolverSettings.NumVelocityShockPropagationIterations = InNumIterations;
 	}
 
+	bool IsShockPropagationEnabled() const
+	{
+		return (SolverSettings.NumPositionShockPropagationIterations > 0) || (SolverSettings.NumVelocityShockPropagationIterations > 0);
+	}
+
 	int32 NumConstraints() const
 	{
 		return GetConstraints().Num();
diff --git a/Engine/Source/Runtime/Experimental/Chaos/Public/Chaos/PBDRigidsEvolutionGBF.h b/Engine/Source/Runtime/Experimental/Chaos/Public/Chaos/PBDRigidsEvolutionGBF.h
--- a/Engine/Source/Runtime/Experimental/Chaos/Public/Chaos/PBDRigidsEvolutionGBF.h
+++ b/Engine/Source/Runtime/Experimental/Chaos/Public/Chaos/PBDRigidsEvolutionGBF.h
@@ -106,6 +106,8 @@ namespace Chaos
 
 		void SetIsDeterministic(const bool bInIsDeterministic);
 
+		void SetShockPropagationIterations(const int32 InPositionIts, const int32 InVelocityIts);
+
 		CHAOS_API void Advance(const FReal Dt, const FReal MaxStepDt, const int32 MaxSteps);
 		CHAOS_API void AdvanceOneTimeStep(const FReal dt, const FSubStepInfo& SubStepInfo = FSubStepInfo());
 
