Hide ISMPool ISMs when empty.

diff --git a/Engine/Source/Runtime/Experimental/GeometryCollectionEngine/Private/GeometryCollection/GeometryCollectionISMPoolComponent.cpp b/Engine/Source/Runtime/Experimental/GeometryCollectionEngine/Private/GeometryCollection/GeometryCollectionISMPoolComponent.cpp
--- a/Engine/Source/Runtime/Experimental/GeometryCollectionEngine/Private/GeometryCollection/GeometryCollectionISMPoolComponent.cpp
+++ b/Engine/Source/Runtime/Experimental/GeometryCollectionEngine/Private/GeometryCollection/GeometryCollectionISMPoolComponent.cpp
@@ -114,6 +114,7 @@ FGeometryCollectionISM::FGeometryCollectionISM(AActor* InOwningActor, const FGeo
 
 	InOwningActor->AddInstanceComponent(ISMC);
 	ISMC->RegisterComponent();
+	ISMC->SetVisibility(false);
 	ISMComponent = ISMC;
 }
 
@@ -128,6 +129,7 @@ FInstanceGroups::FInstanceGroupId FGeometryCollectionISM::AddInstanceGroup(int32
 	TArray<FTransform> ZeroScaleTransforms;
 	ZeroScaleTransforms.Init(ZeroScaleTransform, InstanceCount);
 
+	ISMComponent->SetVisibility(true);
 	ISMComponent->PreAllocateInstancesMemory(InstanceCount);
 	TArray<int32> RenderInstances = ISMComponent->AddInstances(ZeroScaleTransforms, true, true);
 
@@ -304,6 +306,7 @@ void FGeometryCollectionISMPool::RemoveISM(const FGeometryCollectionMeshInfo& Me
 			ISM.InstanceGroups.Reset();
 			ISM.InstanceIndexToRenderIndex.Reset();
 			ISM.RenderIndexToInstanceIndex.Reset();
+			ISM.ISMComponent->SetVisibility(false);
 		}
 
 		if (GUseComponentFreeList && ISM.ISMComponent->PerInstanceSMData.Num() == 0)
