Use counter type directly

#rnx

diff --git a/Engine/Source/Runtime/Experimental/IoStoreOnDemand/Private/Statistics.cpp b/Engine/Source/Runtime/Experimental/IoStoreOnDemand/Private/Statistics.cpp
--- a/Engine/Source/Runtime/Experimental/IoStoreOnDemand/Private/Statistics.cpp
+++ b/Engine/Source/Runtime/Experimental/IoStoreOnDemand/Private/Statistics.cpp
@@ -13,29 +13,48 @@ static int32 BytesToApproxKB(uint64 Bytes) { return int32(Bytes >> 10); }
 
 ////////////////////////////////////////////////////////////////////////////////
 // TRACE STATS
+
+#if COUNTERSTRACE_ENABLED
+	using FCounterInt		= FCountersTrace::FCounterInt;
+	using FCounterAtomicInt = FCountersTrace::FCounterAtomicInt;
+#else
+	template <typename Type>
+	struct TCounterInt
+	{
+		TCounterInt(...)  {}
+		void Set(int64 i) { V = i; }
+		void Add(int64 d) { V += d; }
+		int64 Get() const { return V;}
+		Type V = 0;
+	};
+	using FCounterInt		= TCounterInt<int64>;
+	using FCounterAtomicInt = TCounterInt<std::atomic<int64>>;
+#endif 
+
 // iorequest stats
-TRACE_DECLARE_INT_COUNTER(IoRequestsMade, TEXT("Ias/IoRequestsMade"));
-TRACE_DECLARE_INT_COUNTER(IoRequestsCompleted, TEXT("Ias/IoRequestsCompleted"));
-TRACE_DECLARE_MEMORY_COUNTER(IoRequestsCompletedSize, TEXT("Ias/Size/IoRequestsCompletedSize"));
-TRACE_DECLARE_INT_COUNTER(IoRequestsCancelled, TEXT("Ias/IoRequestsCancelled"));
-TRACE_DECLARE_INT_COUNTER(IoRequestsFailed, TEXT("Ias/IoRequestsFailed"));
+FCounterInt GIoRequestsMade(TEXT("Ias/IoRequestsMade"), TraceCounterDisplayHint_None);
+FCounterInt GIoRequestsCompleted(TEXT("Ias/IoRequestsCompleted"), TraceCounterDisplayHint_None);
+FCounterInt GIoRequestsCompletedSize(TEXT("Ias/Size/IoRequestsCompletedSize"), TraceCounterDisplayHint_Memory);
+FCounterInt GIoRequestsCancelled(TEXT("Ias/IoRequestsCancelled"), TraceCounterDisplayHint_None);
+FCounterInt GIoRequestsFailed(TEXT("Ias/IoRequestsFailed"), TraceCounterDisplayHint_None);
 // chunkrequest stats
-TRACE_DECLARE_INT_COUNTER(ReadRequestsCreated, TEXT("Ias/ReadRequestsCreated"));
-TRACE_DECLARE_INT_COUNTER(ReadRequestsRemoved, TEXT("Ias/ReadRequestsRemoved"));
+FCounterInt GReadRequestsCreated(TEXT("Ias/ReadRequestsCreated"), TraceCounterDisplayHint_None);
+FCounterInt GReadRequestsRemoved(TEXT("Ias/ReadRequestsRemoved"), TraceCounterDisplayHint_None);
 // cache stats
-TRACE_DECLARE_INT_COUNTER(CacheHits, TEXT("Ias/CacheHits"));
-TRACE_DECLARE_MEMORY_COUNTER(CacheHitsSize, TEXT("Ias/Size/CacheHitsSize"));
-TRACE_DECLARE_INT_COUNTER(CachePuts, TEXT("Ias/CachePuts"));
-TRACE_DECLARE_MEMORY_COUNTER(CachePutsSize, TEXT("Ias/Size/CachePutsSize"));
-TRACE_DECLARE_INT_COUNTER(CacheRejects, TEXT("Ias/CacheRejects"));
-TRACE_DECLARE_MEMORY_COUNTER(CacheRejectsSize, TEXT("Ias/Size/CacheRejectsSize"));
+FCounterInt GCacheHits(TEXT("Ias/CacheHits"), TraceCounterDisplayHint_None);
+FCounterInt GCacheHitsSize(TEXT("Ias/Size/CacheHitsSize"), TraceCounterDisplayHint_Memory);
+FCounterInt GCachePuts(TEXT("Ias/CachePuts"), TraceCounterDisplayHint_None);
+FCounterInt GCachePutsSize(TEXT("Ias/Size/CachePutsSize"), TraceCounterDisplayHint_Memory);
+FCounterInt GCacheRejects(TEXT("Ias/CacheRejects"), TraceCounterDisplayHint_None);
+FCounterInt GCacheRejectsSize(TEXT("Ias/Size/CacheRejectsSize"), TraceCounterDisplayHint_Memory);
 // http stats
-TRACE_DECLARE_INT_COUNTER(HttpRequestsCompleted, TEXT("Ias/HttpRequestsCompleted"));
-TRACE_DECLARE_INT_COUNTER(HttpRequestsFailed, TEXT("Ias/HttpRequestsFailed"));
-TRACE_DECLARE_INT_COUNTER(HttpRequestsPending, TEXT("Ias/HttpRequestsPending"));
-TRACE_DECLARE_INT_COUNTER(HttpRequestsInflight, TEXT("Ias/HttpRequestsInflight"));
-TRACE_DECLARE_MEMORY_COUNTER(HttpRequestsCompletedSize, TEXT("Ias/Size/HttpRequestsCompletedSize"));
-///////////////////////////////////////////////////////////////////////////////
+FCounterInt GHttpRequestsCompleted(TEXT("Ias/HttpRequestsCompleted"), TraceCounterDisplayHint_None);
+FCounterInt GHttpRequestsFailed(TEXT("Ias/HttpRequestsFailed"), TraceCounterDisplayHint_None);
+FCounterInt GHttpRequestsPending(TEXT("Ias/HttpRequestsPending"), TraceCounterDisplayHint_None);
+FCounterInt GHttpRequestsInflight(TEXT("Ias/HttpRequestsInflight"), TraceCounterDisplayHint_None);
+FCounterInt GHttpRequestsCompletedSize(TEXT("Ias/Size/HttpRequestsCompletedSize"), TraceCounterDisplayHint_Memory);
+
+////////////////////////////////////////////////////////////////////////////////
 // CSV STATS
 CSV_DEFINE_CATEGORY(Ias, true);
 // iorequest stats
@@ -63,114 +82,104 @@ CSV_DEFINE_STAT(Ias, FrameHttpRequestsCompletedSize);
 
 void FOnDemandIoBackendStats::OnIoRequestEnqueue()
 {
-	++IoRequestsMade;
-	TRACE_COUNTER_SET(IoRequestsMade, IoRequestsMade);
-	CSV_CUSTOM_STAT_DEFINED(FrameIoRequestsMade, IoRequestsMade, ECsvCustomStatOp::Set);
+	GIoRequestsMade.Add(1);
+	CSV_CUSTOM_STAT_DEFINED(FrameIoRequestsMade, int32(GIoRequestsMade.Get()), ECsvCustomStatOp::Set);
 }
+
 void FOnDemandIoBackendStats::OnIoRequestComplete(uint64 RequestSize)
 {
-	int32 Count = IoRequestsCompleted.fetch_add(1, std::memory_order_relaxed);
-	TRACE_COUNTER_SET(IoRequestsCompleted, Count);
-	CSV_CUSTOM_STAT_DEFINED(FrameIoRequestsCompleted, Count, ECsvCustomStatOp::Set);
+	GIoRequestsCompleted.Add(1);
+	CSV_CUSTOM_STAT_DEFINED(FrameIoRequestsCompleted, int32(GIoRequestsCompleted.Get()), ECsvCustomStatOp::Set);
 
-	uint64 Size = IoRequestsCompletedSize.fetch_add(RequestSize);
-	TRACE_COUNTER_SET(IoRequestsCompletedSize, Size);
-	CSV_CUSTOM_STAT_DEFINED(FrameIoRequestsCompletedSize, BytesToApproxKB(Size), ECsvCustomStatOp::Set);
+	GIoRequestsCompletedSize.Add(RequestSize);
+	CSV_CUSTOM_STAT_DEFINED(FrameIoRequestsCompletedSize, BytesToApproxKB(GIoRequestsCompletedSize.Get()), ECsvCustomStatOp::Set);
 }
+
 void FOnDemandIoBackendStats::OnIoRequestCancel()
 {
-	++IoRequestsCancelled;
-	TRACE_COUNTER_SET(IoRequestsCancelled, IoRequestsCancelled);
-	CSV_CUSTOM_STAT_DEFINED(FrameIoRequestsCancelled, IoRequestsCancelled, ECsvCustomStatOp::Set);
+	GIoRequestsCancelled.Add(1);
+	CSV_CUSTOM_STAT_DEFINED(FrameIoRequestsCancelled, int32(GIoRequestsCancelled.Get()), ECsvCustomStatOp::Set);
 }
+
 void FOnDemandIoBackendStats::OnIoRequestFail()
 {
-	int32 Count = IoRequestsFailed.fetch_add(1, std::memory_order_relaxed);
-	TRACE_COUNTER_SET(IoRequestsFailed, Count);
-	CSV_CUSTOM_STAT_DEFINED(FrameIoRequestsFailed, Count, ECsvCustomStatOp::Set);
+	GIoRequestsFailed.Add(1);
+	CSV_CUSTOM_STAT_DEFINED(FrameIoRequestsFailed, int32(GIoRequestsFailed.Get()), ECsvCustomStatOp::Set);
 }
+
 void FOnDemandIoBackendStats::OnChunkRequestCreate()
 {
-	++ReadRequestsCreated;
-	TRACE_COUNTER_SET(ReadRequestsCreated, ReadRequestsCreated);
-	CSV_CUSTOM_STAT_DEFINED(FrameReadRequestsCreated, ReadRequestsCreated, ECsvCustomStatOp::Set);
+	GReadRequestsCreated.Add(1);
+	CSV_CUSTOM_STAT_DEFINED(FrameReadRequestsCreated, int32(GReadRequestsCreated.Get()), ECsvCustomStatOp::Set);
 }
+
 void FOnDemandIoBackendStats::OnChunkRequestRelease()
 {
-	int32 Count = ReadRequestsRemoved.fetch_add(1, std::memory_order_relaxed);
-	TRACE_COUNTER_SET(ReadRequestsRemoved, Count);
-	CSV_CUSTOM_STAT_DEFINED(FrameReadRequestsRemoved, Count, ECsvCustomStatOp::Set);
+	GReadRequestsRemoved.Add(1);
+	CSV_CUSTOM_STAT_DEFINED(FrameReadRequestsRemoved, int32(GReadRequestsRemoved.Get()), ECsvCustomStatOp::Set);
 }
+
 void FOnDemandIoBackendStats::OnCacheHit(uint64 InSize)
 {
-	int32 Count = CacheHits.fetch_add(1, std::memory_order_relaxed);
-	TRACE_COUNTER_SET(CacheHits, Count);
-	CSV_CUSTOM_STAT_DEFINED(FrameCacheHits, Count, ECsvCustomStatOp::Set);
+	GCacheHits.Add(1);
+	CSV_CUSTOM_STAT_DEFINED(FrameCacheHits, int32(GCacheHits.Get()), ECsvCustomStatOp::Set);
 
-	uint64 Size = CacheHitsSize.fetch_add(InSize, std::memory_order_relaxed);
-	TRACE_COUNTER_SET(CacheHitsSize, Size);
-	CSV_CUSTOM_STAT_DEFINED(FrameCacheHitsSize, BytesToApproxKB(Size), ECsvCustomStatOp::Set);
+	GCacheHitsSize.Add(InSize);
+	CSV_CUSTOM_STAT_DEFINED(FrameCacheHitsSize, BytesToApproxKB(GCacheHitsSize.Get()), ECsvCustomStatOp::Set);
 }
+
 void FOnDemandIoBackendStats::OnCachePut(uint64 InSize)
 {
-	int32 Count = CachePuts.fetch_add(1, std::memory_order_relaxed);
-	TRACE_COUNTER_SET(CachePuts, Count);
-	CSV_CUSTOM_STAT_DEFINED(FrameCachePuts, Count, ECsvCustomStatOp::Set);
+	GCachePuts.Add(1);
+	CSV_CUSTOM_STAT_DEFINED(FrameCachePuts, int32(GCachePuts.Get()), ECsvCustomStatOp::Set);
 
-	uint64 Size = CachePutsSize.fetch_add(InSize, std::memory_order_relaxed);
-	TRACE_COUNTER_SET(CachePutsSize, Size);
-	CSV_CUSTOM_STAT_DEFINED(FrameCachePutsSize, BytesToApproxKB(Size), ECsvCustomStatOp::Set);
+	GCachePutsSize.Add(InSize);
+	CSV_CUSTOM_STAT_DEFINED(FrameCachePutsSize, BytesToApproxKB(GCachePutsSize.Get()), ECsvCustomStatOp::Set);
 }
+
 void FOnDemandIoBackendStats::OnCacheReject(uint64 InSize)
 {
-	int32 Count = CacheRejects.fetch_add(1, std::memory_order_relaxed);
-	TRACE_COUNTER_SET(CacheRejects, Count);
-	CSV_CUSTOM_STAT_DEFINED(FrameCacheRejects, Count, ECsvCustomStatOp::Set);
+	GCacheRejects.Add(1);
+	CSV_CUSTOM_STAT_DEFINED(FrameCacheRejects, int32(GCacheRejects.Get()), ECsvCustomStatOp::Set);
 
-	uint64 Size = CacheRejectsSize.fetch_add(InSize, std::memory_order_relaxed);
-	TRACE_COUNTER_SET(CacheRejectsSize, Size);
-	CSV_CUSTOM_STAT_DEFINED(FrameCacheRejectsSize, BytesToApproxKB(Size), ECsvCustomStatOp::Set);
+	GCacheRejectsSize.Add(InSize);
+	CSV_CUSTOM_STAT_DEFINED(FrameCacheRejectsSize, BytesToApproxKB(GCacheRejectsSize.Get()), ECsvCustomStatOp::Set);
 }
+
 void FOnDemandIoBackendStats::OnHttpRequestEnqueue()
 {
-	int32 Pending = HttpRequestsPending.fetch_add(1, std::memory_order_relaxed);
-	TRACE_COUNTER_SET(HttpRequestsPending, Pending);
-	CSV_CUSTOM_STAT_DEFINED(FrameHttpRequestsPending, Pending, ECsvCustomStatOp::Set);
+	GHttpRequestsPending.Add(1);
+	CSV_CUSTOM_STAT_DEFINED(FrameHttpRequestsPending, int32(GHttpRequestsPending.Get()), ECsvCustomStatOp::Set);
 }
+
 void FOnDemandIoBackendStats::OnHttpRequestDequeue()
 {
-	int32 Pending = HttpRequestsPending.fetch_sub(1, std::memory_order_relaxed);
-	TRACE_COUNTER_SET(HttpRequestsPending, Pending);
-	CSV_CUSTOM_STAT_DEFINED(FrameHttpRequestsPending, Pending, ECsvCustomStatOp::Set);
+	GHttpRequestsPending.Add(-1);
+	CSV_CUSTOM_STAT_DEFINED(FrameHttpRequestsPending, int32(GHttpRequestsPending.Get()), ECsvCustomStatOp::Set);
 
-	int32 Inflight = ++HttpRequestsInflight;
-	TRACE_COUNTER_SET(HttpRequestsInflight, Inflight);
-	CSV_CUSTOM_STAT_DEFINED(FrameHttpRequestsPending, Inflight, ECsvCustomStatOp::Set);
+	GHttpRequestsInflight.Add(1);
+	CSV_CUSTOM_STAT_DEFINED(FrameHttpRequestsPending, int32(GHttpRequestsInflight.Get()), ECsvCustomStatOp::Set);
 }
+
 void FOnDemandIoBackendStats::OnHttpRequestComplete(uint64 InSize)
 {
-	int32 Inflight = --HttpRequestsInflight;
-	TRACE_COUNTER_SET(HttpRequestsInflight, Inflight);
-	CSV_CUSTOM_STAT_DEFINED(FrameHttpRequestsPending, Inflight, ECsvCustomStatOp::Set);
+	GHttpRequestsInflight.Add(-1);
+	CSV_CUSTOM_STAT_DEFINED(FrameHttpRequestsPending, int32(GHttpRequestsInflight.Get()), ECsvCustomStatOp::Set);
 
-	++HttpRequestsCompleted;
-	TRACE_COUNTER_SET(HttpRequestsCompleted, HttpRequestsCompleted);
-	CSV_CUSTOM_STAT_DEFINED(FrameHttpRequestsCompleted, HttpRequestsCompleted, ECsvCustomStatOp::Set);
+	GHttpRequestsCompleted.Add(1);
+	CSV_CUSTOM_STAT_DEFINED(FrameHttpRequestsCompleted, int32(GHttpRequestsCompleted.Get()), ECsvCustomStatOp::Set);
 
-	HttpRequestsCompletedSize += InSize;
-	TRACE_COUNTER_SET(HttpRequestsCompletedSize, HttpRequestsCompletedSize);
-	CSV_CUSTOM_STAT_DEFINED(FrameHttpRequestsCompletedSize, BytesToApproxKB(HttpRequestsCompletedSize), ECsvCustomStatOp::Set);
+	GHttpRequestsCompletedSize.Add(InSize);
+	CSV_CUSTOM_STAT_DEFINED(FrameHttpRequestsCompletedSize, BytesToApproxKB(GHttpRequestsCompletedSize.Get()), ECsvCustomStatOp::Set);
 }
+
 void FOnDemandIoBackendStats::OnHttpRequestFail()
 {
-	int32 Inflight = --HttpRequestsInflight;
-	TRACE_COUNTER_SET(HttpRequestsInflight, Inflight);
-	CSV_CUSTOM_STAT_DEFINED(FrameHttpRequestsPending, Inflight, ECsvCustomStatOp::Set);
-
-	++HttpRequestsFailed;
-	TRACE_COUNTER_SET(HttpRequestsFailed, HttpRequestsFailed);
-	CSV_CUSTOM_STAT_DEFINED(FrameHttpRequestsFailed, HttpRequestsFailed, ECsvCustomStatOp::Set);
+	GHttpRequestsInflight.Add(-1);
+	CSV_CUSTOM_STAT_DEFINED(FrameHttpRequestsPending, int32(GHttpRequestsInflight.Get()), ECsvCustomStatOp::Set);
 
+	GHttpRequestsFailed.Add(1);
+	CSV_CUSTOM_STAT_DEFINED(FrameHttpRequestsFailed, int32(GHttpRequestsFailed.Get()), ECsvCustomStatOp::Set);
 }
 
 } // namespace UE::IO::Private
diff --git a/Engine/Source/Runtime/Experimental/IoStoreOnDemand/Private/Statistics.h b/Engine/Source/Runtime/Experimental/IoStoreOnDemand/Private/Statistics.h
--- a/Engine/Source/Runtime/Experimental/IoStoreOnDemand/Private/Statistics.h
+++ b/Engine/Source/Runtime/Experimental/IoStoreOnDemand/Private/Statistics.h
@@ -17,66 +17,29 @@ namespace UE::IO::Private
 {
 
 #if IAS_WITH_STATISTICS
+#	define IAS_STATISTICS_IMPL() ;
+#else
+#	define IAS_STATISTICS_IMPL() {}
+#endif
 
 class FOnDemandIoBackendStats
 {
 public:
-	void OnIoRequestEnqueue();
-	void OnIoRequestComplete(uint64 RequestSize);
-	void OnIoRequestCancel();
-	void OnIoRequestFail();
-	void OnChunkRequestCreate();
-	void OnChunkRequestRelease();
-	void OnCacheHit(uint64 InSize);
-	void OnCachePut(uint64 InSize);
-	void OnCacheReject(uint64 InSize);
-	void OnHttpRequestEnqueue();
-	void OnHttpRequestDequeue();
-	void OnHttpRequestComplete(uint64 InSize);
-	void OnHttpRequestFail();
-
-private:
-	int32 IoRequestsMade = 0;
-	std::atomic<int32> IoRequestsCompleted = 0;
-	int32 IoRequestsCancelled = 0;
-	std::atomic<int32> IoRequestsFailed = 0;
-	int32 ReadRequestsCreated = 0;
-	std::atomic<int32> ReadRequestsRemoved = 0;
-	std::atomic<int32> CacheHits = 0;
-	std::atomic<int32> CachePuts = 0;
-	std::atomic<int32> CacheRejects = 0;
-	int32 HttpRequestsCompleted = 0;
-	int32 HttpRequestsFailed = 0;
-	std::atomic<int32> HttpRequestsPending = 0;
-	int32 HttpRequestsInflight = 0;
-
-	std::atomic<uint64> IoRequestsCompletedSize = 0;
-	std::atomic<uint64> CacheHitsSize = 0;
-	std::atomic<uint64> CachePutsSize = 0;
-	std::atomic<uint64> CacheRejectsSize = 0;
-	uint64 HttpRequestsCompletedSize = 0;
-};
-
-#else // IAS_WITH_STATISTICS
-
-class FOnDemandIoBackendStats
-{
-public:
-	void OnIoRequestEnqueue() {}
-	void OnIoRequestComplete(uint64 RequestSize) {}
-	void OnIoRequestCancel() {}
-	void OnIoRequestFail() {}
-	void OnChunkRequestCreate() {}
-	void OnChunkRequestRelease() {}
-	void OnCacheHit(uint64 InSize) {}
-	void OnCachePut(uint64 InSize) {}
-	void OnCacheReject(uint64 InSize) {}
-	void OnHttpRequestEnqueue() {}
-	void OnHttpRequestDequeue() {}
-	void OnHttpRequestComplete(uint64 InSize) {}
-	void OnHttpRequestFail() {}
+	void OnIoRequestEnqueue() IAS_STATISTICS_IMPL()
+	void OnIoRequestComplete(uint64 RequestSize) IAS_STATISTICS_IMPL()
+	void OnIoRequestCancel() IAS_STATISTICS_IMPL()
+	void OnIoRequestFail() IAS_STATISTICS_IMPL()
+	void OnChunkRequestCreate() IAS_STATISTICS_IMPL()
+	void OnChunkRequestRelease() IAS_STATISTICS_IMPL()
+	void OnCacheHit(uint64 InSize) IAS_STATISTICS_IMPL()
+	void OnCachePut(uint64 InSize) IAS_STATISTICS_IMPL()
+	void OnCacheReject(uint64 InSize) IAS_STATISTICS_IMPL()
+	void OnHttpRequestEnqueue() IAS_STATISTICS_IMPL()
+	void OnHttpRequestDequeue() IAS_STATISTICS_IMPL()
+	void OnHttpRequestComplete(uint64 InSize) IAS_STATISTICS_IMPL()
+	void OnHttpRequestFail() IAS_STATISTICS_IMPL()
 };
 
-#endif // IAS_WITH_STATISTICS
+#undef IAS_STATISTICS_IMPL
 
 } // namespace UE::IO::Private
