Transfer atomic-int values to the counters that replace them.

#rnx

diff --git a/Engine/Source/Runtime/Experimental/IoStoreOnDemand/Private/Statistics.cpp b/Engine/Source/Runtime/Experimental/IoStoreOnDemand/Private/Statistics.cpp
--- a/Engine/Source/Runtime/Experimental/IoStoreOnDemand/Private/Statistics.cpp
+++ b/Engine/Source/Runtime/Experimental/IoStoreOnDemand/Private/Statistics.cpp
@@ -32,27 +32,27 @@ static int32 BytesToApproxKB(uint64 Bytes) { return int32(Bytes >> 10); }
 #endif 
 
 // iorequest stats
-FCounterInt GIoRequestsMade(TEXT("Ias/IoRequestsMade"), TraceCounterDisplayHint_None);
-FCounterInt GIoRequestsCompleted(TEXT("Ias/IoRequestsCompleted"), TraceCounterDisplayHint_None);
-FCounterInt GIoRequestsCompletedSize(TEXT("Ias/Size/IoRequestsCompletedSize"), TraceCounterDisplayHint_Memory);
-FCounterInt GIoRequestsCancelled(TEXT("Ias/IoRequestsCancelled"), TraceCounterDisplayHint_None);
-FCounterInt GIoRequestsFailed(TEXT("Ias/IoRequestsFailed"), TraceCounterDisplayHint_None);
+FCounterInt			GIoRequestsMade(TEXT("Ias/IoRequestsMade"), TraceCounterDisplayHint_None);
+FCounterAtomicInt	GIoRequestsCompleted(TEXT("Ias/IoRequestsCompleted"), TraceCounterDisplayHint_None);
+FCounterAtomicInt	GIoRequestsCompletedSize(TEXT("Ias/Size/IoRequestsCompletedSize"), TraceCounterDisplayHint_Memory);
+FCounterInt			GIoRequestsCancelled(TEXT("Ias/IoRequestsCancelled"), TraceCounterDisplayHint_None);
+FCounterAtomicInt	GIoRequestsFailed(TEXT("Ias/IoRequestsFailed"), TraceCounterDisplayHint_None);
 // chunkrequest stats
-FCounterInt GReadRequestsCreated(TEXT("Ias/ReadRequestsCreated"), TraceCounterDisplayHint_None);
-FCounterInt GReadRequestsRemoved(TEXT("Ias/ReadRequestsRemoved"), TraceCounterDisplayHint_None);
+FCounterInt			GReadRequestsCreated(TEXT("Ias/ReadRequestsCreated"), TraceCounterDisplayHint_None);
+FCounterAtomicInt	GReadRequestsRemoved(TEXT("Ias/ReadRequestsRemoved"), TraceCounterDisplayHint_None);
 // cache stats
-FCounterInt GCacheHits(TEXT("Ias/CacheHits"), TraceCounterDisplayHint_None);
-FCounterInt GCacheHitsSize(TEXT("Ias/Size/CacheHitsSize"), TraceCounterDisplayHint_Memory);
-FCounterInt GCachePuts(TEXT("Ias/CachePuts"), TraceCounterDisplayHint_None);
-FCounterInt GCachePutsSize(TEXT("Ias/Size/CachePutsSize"), TraceCounterDisplayHint_Memory);
-FCounterInt GCacheRejects(TEXT("Ias/CacheRejects"), TraceCounterDisplayHint_None);
-FCounterInt GCacheRejectsSize(TEXT("Ias/Size/CacheRejectsSize"), TraceCounterDisplayHint_Memory);
+FCounterAtomicInt	GCacheHits(TEXT("Ias/CacheHits"), TraceCounterDisplayHint_None);
+FCounterAtomicInt	GCacheHitsSize(TEXT("Ias/Size/CacheHitsSize"), TraceCounterDisplayHint_Memory);
+FCounterAtomicInt	GCachePuts(TEXT("Ias/CachePuts"), TraceCounterDisplayHint_None);
+FCounterAtomicInt	GCachePutsSize(TEXT("Ias/Size/CachePutsSize"), TraceCounterDisplayHint_Memory);
+FCounterAtomicInt	GCacheRejects(TEXT("Ias/CacheRejects"), TraceCounterDisplayHint_None);
+FCounterAtomicInt	GCacheRejectsSize(TEXT("Ias/Size/CacheRejectsSize"), TraceCounterDisplayHint_Memory);
 // http stats
-FCounterInt GHttpRequestsCompleted(TEXT("Ias/HttpRequestsCompleted"), TraceCounterDisplayHint_None);
-FCounterInt GHttpRequestsFailed(TEXT("Ias/HttpRequestsFailed"), TraceCounterDisplayHint_None);
-FCounterInt GHttpRequestsPending(TEXT("Ias/HttpRequestsPending"), TraceCounterDisplayHint_None);
-FCounterInt GHttpRequestsInflight(TEXT("Ias/HttpRequestsInflight"), TraceCounterDisplayHint_None);
-FCounterInt GHttpRequestsCompletedSize(TEXT("Ias/Size/HttpRequestsCompletedSize"), TraceCounterDisplayHint_Memory);
+FCounterInt			GHttpRequestsCompleted(TEXT("Ias/HttpRequestsCompleted"), TraceCounterDisplayHint_None);
+FCounterInt			GHttpRequestsFailed(TEXT("Ias/HttpRequestsFailed"), TraceCounterDisplayHint_None);
+FCounterAtomicInt	GHttpRequestsPending(TEXT("Ias/HttpRequestsPending"), TraceCounterDisplayHint_None);
+FCounterInt			GHttpRequestsInflight(TEXT("Ias/HttpRequestsInflight"), TraceCounterDisplayHint_None);
+FCounterInt			GHttpRequestsCompletedSize(TEXT("Ias/Size/HttpRequestsCompletedSize"), TraceCounterDisplayHint_Memory);
 
 ////////////////////////////////////////////////////////////////////////////////
 // CSV STATS
