[Engine]
* Changed type dllexport to method/staticvar dll export. Various changes for code that upgrade script get wrong

diff --git a/Engine/Source/Runtime/AIModule/Classes/BehaviorTree/BlackboardAssetProvider.h b/Engine/Source/Runtime/AIModule/Classes/BehaviorTree/BlackboardAssetProvider.h
--- a/Engine/Source/Runtime/AIModule/Classes/BehaviorTree/BlackboardAssetProvider.h
+++ b/Engine/Source/Runtime/AIModule/Classes/BehaviorTree/BlackboardAssetProvider.h
@@ -18,18 +18,18 @@ class UBlackboardAssetProvider : public UInterface
 	GENERATED_UINTERFACE_BODY()
 };
 
-class AIMODULE_API IBlackboardAssetProvider
+class IBlackboardAssetProvider
 {
 	GENERATED_IINTERFACE_BODY()
 
 #if WITH_EDITOR
 	/** Delegate to be called by class implementing IBlackboardAssetProvider when the property containing the returned BlackboardData is changed (i.e. on PostEditChangeProperty). */
 	DECLARE_MULTICAST_DELEGATE_TwoParams(FBlackboardOwnerChanged, UObject* /*AssetOwner*/, UBlackboardData* /*Asset*/);
-	static FBlackboardOwnerChanged OnBlackboardOwnerChanged;
+	static AIMODULE_API FBlackboardOwnerChanged OnBlackboardOwnerChanged;
 #endif
 	/** Returns BlackboardData referenced by the owner object. */
 	UFUNCTION(BlueprintCallable, Category = GameplayTags)
-	virtual UBlackboardData* GetBlackboardAsset() const PURE_VIRTUAL(IBlackboardAssetProvider::GetBlackboardAsset, return nullptr; );
+	AIMODULE_API virtual UBlackboardData* GetBlackboardAsset() const PURE_VIRTUAL(IBlackboardAssetProvider::GetBlackboardAsset, return nullptr; );
 };
 
 
diff --git a/Engine/Source/Runtime/Experimental/Dataflow/Core/Public/Dataflow/DataflowOverrideNode.h b/Engine/Source/Runtime/Experimental/Dataflow/Core/Public/Dataflow/DataflowOverrideNode.h
--- a/Engine/Source/Runtime/Experimental/Dataflow/Core/Public/Dataflow/DataflowOverrideNode.h
+++ b/Engine/Source/Runtime/Experimental/Dataflow/Core/Public/Dataflow/DataflowOverrideNode.h
@@ -21,7 +21,7 @@ struct FDataflowOutput;
 *		the asset. They can read the values by the key.
 */
 USTRUCT()
-struct DATAFLOWCORE_API FDataflowOverrideNode : public FDataflowNode
+struct FDataflowOverrideNode : public FDataflowNode
 {
 	GENERATED_USTRUCT_BODY()
 
@@ -38,7 +38,7 @@ public:
 
 	virtual ~FDataflowOverrideNode() { }
 
-	bool ShouldInvalidate(FName InKey) const;
+	DATAFLOWCORE_API bool ShouldInvalidate(FName InKey) const;
 
 	template <class T>
 	T GetDefaultValue(Dataflow::FContext& Context) const
@@ -47,12 +47,12 @@ public:
 	}
 
 	template <>
-	int32 GetDefaultValue(Dataflow::FContext& Context) const;
+	DATAFLOWCORE_API int32 GetDefaultValue(Dataflow::FContext& Context) const;
 
 	template <>
-	float GetDefaultValue(Dataflow::FContext& Context) const;
+	DATAFLOWCORE_API float GetDefaultValue(Dataflow::FContext& Context) const;
 
-	FString GetValueFromAsset(Dataflow::FContext& Context, const UObject* InOwner) const;
+	DATAFLOWCORE_API FString GetValueFromAsset(Dataflow::FContext& Context, const UObject* InOwner) const;
 
 	//
 	// Evaluate
diff --git a/Engine/Source/Runtime/Experimental/GeometryCollectionEngine/Public/GeometryCollection/GeometryCollectionComponent.h b/Engine/Source/Runtime/Experimental/GeometryCollectionEngine/Public/GeometryCollection/GeometryCollectionComponent.h
--- a/Engine/Source/Runtime/Experimental/GeometryCollectionEngine/Public/GeometryCollection/GeometryCollectionComponent.h
+++ b/Engine/Source/Runtime/Experimental/GeometryCollectionEngine/Public/GeometryCollection/GeometryCollectionComponent.h
@@ -199,7 +199,7 @@ namespace GeometryCollection
 *	solver so it can be edited and afterwards will recreate the proxy
 *	The update can also be specified to reset the dynamic collection
 */
-class GEOMETRYCOLLECTIONENGINE_API FGeometryCollectionEdit
+class FGeometryCollectionEdit
 {
 public:
 	/**
@@ -208,10 +208,10 @@ public:
 	 * @param bShapeIsUnchanged				Override indicating the overall shape of the geometry and clusters is unchanged, even if the rest collection changed.  Useful to e.g., not re-compute convex hulls when we don't need to.
 	 * @param bPropagateAcrossComponents	Propagate updates to all components with the same underlying Rest Collection
 	 */
-	FGeometryCollectionEdit(UGeometryCollectionComponent* InComponent, GeometryCollection::EEditUpdate EditUpdate = GeometryCollection::EEditUpdate::RestPhysicsDynamic, bool bShapeIsUnchanged = false, bool bPropagateToAllMatchingComponents = true);
-	~FGeometryCollectionEdit();
+	GEOMETRYCOLLECTIONENGINE_API FGeometryCollectionEdit(UGeometryCollectionComponent* InComponent, GeometryCollection::EEditUpdate EditUpdate = GeometryCollection::EEditUpdate::RestPhysicsDynamic, bool bShapeIsUnchanged = false, bool bPropagateToAllMatchingComponents = true);
+	GEOMETRYCOLLECTIONENGINE_API ~FGeometryCollectionEdit();
 
-	UGeometryCollection* GetRestCollection();
+	GEOMETRYCOLLECTIONENGINE_API UGeometryCollection* GetRestCollection();
 
 private:
 	UGeometryCollectionComponent* Component;
@@ -222,40 +222,40 @@ private:
 };
 
 #if WITH_EDITOR
-class GEOMETRYCOLLECTIONENGINE_API FScopedColorEdit
+class FScopedColorEdit
 {
 public:
-	FScopedColorEdit(UGeometryCollectionComponent* InComponent, bool bForceUpdate = false);
-	~FScopedColorEdit();
-
-	void SetShowBoneColors(bool ShowBoneColorsIn);
-	bool GetShowBoneColors() const;
-
-	void SetEnableBoneSelection(bool ShowSelectedBonesIn);
-	bool GetEnableBoneSelection() const;
-
-	bool IsBoneSelected(int BoneIndex) const;
-	void Sanitize();
-	void SetSelectedBones(const TArray<int32>& SelectedBonesIn);
-	void AppendSelectedBones(const TArray<int32>& SelectedBonesIn);
-	void ToggleSelectedBones(const TArray<int32>& SelectedBonesIn, bool bAdd, bool bSnapToLevel = true);
-	void AddSelectedBone(int32 BoneIndex);
-	void ClearSelectedBone(int32 BoneIndex);
-	const TArray<int32>& GetSelectedBones() const;
-	void ResetBoneSelection();
-	void SelectBones(GeometryCollection::ESelectionMode SelectionMode);
-	void FilterSelectionToLevel(bool bPreferLowestOnly = false);
-	int32 GetMaxSelectedLevel(bool bOnlyRigid) const;
-	bool IsSelectionValidAtLevel(int32 TargetLevel) const;
-
-	bool IsBoneHighlighted(int BoneIndex) const;
-	void SetHighlightedBones(const TArray<int32>& HighlightedBonesIn, bool bHighlightChildren = false);
-	void AddHighlightedBone(int32 BoneIndex);
-	const TArray<int32>& GetHighlightedBones() const;
-	void ResetHighlightedBones();
-
-	void SetLevelViewMode(int ViewLevel);
-	int GetViewLevel();
+	GEOMETRYCOLLECTIONENGINE_API FScopedColorEdit(UGeometryCollectionComponent* InComponent, bool bForceUpdate = false);
+	GEOMETRYCOLLECTIONENGINE_API ~FScopedColorEdit();
+
+	GEOMETRYCOLLECTIONENGINE_API void SetShowBoneColors(bool ShowBoneColorsIn);
+	GEOMETRYCOLLECTIONENGINE_API bool GetShowBoneColors() const;
+
+	GEOMETRYCOLLECTIONENGINE_API void SetEnableBoneSelection(bool ShowSelectedBonesIn);
+	GEOMETRYCOLLECTIONENGINE_API bool GetEnableBoneSelection() const;
+
+	GEOMETRYCOLLECTIONENGINE_API bool IsBoneSelected(int BoneIndex) const;
+	GEOMETRYCOLLECTIONENGINE_API void Sanitize();
+	GEOMETRYCOLLECTIONENGINE_API void SetSelectedBones(const TArray<int32>& SelectedBonesIn);
+	GEOMETRYCOLLECTIONENGINE_API void AppendSelectedBones(const TArray<int32>& SelectedBonesIn);
+	GEOMETRYCOLLECTIONENGINE_API void ToggleSelectedBones(const TArray<int32>& SelectedBonesIn, bool bAdd, bool bSnapToLevel = true);
+	GEOMETRYCOLLECTIONENGINE_API void AddSelectedBone(int32 BoneIndex);
+	GEOMETRYCOLLECTIONENGINE_API void ClearSelectedBone(int32 BoneIndex);
+	GEOMETRYCOLLECTIONENGINE_API const TArray<int32>& GetSelectedBones() const;
+	GEOMETRYCOLLECTIONENGINE_API void ResetBoneSelection();
+	GEOMETRYCOLLECTIONENGINE_API void SelectBones(GeometryCollection::ESelectionMode SelectionMode);
+	GEOMETRYCOLLECTIONENGINE_API void FilterSelectionToLevel(bool bPreferLowestOnly = false);
+	GEOMETRYCOLLECTIONENGINE_API int32 GetMaxSelectedLevel(bool bOnlyRigid) const;
+	GEOMETRYCOLLECTIONENGINE_API bool IsSelectionValidAtLevel(int32 TargetLevel) const;
+
+	GEOMETRYCOLLECTIONENGINE_API bool IsBoneHighlighted(int BoneIndex) const;
+	GEOMETRYCOLLECTIONENGINE_API void SetHighlightedBones(const TArray<int32>& HighlightedBonesIn, bool bHighlightChildren = false);
+	GEOMETRYCOLLECTIONENGINE_API void AddHighlightedBone(int32 BoneIndex);
+	GEOMETRYCOLLECTIONENGINE_API const TArray<int32>& GetHighlightedBones() const;
+	GEOMETRYCOLLECTIONENGINE_API void ResetHighlightedBones();
+
+	GEOMETRYCOLLECTIONENGINE_API void SetLevelViewMode(int ViewLevel);
+	GEOMETRYCOLLECTIONENGINE_API int GetViewLevel();
 
 private:
 	void UpdateBoneColors();
@@ -275,13 +275,13 @@ private:
 //GetArrayRest (gives original rest value)
 //This generates pointers to arrays marked private. Macro assumes getters are public
 //todo(ocohen): may want to take in a static name
-#define COPY_ON_WRITE_ATTRIBUTE(Type, Name, Group)			\
-	const TManagedArray<Type>& Get##Name##Array() const;	\
-	TManagedArray<Type>& Get##Name##ArrayCopyOnWrite();		\
-	void Reset##Name##ArrayDynamic();						\
-	const TManagedArray<Type>& Get##Name##ArrayRest() const;\
-private:													\
-	TManagedArray<Type>* Indirect##Name##Array;				\
+#define COPY_ON_WRITE_ATTRIBUTE(Type, Name, Group)											\
+	GEOMETRYCOLLECTIONENGINE_API const TManagedArray<Type>& Get##Name##Array() const;		\
+	GEOMETRYCOLLECTIONENGINE_API TManagedArray<Type>& Get##Name##ArrayCopyOnWrite();		\
+	GEOMETRYCOLLECTIONENGINE_API void Reset##Name##ArrayDynamic();							\
+	GEOMETRYCOLLECTIONENGINE_API const TManagedArray<Type>& Get##Name##ArrayRest() const;	\
+private:																					\
+	TManagedArray<Type>* Indirect##Name##Array;												\
 public:
 
 /**
@@ -438,8 +438,8 @@ struct TStructOpsTypeTraits<FGeometryCollectionRepData> : public TStructOpsTypeT
 /**
 *	GeometryCollectionComponent
 */
-UCLASS(meta = (BlueprintSpawnableComponent))
-class GEOMETRYCOLLECTIONENGINE_API UGeometryCollectionComponent : public UMeshComponent, public IChaosNotifyHandlerInterface
+UCLASS(meta = (BlueprintSpawnableComponent), MinimalAPI)
+class UGeometryCollectionComponent : public UMeshComponent, public IChaosNotifyHandlerInterface
 {
 	GENERATED_UCLASS_BODY()
 	friend class FGeometryCollectionEdit;
@@ -451,72 +451,72 @@ class GEOMETRYCOLLECTIONENGINE_API UGeometryCollectionComponent : public UMeshCo
 public:
 
 	//~ Begin UActorComponent Interface.
-	virtual void CreateRenderState_Concurrent(FRegisterComponentContext* Context) override;
-	virtual void SendRenderDynamicData_Concurrent() override;
+	GEOMETRYCOLLECTIONENGINE_API virtual void CreateRenderState_Concurrent(FRegisterComponentContext* Context) override;
+	GEOMETRYCOLLECTIONENGINE_API virtual void SendRenderDynamicData_Concurrent() override;
 	FORCEINLINE void SetRenderStateDirty() { bRenderStateDirty = true; }
-	virtual void SetCollisionObjectType(ECollisionChannel Channel) override;
-	virtual void OnActorEnableCollisionChanged() override;
-	virtual void BeginPlay() override;
-	virtual void EndPlay(const EEndPlayReason::Type ReasonEnd) override;
-	virtual void GetLifetimeReplicatedProps(TArray<FLifetimeProperty>& OutLifetimeProps) const override;
-	virtual void InitializeComponent() override;
-	virtual void GetResourceSizeEx(FResourceSizeEx& CumulativeResourceSize) override;
+	GEOMETRYCOLLECTIONENGINE_API virtual void SetCollisionObjectType(ECollisionChannel Channel) override;
+	GEOMETRYCOLLECTIONENGINE_API virtual void OnActorEnableCollisionChanged() override;
+	GEOMETRYCOLLECTIONENGINE_API virtual void BeginPlay() override;
+	GEOMETRYCOLLECTIONENGINE_API virtual void EndPlay(const EEndPlayReason::Type ReasonEnd) override;
+	GEOMETRYCOLLECTIONENGINE_API virtual void GetLifetimeReplicatedProps(TArray<FLifetimeProperty>& OutLifetimeProps) const override;
+	GEOMETRYCOLLECTIONENGINE_API virtual void InitializeComponent() override;
+	GEOMETRYCOLLECTIONENGINE_API virtual void GetResourceSizeEx(FResourceSizeEx& CumulativeResourceSize) override;
 
 #if WITH_EDITOR
-	virtual void PostEditChangeChainProperty(FPropertyChangedChainEvent& PropertyChangedEvent) override;
-	virtual void PostEditChangeProperty(FPropertyChangedEvent& PropertyChangedEvent) override;
+	GEOMETRYCOLLECTIONENGINE_API virtual void PostEditChangeChainProperty(FPropertyChangedChainEvent& PropertyChangedEvent) override;
+	GEOMETRYCOLLECTIONENGINE_API virtual void PostEditChangeProperty(FPropertyChangedEvent& PropertyChangedEvent) override;
 
 	DECLARE_MULTICAST_DELEGATE(FOnGeometryCollectionPropertyChangedMulticaster)
 	FOnGeometryCollectionPropertyChangedMulticaster OnGeometryCollectionPropertyChanged;
 	typedef FOnGeometryCollectionPropertyChangedMulticaster::FDelegate FOnGeometryCollectionPropertyChanged;
 
 	/** Register / Unregister delegates called when the skeletal mesh property is changed */
-	FDelegateHandle RegisterOnGeometryCollectionPropertyChanged(const FOnGeometryCollectionPropertyChanged& Delegate);
-	void UnregisterOnGeometryCollectionPropertyChanged(FDelegateHandle Handle);
+	GEOMETRYCOLLECTIONENGINE_API FDelegateHandle RegisterOnGeometryCollectionPropertyChanged(const FOnGeometryCollectionPropertyChanged& Delegate);
+	GEOMETRYCOLLECTIONENGINE_API void UnregisterOnGeometryCollectionPropertyChanged(FDelegateHandle Handle);
 #endif
 	//~ Begin UActorComponent Interface. 
 
 
 	//~ Begin USceneComponent Interface.
-	virtual FBoxSphereBounds CalcBounds(const FTransform& LocalToWorld) const override;
+	GEOMETRYCOLLECTIONENGINE_API virtual FBoxSphereBounds CalcBounds(const FTransform& LocalToWorld) const override;
 	virtual FBoxSphereBounds CalcLocalBounds() const { return LocalBounds; }
 
-	virtual bool HasAnySockets() const override;
-	virtual bool DoesSocketExist(FName InSocketName) const override;
-	virtual FTransform GetSocketTransform(FName InSocketName, ERelativeTransformSpace TransformSpace = RTS_World) const override;
-	virtual void QuerySupportedSockets(TArray<FComponentSocketDescription>& OutSockets) const override;
+	GEOMETRYCOLLECTIONENGINE_API virtual bool HasAnySockets() const override;
+	GEOMETRYCOLLECTIONENGINE_API virtual bool DoesSocketExist(FName InSocketName) const override;
+	GEOMETRYCOLLECTIONENGINE_API virtual FTransform GetSocketTransform(FName InSocketName, ERelativeTransformSpace TransformSpace = RTS_World) const override;
+	GEOMETRYCOLLECTIONENGINE_API virtual void QuerySupportedSockets(TArray<FComponentSocketDescription>& OutSockets) const override;
 	
-	virtual void TickComponent(float DeltaTime, enum ELevelTick TickType, FActorComponentTickFunction *ThisTickFunction) override;
-	virtual void AsyncPhysicsTickComponent(float DeltaTime, float SimTime) override;
+	GEOMETRYCOLLECTIONENGINE_API virtual void TickComponent(float DeltaTime, enum ELevelTick TickType, FActorComponentTickFunction *ThisTickFunction) override;
+	GEOMETRYCOLLECTIONENGINE_API virtual void AsyncPhysicsTickComponent(float DeltaTime, float SimTime) override;
 
-	virtual void OnHiddenInGameChanged() override;
+	GEOMETRYCOLLECTIONENGINE_API virtual void OnHiddenInGameChanged() override;
 	//~ Begin USceneComponent Interface.
 
 
 	//~ Begin UPrimitiveComponent Interface.
-	virtual FPrimitiveSceneProxy* CreateSceneProxy() override;
-	virtual void OnRegister() override;
-	virtual void OnUnregister() override;
-	virtual FBodyInstance* GetBodyInstance(FName BoneName = NAME_None, bool bGetWelded = true, int32 Index = -1) const override;
-	virtual void SetNotifyRigidBodyCollision(bool bNewNotifyRigidBodyCollision) override;
-	virtual bool CanEditSimulatePhysics() override;
-	virtual void SetSimulatePhysics(bool bEnabled) override;
-	virtual void AddForce(FVector Force, FName BoneName = NAME_None, bool bAccelChange = false) override;
-	virtual void AddForceAtLocation(FVector Force, FVector WorldLocation, FName BoneName = NAME_None) override;
-	virtual void AddImpulse(FVector Impulse, FName BoneName = NAME_None, bool bVelChange = false) override;
-	virtual void AddImpulseAtLocation(FVector Impulse, FVector WorldLocation, FName BoneName = NAME_None) override;
-	virtual void AddRadialForce(FVector Origin, float Radius, float Strength, ERadialImpulseFalloff Falloff, bool bAccelChange = false) override;
-	virtual void AddRadialImpulse(FVector Origin, float Radius, float Strength, enum ERadialImpulseFalloff Falloff, bool bVelChange = false) override;
-	virtual void AddTorqueInRadians(FVector Torque, FName BoneName = NAME_None, bool bAccelChange = false) override;
-	virtual void PostLoad() override;
+	GEOMETRYCOLLECTIONENGINE_API virtual FPrimitiveSceneProxy* CreateSceneProxy() override;
+	GEOMETRYCOLLECTIONENGINE_API virtual void OnRegister() override;
+	GEOMETRYCOLLECTIONENGINE_API virtual void OnUnregister() override;
+	GEOMETRYCOLLECTIONENGINE_API virtual FBodyInstance* GetBodyInstance(FName BoneName = NAME_None, bool bGetWelded = true, int32 Index = -1) const override;
+	GEOMETRYCOLLECTIONENGINE_API virtual void SetNotifyRigidBodyCollision(bool bNewNotifyRigidBodyCollision) override;
+	GEOMETRYCOLLECTIONENGINE_API virtual bool CanEditSimulatePhysics() override;
+	GEOMETRYCOLLECTIONENGINE_API virtual void SetSimulatePhysics(bool bEnabled) override;
+	GEOMETRYCOLLECTIONENGINE_API virtual void AddForce(FVector Force, FName BoneName = NAME_None, bool bAccelChange = false) override;
+	GEOMETRYCOLLECTIONENGINE_API virtual void AddForceAtLocation(FVector Force, FVector WorldLocation, FName BoneName = NAME_None) override;
+	GEOMETRYCOLLECTIONENGINE_API virtual void AddImpulse(FVector Impulse, FName BoneName = NAME_None, bool bVelChange = false) override;
+	GEOMETRYCOLLECTIONENGINE_API virtual void AddImpulseAtLocation(FVector Impulse, FVector WorldLocation, FName BoneName = NAME_None) override;
+	GEOMETRYCOLLECTIONENGINE_API virtual void AddRadialForce(FVector Origin, float Radius, float Strength, ERadialImpulseFalloff Falloff, bool bAccelChange = false) override;
+	GEOMETRYCOLLECTIONENGINE_API virtual void AddRadialImpulse(FVector Origin, float Radius, float Strength, enum ERadialImpulseFalloff Falloff, bool bVelChange = false) override;
+	GEOMETRYCOLLECTIONENGINE_API virtual void AddTorqueInRadians(FVector Torque, FName BoneName = NAME_None, bool bAccelChange = false) override;
+	GEOMETRYCOLLECTIONENGINE_API virtual void PostLoad() override;
 	//~ End UPrimitiveComponent Interface.
 
 
 	//~ Begin UMeshComponent Interface.	
-	virtual int32 GetNumMaterials() const override;
-	virtual UMaterialInterface* GetMaterial(int32 MaterialIndex) const override;
-	virtual void GetUsedMaterials(TArray<UMaterialInterface*>& OutMaterials, bool bGetDebugMaterials = false) const override;
-	virtual FMaterialRelevance GetMaterialRelevance(ERHIFeatureLevel::Type InFeatureLevel) const override;
+	GEOMETRYCOLLECTIONENGINE_API virtual int32 GetNumMaterials() const override;
+	GEOMETRYCOLLECTIONENGINE_API virtual UMaterialInterface* GetMaterial(int32 MaterialIndex) const override;
+	GEOMETRYCOLLECTIONENGINE_API virtual void GetUsedMaterials(TArray<UMaterialInterface*>& OutMaterials, bool bGetDebugMaterials = false) const override;
+	GEOMETRYCOLLECTIONENGINE_API virtual FMaterialRelevance GetMaterialRelevance(ERHIFeatureLevel::Type InFeatureLevel) const override;
 	//~ End UMeshComponent Interface.
 
 	/** Chaos RBD Solver override. Will use the world's default solver actor if null. */
@@ -539,7 +539,7 @@ public:
 	 * @param Strain strain / damage to apply 
 	 */
 	UFUNCTION(BlueprintCallable, Category = "ChaosPhysics")
-	void ApplyExternalStrain(int32 ItemIndex, const FVector& Location, float Radius = 0, int32 PropagationDepth = 0, float PropagationFactor = 1, float Strain = 0);
+	GEOMETRYCOLLECTIONENGINE_API void ApplyExternalStrain(int32 ItemIndex, const FVector& Location, float Radius = 0, int32 PropagationDepth = 0, float PropagationFactor = 1, float Strain = 0);
 
 	/**
 	 * Apply an internal strain to specific piece of the geometry collection
@@ -551,45 +551,45 @@ public:
 	 * @param Strain strain / damage to apply 
 	 */
 	UFUNCTION(BlueprintCallable, Category = "ChaosPhysics")
-	void ApplyInternalStrain(int32 ItemIndex, const FVector& Location, float Radius = 0, int32 PropagationDepth = 0, float PropagationFactor = 1, float Strain = 0);
+	GEOMETRYCOLLECTIONENGINE_API void ApplyInternalStrain(int32 ItemIndex, const FVector& Location, float Radius = 0, int32 PropagationDepth = 0, float PropagationFactor = 1, float Strain = 0);
 	
 	/**
 	 * Crumbe a cluster into all its pieces
 	 * @param ItemIndex item index ( from HitResult) of the cluster to crumble
 	 */
 	UFUNCTION(BlueprintCallable, Category = "ChaosPhysics")
-	void CrumbleCluster(int32 ItemIndex);
+	GEOMETRYCOLLECTIONENGINE_API void CrumbleCluster(int32 ItemIndex);
 
 	/**
 	* Crumbe active clusters for this entire geometry collection
 	* this will apply to internal and regular clusters
 	*/
 	UFUNCTION(BlueprintCallable, Category = "ChaosPhysics")
-	void CrumbleActiveClusters();
+	GEOMETRYCOLLECTIONENGINE_API void CrumbleActiveClusters();
 
 	/**
 	* Set a piece or cluster to be anchored or not 
 	*/
 	UFUNCTION(BlueprintCallable, Category = "ChaosPhysics")
-	void SetAnchoredByIndex(int32 Index, bool bAnchored);
+	GEOMETRYCOLLECTIONENGINE_API void SetAnchoredByIndex(int32 Index, bool bAnchored);
 
 	/**
 	* Set all pieces within a world space bounding box to be anchored or not 
 	*/
 	UFUNCTION(BlueprintCallable, Category = "ChaosPhysics")
-	void SetAnchoredByBox(FBox WorldSpaceBox, bool bAnchored, int32 MaxLevel = -1);
+	GEOMETRYCOLLECTIONENGINE_API void SetAnchoredByBox(FBox WorldSpaceBox, bool bAnchored, int32 MaxLevel = -1);
 
 	/**
 	* Set all pieces within a world transformed bounding box to be anchored or not
 	*/
 	UFUNCTION(BlueprintCallable, Category = "ChaosPhysics")
-	void SetAnchoredByTransformedBox(FBox Box, FTransform Transform, bool bAnchored, int32 MaxLevel = -1);
+	GEOMETRYCOLLECTIONENGINE_API void SetAnchoredByTransformedBox(FBox Box, FTransform Transform, bool bAnchored, int32 MaxLevel = -1);
 
 	/**
 	* this will remove anchors on all the pieces ( including the static and kinematic initial states ones ) of the geometry colection
 	*/
 	UFUNCTION(BlueprintCallable, Category = "ChaosPhysics")
-	void RemoveAllAnchors();
+	GEOMETRYCOLLECTIONENGINE_API void RemoveAllAnchors();
 
 	/**
 	 * Apply linear velocity on breaking pieces for a specific cluster
@@ -598,7 +598,7 @@ public:
 	 * @param LinearVelocity linear velocity to apply
 	 */
 	UFUNCTION(BlueprintCallable, Category = "ChaosPhysics")
-	void ApplyBreakingLinearVelocity(int32 ItemIndex, const FVector& LinearVelocity);
+	GEOMETRYCOLLECTIONENGINE_API void ApplyBreakingLinearVelocity(int32 ItemIndex, const FVector& LinearVelocity);
 
 	/**
 	 * Apply linear velocity on breaking pieces for a specific cluster
@@ -607,7 +607,7 @@ public:
 	 * @param AngularVelocity linear velocity to apply
 	 */
 	UFUNCTION(BlueprintCallable, Category = "ChaosPhysics")
-	void ApplyBreakingAngularVelocity(int32 ItemIndex, const FVector& AngularVelocity);
+	GEOMETRYCOLLECTIONENGINE_API void ApplyBreakingAngularVelocity(int32 ItemIndex, const FVector& AngularVelocity);
 	
 	/**
 	 * Apply linear velocity on specific piece 
@@ -615,7 +615,7 @@ public:
 	* @param LinearVelocity linear velocity to apply
 	 */
 	UFUNCTION(BlueprintCallable, Category = "ChaosPhysics")
-	void ApplyLinearVelocity(int32 ItemIndex, const FVector& LinearVelocity);
+	GEOMETRYCOLLECTIONENGINE_API void ApplyLinearVelocity(int32 ItemIndex, const FVector& LinearVelocity);
 
 	/**
 	 * Apply angular velocity on specific piece 
@@ -623,7 +623,7 @@ public:
 	* @param AngularVelocity linear velocity to apply
 	 */
 	UFUNCTION(BlueprintCallable, Category = "ChaosPhysics")
-	void ApplyAngularVelocity(int32 ItemIndex, const FVector& AngularVelocity);
+	GEOMETRYCOLLECTIONENGINE_API void ApplyAngularVelocity(int32 ItemIndex, const FVector& AngularVelocity);
 
 	/**
 	 * Get the initial level of a specific piece
@@ -632,33 +632,33 @@ public:
 	 * @param Level of the piece ( 0 for root level and positive for the rest ) 
 	 */
 	UFUNCTION(BlueprintCallable, Category = "ChaosPhysics")
-	int32 GetInitialLevel(int32 ItemIndex);
+	GEOMETRYCOLLECTIONENGINE_API int32 GetInitialLevel(int32 ItemIndex);
 	
 	/** Get the root item index of the hierarchy */
 	UFUNCTION(BlueprintCallable, Category = "ChaosPhysics")
-	int32 GetRootIndex() const;
+	GEOMETRYCOLLECTIONENGINE_API int32 GetRootIndex() const;
 
 	/** Get the root item initial transform in world space */
 	UFUNCTION(BlueprintCallable, Category = "ChaosPhysics")
-	FTransform GetRootInitialTransform() const;
+	GEOMETRYCOLLECTIONENGINE_API FTransform GetRootInitialTransform() const;
 
 	/** Get the root item current world transform */
 	UFUNCTION(BlueprintCallable, Category = "ChaosPhysics")
-	FTransform GetRootCurrentTransform() const;
+	GEOMETRYCOLLECTIONENGINE_API FTransform GetRootCurrentTransform() const;
 
 	/** 
 	* Get the initial rest transforms in component (local) space  space, 
 	* they are the transforms as defined in the rest collection asset 
 	*/
 	UFUNCTION(BlueprintCallable, Category = "ChaosPhysics")
-	TArray<FTransform> GetInitialLocalRestTransforms() const;
+	GEOMETRYCOLLECTIONENGINE_API TArray<FTransform> GetInitialLocalRestTransforms() const;
 
 	/** 
 	* Set the local rest transform, this may be different from the rest collection 
 	* If the geometry collection is already simulating those matrices will be overriden by the physics state updates
 	*/
 	UFUNCTION(BlueprintCallable, Category = "ChaosPhysics")
-	void SetLocalRestTransforms(const TArray<FTransform>& Transforms, bool bOnlyLeaves);
+	GEOMETRYCOLLECTIONENGINE_API void SetLocalRestTransforms(const TArray<FTransform>& Transforms, bool bOnlyLeaves);
 
 	/**
 	* Get mass and extent of a specific piece
@@ -666,15 +666,15 @@ public:
 	* @param Level of the piece ( 0 for root level and positive for the rest )
 	*/
 	UFUNCTION(BlueprintCallable, Category = "ChaosPhysics")
-	void GetMassAndExtents(int32 ItemIndex, float& OutMass, FBox& OutExtents);
+	GEOMETRYCOLLECTIONENGINE_API void GetMassAndExtents(int32 ItemIndex, float& OutMass, FBox& OutExtents);
 
 	/** RestCollection */
 	UFUNCTION(BlueprintCallable, Category = "ChaosPhysics")
-	void SetRestCollection(const UGeometryCollection * RestCollectionIn, bool bApplyAssetDefaults = true);
+	GEOMETRYCOLLECTIONENGINE_API void SetRestCollection(const UGeometryCollection * RestCollectionIn, bool bApplyAssetDefaults = true);
 
 	/** RestCollection */
 	UFUNCTION(BlueprintCallable, Category = "ChaosPhysics")
-	FString GetDebugInfo();
+	GEOMETRYCOLLECTIONENGINE_API FString GetDebugInfo();
 
 	FORCEINLINE const UGeometryCollection* GetRestCollection() const { return RestCollection; }
 
@@ -683,27 +683,27 @@ public:
 	FORCEINLINE FScopedColorEdit EditBoneSelection(bool bForceUpdate = false) { return FScopedColorEdit(this, bForceUpdate); }
 
 	/** Propagate bone selection to embedded geometry components. */
-	void SelectEmbeddedGeometry();
+	GEOMETRYCOLLECTIONENGINE_API void SelectEmbeddedGeometry();
 #endif
 
 	UFUNCTION(BlueprintCallable, Category = "ChaosPhysics")
-	void SetEnableDamageFromCollision(bool bValue);
+	GEOMETRYCOLLECTIONENGINE_API void SetEnableDamageFromCollision(bool bValue);
 
 
 	UFUNCTION(BlueprintCallable, Category = "ChaosPhysics")
-	void SetAbandonedParticleCollisionProfileName(FName CollisionProfile);
+	GEOMETRYCOLLECTIONENGINE_API void SetAbandonedParticleCollisionProfileName(FName CollisionProfile);
 
 	UFUNCTION(BlueprintCallable, Category = "ChaosPhysics")
-	void SetPerLevelCollisionProfileNames(const TArray<FName>& ProfileNames);
+	GEOMETRYCOLLECTIONENGINE_API void SetPerLevelCollisionProfileNames(const TArray<FName>& ProfileNames);
 
 	UFUNCTION(BlueprintCallable, Category = "ChaosPhysics")
-	void SetPerParticleCollisionProfileName(const TArray<int32>& BoneIds, FName ProfileName);
+	GEOMETRYCOLLECTIONENGINE_API void SetPerParticleCollisionProfileName(const TArray<int32>& BoneIds, FName ProfileName);
 
 	/** API for getting at geometry collection data */
-	int32 GetNumElements(FName Group) const;
+	GEOMETRYCOLLECTIONENGINE_API int32 GetNumElements(FName Group) const;
 
 	// Update cached bounds; used e.g. when updating the exploded view of the geometry collection
-	void UpdateCachedBounds();
+	GEOMETRYCOLLECTIONENGINE_API void UpdateCachedBounds();
 
 #define COPY_ON_WRITE_ATTRIBUTES \
 	/* Vertices Group */ \
@@ -753,7 +753,7 @@ public:
 
 	/** Apply default values from asset ( damage related data and physics material ) */
 	UFUNCTION(BlueprintCallable, CallInEditor, Category = "ChaosPhysics", meta = (EditCondition = "RestCollection != nullptr"))
-	void ApplyAssetDefaults();
+	GEOMETRYCOLLECTIONENGINE_API void ApplyAssetDefaults();
 
 	UPROPERTY(EditAnywhere, NoClear, BlueprintReadOnly, Category = "ChaosPhysics")
 	TArray<TObjectPtr<const AFieldSystemActor>> InitializationFields;
@@ -806,7 +806,7 @@ public:
 	TArray<float> GetDamageThreshold() const { return DamageThreshold; }
 
 	UFUNCTION(BlueprintCallable, BlueprintInternalUseOnly)
-	void SetDamageThreshold(const TArray<float>& InDamageThreshold);
+	GEOMETRYCOLLECTIONENGINE_API void SetDamageThreshold(const TArray<float>& InDamageThreshold);
 
 	/** Damage threshold for clusters at different levels. */
 	UPROPERTY(EditAnywhere, BlueprintReadWrite, Category = "ChaosPhysics|Damage", meta = (EditCondition = "DamageModel == EDamageModelTypeEnum::Chaos_Damage_Model_UserDefined_Damage_Threshold"))
@@ -884,7 +884,7 @@ public:
 	*
 	*/
 	UFUNCTION(BlueprintCallable, Category = "Field", DisplayName = "Set Dynamic State")
-	void ApplyKinematicField(UPARAM(DisplayName = "Field Radius") float Radius, 
+	GEOMETRYCOLLECTIONENGINE_API void ApplyKinematicField(UPARAM(DisplayName = "Field Radius") float Radius, 
 							 UPARAM(DisplayName = "Center Position") FVector Position);
 
 	/**
@@ -901,7 +901,7 @@ public:
 	*
 	*/
 	UFUNCTION(BlueprintCallable, Category = "Field", DisplayName = "Add Physics Field")
-	void ApplyPhysicsField(UPARAM(DisplayName = "Enable Field") bool Enabled, 
+	GEOMETRYCOLLECTIONENGINE_API void ApplyPhysicsField(UPARAM(DisplayName = "Enable Field") bool Enabled, 
 						   UPARAM(DisplayName = "Physics Type") EGeometryCollectionPhysicsTypeEnum Target, 
 						   UPARAM(DisplayName = "Meta Data") UFieldSystemMetaData* MetaData, 
 						   UPARAM(DisplayName = "Field Node") UFieldNodeBase* Field);
@@ -915,7 +915,7 @@ public:
 	UPROPERTY(BlueprintAssignable, Category = "Game|Damage")
 	FNotifyGeometryCollectionPhysicsStateChange NotifyGeometryCollectionPhysicsStateChange;
 
-	bool GetIsObjectDynamic() const;
+	GEOMETRYCOLLECTIONENGINE_API bool GetIsObjectDynamic() const;
 
 	DECLARE_DYNAMIC_MULTICAST_DELEGATE_OneParam(FNotifyGeometryCollectionPhysicsLoadingStateChange, UGeometryCollectionComponent*, FracturedComponent);
 	UPROPERTY(BlueprintAssignable, Category = "Game|Loading")
@@ -934,24 +934,24 @@ public:
 	bool GetShowBoneColors() const { return bShowBoneColors; }
 	bool GetEnableBoneSelection() const { return bEnableBoneSelection; }
 	
-	bool GetSuppressSelectionMaterial() const;
+	GEOMETRYCOLLECTIONENGINE_API bool GetSuppressSelectionMaterial() const;
 	
-	const int GetBoneSelectedMaterialID() const;
+	GEOMETRYCOLLECTIONENGINE_API const int GetBoneSelectedMaterialID() const;
 	
 #if WITH_EDITORONLY_DATA
 	FORCEINLINE const TArray<int32>& GetSelectedBones() const { return SelectedBones; }
 	FORCEINLINE const TArray<int32>& GetHighlightedBones() const { return HighlightedBones; }
 #endif
 
-	FPhysScene_Chaos* GetInnerChaosScene() const;
-	AChaosSolverActor* GetPhysicsSolverActor() const;
+	GEOMETRYCOLLECTIONENGINE_API FPhysScene_Chaos* GetInnerChaosScene() const;
+	GEOMETRYCOLLECTIONENGINE_API AChaosSolverActor* GetPhysicsSolverActor() const;
 
 	const FGeometryCollectionPhysicsProxy* GetPhysicsProxy() const { return PhysicsProxy; }
 	FGeometryCollectionPhysicsProxy* GetPhysicsProxy() { return PhysicsProxy; }
 
 #if GEOMETRYCOLLECTION_EDITOR_SELECTION
 	/** Enable/disable the scene proxy per transform selection mode. When disabled the per material id default selection is used instead. */
-	void EnableTransformSelectionMode(bool bEnable);
+	GEOMETRYCOLLECTIONENGINE_API void EnableTransformSelectionMode(bool bEnable);
 	bool GetIsTransformSelectionMode() const { return bIsTransformSelectionModeEnabled; }
 #endif  // #if GEOMETRYCOLLECTION_EDITOR_SELECTION
 
@@ -966,13 +966,13 @@ public:
 	/**/
 	const TArray<bool>& GetDisabledFlags() const { return DisabledFlags; }
 
-	virtual void OnCreatePhysicsState() override;
-	void RegisterAndInitializePhysicsProxy();
-	virtual void OnDestroyPhysicsState() override;
-	virtual bool ShouldCreatePhysicsState() const override;
-	virtual bool HasValidPhysicsState() const override;
+	GEOMETRYCOLLECTIONENGINE_API virtual void OnCreatePhysicsState() override;
+	GEOMETRYCOLLECTIONENGINE_API void RegisterAndInitializePhysicsProxy();
+	GEOMETRYCOLLECTIONENGINE_API virtual void OnDestroyPhysicsState() override;
+	GEOMETRYCOLLECTIONENGINE_API virtual bool ShouldCreatePhysicsState() const override;
+	GEOMETRYCOLLECTIONENGINE_API virtual bool HasValidPhysicsState() const override;
 
-	virtual void OnUpdateTransform(EUpdateTransformFlags UpdateTransformFlags, ETeleportType Teleport = ETeleportType::None) override;
+	GEOMETRYCOLLECTIONENGINE_API virtual void OnUpdateTransform(EUpdateTransformFlags UpdateTransformFlags, ETeleportType Teleport = ETeleportType::None) override;
 
 	// Mirrored from the proxy on a sync
 	//TManagedArray<int32> RigidBodyIds;
@@ -982,31 +982,31 @@ public:
 
 	/** Changes whether or not this component will get future break notifications. */
 	UFUNCTION(BlueprintCallable, Category = "Physics")
-	void SetNotifyBreaks(bool bNewNotifyBreaks);
+	GEOMETRYCOLLECTIONENGINE_API void SetNotifyBreaks(bool bNewNotifyBreaks);
 
 	/** Changes whether or not this component will get future removal notifications. */
 	UFUNCTION(BlueprintCallable, Category = "Physics")
-	void SetNotifyRemovals(bool bNewNotifyRemovals);
+	GEOMETRYCOLLECTIONENGINE_API void SetNotifyRemovals(bool bNewNotifyRemovals);
 
 	/** Changes whether or not this component will get future crumbling notifications. */
 	UFUNCTION(BlueprintCallable, Category = "Physics")
-	void SetNotifyCrumblings(bool bNewNotifyCrumblings, bool bNewCrumblingEventIncludesChildren = false);
+	GEOMETRYCOLLECTIONENGINE_API void SetNotifyCrumblings(bool bNewNotifyCrumblings, bool bNewCrumblingEventIncludesChildren = false);
 
 	/** Changes whether or not this component will get future global break notifications. */
 	UFUNCTION(BlueprintCallable, Category = "Physics")
-	void SetNotifyGlobalBreaks(bool bNewNotifyGlobalBreaks);
+	GEOMETRYCOLLECTIONENGINE_API void SetNotifyGlobalBreaks(bool bNewNotifyGlobalBreaks);
 
 	/** Changes whether or not this component will get future global collision notifications. */
 	UFUNCTION(BlueprintCallable, Category = "Physics")
-	void SetNotifyGlobalCollision(bool bNewNotifyGlobalCollisions);
+	GEOMETRYCOLLECTIONENGINE_API void SetNotifyGlobalCollision(bool bNewNotifyGlobalCollisions);
 
 	/** Changes whether or not this component will get future global removal notifications. */
 	UFUNCTION(BlueprintCallable, Category = "Physics")
-	void SetNotifyGlobalRemovals(bool bNewNotifyGlobalRemovals);
+	GEOMETRYCOLLECTIONENGINE_API void SetNotifyGlobalRemovals(bool bNewNotifyGlobalRemovals);
 
 	/** Changes whether or not this component will get future global crumbling notifications. */
 	UFUNCTION(BlueprintCallable, Category = "Physics")
-	void SetNotifyGlobalCrumblings(bool bNewNotifyGlobalCrumblings, bool bGlobalNewCrumblingEventIncludesChildren);
+	GEOMETRYCOLLECTIONENGINE_API void SetNotifyGlobalCrumblings(bool bNewNotifyGlobalCrumblings, bool bGlobalNewCrumblingEventIncludesChildren);
 	
 	/** Overrideable native notification */
 	virtual void NotifyBreak(const FChaosBreakEvent& Event) {};
@@ -1026,11 +1026,11 @@ public:
 	// todo(chaos) remove when no longer necessary
 	FOnChaosBreakEvent OnRootBreakEvent;
 
-	void DispatchBreakEvent(const FChaosBreakEvent& Event);
+	GEOMETRYCOLLECTIONENGINE_API void DispatchBreakEvent(const FChaosBreakEvent& Event);
 
-	void DispatchRemovalEvent(const FChaosRemovalEvent& Event);
+	GEOMETRYCOLLECTIONENGINE_API void DispatchRemovalEvent(const FChaosRemovalEvent& Event);
 
-	void DispatchCrumblingEvent(const FChaosCrumblingEvent& Event);
+	GEOMETRYCOLLECTIONENGINE_API void DispatchCrumblingEvent(const FChaosCrumblingEvent& Event);
 
 	UPROPERTY(Transient, VisibleAnywhere, BlueprintReadWrite, Interp, Category = "Chaos")
 	float DesiredCacheTime;
@@ -1038,41 +1038,41 @@ public:
 	UPROPERTY(Transient, VisibleAnywhere, BlueprintReadWrite, Category = "Chaos")
 	bool CachePlayback;
 
-	bool DoCustomNavigableGeometryExport(FNavigableGeometryExport& GeomExport) const override;
+	GEOMETRYCOLLECTIONENGINE_API bool DoCustomNavigableGeometryExport(FNavigableGeometryExport& GeomExport) const override;
 
 	/** Gets the physical material to use for this geometry collection, taking into account instance overrides and render materials */
-	UPhysicalMaterial* GetPhysicalMaterial() const;
+	GEOMETRYCOLLECTIONENGINE_API UPhysicalMaterial* GetPhysicalMaterial() const;
 
 	/** Update component structure to reflect any changes to the embedded geometry */
-	void InitializeEmbeddedGeometry();
+	GEOMETRYCOLLECTIONENGINE_API void InitializeEmbeddedGeometry();
 	
 	/** Update instanced static mesh components to reflect internal embedded geometry state. */
-	void RefreshEmbeddedGeometry();
+	GEOMETRYCOLLECTIONENGINE_API void RefreshEmbeddedGeometry();
 
 #if WITH_EDITOR
-	void SetEmbeddedGeometrySelectable(bool bSelectableIn);
-	int32 EmbeddedIndexToTransformIndex(const UInstancedStaticMeshComponent* ISMComponent, int32 InstanceIndex) const;
-	void GetBoneColors(TArray<FColor>& OutColors) const;
-	void GetHiddenTransforms(TArray<bool>& OutHiddenTransforms) const;
+	GEOMETRYCOLLECTIONENGINE_API void SetEmbeddedGeometrySelectable(bool bSelectableIn);
+	GEOMETRYCOLLECTIONENGINE_API int32 EmbeddedIndexToTransformIndex(const UInstancedStaticMeshComponent* ISMComponent, int32 InstanceIndex) const;
+	GEOMETRYCOLLECTIONENGINE_API void GetBoneColors(TArray<FColor>& OutColors) const;
+	GEOMETRYCOLLECTIONENGINE_API void GetHiddenTransforms(TArray<bool>& OutHiddenTransforms) const;
 #endif
 
-	void GetRestTransforms(TArray<FMatrix44f>& OutRestTransforms) const;
+	GEOMETRYCOLLECTIONENGINE_API void GetRestTransforms(TArray<FMatrix44f>& OutRestTransforms) const;
 
 #if WITH_EDITORONLY_DATA
-	const FDamageCollector* GetRunTimeDataCollector() const;
+	GEOMETRYCOLLECTIONENGINE_API const FDamageCollector* GetRunTimeDataCollector() const;
 #endif
 	
 	// #todo should this only be available in editor?
-	void SetRestState(TArray<FTransform>&& InRestTransforms);
+	GEOMETRYCOLLECTIONENGINE_API void SetRestState(TArray<FTransform>&& InRestTransforms);
 
 	/** Set the dynamic state for all bodies in the DynamicCollection. */
-	void SetDynamicState(const Chaos::EObjectStateType& NewDynamicState);
+	GEOMETRYCOLLECTIONENGINE_API void SetDynamicState(const Chaos::EObjectStateType& NewDynamicState);
 
 	/** Set transforms for all bodies in the DynamicCollection. */
-	void SetInitialTransforms(const TArray<FTransform>& InitialTransforms);
+	GEOMETRYCOLLECTIONENGINE_API void SetInitialTransforms(const TArray<FTransform>& InitialTransforms);
 
 	/** Modify DynamicCollection transform hierarchy to effect cluster breaks releasing the specified indices. */
-	void SetInitialClusterBreaks(const TArray<int32>& ReleaseIndices);
+	GEOMETRYCOLLECTIONENGINE_API void SetInitialClusterBreaks(const TArray<int32>& ReleaseIndices);
 
 	/** Used by Niagara DI to query global matrices rather than recalculating them again */
 	UE_DEPRECATED(5.3, "Use GetComponentSpaceTransforms instead")
@@ -1080,23 +1080,23 @@ public:
 
 	const TArray<FTransform> GetComponentSpaceTransforms() { return ComponentSpaceTransforms; }
 
-	const FGeometryDynamicCollection* GetDynamicCollection() const;
-	FGeometryDynamicCollection* GetDynamicCollection();  // TEMP HACK?
+	GEOMETRYCOLLECTIONENGINE_API const FGeometryDynamicCollection* GetDynamicCollection() const;
+	GEOMETRYCOLLECTIONENGINE_API FGeometryDynamicCollection* GetDynamicCollection();  // TEMP HACK?
 
-	TArray<UStaticMeshComponent*> CreateProxyComponents() const;
+	GEOMETRYCOLLECTIONENGINE_API TArray<UStaticMeshComponent*> CreateProxyComponents() const;
 
 	/** Force all GC components to reregister their custom renderer objects. */
-	static void ReregisterAllCustomRenderers();
+	static GEOMETRYCOLLECTIONENGINE_API void ReregisterAllCustomRenderers();
 
 public:
 	UPROPERTY(BlueprintAssignable, Category = "Collision")
 	FOnChaosPhysicsCollision OnChaosPhysicsCollision;
 	
 	UFUNCTION(BlueprintImplementableEvent, meta = (DisplayName = "Physics Collision"), Category = "Collision")
-	void ReceivePhysicsCollision(const FChaosPhysicsCollisionInfo& CollisionInfo);
+	GEOMETRYCOLLECTIONENGINE_API void ReceivePhysicsCollision(const FChaosPhysicsCollisionInfo& CollisionInfo);
 
 	// IChaosNotifyHandlerInterface
-	virtual void DispatchChaosPhysicsCollisionBlueprintEvents(const FChaosPhysicsCollisionInfo& CollisionInfo) override;
+	GEOMETRYCOLLECTIONENGINE_API virtual void DispatchChaosPhysicsCollisionBlueprintEvents(const FChaosPhysicsCollisionInfo& CollisionInfo) override;
 	
 	/** If true, this component will generate breaking events that other systems may subscribe to. */
 	UPROPERTY(EditAnywhere, BlueprintReadOnly, Category = "ChaosPhysics|Events")
@@ -1187,38 +1187,38 @@ protected:
 	TObjectPtr<UGeometryCollectionExternalRenderInterface> CustomRenderer;
 
 	/** Populate the dynamic particle data for the render thread. */
-	FGeometryCollectionDynamicData* InitDynamicData(bool bInitialization);
+	GEOMETRYCOLLECTIONENGINE_API FGeometryCollectionDynamicData* InitDynamicData(bool bInitialization);
 
 	/** Reset the dynamic collection from the current rest state. */
-	void ResetDynamicCollection();
+	GEOMETRYCOLLECTIONENGINE_API void ResetDynamicCollection();
 
 	/** Combine the commands from the input field assets */
-	void GetInitializationCommands(TArray<FFieldSystemCommand>& CombinedCommmands);
+	GEOMETRYCOLLECTIONENGINE_API void GetInitializationCommands(TArray<FFieldSystemCommand>& CombinedCommmands);
 
 	/** Issue a field command for the physics thread */
-	void DispatchFieldCommand(const FFieldSystemCommand& InCommand);
+	GEOMETRYCOLLECTIONENGINE_API void DispatchFieldCommand(const FFieldSystemCommand& InCommand);
 
-	Chaos::FPhysicsSolver* GetSolver(const UGeometryCollectionComponent& GeometryCollectionComponent);
-	void CalculateLocalBounds();
-	void CalculateGlobalMatrices();
+	GEOMETRYCOLLECTIONENGINE_API Chaos::FPhysicsSolver* GetSolver(const UGeometryCollectionComponent& GeometryCollectionComponent);
+	GEOMETRYCOLLECTIONENGINE_API void CalculateLocalBounds();
+	GEOMETRYCOLLECTIONENGINE_API void CalculateGlobalMatrices();
 	
 	UE_DEPRECATED(5.3, "Use ComputeBoundsFromComponentSpaceTransforms instead")
-	FBox ComputeBoundsFromGlobalMatrices(const FMatrix& LocalToWorldWithScale, const TArray<FMatrix>& GlobalMatricesArray) const;
+	GEOMETRYCOLLECTIONENGINE_API FBox ComputeBoundsFromGlobalMatrices(const FMatrix& LocalToWorldWithScale, const TArray<FMatrix>& GlobalMatricesArray) const;
 
-	FBox ComputeBoundsFromComponentSpaceTransforms(const FTransform& LocalToWorldWithScale, const TArray<FTransform>& ComponentSpaceTransformsArray) const;
+	GEOMETRYCOLLECTIONENGINE_API FBox ComputeBoundsFromComponentSpaceTransforms(const FTransform& LocalToWorldWithScale, const TArray<FTransform>& ComponentSpaceTransformsArray) const;
 
 	UE_DEPRECATED(5.3, "Use FTransform version of ComputeBounds instead")
-	FBox ComputeBounds(const FMatrix& LocalToWorldWithScale) const;
+	GEOMETRYCOLLECTIONENGINE_API FBox ComputeBounds(const FMatrix& LocalToWorldWithScale) const;
 
-	FBox ComputeBounds(const FTransform& LocalToWorldWithScale) const;
+	GEOMETRYCOLLECTIONENGINE_API FBox ComputeBounds(const FTransform& LocalToWorldWithScale) const;
 
-	void RegisterForEvents();
-	void UpdateRBCollisionEventRegistration();
-	void UpdateGlobalCollisionEventRegistration();
-	void UpdateBreakEventRegistration();
-	void UpdateRemovalEventRegistration();
-	void UpdateGlobalRemovalEventRegistration();
-	void UpdateCrumblingEventRegistration();
+	GEOMETRYCOLLECTIONENGINE_API void RegisterForEvents();
+	GEOMETRYCOLLECTIONENGINE_API void UpdateRBCollisionEventRegistration();
+	GEOMETRYCOLLECTIONENGINE_API void UpdateGlobalCollisionEventRegistration();
+	GEOMETRYCOLLECTIONENGINE_API void UpdateBreakEventRegistration();
+	GEOMETRYCOLLECTIONENGINE_API void UpdateRemovalEventRegistration();
+	GEOMETRYCOLLECTIONENGINE_API void UpdateGlobalRemovalEventRegistration();
+	GEOMETRYCOLLECTIONENGINE_API void UpdateCrumblingEventRegistration();
 	
 	/* Per-instance override to enable/disable replication for the geometry collection */
 	UPROPERTY(EditAnywhere, BlueprintReadWrite, Category=Network)
@@ -1283,18 +1283,18 @@ protected:
 
 	/** Called post solve to allow authoritative components to update their replication data */
 	UFUNCTION()
-	void OnRep_RepData();
+	GEOMETRYCOLLECTIONENGINE_API void OnRep_RepData();
 
-	void RequestUpdateRepData();
-	virtual void UpdateRepData();
+	GEOMETRYCOLLECTIONENGINE_API void RequestUpdateRepData();
+	GEOMETRYCOLLECTIONENGINE_API virtual void UpdateRepData();
 
 	/** Clear all rep data, this is required if the physics proxy has been recreated */
-	virtual void ResetRepData();
+	GEOMETRYCOLLECTIONENGINE_API virtual void ResetRepData();
  
 	UE_DEPRECATED(5.3, "The argument-free version of ProcessRepData will be removed. Please use the version which takes DeltaTime and SimTime instead.")
-	virtual void ProcessRepData();
+	GEOMETRYCOLLECTIONENGINE_API virtual void ProcessRepData();
 
-	virtual bool ProcessRepData(float DeltaTime, float SimTime);
+	GEOMETRYCOLLECTIONENGINE_API virtual bool ProcessRepData(float DeltaTime, float SimTime);
 
 	int32 VersionProcessed = INDEX_NONE;
 
@@ -1304,8 +1304,8 @@ protected:
 	int64 LastAsyncPhysicsTickMs = 0;
 
 private:
-	void ProcessRepDataOnPT();
-	void ResetRepDataCommon();
+	GEOMETRYCOLLECTIONENGINE_API void ProcessRepDataOnPT();
+	GEOMETRYCOLLECTIONENGINE_API void ResetRepDataCommon();
 
 	bool bRenderStateDirty;
 	bool bEnableBoneSelection;
@@ -1357,7 +1357,7 @@ private:
 	UPROPERTY(Transient)
 	TObjectPtr<AActor> EditorActor;
 #endif
-	void SwitchRenderModels(const AActor* Actor);
+	GEOMETRYCOLLECTIONENGINE_API void SwitchRenderModels(const AActor* Actor);
 
 	// Event dispatcher for break, crumble, removal and collision events
 	UPROPERTY(transient)
@@ -1376,41 +1376,41 @@ private:
 	TArray<int32> EmbeddedInstanceIndex;
 #endif
 
-	bool IsEmbeddedGeometryValid() const;
-	void ClearEmbeddedGeometry();
+	GEOMETRYCOLLECTIONENGINE_API bool IsEmbeddedGeometryValid() const;
+	GEOMETRYCOLLECTIONENGINE_API void ClearEmbeddedGeometry();
 
 	/** return true if a a custom renderer has been set and the feature is enabled */
-	bool CanUseCustomRenderer() const;
+	GEOMETRYCOLLECTIONENGINE_API bool CanUseCustomRenderer() const;
 
-	void RegisterCustomRenderer();
-	void UnregisterCustomRenderer();
-	void RefreshCustomRenderer();
+	GEOMETRYCOLLECTIONENGINE_API void RegisterCustomRenderer();
+	GEOMETRYCOLLECTIONENGINE_API void UnregisterCustomRenderer();
+	GEOMETRYCOLLECTIONENGINE_API void RefreshCustomRenderer();
 
 	/** return true if the root cluster is not longer active at runtime */
 	UFUNCTION(BlueprintCallable, Category = "ChaosPhysics")
-	bool IsRootBroken() const;
+	GEOMETRYCOLLECTIONENGINE_API bool IsRootBroken() const;
 
-	void IncrementSleepTimer(float DeltaTime);
-	void IncrementBreakTimer(float DeltaTime);
-	bool CalculateInnerSphere(int32 TransformIndex, UE::Math::TSphere<double>& SphereOut) const;
-	void UpdateDecay(int32 TransformIdx, float UpdatedDecay, bool UseClusterCrumbling, bool HasDynamicInternalClusterParent, FGeometryCollectionDecayContext& ContextInOut);
+	GEOMETRYCOLLECTIONENGINE_API void IncrementSleepTimer(float DeltaTime);
+	GEOMETRYCOLLECTIONENGINE_API void IncrementBreakTimer(float DeltaTime);
+	GEOMETRYCOLLECTIONENGINE_API bool CalculateInnerSphere(int32 TransformIndex, UE::Math::TSphere<double>& SphereOut) const;
+	GEOMETRYCOLLECTIONENGINE_API void UpdateDecay(int32 TransformIdx, float UpdatedDecay, bool UseClusterCrumbling, bool HasDynamicInternalClusterParent, FGeometryCollectionDecayContext& ContextInOut);
 
-	void UpdateAttachedChildrenTransform() const;
+	GEOMETRYCOLLECTIONENGINE_API void UpdateAttachedChildrenTransform() const;
 	
-	void UpdateRenderSystemsIfNeeded(bool bDynamicCollectionDirty);
-	void UpdateNavigationDataIfNeeded(bool bDynamicCollectionDirty);
-	void UpdateRemovalIfNeeded();
+	GEOMETRYCOLLECTIONENGINE_API void UpdateRenderSystemsIfNeeded(bool bDynamicCollectionDirty);
+	GEOMETRYCOLLECTIONENGINE_API void UpdateNavigationDataIfNeeded(bool bDynamicCollectionDirty);
+	GEOMETRYCOLLECTIONENGINE_API void UpdateRemovalIfNeeded();
 
-	void BuildInitialFilterData();
+	GEOMETRYCOLLECTIONENGINE_API void BuildInitialFilterData();
 
-	void LoadCollisionProfiles();
+	GEOMETRYCOLLECTIONENGINE_API void LoadCollisionProfiles();
 
-	void OnPostPhysicsSync();
+	GEOMETRYCOLLECTIONENGINE_API void OnPostPhysicsSync();
 
-	bool HasVisibleGeometry() const;
+	GEOMETRYCOLLECTIONENGINE_API bool HasVisibleGeometry() const;
 
 	/** backward compatibility method, until we can remove GlobalMatrices */
-	TArray<FMatrix> ComputeGlobalMatricesFromComponentSpaceTransforms() const;
+	GEOMETRYCOLLECTIONENGINE_API TArray<FMatrix> ComputeGlobalMatricesFromComponentSpaceTransforms() const;
 
 	/** The clusters we need to replicate */
 	TUniquePtr<TSet<Chaos::FPBDRigidClusteredParticleHandle*>> ClustersToRep;
@@ -1424,15 +1424,15 @@ private:
 
 	//~ Begin IPhysicsComponent Interface.
 public:
-	virtual Chaos::FPhysicsObject* GetPhysicsObjectById(Chaos::FPhysicsObjectId Id) const override;
-	virtual Chaos::FPhysicsObject* GetPhysicsObjectByName(const FName& Name) const override;
-	virtual TArray<Chaos::FPhysicsObject*> GetAllPhysicsObjects() const override;
-	virtual Chaos::FPhysicsObjectId GetIdFromGTParticle(Chaos::FGeometryParticle* Particle) const override;
+	GEOMETRYCOLLECTIONENGINE_API virtual Chaos::FPhysicsObject* GetPhysicsObjectById(Chaos::FPhysicsObjectId Id) const override;
+	GEOMETRYCOLLECTIONENGINE_API virtual Chaos::FPhysicsObject* GetPhysicsObjectByName(const FName& Name) const override;
+	GEOMETRYCOLLECTIONENGINE_API virtual TArray<Chaos::FPhysicsObject*> GetAllPhysicsObjects() const override;
+	GEOMETRYCOLLECTIONENGINE_API virtual Chaos::FPhysicsObjectId GetIdFromGTParticle(Chaos::FGeometryParticle* Particle) const override;
 	//~ End IPhysicsComponent Interface.
 
 #if WITH_EDITOR
 	//~ Begin UActorComponent interface.
-	virtual bool IsHLODRelevant() const override;
+	GEOMETRYCOLLECTIONENGINE_API virtual bool IsHLODRelevant() const override;
 	//~ End UActorComponent interface.
 #endif
 };
diff --git a/Engine/Source/Runtime/Experimental/Iris/Core/Public/Iris/ReplicationSystem/ReplicationSystem.h b/Engine/Source/Runtime/Experimental/Iris/Core/Public/Iris/ReplicationSystem/ReplicationSystem.h
--- a/Engine/Source/Runtime/Experimental/Iris/Core/Public/Iris/ReplicationSystem/ReplicationSystem.h
+++ b/Engine/Source/Runtime/Experimental/Iris/Core/Public/Iris/ReplicationSystem/ReplicationSystem.h
@@ -508,7 +508,7 @@ public:
 	 */
 	IRISCORE_API UObject* GetConnectionUserData(uint32 ConnectionId) const;
 
-	IRISCORE_API int32 GetPIEInstanceID() const { return PIEInstanceID; }
+	int32 GetPIEInstanceID() const { return PIEInstanceID; }
 
 	/** Set the squared cull distance for an object. This can be used by prioritizers and filters like UNetObjectGridFilter for example. For Actors this will override, not overwrite, the NetCullDistanceSquared property. */
 	IRISCORE_API void SetCullDistanceSqrOverride(FNetRefHandle Handle, float DistSqr);
@@ -546,7 +546,7 @@ private:
 	IRISCORE_API void ResetGameWorldState();
 	IRISCORE_API void NotifyStreamingLevelUnload(const UObject* Level);
 
-	IRISCORE_API void SetPIEInstanceID(int32 InPIEInstanceID) { PIEInstanceID = InPIEInstanceID; }
+	void SetPIEInstanceID(int32 InPIEInstanceID) { PIEInstanceID = InPIEInstanceID; }
 
 private:
 
diff --git a/Engine/Source/Runtime/Experimental/Iris/Core/Public/Iris/Serialization/NetSerializerArrayStorage.h b/Engine/Source/Runtime/Experimental/Iris/Core/Public/Iris/Serialization/NetSerializerArrayStorage.h
--- a/Engine/Source/Runtime/Experimental/Iris/Core/Public/Iris/Serialization/NetSerializerArrayStorage.h
+++ b/Engine/Source/Runtime/Experimental/Iris/Core/Public/Iris/Serialization/NetSerializerArrayStorage.h
@@ -51,8 +51,8 @@ public:
 
 private:
 
-	IRISCORE_API static void Free(FNetSerializationContext& Context, void* Ptr);
-	IRISCORE_API static void* Realloc(FNetSerializationContext& Context, void* Original, SIZE_T Size, uint32 Alignment);
+	static IRISCORE_API void Free(FNetSerializationContext& Context, void* Ptr);
+	static IRISCORE_API void* Realloc(FNetSerializationContext& Context, void* Original, SIZE_T Size, uint32 Alignment);
 };
 
 /**
diff --git a/Engine/Source/Runtime/GeometryCore/Public/CompGeom/ConvexHull2.h b/Engine/Source/Runtime/GeometryCore/Public/CompGeom/ConvexHull2.h
--- a/Engine/Source/Runtime/GeometryCore/Public/CompGeom/ConvexHull2.h
+++ b/Engine/Source/Runtime/GeometryCore/Public/CompGeom/ConvexHull2.h
@@ -25,7 +25,7 @@ using namespace UE::Math;
 
 
 template<typename RealType>
-class GEOMETRYCORE_API TConvexHull2
+class TConvexHull2
 {
 public:
 
@@ -39,7 +39,7 @@ public:
 	 * @param bIsKnownCCW		If true, will save some processing time by assuming input is wound counter-clockwise
 	 * @return true if a hull was generated
 	 */
-	bool SolveSimplePolygon(int32 NumPolygonPoints, TFunctionRef<TVector2<RealType>(int32)> GetPointFunc, bool bIsKnownCCW = false);
+	GEOMETRYCORE_API bool SolveSimplePolygon(int32 NumPolygonPoints, TFunctionRef<TVector2<RealType>(int32)> GetPointFunc, bool bIsKnownCCW = false);
 
 	/**
 	 * Generate convex hull as long as input is not degenerate
@@ -51,7 +51,7 @@ public:
 	 * @param FilterFunc Optional filter to include only a subset of the points in the output hull
 	 * @return true if hull was generated, false if points span < 2 dimensions
 	 */
-	bool Solve(int32 NumPoints, TFunctionRef<TVector2<RealType>(int32)> GetPointFunc, TFunctionRef<bool(int32)> FilterFunc = [](int32 Idx) {return true;});
+	GEOMETRYCORE_API bool Solve(int32 NumPoints, TFunctionRef<TVector2<RealType>(int32)> GetPointFunc, TFunctionRef<bool(int32)> FilterFunc = [](int32 Idx) {return true;});
 
 	/**
 	 * Generate convex hull as long as input is not degenerate
@@ -152,4 +152,4 @@ typedef TConvexHull2<double> FConvexHull2d;
 
 
 } // end namespace UE::Geometry
-} // end namespace UE
\ No newline at end of file
+} // end namespace UE
diff --git a/Engine/Source/Runtime/GeometryCore/Public/CompGeom/ConvexHull3.h b/Engine/Source/Runtime/GeometryCore/Public/CompGeom/ConvexHull3.h
--- a/Engine/Source/Runtime/GeometryCore/Public/CompGeom/ConvexHull3.h
+++ b/Engine/Source/Runtime/GeometryCore/Public/CompGeom/ConvexHull3.h
@@ -34,7 +34,7 @@ using namespace UE::Math;
  * incremental convex hull construction.
  */
 template<typename RealType>
-struct GEOMETRYCORE_API TExtremePoints3
+struct TExtremePoints3
 {
 	int Dimension = 0;
 	int Extreme[4]{ 0, 0, 0, 0 };
@@ -49,14 +49,14 @@ struct GEOMETRYCORE_API TExtremePoints3
 	}
 
 private:
-	void Init(int32 NumPoints, TFunctionRef<TVector<RealType>(int32)> GetPointFunc, TFunctionRef<bool(int32)> FilterFunc, RealType Epsilon = TMathUtil<RealType>::Epsilon);
+	GEOMETRYCORE_API void Init(int32 NumPoints, TFunctionRef<TVector<RealType>(int32)> GetPointFunc, TFunctionRef<bool(int32)> FilterFunc, RealType Epsilon = TMathUtil<RealType>::Epsilon);
 };
 
 /**
  * Calculate the Convex Hull of a 3D point set as a Triangle Mesh
  */
 template<typename RealType>
-class GEOMETRYCORE_API TConvexHull3
+class TConvexHull3
 {
 public:
 
@@ -74,7 +74,7 @@ public:
 	 * @param Filter Optional filter to include only a subset of the points in the output hull
 	 * @return true if hull was generated, false if points span < 2 dimensions
 	 */
-	bool Solve(int32 NumPoints, TFunctionRef<TVector<RealType>(int32)> GetPointFunc, TFunctionRef<bool(int32)> FilterFunc = [](int32 Idx) {return true;});
+	GEOMETRYCORE_API bool Solve(int32 NumPoints, TFunctionRef<TVector<RealType>(int32)> GetPointFunc, TFunctionRef<bool(int32)> FilterFunc = [](int32 Idx) {return true;});
 
 	/**
 	 * Generate convex hull as long as input is not degenerate
@@ -141,7 +141,7 @@ public:
 	 * @param GetPointFunc	Function providing array-style access into points
 	 * TODO: Provide an optional Epsilon parameter to simplify the hull within some tolerance
 	 */
-	void GetFaces(TFunctionRef<void(TArray<int32>&, TVector<RealType>)> PolygonFunc, TFunctionRef<TVector<RealType>(int32)> GetPointFunc) const;
+	GEOMETRYCORE_API void GetFaces(TFunctionRef<void(TArray<int32>&, TVector<RealType>)> PolygonFunc, TFunctionRef<TVector<RealType>(int32)> GetPointFunc) const;
 
 	/**
 	 * Only valid if bSaveTriangleNeighbors was true when Solve() was called
diff --git a/Engine/Source/Runtime/GeometryCore/Public/DynamicMesh/DynamicMesh3.h b/Engine/Source/Runtime/GeometryCore/Public/DynamicMesh/DynamicMesh3.h
--- a/Engine/Source/Runtime/GeometryCore/Public/DynamicMesh/DynamicMesh3.h
+++ b/Engine/Source/Runtime/GeometryCore/Public/DynamicMesh/DynamicMesh3.h
@@ -103,7 +103,7 @@ enum class EMeshComponents : uint8
 * The function CheckValidity() does extensive sanity checking on the mesh data structure.
 * Use this to test your code, both for mesh construction and editing!!
 */
-class GEOMETRYCORE_API FDynamicMesh3
+class FDynamicMesh3
 {
 
 // TODO:
@@ -155,7 +155,7 @@ public:
 	/** DuplicateTriangleID is returned by AppendTriangle() to indicate that the added triangle already exists in the mesh, and was ignored because we do not support duplicate triangles */
 	constexpr static int DuplicateTriangleID = -3;
 
-	const static FVector3d InvalidVertex;
+	GEOMETRYCORE_API const static FVector3d InvalidVertex;
 	constexpr static FIndex3i InvalidTriangle{InvalidID, InvalidID, InvalidID};
 	constexpr static FIndex2i InvalidEdge{InvalidID, InvalidID};
 
@@ -204,23 +204,23 @@ protected:
 
 public:
 	/** Default constructor */
-	FDynamicMesh3() : FDynamicMesh3(false, false, false, false) {}
+	GEOMETRYCORE_API FDynamicMesh3() : FDynamicMesh3(false, false, false, false) {}
 
 	/** Copy/Move construction */
-	FDynamicMesh3(const FDynamicMesh3& CopyMesh);
-	FDynamicMesh3(FDynamicMesh3&& MoveMesh);
+	GEOMETRYCORE_API FDynamicMesh3(const FDynamicMesh3& CopyMesh);
+	GEOMETRYCORE_API FDynamicMesh3(FDynamicMesh3&& MoveMesh);
 
 	/** Copy and move assignment */
-	const FDynamicMesh3& operator=(const FDynamicMesh3& CopyMesh);
-	const FDynamicMesh3& operator=(FDynamicMesh3&& MoveMesh);
+	GEOMETRYCORE_API const FDynamicMesh3& operator=(const FDynamicMesh3& CopyMesh);
+	GEOMETRYCORE_API const FDynamicMesh3& operator=(FDynamicMesh3&& MoveMesh);
 
 	/** Destructor */
-	virtual ~FDynamicMesh3();
+	GEOMETRYCORE_API virtual ~FDynamicMesh3();
 
 
 	/** Construct an empty mesh with specified attributes */
-	explicit FDynamicMesh3(bool bWantNormals, bool bWantColors, bool bWantUVs, bool bWantTriGroups);
-	explicit FDynamicMesh3(EMeshComponents flags)
+	GEOMETRYCORE_API explicit FDynamicMesh3(bool bWantNormals, bool bWantColors, bool bWantUVs, bool bWantTriGroups);
+	GEOMETRYCORE_API explicit FDynamicMesh3(EMeshComponents flags)
 	    : FDynamicMesh3(((int)flags & (int)EMeshComponents::VertexNormals) != 0,
 	                    ((int)flags & (int)EMeshComponents::VertexColors) != 0,
 	                    ((int)flags & (int)EMeshComponents::VertexUVs) != 0,
@@ -228,17 +228,17 @@ public:
 	{}
 
 	/** Construction from Mesh Generator */
-	FDynamicMesh3(const FMeshShapeGenerator* Generator);
+	GEOMETRYCORE_API FDynamicMesh3(const FMeshShapeGenerator* Generator);
 
 	/** Set internal data structures to be a copy of input mesh using the specified attributes*/
-	void Copy(const FDynamicMesh3& CopyMesh, bool bNormals = true, bool bColors = true, bool bUVs = true,
+	GEOMETRYCORE_API void Copy(const FDynamicMesh3& CopyMesh, bool bNormals = true, bool bColors = true, bool bUVs = true,
 	          bool bAttributes = true);
 
 	/** 
 	  * Initialize mesh from the output of a MeshShapeGenerator (assumes Generate() was already called) 
 	  * @return false if any triangles defined in Generator were skipped due to invalid topology
 	  */
-	bool Copy(const FMeshShapeGenerator* Generator);
+	GEOMETRYCORE_API bool Copy(const FMeshShapeGenerator* Generator);
 
 	/** 
 	 * Copy input mesh while compacting, i.e. removing unused vertices/triangles/edges. 
@@ -252,11 +252,11 @@ public:
 	 * @param bAttributes if true, will copy attributes
 	 * @param CompactInfo if not nullptr, will be filled with mapping indicating how vertex and triangle IDs were changed during compaction
 	 */
-	void CompactCopy(const FDynamicMesh3& CopyMesh, bool bNormals = true, bool bColors = true, bool bUVs = true,
+	GEOMETRYCORE_API void CompactCopy(const FDynamicMesh3& CopyMesh, bool bNormals = true, bool bColors = true, bool bUVs = true,
 	                 bool bAttributes = true, FCompactMaps* CompactInfo = nullptr);
 
 	/** Discard all data */
-	void Clear();
+	GEOMETRYCORE_API void Clear();
 
 	/**
 	 * Ensure that all the same extended attributes available in ToMatch are also enabled.
@@ -264,7 +264,7 @@ public:
 	 * If bClearExisting is passed as false, existing attributes are not removed/cleared.
 	 * If bDiscardExtraAttributes=true and bClearExisting=false, extra attributes not in ToMatch are discarded, but existing attributes are not cleared/reset
 	 */
-	void EnableMatchingAttributes(const FDynamicMesh3& ToMatch, bool bClearExisting = true, bool bDiscardExtraAttributes = false);
+	GEOMETRYCORE_API void EnableMatchingAttributes(const FDynamicMesh3& ToMatch, bool bClearExisting = true, bool bDiscardExtraAttributes = false);
 
 	/**
 	 * Serialization operator for FDynamicMesh3.
@@ -280,7 +280,7 @@ public:
 	}
 
 	/** Serialize the mesh to an archive. */
-	void Serialize(FArchive& Ar);
+	GEOMETRYCORE_API void Serialize(FArchive& Ar);
 
 public:
 	/** @return number of vertices in the mesh */
@@ -348,7 +348,7 @@ public:
 	}
 
 	/** @return bitwise-or of EMeshComponents flags specifying which extra data this mesh has */
-	int GetComponentsFlags() const;
+	GEOMETRYCORE_API int GetComponentsFlags() const;
 
 
 	/** @return true if VertexID is a valid vertex in this mesh */
@@ -540,10 +540,10 @@ public:
 	}
 
 	/** Call ApplyFunc for each one-ring triangle of a vertex. Currently this is significantly more efficient than VtxTrianglesItr() in many use cases. */
-	void EnumerateVertexTriangles(int32 VertexID, TFunctionRef<void(int32)> ApplyFunc) const;
+	GEOMETRYCORE_API void EnumerateVertexTriangles(int32 VertexID, TFunctionRef<void(int32)> ApplyFunc) const;
 
 	/** Call ApplyFunc for each triangle connected to an Edge (1 or 2 triangles) */
-	void EnumerateEdgeTriangles(int32 EdgeID, TFunctionRef<void(int32)> ApplyFunc) const;
+	GEOMETRYCORE_API void EnumerateEdgeTriangles(int32 EdgeID, TFunctionRef<void(int32)> ApplyFunc) const;
 
 
 	//
@@ -551,7 +551,7 @@ public:
 	//
 public:
 	/** Append vertex at position and other fields, returns vid */
-	int AppendVertex(const FVertexInfo& VertInfo);
+	GEOMETRYCORE_API int AppendVertex(const FVertexInfo& VertInfo);
 
 	/** Append vertex at position, returns vid */
 	int AppendVertex(const FVector3d& Position)
@@ -560,10 +560,10 @@ public:
 	}
 
 	/** Copy vertex SourceVertexID from existing SourceMesh, returns new vertex id */
-	int AppendVertex(const FDynamicMesh3& SourceMesh, int SourceVertexID);
+	GEOMETRYCORE_API int AppendVertex(const FDynamicMesh3& SourceMesh, int SourceVertexID);
 
 	/** TriVertices must be distinct and refer to existing, valid vertices */
-	int AppendTriangle(const FIndex3i& TriVertices, int GroupID = 0);
+	GEOMETRYCORE_API int AppendTriangle(const FIndex3i& TriVertices, int GroupID = 0);
 
 	/** Vertex0, Vertex1, and Vertex2 must be distinct and refer to existing, valid vertices */
 	inline int AppendTriangle(int Vertex0, int Vertex1, int Vertex2, int GroupID = 0)
@@ -596,7 +596,7 @@ public:
 	 * If bUnsafe, we use fast id allocation that does not update free list.
 	 * You should only be using this between BeginUnsafeVerticesInsert() / EndUnsafeVerticesInsert() calls
 	 */
-	EMeshResult InsertVertex(int VertexID, const FVertexInfo& VertInfo, bool bUnsafe = false);
+	GEOMETRYCORE_API EMeshResult InsertVertex(int VertexID, const FVertexInfo& VertInfo, bool bUnsafe = false);
 
 	/** Call this before a set of unsafe InsertTriangle() calls */
 	virtual void BeginUnsafeTrianglesInsert()
@@ -615,7 +615,7 @@ public:
 	 * If bUnsafe, we use fast id allocation that does not update free list.
 	 * You should only be using this between BeginUnsafeTrianglesInsert() / EndUnsafeTrianglesInsert() calls
 	 */
-	EMeshResult InsertTriangle(int TriangleID, const FIndex3i& TriVertices, int GroupID = 0, bool bUnsafe = false);
+	GEOMETRYCORE_API EMeshResult InsertTriangle(int TriangleID, const FIndex3i& TriVertices, int GroupID = 0, bool bUnsafe = false);
 
 	//
 	// Vertex/Tri/Edge accessors
@@ -655,10 +655,10 @@ public:
 
 
 	/** Get extended vertex information */
-	bool GetVertex(int VertexID, FVertexInfo& VertInfo, bool bWantNormals, bool bWantColors, bool bWantUVs) const;
+	GEOMETRYCORE_API bool GetVertex(int VertexID, FVertexInfo& VertInfo, bool bWantNormals, bool bWantColors, bool bWantUVs) const;
 
 	/** Get all vertex information available */
-	FVertexInfo GetVertexInfo(int VertexID) const;
+	GEOMETRYCORE_API FVertexInfo GetVertexInfo(int VertexID) const;
 
 	/** @return the valence of a vertex (the number of connected edges) */
 	int GetVtxEdgeCount(int VertexID) const
@@ -667,7 +667,7 @@ public:
 	}
 
 	/** @return the max valence of all vertices in the mesh */
-	int GetMaxVtxEdgeCount() const;
+	GEOMETRYCORE_API int GetMaxVtxEdgeCount() const;
 
 	/** Get triangle vertices */
 	inline FIndex3i GetTriangle(int TriangleID) const
@@ -705,7 +705,7 @@ public:
 	}
 
 	/** Find the neighbour triangles of a triangle (any of them might be InvalidID) */
-	FIndex3i GetTriNeighbourTris(int TriangleID) const;
+	GEOMETRYCORE_API FIndex3i GetTriNeighbourTris(int TriangleID) const;
 
 	/** Get the three vertex positions of a triangle */
 	template<typename VecType>
@@ -765,7 +765,7 @@ public:
 	}
 
 	/** Return edge vertex indices, but oriented based on attached triangle (rather than min-sorted) */
-	FIndex2i GetOrientedBoundaryEdgeV(int EdgeID) const;
+	GEOMETRYCORE_API FIndex2i GetOrientedBoundaryEdgeV(int EdgeID) const;
 
 	/** Return (triangle, edge_index) representation for given Edge ID */
 	inline FMeshTriEdgeID GetTriEdgeIDFromEdgeID(int EdgeID) const
@@ -791,10 +791,10 @@ public:
 	 * and discard any that are not requested
 	 * @param MeshComponentsFlags A 'bitwise or' of requested EMeshComponents flags
 	 */
-	void EnableMeshComponents(int MeshComponentsFlags);
+	GEOMETRYCORE_API void EnableMeshComponents(int MeshComponentsFlags);
 
-	void EnableVertexNormals(const FVector3f& InitialNormal);
-	void DiscardVertexNormals();
+	GEOMETRYCORE_API void EnableVertexNormals(const FVector3f& InitialNormal);
+	GEOMETRYCORE_API void DiscardVertexNormals();
 
 	FVector3f GetVertexNormal(int vID) const
 	{
@@ -817,8 +817,8 @@ public:
 		}
 	}
 
-	void EnableVertexColors(const FVector3f& InitialColor);
-	void DiscardVertexColors();
+	GEOMETRYCORE_API void EnableVertexColors(const FVector3f& InitialColor);
+	GEOMETRYCORE_API void DiscardVertexColors();
 
 
 	FVector3f GetVertexColor(int vID) const
@@ -843,8 +843,8 @@ public:
 		}
 	}
 
-	void EnableVertexUVs(const FVector2f& InitialUV);
-	void DiscardVertexUVs();
+	GEOMETRYCORE_API void EnableVertexUVs(const FVector2f& InitialUV);
+	GEOMETRYCORE_API void DiscardVertexUVs();
 
 	FVector2f GetVertexUV(int vID) const
 	{
@@ -867,8 +867,8 @@ public:
 		}
 	}
 
-	void EnableTriangleGroups(int InitialGroupID = 0);
-	void DiscardTriangleGroups();
+	GEOMETRYCORE_API void EnableTriangleGroups(int InitialGroupID = 0);
+	GEOMETRYCORE_API void DiscardTriangleGroups();
 
 	int AllocateTriangleGroup()
 	{
@@ -900,8 +900,8 @@ public:
 		return HasAttributes() ? AttributeSet.Get() : nullptr;
 	}
 
-	void EnableAttributes();
-	void DiscardAttributes();
+	GEOMETRYCORE_API void EnableAttributes();
+	GEOMETRYCORE_API void DiscardAttributes();
 
 
 	//
@@ -916,40 +916,40 @@ public:
 	}
 
 	/** Returns true if the vertex is part of any boundary edges */
-	bool IsBoundaryVertex(int VertexID) const;
+	GEOMETRYCORE_API bool IsBoundaryVertex(int VertexID) const;
 
 	/** Returns true if any edge of triangle is a boundary edge */
-	bool IsBoundaryTriangle(int TriangleID) const;
+	GEOMETRYCORE_API bool IsBoundaryTriangle(int TriangleID) const;
 
 	/** Find id of edge connecting A and B */
-	int FindEdge(int VertexA, int VertexB) const;
+	GEOMETRYCORE_API int FindEdge(int VertexA, int VertexB) const;
 
 	/** Find edgeid for edge [a,b] from triangle that contains the edge. Faster than FindEdge() because it is constant-time. */
-	int FindEdgeFromTri(int VertexA, int VertexB, int TriangleID) const;
+	GEOMETRYCORE_API int FindEdgeFromTri(int VertexA, int VertexB, int TriangleID) const;
 
 	/** Find edgeid for edge connecting two triangles */
-	int FindEdgeFromTriPair(int TriangleA, int TriangleB) const;
+	GEOMETRYCORE_API int FindEdgeFromTriPair(int TriangleA, int TriangleB) const;
 
 	/** Find triangle made up of any permutation of vertices [a,b,c] */
-	int FindTriangle(int A, int B, int C) const;
+	GEOMETRYCORE_API int FindTriangle(int A, int B, int C) const;
 
 	/**
 	 * If edge has vertices [a,b], and is connected two triangles [a,b,c] and [a,b,d],
 	 * this returns [c,d], or [c,InvalidID] for a boundary edge
 	 */
-	FIndex2i GetEdgeOpposingV(int EdgeID) const;
+	GEOMETRYCORE_API FIndex2i GetEdgeOpposingV(int EdgeID) const;
 
 	/**
 	 * Given an edge and vertex on that edge, returns other vertex of edge, the two opposing verts, and the two connected triangles (OppVert2Out and Tri2Out are be InvalidID for boundary edge)
 	 */
-	void GetVtxNbrhood(int EdgeID, int VertexID, int& OtherVertOut, int& OppVert1Out, int& OppVert2Out, int& Tri1Out,
+	GEOMETRYCORE_API void GetVtxNbrhood(int EdgeID, int VertexID, int& OtherVertOut, int& OppVert1Out, int& OppVert2Out, int& Tri1Out,
 	                   int& Tri2Out) const;
 
 	/**
 	 * Returns count of boundary edges at vertex, and the first two boundary
 	 * edges if found. If return is > 2, call GetAllVtxBoundaryEdges
 	 */
-	int GetVtxBoundaryEdges(int VertexID, int& Edge0Out, int& Edge1Out) const;
+	GEOMETRYCORE_API int GetVtxBoundaryEdges(int VertexID, int& Edge0Out, int& Edge1Out) const;
 
 	/**
 	 * Find edge ids of boundary edges connected to vertex.
@@ -957,22 +957,22 @@ public:
 	 * @param EdgeListOut boundary edge IDs are appended to this list
 	 * @return count of number of elements of e that were filled
 	 */
-	int GetAllVtxBoundaryEdges(int VertexID, TArray<int>& EdgeListOut) const;
+	GEOMETRYCORE_API int GetAllVtxBoundaryEdges(int VertexID, TArray<int>& EdgeListOut) const;
 
 	/**
 	 * return # of triangles attached to vID, or -1 if invalid vertex
 	 */
-	int GetVtxTriangleCount(int VertexID) const;
+	GEOMETRYCORE_API int GetVtxTriangleCount(int VertexID) const;
 
 	/**
 	 * Get triangle one-ring at vertex.
 	 */
-	EMeshResult GetVtxTriangles(int VertexID, TArray<int>& TrianglesOut) const;
+	GEOMETRYCORE_API EMeshResult GetVtxTriangles(int VertexID, TArray<int>& TrianglesOut) const;
 
 	/**
 	 * @return Triangle ID for a single triangle connected to VertexID, or InvalidID if VertexID does not exist or has no attached triangles
 	 */
-	int GetVtxSingleTriangle(int VertexID) const;
+	GEOMETRYCORE_API int GetVtxSingleTriangle(int VertexID) const;
 
 	/**
 	* Get triangles connected to vertex in contiguous order, with multiple groups if vertex is a bowtie.
@@ -981,26 +981,26 @@ public:
 	* @param ContiguousGroupLengths Lengths of contiguous groups packed into TrianglesOut (if not a bowtie, this will just be a length-one array w/ {TrianglesOut.Num()})
 	* @param GroupIsLoop Indicates whether each contiguous group is a loop (first triangle connected to last) or not
 	*/
-	EMeshResult GetVtxContiguousTriangles(int VertexID, TArray<int>& TrianglesOut, TArray<int>& ContiguousGroupLengths,
+	GEOMETRYCORE_API EMeshResult GetVtxContiguousTriangles(int VertexID, TArray<int>& TrianglesOut, TArray<int>& ContiguousGroupLengths,
 	                                      TArray<bool>& GroupIsLoop) const;
 
 	/** Returns true if the two triangles connected to edge have different group IDs */
-	bool IsGroupBoundaryEdge(int EdgeID) const;
+	GEOMETRYCORE_API bool IsGroupBoundaryEdge(int EdgeID) const;
 
 	/** Returns true if vertex has more than one tri group in its tri nbrhood */
-	bool IsGroupBoundaryVertex(int VertexID) const;
+	GEOMETRYCORE_API bool IsGroupBoundaryVertex(int VertexID) const;
 
 	/** Returns true if more than two group boundary edges meet at vertex (ie 3+ groups meet at this vertex) */
-	bool IsGroupJunctionVertex(int VertexID) const;
+	GEOMETRYCORE_API bool IsGroupJunctionVertex(int VertexID) const;
 
 	/** Returns up to 4 group IDs at vertex. Returns false if > 4 encountered */
-	bool GetVertexGroups(int VertexID, FIndex4i& GroupsOut) const;
+	GEOMETRYCORE_API bool GetVertexGroups(int VertexID, FIndex4i& GroupsOut) const;
 
 	/** Returns all group IDs at vertex */
-	bool GetAllVertexGroups(int VertexID, TArray<int>& GroupsOut) const;
+	GEOMETRYCORE_API bool GetAllVertexGroups(int VertexID, TArray<int>& GroupsOut) const;
 
 	/** returns true if vID is a "bowtie" vertex, ie multiple disjoint triangle sets in one-ring */
-	bool IsBowtieVertex(int VertexID) const;
+	GEOMETRYCORE_API bool IsBowtieVertex(int VertexID) const;
 
 	/** returns true if vertices, edges, and triangles are all dense (Count == MaxID) **/
 	bool IsCompact() const
@@ -1027,7 +1027,7 @@ public:
 	}
 
 	/** @return true if mesh has no boundary edges */
-	bool IsClosed() const;
+	GEOMETRYCORE_API bool IsClosed() const;
 
 
 	//
@@ -1035,7 +1035,7 @@ public:
 	//
 public:
 	/** Returns bounding box of all mesh vertices (including unreferenced vertices) */
-	FAxisAlignedBox3d GetBounds(bool bParallel = false) const;
+	GEOMETRYCORE_API FAxisAlignedBox3d GetBounds(bool bParallel = false) const;
 
 	/**
 	 * Compute a normal/tangent frame at vertex that is "stable" as long as
@@ -1046,59 +1046,59 @@ public:
 	 * @param bFrameNormalY if true, then frame.Y is normal (X still points along mesh edge)
 	 * @param UseNormal if defined, this normal is used instead of VertexNormals normal
 	 */
-	FFrame3d GetVertexFrame(int VertexID, bool bFrameNormalY = false, FVector3d* UseNormal = nullptr) const;
+	GEOMETRYCORE_API FFrame3d GetVertexFrame(int VertexID, bool bFrameNormalY = false, FVector3d* UseNormal = nullptr) const;
 
 	/** Calculate face normal of triangle */
-	FVector3d GetTriNormal(int TriangleID) const;
+	GEOMETRYCORE_API FVector3d GetTriNormal(int TriangleID) const;
 
 	/** Calculate area triangle */
-	double GetTriArea(int TriangleID) const;
+	GEOMETRYCORE_API double GetTriArea(int TriangleID) const;
 
 	/**
 	 * Compute triangle normal, area, and centroid all at once. Re-uses vertex
 	 * lookups and computes normal & area simultaneously. *However* does not produce
 	 * the same normal/area as separate calls, because of this.
 	 */
-	void GetTriInfo(int TriangleID, FVector3d& Normal, double& Area, FVector3d& Centroid) const;
+	GEOMETRYCORE_API void GetTriInfo(int TriangleID, FVector3d& Normal, double& Area, FVector3d& Centroid) const;
 
 	/** Compute centroid of triangle */
-	FVector3d GetTriCentroid(int TriangleID) const;
+	GEOMETRYCORE_API FVector3d GetTriCentroid(int TriangleID) const;
 
 	/** Interpolate vertex positions of triangle using barycentric coordinates */
-	FVector3d GetTriBaryPoint(int TriangleID, double Bary0, double Bary1, double Bary2) const;
+	GEOMETRYCORE_API FVector3d GetTriBaryPoint(int TriangleID, double Bary0, double Bary1, double Bary2) const;
 
 	/** Interpolate vertex normals of triangle using barycentric coordinates */
-	FVector3d GetTriBaryNormal(int TriangleID, double Bary0, double Bary1, double Bary2) const;
+	GEOMETRYCORE_API FVector3d GetTriBaryNormal(int TriangleID, double Bary0, double Bary1, double Bary2) const;
 
 	/** Compute interpolated vertex attributes at point of triangle */
-	void GetTriBaryPoint(int TriangleID, double Bary0, double Bary1, double Bary2, FVertexInfo& VertInfo) const;
+	GEOMETRYCORE_API void GetTriBaryPoint(int TriangleID, double Bary0, double Bary1, double Bary2, FVertexInfo& VertInfo) const;
 
 	/** Construct bounding box of triangle as efficiently as possible */
-	FAxisAlignedBox3d GetTriBounds(int TriangleID) const;
+	GEOMETRYCORE_API FAxisAlignedBox3d GetTriBounds(int TriangleID) const;
 
 	/** Construct stable frame at triangle centroid, where frame.Z is face normal, and frame.X is aligned with edge nEdge of triangle. */
-	FFrame3d GetTriFrame(int TriangleID, int Edge = 0) const;
+	GEOMETRYCORE_API FFrame3d GetTriFrame(int TriangleID, int Edge = 0) const;
 
 	/** Compute solid angle of oriented triangle tID relative to point p - see WindingNumber() */
-	double GetTriSolidAngle(int TriangleID, const FVector3d& p) const;
+	GEOMETRYCORE_API double GetTriSolidAngle(int TriangleID, const FVector3d& p) const;
 
 	/** Compute internal angle at vertex i of triangle (where i is 0,1,2); */
-	double GetTriInternalAngleR(int TriangleID, int i) const;
+	GEOMETRYCORE_API double GetTriInternalAngleR(int TriangleID, int i) const;
 
 	/** Compute internal angles at all vertices of triangle */
-	FVector3d GetTriInternalAnglesR(int TriangleID) const;
+	GEOMETRYCORE_API FVector3d GetTriInternalAnglesR(int TriangleID) const;
 
 	/** Returns average normal of connected face normals */
-	FVector3d GetEdgeNormal(int EdgeID) const;
+	GEOMETRYCORE_API FVector3d GetEdgeNormal(int EdgeID) const;
 
 	/** Get point along edge, t clamped to range [0,1] */
-	FVector3d GetEdgePoint(int EdgeID, double ParameterT) const;
+	GEOMETRYCORE_API FVector3d GetEdgePoint(int EdgeID, double ParameterT) const;
 
 	/**
 	 * Fastest possible one-ring centroid. This is used inside many other algorithms
 	 * so it helps to have it be maximally efficient
 	 */
-	void GetVtxOneRingCentroid(int VertexID, FVector3d& CentroidOut) const;
+	GEOMETRYCORE_API void GetVtxOneRingCentroid(int VertexID, FVector3d& CentroidOut) const;
 
 	/**
 	 * Compute mesh winding number, from Jacobson et. al., Robust Inside-Outside Segmentation using Generalized Winding Numbers
@@ -1106,7 +1106,7 @@ public:
 	 * returns ~0 for points outside a closed, consistently oriented mesh, and a positive or negative integer
 	 * for points inside, with value > 1 depending on how many "times" the point inside the mesh (like in 2D polygon winding)
 	 */
-	double CalculateWindingNumber(const FVector3d& QueryPoint) const;
+	GEOMETRYCORE_API double CalculateWindingNumber(const FVector3d& QueryPoint) const;
 
 	//
 	// direct buffer access
@@ -1171,18 +1171,18 @@ public:
 	 * @param CompactInfo if not nullptr, will be filled with mapping indicating how vertex and triangle IDs were changed during compaction
 	 * @todo VertexEdgeLists is not compacted. does not affect indices, but does keep memory.
 	 */
-	void CompactInPlace(FCompactMaps* CompactInfo = nullptr);
+	GEOMETRYCORE_API void CompactInPlace(FCompactMaps* CompactInfo = nullptr);
 
 	/**
 	 * Reverse the ccw/cw orientation of all triangles in the mesh, and
 	 * optionally flip the vertex normals if they exist
 	 */
-	void ReverseOrientation(bool bFlipNormals = true);
+	GEOMETRYCORE_API void ReverseOrientation(bool bFlipNormals = true);
 
 	/**
 	 * Reverse the ccw/cw orientation of a triangle
 	 */
-	EMeshResult ReverseTriOrientation(int TriangleID);
+	GEOMETRYCORE_API EMeshResult ReverseTriOrientation(int TriangleID);
 
 	/**
 	 * Remove vertex VertexID and all connected triangles. 
@@ -1190,7 +1190,7 @@ public:
 	 * If bPreserveManifold is true, checks that we will not create a bowtie vertex first.
 	 * In this case, returns Failed_WouldCreateBowtie if removing the triangles would create a bowtie.
 	 */
-	EMeshResult RemoveVertex(int VertexID, bool bPreserveManifold = false);
+	GEOMETRYCORE_API EMeshResult RemoveVertex(int VertexID, bool bPreserveManifold = false);
 
 	/**
 	* Remove a triangle from the mesh. Also removes any unreferenced edges after tri is removed.
@@ -1199,14 +1199,14 @@ public:
 	* If this check is not done, you have to make sure you don't create a bow tie, because other
 	* code assumes we don't have bow ties, and will not handle it properly
 	*/
-	EMeshResult RemoveTriangle(int TriangleID, bool bRemoveIsolatedVertices = true, bool bPreserveManifold = false);
+	GEOMETRYCORE_API EMeshResult RemoveTriangle(int TriangleID, bool bRemoveIsolatedVertices = true, bool bPreserveManifold = false);
 
 	/**
 	 * Rewrite the triangle to reference the new tuple of vertices.
 	 *
 	 * @todo this function currently does not guarantee that the returned mesh is well-formed. Only call if you know it's OK.
 	 */
-	virtual EMeshResult SetTriangle(int TriangleID, const FIndex3i& NewVertices, bool bRemoveIsolatedVertices = true);
+	GEOMETRYCORE_API virtual EMeshResult SetTriangle(int TriangleID, const FIndex3i& NewVertices, bool bRemoveIsolatedVertices = true);
 
 	/**
 	 * Split an edge of the mesh by inserting a vertex. This creates a new triangle on either side of the edge (ie a 2-4 split).
@@ -1218,7 +1218,7 @@ public:
 	 * @param SplitParameterT defines the position along the edge that we split at, must be between 0 and 1, and is assumed to be based on the order of vertices returned by GetEdgeV()
 	 * @return Ok on success, or enum value indicates why operation cannot be applied. Mesh remains unmodified on error.
 	 */
-	virtual EMeshResult SplitEdge(int EdgeAB, FEdgeSplitInfo& SplitInfo, double SplitParameterT = 0.5);
+	GEOMETRYCORE_API virtual EMeshResult SplitEdge(int EdgeAB, FEdgeSplitInfo& SplitInfo, double SplitParameterT = 0.5);
 
 	/** 
 	 * Splits the edge between two vertices at the midpoint, if this edge exists 
@@ -1227,7 +1227,7 @@ public:
 	 * @param SplitInfo returned information about new and modified mesh elements
 	 * @return Ok on success, or enum value indicates why operation cannot be applied. Mesh remains unmodified on error.
 	 */
-	EMeshResult SplitEdge(int EdgeVertA, int EdgeVertB, FEdgeSplitInfo& SplitInfo);
+	GEOMETRYCORE_API EMeshResult SplitEdge(int EdgeVertA, int EdgeVertB, FEdgeSplitInfo& SplitInfo);
 
 	/**
 	 * Flip/Rotate an edge of the mesh. This does not change the number of edges, vertices, or triangles.
@@ -1236,7 +1236,7 @@ public:
 	 * @param FlipInfo returned information about new and modified mesh elements
 	 * @return Ok on success, or enum value indicates why operation cannot be applied. Mesh remains unmodified on error.
 	 */
-	virtual EMeshResult FlipEdge(int EdgeAB, FEdgeFlipInfo& FlipInfo);
+	GEOMETRYCORE_API virtual EMeshResult FlipEdge(int EdgeAB, FEdgeFlipInfo& FlipInfo);
 
 	/** calls FlipEdge() on the edge between two vertices, if it exists
 	 * @param EdgeVertA index of first vertex
@@ -1244,7 +1244,7 @@ public:
 	 * @param FlipInfo returned information about new and modified mesh elements
 	 * @return Ok on success, or enum value indicates why operation cannot be applied. Mesh remains unmodified on error.
 	 */
-	virtual EMeshResult FlipEdge(int EdgeVertA, int EdgeVertB, FEdgeFlipInfo& FlipInfo);
+	GEOMETRYCORE_API virtual EMeshResult FlipEdge(int EdgeVertA, int EdgeVertB, FEdgeFlipInfo& FlipInfo);
 
 
 	/**
@@ -1254,7 +1254,7 @@ public:
 	 * @param SplitInfo returned info about the new and modified mesh elements
 	 * @return Ok on success, or enum value indicates why operation cannot be applied. Mesh remains unmodified on error.
 	 */
-	virtual EMeshResult SplitVertex(int VertexID, const TArrayView<const int>& TrianglesToUpdate,
+	GEOMETRYCORE_API virtual EMeshResult SplitVertex(int VertexID, const TArrayView<const int>& TrianglesToUpdate,
 	                                FVertexSplitInfo& SplitInfo);
 
 	/**
@@ -1263,7 +1263,7 @@ public:
 	 * @param TrianglesToUpdate triangles that should be updated to use the new vertex anywhere they previously had the old one
 	 * @return true if calling SplitVertex with these arguments would leave an isolated vertex at the original VertexID
 	 */
-	virtual bool SplitVertexWouldLeaveIsolated(int VertexID, const TArrayView<const int>& TrianglesToUpdate);
+	GEOMETRYCORE_API virtual bool SplitVertexWouldLeaveIsolated(int VertexID, const TArrayView<const int>& TrianglesToUpdate);
 
 
 	/**
@@ -1274,7 +1274,7 @@ public:
 	 * @param CollapseInfo returned information about new and modified mesh elements
 	 * @return Ok on success, or enum value indicates why operation cannot be applied. Mesh remains unmodified on error.
 	 */
-	virtual EMeshResult CollapseEdge(int KeepVertID, int RemoveVertID, double EdgeParameterT,
+	GEOMETRYCORE_API virtual EMeshResult CollapseEdge(int KeepVertID, int RemoveVertID, double EdgeParameterT,
 	                                 FEdgeCollapseInfo& CollapseInfo);
 	virtual EMeshResult CollapseEdge(int KeepVertID, int RemoveVertID, FEdgeCollapseInfo& CollapseInfo)
 	{
@@ -1296,7 +1296,7 @@ public:
 	 * @param CheckValidOrientation perform edge consistency orientation checks before merging.
 	 * @return Ok on success, or enum value indicates why operation cannot be applied. Mesh remains unmodified on error.
 	 */
-	virtual EMeshResult MergeEdges(int KeepEdgeID, int DiscardEdgeID, FMergeEdgesInfo& MergeInfo, bool bCheckValidOrientation=true);
+	GEOMETRYCORE_API virtual EMeshResult MergeEdges(int KeepEdgeID, int DiscardEdgeID, FMergeEdgesInfo& MergeInfo, bool bCheckValidOrientation=true);
 
 
 	/**
@@ -1306,7 +1306,7 @@ public:
 	 * @param PokeInfo returned information about new and modified mesh elements
 	 * @return Ok on success, or enum value indicates why operation cannot be applied. Mesh remains unmodified on error.
 	 */
-	virtual EMeshResult PokeTriangle(int TriangleID, const FVector3d& BaryCoordinates, FPokeTriangleInfo& PokeInfo);
+	GEOMETRYCORE_API virtual EMeshResult PokeTriangle(int TriangleID, const FVector3d& BaryCoordinates, FPokeTriangleInfo& PokeInfo);
 
 	/** Call PokeTriangle at the centroid of the triangle */
 	virtual EMeshResult PokeTriangle(int TriangleID, FPokeTriangleInfo& PokeInfo)
@@ -1321,7 +1321,7 @@ public:
 	/**
 	 * Returns a debug string that contains mesh statistics and other information
 	 */
-	virtual FString MeshInfoString() const;
+	GEOMETRYCORE_API virtual FString MeshInfoString() const;
 
 	/**
 	 * Options for the IsSameAs check
@@ -1346,7 +1346,7 @@ public:
 	 * Check if another mesh is the same as this mesh. By default only checks
 	 * vertices and triangles, turn on other parameters w/ flags
 	 */
-	virtual bool IsSameAs(const FDynamicMesh3& OtherMesh, const FSameAsOptions& Options) const;
+	GEOMETRYCORE_API virtual bool IsSameAs(const FDynamicMesh3& OtherMesh, const FSameAsOptions& Options) const;
 
 	/**
 	 * Options for what the validity check will permit
@@ -1379,7 +1379,7 @@ public:
 	/**
 	 * Checks that the mesh is well-formed, ie all internal data structures are consistent
 	 */
-	virtual bool CheckValidity(FValidityOptions ValidityOptions = FValidityOptions(),
+	GEOMETRYCORE_API virtual bool CheckValidity(FValidityOptions ValidityOptions = FValidityOptions(),
 	                           EValidityCheckFailMode FailMode  = EValidityCheckFailMode::Check) const;
 
 	//
@@ -1395,8 +1395,8 @@ protected:
 		TriangleEdges[TriangleID] = FIndex3i(e0, e1, e2);
 	}
 
-	int AddEdgeInternal(int vA, int vB, int tA, int tB = InvalidID);
-	int AddTriangleInternal(int a, int b, int c, int e0, int e1, int e2);
+	GEOMETRYCORE_API int AddEdgeInternal(int vA, int vB, int tA, int tB = InvalidID);
+	GEOMETRYCORE_API int AddTriangleInternal(int a, int b, int c, int e0, int e1, int e2);
 
 	inline int ReplaceTriangleVertex(int TriangleID, int vOld, int vNew)
 	{
@@ -1445,9 +1445,9 @@ protected:
 		Edges[EdgeID].Tri[1] = t1;
 	}
 
-	int ReplaceEdgeVertex(int EdgeID, int vOld, int vNew);
-	int ReplaceEdgeTriangle(int EdgeID, int tOld, int tNew);
-	int ReplaceTriangleEdge(int EdgeID, int eOld, int eNew);
+	GEOMETRYCORE_API int ReplaceEdgeVertex(int EdgeID, int vOld, int vNew);
+	GEOMETRYCORE_API int ReplaceEdgeTriangle(int EdgeID, int tOld, int tNew);
+	GEOMETRYCORE_API int ReplaceTriangleEdge(int EdgeID, int eOld, int eNew);
 
 	inline bool TriangleHasVertex(int TriangleID, int VertexID) const
 	{
@@ -1468,9 +1468,9 @@ protected:
 		return ((Tri.A == vA && Tri.B == vB) || (Tri.B == vA && Tri.C == vB) || (Tri.C == vA && Tri.A == vB));
 	}
 
-	int FindTriangleEdge(int TriangleID, int vA, int vB) const;
+	GEOMETRYCORE_API int FindTriangleEdge(int TriangleID, int vA, int vB) const;
 
-	int32 FindEdgeInternal(int32 vA, int32 vB, bool& bIsBoundary) const;
+	GEOMETRYCORE_API int32 FindEdgeInternal(int32 vA, int32 vB, bool& bIsBoundary) const;
 
 	inline bool EdgeHasVertex(int EdgeID, int VertexID) const
 	{
@@ -1524,7 +1524,7 @@ protected:
 	}
 
 
-	void ReverseTriOrientationInternal(int TriangleID);
+	GEOMETRYCORE_API void ReverseTriOrientationInternal(int TriangleID);
 
 	/**
 	* Internal implementations for serialization to allow for better code separation between different versions and other implementation details.
diff --git a/Engine/Source/Runtime/GeometryFramework/Public/UDynamicMesh.h b/Engine/Source/Runtime/GeometryFramework/Public/UDynamicMesh.h
--- a/Engine/Source/Runtime/GeometryFramework/Public/UDynamicMesh.h
+++ b/Engine/Source/Runtime/GeometryFramework/Public/UDynamicMesh.h
@@ -18,8 +18,8 @@
  * mutate an input mesh into an output mesh. A subclass of this class can be attached to
  * a UDynamicMesh to allow for arbitrarily-complex procedural generation
  */
-UCLASS(Abstract)
-class GEOMETRYFRAMEWORK_API UDynamicMeshGenerator : public UObject
+UCLASS(Abstract, MinimalAPI)
+class UDynamicMeshGenerator : public UObject
 {
 	GENERATED_BODY()
 public:
@@ -65,7 +65,7 @@ ENUM_CLASS_FLAGS(EDynamicMeshAttributeChangeFlags)
  * This struct is emitted by the UDynamicMesh OnPreMeshChanged() and OnMeshChanged() delegates.
  */
 USTRUCT(BlueprintType)
-struct GEOMETRYFRAMEWORK_API FDynamicMeshChangeInfo
+struct FDynamicMeshChangeInfo
 {
 	GENERATED_BODY()
 
@@ -96,8 +96,8 @@ DECLARE_DYNAMIC_MULTICAST_DELEGATE_OneParam(FOnDynamicMeshModifiedBP, UDynamicMe
 /**
  * UDynamicMesh is a UObject container for a FDynamicMesh3. 
  */
-UCLASS(BlueprintType)
-class GEOMETRYFRAMEWORK_API UDynamicMesh : public UObject,
+UCLASS(BlueprintType, MinimalAPI)
+class UDynamicMesh : public UObject,
 	public IMeshVertexCommandChangeTarget, 
 	public IMeshCommandChangeTarget, 
 	public IMeshReplacementCommandChangeTarget
@@ -110,16 +110,14 @@ public:
 	 * This *does not* allocate a new mesh, so any existing mesh pointers/refs are still valid
 	 */
 	UFUNCTION(BlueprintCallable, Category = "Dynamic Mesh")
-	UPARAM(DisplayName = "Target") UDynamicMesh* 
-	Reset();
+	GEOMETRYFRAMEWORK_API UPARAM(DisplayName = "TarGEOMETRYFRAMEWORK_API get") UDynamicMesh* Reset();
 
 	/**
 	 * Clear the internal mesh to a 100x100x100 cube with base at the origin.
 	 * This this instead of Reset() if an initially-empty mesh is undesirable (eg for a Component)
 	 */
 	UFUNCTION(BlueprintCallable, Category = "Dynamic Mesh")
-	UPARAM(DisplayName = "Target") UDynamicMesh* 
-	ResetToCube();
+	GEOMETRYFRAMEWORK_API UPARAM(DisplayName = "TarGEOMETRYFRAMEWORK_API get") UDynamicMesh* ResetToCube();
 
 	//
 	// Native access/modification functions
@@ -128,19 +126,19 @@ public:
 	/**
 	 * Reset the internal mesh data and then optionally run the MeshGenerator
 	 */
-	virtual void InitializeMesh();
+	GEOMETRYFRAMEWORK_API virtual void InitializeMesh();
 
 	/**
 	 * @return true if the mesh has no triangles
 	 */
 	UFUNCTION(BlueprintCallable, Category = "Dynamic Mesh")
-	bool IsEmpty() const;
+	GEOMETRYFRAMEWORK_API bool IsEmpty() const;
 
 	/**
 	 * @return number of triangles in the mesh
 	 */
 	UFUNCTION(BlueprintCallable, Category = "Dynamic Mesh")
-	UPARAM(DisplayName = "Triangle Count") int32 GetTriangleCount() const;
+	GEOMETRYFRAMEWORK_API UPARAM(DisplayName = "TriGEOMETRYFRAMEWORK_API angle Count") int32 GetTriangleCount() const;
 
 
 	/** 
@@ -167,21 +165,21 @@ public:
 
 
 	/** Replace the internal mesh with a copy of MoveMesh */
-	void SetMesh(const UE::Geometry::FDynamicMesh3& MoveMesh);
+	GEOMETRYFRAMEWORK_API void SetMesh(const UE::Geometry::FDynamicMesh3& MoveMesh);
 
 	/** Replace the internal mesh with the data in MoveMesh */
-	void SetMesh(UE::Geometry::FDynamicMesh3&& MoveMesh);
+	GEOMETRYFRAMEWORK_API void SetMesh(UE::Geometry::FDynamicMesh3&& MoveMesh);
 
 	/**
 	 * Apply ProcessFunc to the internal Mesh
 	 */
-	void ProcessMesh(TFunctionRef<void(const UE::Geometry::FDynamicMesh3&)> ProcessFunc) const;
+	GEOMETRYFRAMEWORK_API void ProcessMesh(TFunctionRef<void(const UE::Geometry::FDynamicMesh3&)> ProcessFunc) const;
 
 	/**
 	 * Apply EditFunc to the internal mesh.
 	 * This will broadcast PreMeshChangedEvent, then call EditFunc(), then broadcast MeshChangedEvent and MeshModifiedBPEvent
 	 */
-	void EditMesh(TFunctionRef<void(UE::Geometry::FDynamicMesh3&)> EditFunc,
+	GEOMETRYFRAMEWORK_API void EditMesh(TFunctionRef<void(UE::Geometry::FDynamicMesh3&)> EditFunc,
 				  EDynamicMeshChangeType ChangeType = EDynamicMeshChangeType::GeneralEdit,
 				  EDynamicMeshAttributeChangeFlags ChangeFlags = EDynamicMeshAttributeChangeFlags::Unknown,
 				  bool bDeferChangeEvents = false);
@@ -189,7 +187,7 @@ public:
 	/**
 	 * Take ownership of the internal Mesh, and have it replaced with a new mesh
 	 */
-	TUniquePtr<UE::Geometry::FDynamicMesh3> ExtractMesh();
+	GEOMETRYFRAMEWORK_API TUniquePtr<UE::Geometry::FDynamicMesh3> ExtractMesh();
 
 
 	//
@@ -198,13 +196,13 @@ public:
 
 
 	// IMeshVertexCommandChangeTarget implementation, allows a FVertexChange to be applied to the mesh
-	void ApplyChange(const FMeshVertexChange* Change, bool bRevert);
+	GEOMETRYFRAMEWORK_API void ApplyChange(const FMeshVertexChange* Change, bool bRevert);
 
 	// IMeshCommandChangeTarget implementation, allows a FMeshChange to be applied to the mesh
-	void ApplyChange(const FMeshChange* Change, bool bRevert);
+	GEOMETRYFRAMEWORK_API void ApplyChange(const FMeshChange* Change, bool bRevert);
 
 	// IMeshReplacementCommandChangeTarget implementation, allows a FMeshReplacementChange to be applied to the mesh
-	void ApplyChange(const FMeshReplacementChange* Change, bool bRevert);
+	GEOMETRYFRAMEWORK_API void ApplyChange(const FMeshReplacementChange* Change, bool bRevert);
 
 
 	//
@@ -249,7 +247,7 @@ public:
 	/**
 	 * Broadcasts FOnMeshRealtimeUpdate
 	 */
-	virtual void PostRealtimeUpdate();
+	GEOMETRYFRAMEWORK_API virtual void PostRealtimeUpdate();
 
 protected:
 	FOnMeshRealtimeUpdate MeshRealtimeUpdateEvent;
@@ -266,21 +264,21 @@ protected:
 	/**
 	 * Allocate a new Mesh (ie pointer will change) and then call InitializeMesh()
 	 */
-	void InitializeNewMesh();
+	GEOMETRYFRAMEWORK_API void InitializeNewMesh();
 
 	/**
 	 * Internal function that edits the Mesh, but broadcasts PreMeshChangedEvent and MeshChangedEvent, and then MeshModifiedBPEvent
 	 */
-	void EditMeshInternal(TFunctionRef<void(UE::Geometry::FDynamicMesh3&)> EditFunc, const FDynamicMeshChangeInfo& ChangeInfo, bool bDeferChangeEvents = false);
+	GEOMETRYFRAMEWORK_API void EditMeshInternal(TFunctionRef<void(UE::Geometry::FDynamicMesh3&)> EditFunc, const FDynamicMeshChangeInfo& ChangeInfo, bool bDeferChangeEvents = false);
 
 
 public:
 	// serialize Mesh to an Archive
-	virtual void Serialize(FArchive& Archive) override;
+	GEOMETRYFRAMEWORK_API virtual void Serialize(FArchive& Archive) override;
 
 	// serialize Mesh to/from T3D
-	virtual void ExportCustomProperties(FOutputDevice& Out, uint32 Indent) override;
-	virtual void ImportCustomProperties(const TCHAR* SourceText, FFeedbackContext* Warn) override;
+	GEOMETRYFRAMEWORK_API virtual void ExportCustomProperties(FOutputDevice& Out, uint32 Indent) override;
+	GEOMETRYFRAMEWORK_API virtual void ImportCustomProperties(const TCHAR* SourceText, FFeedbackContext* Warn) override;
 
 
 
@@ -312,17 +310,17 @@ public:
 	/**
 	 * Set the active mesh generator. Clears if nullptr
 	 */
-	virtual void SetMeshGenerator(TObjectPtr<UDynamicMeshGenerator> NewGenerator);
+	GEOMETRYFRAMEWORK_API virtual void SetMeshGenerator(TObjectPtr<UDynamicMeshGenerator> NewGenerator);
 
 	/**
 	 * Set the active mesh generator. Clears if nullptr
 	 */
-	virtual void ClearMeshGenerator();
+	GEOMETRYFRAMEWORK_API virtual void ClearMeshGenerator();
 
 	/**
 	 * Reset() the mesh, which will re-run the active MeshGenerator, if bEnableMeshGenerator
 	 */
-	virtual void Regenerate();
+	GEOMETRYFRAMEWORK_API virtual void Regenerate();
 
 };
 
@@ -366,26 +364,26 @@ public:
  * the internal FDynamicMesh3 (which uses normal C++ memory management, so no garbage collection involved)
  * So the Pool does not re-use mesh memory, only the UObject containers.
  */
-UCLASS(BlueprintType, Transient)
-class GEOMETRYFRAMEWORK_API UDynamicMeshPool : public UObject
+UCLASS(BlueprintType, Transient, MinimalAPI)
+class UDynamicMeshPool : public UObject
 {
 	GENERATED_BODY()
 public:
 	/** @return an available UDynamicMesh from the pool (possibly allocating a new mesh) */
 	UFUNCTION(BlueprintCallable, Category="Dynamic Mesh")
-	UDynamicMesh* RequestMesh();
+	GEOMETRYFRAMEWORK_API UDynamicMesh* RequestMesh();
 
 	/** Release a UDynamicMesh returned by RequestMesh() back to the pool */
 	UFUNCTION(BlueprintCallable, Category = "Dynamic Mesh")
-	void ReturnMesh(UDynamicMesh* Mesh);
+	GEOMETRYFRAMEWORK_API void ReturnMesh(UDynamicMesh* Mesh);
 
 	/** Release all GeneratedMeshes back to the pool */
 	UFUNCTION(BlueprintCallable, Category = "Dynamic Mesh")
-	void ReturnAllMeshes();
+	GEOMETRYFRAMEWORK_API void ReturnAllMeshes();
 
 	/** Release all GeneratedMeshes back to the pool and allow them to be garbage collected */
 	UFUNCTION(BlueprintCallable, Category = "Dynamic Mesh")
-	void FreeAllMeshes();
+	GEOMETRYFRAMEWORK_API void FreeAllMeshes();
 
 protected:
 	/** Meshes in the pool that are available */
@@ -395,4 +393,4 @@ protected:
 	/** All meshes the pool has allocated */
 	UPROPERTY()
 	TArray<TObjectPtr<UDynamicMesh>> AllCreatedMeshes;
-};
\ No newline at end of file
+};
diff --git a/Engine/Source/Runtime/HeadMountedDisplay/Public/IXRTrackingSystem.h b/Engine/Source/Runtime/HeadMountedDisplay/Public/IXRTrackingSystem.h
--- a/Engine/Source/Runtime/HeadMountedDisplay/Public/IXRTrackingSystem.h
+++ b/Engine/Source/Runtime/HeadMountedDisplay/Public/IXRTrackingSystem.h
@@ -43,7 +43,7 @@ struct FXRSensorProperties
 /**
  * Main access point to an XR tracking system. Use it to enumerate devices and query their poses.
  */
-class HEADMOUNTEDDISPLAY_API  IXRTrackingSystem : public IModularFeature, public IXRSystemIdentifier
+class  IXRTrackingSystem : public IModularFeature, public IXRSystemIdentifier
 {
 public:
 	static FName GetModularFeatureName()
@@ -67,7 +67,7 @@ public:
 	 * Other devices can have arbitrary ids defined by each system.
 	 * If a tracking system does not support tracking HMDs, device ID zero should be treated as invalid.
 	 */
-	static const int32 HMDDeviceId = 0;
+	static HEADMOUNTEDDISPLAY_API const int32 HMDDeviceId = 0;
 
 	/**
 	 * Whether or not the system supports positional tracking (either via sensor or other means)
@@ -365,7 +365,7 @@ public:
 	/**
 	 * Same as IsHeadTrackingAllowed, but returns false if the World is not using VR (such as with the non-VR PIE instances when using VR Preview)
 	 **/
-	virtual bool IsHeadTrackingAllowedForWorld(UWorld & World) const;
+	HEADMOUNTEDDISPLAY_API virtual bool IsHeadTrackingAllowedForWorld(UWorld & World) const;
 
 	/** 
 	* Can be used to enforce tracking even when stereo rendering is disabled. 
@@ -423,7 +423,7 @@ public:
 	/**
 	 * Platform Agnostic Query about HMD details
 	 */
-	virtual void GetHMDData(UObject* WorldContext, FXRHMDData& HMDData);
+	HEADMOUNTEDDISPLAY_API virtual void GetHMDData(UObject* WorldContext, FXRHMDData& HMDData);
 
 	/**
 	 * Platform Agnostic Query about MotionControllers details
diff --git a/Engine/Source/Runtime/Landscape/Classes/LandscapeProxy.h b/Engine/Source/Runtime/Landscape/Classes/LandscapeProxy.h
--- a/Engine/Source/Runtime/Landscape/Classes/LandscapeProxy.h
+++ b/Engine/Source/Runtime/Landscape/Classes/LandscapeProxy.h
@@ -1067,7 +1067,7 @@ public:
 	LANDSCAPE_API FIntPoint GetSectionBaseOffset() const;
 
 	// ILandscapeSplineInterface
-	LANDSCAPE_API virtual ULandscapeSplinesComponent* GetSplinesComponent() const override { return SplineComponent; }
+	virtual ULandscapeSplinesComponent* GetSplinesComponent() const override { return SplineComponent; }
 	LANDSCAPE_API virtual void UpdateSharedProperties(ULandscapeInfo* InLandscapeInfo) override;
 
 	// Retrieve the screen size at which each LOD should be rendered
@@ -1101,10 +1101,10 @@ public:
 	LANDSCAPE_API bool IsPropertyOverridable(const FProperty* InProperty) const;
 	
 	// Returns true if the shared property is overridden by the object.
-	LANDSCAPE_API virtual bool IsSharedPropertyOverridden(const FString& InPropertyName) const { return false; }
+	virtual bool IsSharedPropertyOverridden(const FString& InPropertyName) const { return false; }
 
 	// Modifies the override state of the property given as argument.
-	LANDSCAPE_API virtual void SetSharedPropertyOverride(const FString& InPropertyName, const bool bIsOverriden) { }
+	virtual void SetSharedPropertyOverride(const FString& InPropertyName, const bool bIsOverriden) { }
 #endif // WITH_EDITOR
 
 	// Get Landscape Material assigned to this Landscape
@@ -1117,16 +1117,16 @@ public:
 	/* Serialize all hashes/guids that record the current state of this proxy */
 	void SerializeStateHashes(FArchive& Ar);
 
-	LANDSCAPE_API void SetSplinesComponent(ULandscapeSplinesComponent* InSplineComponent) { check(!SplineComponent || (SplineComponent == InSplineComponent)); SplineComponent = InSplineComponent; }
+	void SetSplinesComponent(ULandscapeSplinesComponent* InSplineComponent) { check(!SplineComponent || (SplineComponent == InSplineComponent)); SplineComponent = InSplineComponent; }
 
-	LANDSCAPE_API virtual bool SupportsForeignSplineMesh() const override { return true; }
+	virtual bool SupportsForeignSplineMesh() const override { return true; }
 
 	LANDSCAPE_API int32 GetOutdatedGrassMapCount() const;
 	LANDSCAPE_API void BuildGrassMaps(struct FScopedSlowTask* InSlowTask = nullptr);
 	UE_DEPRECATED(5.3, "BuildGIBakedTextures is officially deprecated now")
-	LANDSCAPE_API void BuildGIBakedTextures(struct FScopedSlowTask* InSlowTask = nullptr) {}
+	void BuildGIBakedTextures(struct FScopedSlowTask* InSlowTask = nullptr) {}
 	UE_DEPRECATED(5.3, "GetOutdatedGIBakedTextureComponentsCount is officially deprecated now returns 0")
-	LANDSCAPE_API int32 GetOutdatedGIBakedTextureComponentsCount() const { return 0; }
+	int32 GetOutdatedGIBakedTextureComponentsCount() const { return 0; }
 	LANDSCAPE_API void BuildPhysicalMaterial(struct FScopedSlowTask* InSlowTask = nullptr);
 	LANDSCAPE_API int32 GetOudatedPhysicalMaterialComponentsCount() const;
 	LANDSCAPE_API virtual void CreateSplineComponent() override;
@@ -1236,7 +1236,7 @@ public:
 	UE_DEPRECATED(5.2, "Use the version of this function taking a FRawMeshExportParams as a parameter")
 	LANDSCAPE_API bool ExportToRawMesh(const TArrayView<ULandscapeComponent*>& InComponents, int32 InExportLOD, FMeshDescription& OutRawMesh, const FBoxSphereBounds& InBounds, bool bIgnoreBounds = false, bool bGenerateOnePolygonGroupPerComponent = false) const;
 
-	struct LANDSCAPE_API FRawMeshExportParams
+	struct FRawMeshExportParams
 	{
 		static constexpr int32 MaxUVCount = 6;
 
@@ -1261,8 +1261,8 @@ public:
 		/** Describes what to export on each UV channel */
 		struct FUVConfiguration
 		{
-			FUVConfiguration();
-			int32 GetNumUVChannelsNeeded() const;
+			LANDSCAPE_API FUVConfiguration();
+			LANDSCAPE_API int32 GetNumUVChannelsNeeded() const;
 
 		public:
 			// Index 0 = UVChannel 0, Index 1 = UVChannel 1... 
@@ -1280,9 +1280,9 @@ public:
 
 	public:
 		FRawMeshExportParams() = default;
-		const FUVConfiguration& GetUVConfiguration(int32 InComponentIndex) const;
-		const FName& GetMaterialSlotName(int32 InComponentIndex) const;
-		int32 GetNumUVChannelsNeeded() const;
+		LANDSCAPE_API const FUVConfiguration& GetUVConfiguration(int32 InComponentIndex) const;
+		LANDSCAPE_API const FName& GetMaterialSlotName(int32 InComponentIndex) const;
+		LANDSCAPE_API int32 GetNumUVChannelsNeeded() const;
 
 	public:
 		/** LOD level to export. If none specified, LOD 0 will be used */
@@ -1448,7 +1448,7 @@ protected:
 	LANDSCAPE_API void InitializeLayerWithEmptyContent(const FGuid& InLayerGuid);
 
 	/** Fixup any internal representation for shared properties. Will be used when deprecating or renaming a property. */
-	LANDSCAPE_API virtual void FixupOverriddenSharedProperties() {}
+	virtual void FixupOverriddenSharedProperties() {}
 
 protected:
 	FLandscapeMaterialChangedDelegate LandscapeMaterialChangedDelegate;
@@ -1506,12 +1506,12 @@ private:
 /**
  * Helper class used to Build or monitor outdated Grass maps of a world
  */
-class LANDSCAPE_API FLandscapeGrassMapsBuilder
+class FLandscapeGrassMapsBuilder
 {
 public:
-	FLandscapeGrassMapsBuilder(UWorld* InWorld);
-	void Build();
-	int32 GetOutdatedGrassMapCount(bool bInForceUpdate = true) const;
+	LANDSCAPE_API FLandscapeGrassMapsBuilder(UWorld* InWorld);
+	LANDSCAPE_API void Build();
+	LANDSCAPE_API int32 GetOutdatedGrassMapCount(bool bInForceUpdate = true) const;
 private:
 	UWorld* World;
 	mutable int32 OutdatedGrassMapCount;
@@ -1540,12 +1540,12 @@ private:
 /**
  * Helper class used to Build or monitor Landscape Physical Material
  */
-class LANDSCAPE_API FLandscapePhysicalMaterialBuilder
+class FLandscapePhysicalMaterialBuilder
 {
 public:
-	FLandscapePhysicalMaterialBuilder(UWorld* InWorld);
-	void Build();
-	int32 GetOudatedPhysicalMaterialComponentsCount();
+	LANDSCAPE_API FLandscapePhysicalMaterialBuilder(UWorld* InWorld);
+	LANDSCAPE_API void Build();
+	LANDSCAPE_API int32 GetOudatedPhysicalMaterialComponentsCount();
 private:
 	UWorld* World;
 	int32 OudatedPhysicalMaterialComponentsCount;
@@ -1554,11 +1554,11 @@ private:
 /**
 * Helper class to store proxy changes information 
 */
-class LANDSCAPE_API FLandscapeProxyComponentDataChangedParams
+class FLandscapeProxyComponentDataChangedParams
 {
 public:
-	FLandscapeProxyComponentDataChangedParams(const TSet<ULandscapeComponent*>& InComponents);
-	void ForEachComponent(TFunctionRef<void(const ULandscapeComponent*)> Func) const;
+	LANDSCAPE_API FLandscapeProxyComponentDataChangedParams(const TSet<ULandscapeComponent*>& InComponents);
+	LANDSCAPE_API void ForEachComponent(TFunctionRef<void(const ULandscapeComponent*)> Func) const;
 	const TArray<ULandscapeComponent*>& GetComponents() const { return Components; }
 
 private:
diff --git a/Engine/Source/Runtime/Landscape/Public/LandscapeRender.h b/Engine/Source/Runtime/Landscape/Public/LandscapeRender.h
--- a/Engine/Source/Runtime/Landscape/Public/LandscapeRender.h
+++ b/Engine/Source/Runtime/Landscape/Public/LandscapeRender.h
@@ -201,13 +201,13 @@ public:
 };
 
 /** vertex factory for VTF-heightmap terrain  */
-class LANDSCAPE_API FLandscapeVertexFactory : public FVertexFactory
+class FLandscapeVertexFactory : public FVertexFactory
 {
-	DECLARE_VERTEX_FACTORY_TYPE(FLandscapeVertexFactory);
+	DECLARE_VERTEX_FACTORY_TYPE_API(FLandscapeVertexFactory, LANDSCAPE_API);
 
 public:
 
-	FLandscapeVertexFactory(ERHIFeatureLevel::Type InFeatureLevel);
+	LANDSCAPE_API FLandscapeVertexFactory(ERHIFeatureLevel::Type InFeatureLevel);
 
 	virtual ~FLandscapeVertexFactory()
 	{
@@ -224,26 +224,26 @@ public:
 	/**
 	* Should we cache the material's shadertype on this platform with this vertex factory?
 	*/
-	static bool ShouldCompilePermutation(const FVertexFactoryShaderPermutationParameters& Parameters);
+	static LANDSCAPE_API bool ShouldCompilePermutation(const FVertexFactoryShaderPermutationParameters& Parameters);
 
 	/**
 	* Can be overridden by FVertexFactory subclasses to modify their compile environment just before compilation occurs.
 	*/
-	static void ModifyCompilationEnvironment(const FVertexFactoryShaderPermutationParameters& Parameters, FShaderCompilerEnvironment& OutEnvironment);
+	static LANDSCAPE_API void ModifyCompilationEnvironment(const FVertexFactoryShaderPermutationParameters& Parameters, FShaderCompilerEnvironment& OutEnvironment);
 	
 	/**
 	* Get vertex elements used when during PSO precaching materials using this vertex factory type
 	*/
-	static void GetPSOPrecacheVertexFetchElements(EVertexInputStreamType VertexInputStreamType, FVertexDeclarationElementList& Elements);
+	static LANDSCAPE_API void GetPSOPrecacheVertexFetchElements(EVertexInputStreamType VertexInputStreamType, FVertexDeclarationElementList& Elements);
 
 	/**
 	* Copy the data from another vertex factory
 	* @param Other - factory to copy from
 	*/
-	void Copy(const FLandscapeVertexFactory& Other);
+	LANDSCAPE_API void Copy(const FLandscapeVertexFactory& Other);
 
 	// FRenderResource interface.
-	virtual void InitRHI() override;
+	LANDSCAPE_API virtual void InitRHI() override;
 	virtual void ReleaseResource() override final { FVertexFactory::ReleaseResource(); }
 
 	/**
@@ -263,7 +263,7 @@ public:
 /** vertex factory for VTF-heightmap terrain  */
 class FLandscapeXYOffsetVertexFactory : public FLandscapeVertexFactory
 {
-	DECLARE_VERTEX_FACTORY_TYPE(FLandscapeXYOffsetVertexFactory);
+	DECLARE_VERTEX_FACTORY_TYPE_API(FLandscapeXYOffsetVertexFactory, LANDSCAPE_API);
 
 public:
 	FLandscapeXYOffsetVertexFactory(ERHIFeatureLevel::Type InFeatureLevel)
@@ -278,9 +278,9 @@ public:
 
 
 /** Vertex factory for fixed grid runtime virtual texture lod  */
-class LANDSCAPE_API FLandscapeFixedGridVertexFactory : public FLandscapeVertexFactory
+class FLandscapeFixedGridVertexFactory : public FLandscapeVertexFactory
 {
-	DECLARE_VERTEX_FACTORY_TYPE(FLandscapeFixedGridVertexFactory);
+	DECLARE_VERTEX_FACTORY_TYPE_API(FLandscapeFixedGridVertexFactory, LANDSCAPE_API);
 
 public:
 	FLandscapeFixedGridVertexFactory(ERHIFeatureLevel::Type InFeatureLevel)
@@ -288,7 +288,7 @@ public:
 	{
 	}
 
-	static void ModifyCompilationEnvironment(const FVertexFactoryShaderPermutationParameters& Parameters, FShaderCompilerEnvironment& OutEnvironment);
+	static LANDSCAPE_API void ModifyCompilationEnvironment(const FVertexFactoryShaderPermutationParameters& Parameters, FShaderCompilerEnvironment& OutEnvironment);
 };
 
 
@@ -337,7 +337,7 @@ public:
 //
 // FLandscapeSharedBuffers
 //
-class LANDSCAPE_API FLandscapeSharedBuffers : public FRefCountedObject
+class FLandscapeSharedBuffers : public FRefCountedObject
 {
 public:
 	struct FLandscapeIndexRanges
@@ -374,7 +374,7 @@ public:
 	TArray<FIndexBuffer*> ZeroOffsetIndexBuffers;
 #endif
 
-	FLandscapeSharedBuffers(FRHICommandListBase& RHICmdList, int32 SharedBuffersKey, int32 SubsectionSizeQuads, int32 NumSubsections, ERHIFeatureLevel::Type FeatureLevel, const FName& OwnerName = NAME_None);
+	LANDSCAPE_API FLandscapeSharedBuffers(FRHICommandListBase& RHICmdList, int32 SharedBuffersKey, int32 SubsectionSizeQuads, int32 NumSubsections, ERHIFeatureLevel::Type FeatureLevel, const FName& OwnerName = NAME_None);
 
 	template <typename INDEX_TYPE>
 	void CreateIndexBuffers(FRHICommandListBase& RHICmdList, const FName& OwnerName);
@@ -384,7 +384,7 @@ public:
 	void CreateGrassIndexBuffer(FRHICommandListBase& RHICmdList, const FName& InOwnerName);
 #endif
 
-	virtual ~FLandscapeSharedBuffers();
+	LANDSCAPE_API virtual ~FLandscapeSharedBuffers();
 };
 
 //
@@ -568,7 +568,7 @@ public:
 	virtual void PreInitViews_RenderThread(FRDGBuilder& GraphBuilder) override;
 
 	LANDSCAPE_API const TMap<uint32, FLandscapeRenderSystem*>& GetLandscapeRenderSystems() const;
-	LANDSCAPE_API int32 GetNumViewsWithShowCollision() const { return NumViewsWithShowCollision; }
+	int32 GetNumViewsWithShowCollision() const { return NumViewsWithShowCollision; }
 private:
 	FBufferRHIRef LandscapeLODDataBuffer;
 	FBufferRHIRef LandscapeIndirectionBuffer;
@@ -612,9 +612,9 @@ private:
 //
 // FLandscapeDebugOptions
 //
-struct LANDSCAPE_API FLandscapeDebugOptions
+struct FLandscapeDebugOptions
 {
-	FLandscapeDebugOptions();
+	LANDSCAPE_API FLandscapeDebugOptions();
 
 	enum eCombineMode
 	{
@@ -665,11 +665,11 @@ private:
 //
 // FLandscapeComponentSceneProxy
 //
-class LANDSCAPE_API FLandscapeComponentSceneProxy : public FPrimitiveSceneProxy, public FLandscapeSectionInfo
+class FLandscapeComponentSceneProxy : public FPrimitiveSceneProxy, public FLandscapeSectionInfo
 {
 	friend class FLandscapeSharedBuffers;
 
-	SIZE_T GetTypeHash() const override;
+	LANDSCAPE_API SIZE_T GetTypeHash() const override;
 	class FLandscapeLCI final : public FLightCacheInterface
 	{
 	public:
@@ -714,7 +714,7 @@ public:
 	// Reference counted vertex and index buffer shared among all landscape scene proxies of the same component size
 	// Key is the component size and number of subsections.
 	// Also being reused by GPULightmass currently to save mem
-	static TMap<uint32, FLandscapeSharedBuffers*> SharedBuffersMap;
+	static LANDSCAPE_API TMap<uint32, FLandscapeSharedBuffers*> SharedBuffersMap;
 
 protected:
 	int8						MaxLOD;						// Maximum LOD level, user override possible
@@ -838,31 +838,31 @@ protected:
 	FLandscapeVisibilityHelper VisibilityHelper;
 
 protected:
-	virtual ~FLandscapeComponentSceneProxy();
+	LANDSCAPE_API virtual ~FLandscapeComponentSceneProxy();
 	
-	int8 GetLODFromScreenSize(float InScreenSizeSquared, float InViewLODScale) const;
+	LANDSCAPE_API int8 GetLODFromScreenSize(float InScreenSizeSquared, float InViewLODScale) const;
 
-	bool GetMeshElementForVirtualTexture(int32 InLodIndex, ERuntimeVirtualTextureMaterialType MaterialType, FMaterialRenderProxy* InMaterialInterface, FMeshBatch& OutMeshBatch, TArray<FLandscapeBatchElementParams>& OutStaticBatchParamArray) const;
+	LANDSCAPE_API bool GetMeshElementForVirtualTexture(int32 InLodIndex, ERuntimeVirtualTextureMaterialType MaterialType, FMaterialRenderProxy* InMaterialInterface, FMeshBatch& OutMeshBatch, TArray<FLandscapeBatchElementParams>& OutStaticBatchParamArray) const;
 	template<class ArrayType> bool GetStaticMeshElement(int32 LODIndex, bool bForToolMesh, FMeshBatch& MeshBatch, ArrayType& OutStaticBatchParamArray) const;
 
 public:
 	// constructor
-	FLandscapeComponentSceneProxy(ULandscapeComponent* InComponent);
+	LANDSCAPE_API FLandscapeComponentSceneProxy(ULandscapeComponent* InComponent);
 
 	// FPrimitiveSceneProxy interface.
-	virtual void ApplyWorldOffset(FVector InOffset) override;
-	virtual void DrawStaticElements(FStaticPrimitiveDrawInterface* PDI) override;
-	virtual void GetDynamicMeshElements(const TArray<const FSceneView*>& Views, const FSceneViewFamily& ViewFamily, uint32 VisibilityMap, FMeshElementCollector& Collector) const override;
-	virtual void ApplyViewDependentMeshArguments(const FSceneView& View, FMeshBatch& ViewDependentMeshBatch) const override;
+	LANDSCAPE_API virtual void ApplyWorldOffset(FVector InOffset) override;
+	LANDSCAPE_API virtual void DrawStaticElements(FStaticPrimitiveDrawInterface* PDI) override;
+	LANDSCAPE_API virtual void GetDynamicMeshElements(const TArray<const FSceneView*>& Views, const FSceneViewFamily& ViewFamily, uint32 VisibilityMap, FMeshElementCollector& Collector) const override;
+	LANDSCAPE_API virtual void ApplyViewDependentMeshArguments(const FSceneView& View, FMeshBatch& ViewDependentMeshBatch) const override;
 	virtual uint32 GetMemoryFootprint() const override { return(sizeof(*this) + GetAllocatedSize()); }
-	virtual FPrimitiveViewRelevance GetViewRelevance(const FSceneView* View) const override;
-	virtual bool CanBeOccluded() const override;
-	virtual void GetLightRelevance(const FLightSceneProxy* LightSceneProxy, bool& bDynamic, bool& bRelevant, bool& bLightMapped, bool& bShadowMapped) const override;
-	virtual void OnTransformChanged() override;
-	virtual void CreateRenderThreadResources() override;
-	virtual void DestroyRenderThreadResources() override;
-	virtual bool OnLevelAddedToWorld_RenderThread() override;
-	virtual void OnLevelRemovedFromWorld_RenderThread() override;
+	LANDSCAPE_API virtual FPrimitiveViewRelevance GetViewRelevance(const FSceneView* View) const override;
+	LANDSCAPE_API virtual bool CanBeOccluded() const override;
+	LANDSCAPE_API virtual void GetLightRelevance(const FLightSceneProxy* LightSceneProxy, bool& bDynamic, bool& bRelevant, bool& bLightMapped, bool& bShadowMapped) const override;
+	LANDSCAPE_API virtual void OnTransformChanged() override;
+	LANDSCAPE_API virtual void CreateRenderThreadResources() override;
+	LANDSCAPE_API virtual void DestroyRenderThreadResources() override;
+	LANDSCAPE_API virtual bool OnLevelAddedToWorld_RenderThread() override;
+	LANDSCAPE_API virtual void OnLevelRemovedFromWorld_RenderThread() override;
 	
 	friend class ULandscapeComponent;
 	friend class FLandscapeVertexFactoryVertexShaderParameters;
@@ -875,32 +875,32 @@ public:
 #endif
 
 	// FLandcapeSceneProxy
-	void ChangeComponentScreenSizeToUseSubSections_RenderThread(float InComponentScreenSizeToUseSubSections);
+	LANDSCAPE_API void ChangeComponentScreenSizeToUseSubSections_RenderThread(float InComponentScreenSizeToUseSubSections);
 
-	virtual bool HeightfieldHasPendingStreaming() const override;
+	LANDSCAPE_API virtual bool HeightfieldHasPendingStreaming() const override;
 
-	virtual void GetHeightfieldRepresentation(UTexture2D*& OutHeightmapTexture, UTexture2D*& OutVisibilityTexture, FHeightfieldComponentDescription& OutDescription) const override;
+	LANDSCAPE_API virtual void GetHeightfieldRepresentation(UTexture2D*& OutHeightmapTexture, UTexture2D*& OutVisibilityTexture, FHeightfieldComponentDescription& OutDescription) const override;
 
-	virtual void GetLCIs(FLCIArray& LCIs) override;
+	LANDSCAPE_API virtual void GetLCIs(FLCIArray& LCIs) override;
 
 #if !(UE_BUILD_SHIPPING || UE_BUILD_TEST)
 	virtual int32 GetLightMapResolution() const override { return LightMapResolution; }
 #endif
 
 #if RHI_RAYTRACING
-	virtual void GetDynamicRayTracingInstances(FRayTracingMaterialGatheringContext& Context, TArray<FRayTracingInstance>& OutRayTracingInstances) override final;
+	LANDSCAPE_API virtual void GetDynamicRayTracingInstances(FRayTracingMaterialGatheringContext& Context, TArray<FRayTracingInstance>& OutRayTracingInstances) override final;
 	virtual bool HasRayTracingRepresentation() const override { return true; }
 	virtual bool IsRayTracingRelevant() const override { return true; }
 #endif
 
 	// FLandscapeSectionInfo interface
-	virtual float ComputeLODForView(const FSceneView& InView) const override;
-	virtual float ComputeLODBias() const override;
-	virtual void OnRenderCoordsChanged() override;
-	virtual int32 GetComponentResolution() const override;
+	LANDSCAPE_API virtual float ComputeLODForView(const FSceneView& InView) const override;
+	LANDSCAPE_API virtual float ComputeLODBias() const override;
+	LANDSCAPE_API virtual void OnRenderCoordsChanged() override;
+	LANDSCAPE_API virtual int32 GetComponentResolution() const override;
 
-	virtual double ComputeSectionResolution() const override;
-	virtual void GetSectionBoundsAndLocalToWorld(FBoxSphereBounds& LocalBounds, FMatrix& LocalToWorld) const override;
+	LANDSCAPE_API virtual double ComputeSectionResolution() const override;
+	LANDSCAPE_API virtual void GetSectionBoundsAndLocalToWorld(FBoxSphereBounds& LocalBounds, FMatrix& LocalToWorld) const override;
 };
 
 class FLandscapeDebugMaterialRenderProxy : public FMaterialRenderProxy
diff --git a/Engine/Source/Runtime/NavigationSystem/Public/NavLinkCustomInterface.h b/Engine/Source/Runtime/NavigationSystem/Public/NavLinkCustomInterface.h
--- a/Engine/Source/Runtime/NavigationSystem/Public/NavLinkCustomInterface.h
+++ b/Engine/Source/Runtime/NavigationSystem/Public/NavLinkCustomInterface.h
@@ -31,7 +31,7 @@ class UNavLinkCustomInterface : public UInterface
 	GENERATED_UINTERFACE_BODY()
 };
 
-class NAVIGATIONSYSTEM_API INavLinkCustomInterface
+class INavLinkCustomInterface
 {
 	GENERATED_IINTERFACE_BODY()
 
@@ -61,7 +61,7 @@ class NAVIGATIONSYSTEM_API INavLinkCustomInterface
 	virtual void UpdateLinkId(FNavLinkId NewUniqueId) {}
 
 	/** Get object owner of navigation link, used for creating containers with multiple links */
-	virtual UObject* GetLinkOwner() const;
+	NAVIGATIONSYSTEM_API virtual UObject* GetLinkOwner() const;
 
 	/** Check if link allows path finding
 	 *  Querier is usually an AIController trying to find path
@@ -77,22 +77,22 @@ class NAVIGATIONSYSTEM_API INavLinkCustomInterface
 	virtual void OnLinkMoveFinished(class UObject* PathComp) {}
 
 	UE_DEPRECATED(5.3, "LinkIds are now based on FNavLinkId using FNavLinkId::GenerateUniqueId(). This function will still generate an incremental Id however it does not work well in all circumstances.")
-	static uint32 GetUniqueId();
+	static NAVIGATIONSYSTEM_API uint32 GetUniqueId();
 
 	/** Helper function: bump unique ID numbers above given one */
 	UE_DEPRECATED(5.3, "LinkIds are now based on FNavLinkId. If your project is still using any of the old incremental Ids (saved in actors in levels or licensee code) then this function must be called still (typically by existing engine code), otherwise it is not necessary.")
-	static void UpdateUniqueId(FNavLinkId AlreadyUsedId);
+	static NAVIGATIONSYSTEM_API void UpdateUniqueId(FNavLinkId AlreadyUsedId);
 
 	UE_DEPRECATED(5.3, "LinkIds are now based on FNavLinkId. You may need to call the other version of this function that takes a FNavLinkId. This function has no effect.")
 	static void UpdateUniqueId(uint32 AlreadyUsedId) {}
 
 	/** Helper function: create modifier for navigation data export */
-	static FNavigationLink GetModifier(const INavLinkCustomInterface* CustomNavLink);
+	static NAVIGATIONSYSTEM_API FNavigationLink GetModifier(const INavLinkCustomInterface* CustomNavLink);
 	
 	UE_DEPRECATED(5.3, "LinkIds are now based on a FNavLinkId Hash. If your project is still using any of the old incremental Ids then this function must be called still (typically by existing engine code), otherwise it is not necessary.")
-	static void ResetUniqueId();
+	static NAVIGATIONSYSTEM_API void ResetUniqueId();
 
-	static void OnPreWorldInitialization(UWorld* World, const UWorld::InitializationValues IVS);
+	static NAVIGATIONSYSTEM_API void OnPreWorldInitialization(UWorld* World, const UWorld::InitializationValues IVS);
 
 	UE_DEPRECATED(5.3, "LinkIds are now based on FNavLinkId using FNavLinkId::GenerateUniqueId().")
 	static uint32 NextUniqueId;
diff --git a/Engine/Source/Runtime/StaticMeshDescription/Public/StaticMeshAttributes.h b/Engine/Source/Runtime/StaticMeshDescription/Public/StaticMeshAttributes.h
--- a/Engine/Source/Runtime/StaticMeshDescription/Public/StaticMeshAttributes.h
+++ b/Engine/Source/Runtime/StaticMeshDescription/Public/StaticMeshAttributes.h
@@ -69,7 +69,7 @@ namespace MeshAttribute
 
 PRAGMA_DISABLE_DEPRECATION_WARNINGS
 
-class STATICMESHDESCRIPTION_API FStaticMeshAttributes : public FMeshAttributes
+class FStaticMeshAttributes : public FMeshAttributes
 {
 public:
 
@@ -77,7 +77,7 @@ public:
 		: FMeshAttributes(InMeshDescription)
 	{}
 
-	virtual void Register(bool bKeepExistingAttribute = false) override;
+	STATICMESHDESCRIPTION_API virtual void Register(bool bKeepExistingAttribute = false) override;
 	
 	static bool IsReservedAttributeName(const FName InAttributeName)
 	{
@@ -96,9 +96,9 @@ public:
 	}	
 
 	UE_DEPRECATED(4.26, "Please use RegisterTriangleNormalAndTangentAttributes() instead.")
-	void RegisterPolygonNormalAndTangentAttributes();
+	STATICMESHDESCRIPTION_API void RegisterPolygonNormalAndTangentAttributes();
 
-	void RegisterTriangleNormalAndTangentAttributes();
+	STATICMESHDESCRIPTION_API void RegisterTriangleNormalAndTangentAttributes();
 
 	UE_DEPRECATED(4.26, "This attribute is no longer supported, please remove code pertaining to it.")
 	TVertexAttributesRef<float> GetVertexCornerSharpnesses() { return MeshDescription.VertexAttributes().GetAttributesRef<float>(MeshAttribute::Vertex::CornerSharpness); }
diff --git a/Engine/Source/Runtime/UMG/Public/Components/CheckBox.h b/Engine/Source/Runtime/UMG/Public/Components/CheckBox.h
--- a/Engine/Source/Runtime/UMG/Public/Components/CheckBox.h
+++ b/Engine/Source/Runtime/UMG/Public/Components/CheckBox.h
@@ -27,8 +27,8 @@ DECLARE_DYNAMIC_MULTICAST_DELEGATE_OneParam( FOnCheckBoxComponentStateChanged, b
  * * Single Child
  * * Toggle
  */
-UCLASS()
-class UMG_API UCheckBox : public UContentWidget
+UCLASS(MinimalAPI)
+class UCheckBox : public UContentWidget
 {
 	GENERATED_UCLASS_BODY()
 
@@ -82,82 +82,82 @@ public:
 
 	/** Returns true if this button is currently pressed */
 	UFUNCTION(BlueprintCallable, Category="Widget")
-	bool IsPressed() const;
+	UMG_API bool IsPressed() const;
 	
 	/** Returns true if the checkbox is currently checked */
 	UFUNCTION(BlueprintCallable, Category="Widget")
-	bool IsChecked() const;
+	UMG_API bool IsChecked() const;
 
 	/** Returns the full current checked state. */
 	UFUNCTION(BlueprintCallable, Category="Widget")
-	ECheckBoxState GetCheckedState() const;
+	UMG_API ECheckBoxState GetCheckedState() const;
 
 	/** Sets the checked state. */
 	UFUNCTION(BlueprintCallable, Category="Widget")
-	void SetIsChecked(bool InIsChecked);
+	UMG_API void SetIsChecked(bool InIsChecked);
 
 	/** Sets the checked state. */
 	UFUNCTION(BlueprintCallable, Category="Widget")
-	void SetCheckedState(ECheckBoxState InCheckedState);
+	UMG_API void SetCheckedState(ECheckBoxState InCheckedState);
 
 	/** Returns the local style. */
-	const FCheckBoxStyle& GetWidgetStyle() const;
+	UMG_API const FCheckBoxStyle& GetWidgetStyle() const;
 
 	/** Sets the style. */
-	void SetWidgetStyle(const FCheckBoxStyle& InStyle);
+	UMG_API void SetWidgetStyle(const FCheckBoxStyle& InStyle);
 
 	/** Returns the click method. */
-	EButtonClickMethod::Type GetClickMethod() const;
+	UMG_API EButtonClickMethod::Type GetClickMethod() const;
 
 	/** Sets the click method. */
 	UFUNCTION(BlueprintCallable, Category="Button")
-	void SetClickMethod(EButtonClickMethod::Type InClickMethod);
+	UMG_API void SetClickMethod(EButtonClickMethod::Type InClickMethod);
 
 	/** Returns the touch method. */
-	EButtonTouchMethod::Type GetTouchMethod() const;
+	UMG_API EButtonTouchMethod::Type GetTouchMethod() const;
 
 	/** Sets the touch method. */
 	UFUNCTION(BlueprintCallable, Category="Button")
-	void SetTouchMethod(EButtonTouchMethod::Type InTouchMethod);
+	UMG_API void SetTouchMethod(EButtonTouchMethod::Type InTouchMethod);
 
 	/** Returns the press method. */
-	EButtonPressMethod::Type GetPressMethod() const;
+	UMG_API EButtonPressMethod::Type GetPressMethod() const;
 
 	/** Sets the press method. */
 	UFUNCTION(BlueprintCallable, Category="Button")
-	void SetPressMethod(EButtonPressMethod::Type InPressMethod);
+	UMG_API void SetPressMethod(EButtonPressMethod::Type InPressMethod);
 
 	/** Is the checkbox focusable. */
-	bool GetIsFocusable() const;
+	UMG_API bool GetIsFocusable() const;
 
 public:
 	
 	//~ Begin UWidget Interface
-	virtual void SynchronizeProperties() override;
+	UMG_API virtual void SynchronizeProperties() override;
 	//~ End UWidget Interface
 
 	//~ Begin UVisual Interface
-	virtual void ReleaseSlateResources(bool bReleaseChildren) override;
+	UMG_API virtual void ReleaseSlateResources(bool bReleaseChildren) override;
 	//~ End UVisual Interface
 
 #if WITH_EDITOR
-	virtual const FText GetPaletteCategory() override;
+	UMG_API virtual const FText GetPaletteCategory() override;
 #endif
 
 protected:
 
 	// UPanelWidget
-	virtual void OnSlotAdded(UPanelSlot* Slot) override;
-	virtual void OnSlotRemoved(UPanelSlot* Slot) override;
+	UMG_API virtual void OnSlotAdded(UPanelSlot* Slot) override;
+	UMG_API virtual void OnSlotRemoved(UPanelSlot* Slot) override;
 	// End UPanelWidget
 
 	/** Initialize IsFocusable in the constructor before the SWidget is constructed. */
-	void InitIsFocusable(bool InIsFocusable);
+	UMG_API void InitIsFocusable(bool InIsFocusable);
 
-	void InitCheckedStateDelegate(FGetCheckBoxState InCheckedStateDelegate);
+	UMG_API void InitCheckedStateDelegate(FGetCheckBoxState InCheckedStateDelegate);
 protected:
 	//~ Begin UWidget Interface
-	virtual TSharedRef<SWidget> RebuildWidget() override;
+	UMG_API virtual TSharedRef<SWidget> RebuildWidget() override;
 #if WITH_EDITOR
 	virtual TSharedRef<SWidget> RebuildDesignWidget(TSharedRef<SWidget> Content) override
 	{
@@ -166,10 +166,10 @@ protected:
 #endif
 	//~ End UWidget Interface
 
-	void SlateOnCheckStateChangedCallback(ECheckBoxState NewState);
+	UMG_API void SlateOnCheckStateChangedCallback(ECheckBoxState NewState);
 
 #if WITH_ACCESSIBILITY
-	virtual TSharedPtr<SWidget> GetAccessibleWidget() const override;
+	UMG_API virtual TSharedPtr<SWidget> GetAccessibleWidget() const override;
 #endif
 	
 protected:
@@ -180,32 +180,32 @@ PRAGMA_DISABLE_DEPRECATION_WARNINGS
 PRAGMA_ENABLE_DEPRECATION_WARNINGS
 };
 
-UCLASS(Transient)
-class UMG_API UWidgetCheckedStateRegistration : public UWidgetEnumStateRegistration
+UCLASS(Transient, MinimalAPI)
+class UWidgetCheckedStateRegistration : public UWidgetEnumStateRegistration
 {
 	GENERATED_BODY()
 
 public:
 
 	/** Post-load initialized values corresponding to this enum state */
-	static inline FWidgetStateBitfield Unchecked;
-	static inline FWidgetStateBitfield Checked;
-	static inline FWidgetStateBitfield Undetermined;
+	static UMG_API inline FWidgetStateBitfield Unchecked;
+	static UMG_API inline FWidgetStateBitfield Checked;
+	static UMG_API inline FWidgetStateBitfield Undetermined;
 
 	static const inline FName StateName = FName("CheckedState");
 
 	//~ Begin UWidgetEnumStateRegistration Interface.
-	virtual FName GetStateName() const override;
-	virtual uint8 GetRegisteredWidgetState(const UWidget* InWidget) const override;
+	UMG_API virtual FName GetStateName() const override;
+	UMG_API virtual uint8 GetRegisteredWidgetState(const UWidget* InWidget) const override;
 	//~ End UWidgetEnumStateRegistration Interface
 
 	/** Convenience method to get widget state bitfield from enum value */
-	static const FWidgetStateBitfield& GetBitfieldFromValue(uint8 InValue);
+	static UMG_API const FWidgetStateBitfield& GetBitfieldFromValue(uint8 InValue);
 
 protected:
 	friend UWidgetStateSettings;
 
 	//~ Begin UWidgetEnumStateRegistration Interface.
-	virtual void InitializeStaticBitfields() const override;
+	UMG_API virtual void InitializeStaticBitfields() const override;
 	//~ End UWidgetEnumStateRegistration Interface
-};
\ No newline at end of file
+};
diff --git a/Engine/Source/Runtime/UMG/Public/Components/Image.h b/Engine/Source/Runtime/UMG/Public/Components/Image.h
--- a/Engine/Source/Runtime/UMG/Public/Components/Image.h
+++ b/Engine/Source/Runtime/UMG/Public/Components/Image.h
@@ -26,8 +26,8 @@ struct FStreamableHandle;
  *
  * * No Children
  */
-UCLASS()
-class UMG_API UImage : public UWidget
+UCLASS(MinimalAPI)
+class UImage : public UWidget
 {
 	GENERATED_UCLASS_BODY()
 
@@ -65,38 +65,38 @@ public:
 
 	/**  */
 	UFUNCTION(BlueprintCallable, Category="Appearance")
-	void SetColorAndOpacity(FLinearColor InColorAndOpacity);
+	UMG_API void SetColorAndOpacity(FLinearColor InColorAndOpacity);
 
-	const FLinearColor& GetColorAndOpacity() const;
+	UMG_API const FLinearColor& GetColorAndOpacity() const;
 
 	/**  */
 	UFUNCTION(BlueprintCallable, Category="Appearance")
-	void SetOpacity(float InOpacity);
+	UMG_API void SetOpacity(float InOpacity);
 
 	/**  */
 	UE_DEPRECATED(5.0, "Deprecated. Use SetDesiredSizeOverride instead.")
-	void SetBrushSize(FVector2D DesiredSize);
+	UMG_API void SetBrushSize(FVector2D DesiredSize);
 
 	UFUNCTION(BlueprintCallable, Category = "Appearance")
-	void SetDesiredSizeOverride(FVector2D DesiredSize);
+	UMG_API void SetDesiredSizeOverride(FVector2D DesiredSize);
 
 	/**  */
 	UFUNCTION(BlueprintCallable, Category = "Appearance")
-	void SetBrushTintColor(FSlateColor TintColor);
+	UMG_API void SetBrushTintColor(FSlateColor TintColor);
 
 	/**  */
 	UFUNCTION(BlueprintCallable, Category = "Appearance")
-	void SetBrushResourceObject(UObject* ResourceObject);
+	UMG_API void SetBrushResourceObject(UObject* ResourceObject);
 	
 	/**  */
 	UFUNCTION(BlueprintCallable, Category="Appearance")
-	virtual void SetBrush(const FSlateBrush& InBrush);
+	UMG_API virtual void SetBrush(const FSlateBrush& InBrush);
 
-	const FSlateBrush& GetBrush() const;
+	UMG_API const FSlateBrush& GetBrush() const;
 
 	/**  */
 	UFUNCTION(BlueprintCallable, Category="Appearance")
-	virtual void SetBrushFromAsset(USlateBrushAsset* Asset);
+	UMG_API virtual void SetBrushFromAsset(USlateBrushAsset* Asset);
 
 	/**
 	* Sets the Brush to the specified Texture.
@@ -105,7 +105,7 @@ public:
 	*	@param bMatchSize If true, image will change its size to texture size. If false, texture will be stretched to image size.
 	*/
 	UFUNCTION(BlueprintCallable, Category="Appearance")
-	virtual void SetBrushFromTexture(UTexture2D* Texture, bool bMatchSize = false);
+	UMG_API virtual void SetBrushFromTexture(UTexture2D* Texture, bool bMatchSize = false);
 
 	/**
 	* Sets the Brush to the specified Atlas Region.
@@ -114,7 +114,7 @@ public:
 	*	@param bMatchSize If true, image will change its size to atlas region size. If false, atlas region will be stretched to image size.
 	*/
 	UFUNCTION(BlueprintCallable, Category="Appearance")
-	virtual void SetBrushFromAtlasInterface(TScriptInterface<ISlateTextureAtlasInterface> AtlasRegion, bool bMatchSize = false);
+	UMG_API virtual void SetBrushFromAtlasInterface(TScriptInterface<ISlateTextureAtlasInterface> AtlasRegion, bool bMatchSize = false);
 
 	/**
 	* Sets the Brush to the specified Dynamic Texture.
@@ -123,11 +123,11 @@ public:
 	*	@param bMatchSize If true, image will change its size to texture size. If false, texture will be stretched to image size.
 	*/
 	UFUNCTION(BlueprintCallable, Category = "Appearance")
-	virtual void SetBrushFromTextureDynamic(UTexture2DDynamic* Texture, bool bMatchSize = false);
+	UMG_API virtual void SetBrushFromTextureDynamic(UTexture2DDynamic* Texture, bool bMatchSize = false);
 
 	/**  */
 	UFUNCTION(BlueprintCallable, Category="Appearance")
-	virtual void SetBrushFromMaterial(UMaterialInterface* Material);
+	UMG_API virtual void SetBrushFromMaterial(UMaterialInterface* Material);
 
 	/**
 	* Sets the Brush to the specified Soft Texture.
@@ -136,60 +136,60 @@ public:
 	*	@param bMatchSize If true, image will change its size to texture size. If false, texture will be stretched to image size.
 	*/
 	UFUNCTION(BlueprintCallable, Category="Appearance")
-	virtual void SetBrushFromSoftTexture(TSoftObjectPtr<UTexture2D> SoftTexture, bool bMatchSize = false);
+	UMG_API virtual void SetBrushFromSoftTexture(TSoftObjectPtr<UTexture2D> SoftTexture, bool bMatchSize = false);
 
 	/**  */
 	UFUNCTION(BlueprintCallable, Category="Appearance")
-	virtual void SetBrushFromSoftMaterial(TSoftObjectPtr<UMaterialInterface> SoftMaterial);
+	UMG_API virtual void SetBrushFromSoftMaterial(TSoftObjectPtr<UMaterialInterface> SoftMaterial);
 
 	/**  */
 	UFUNCTION(BlueprintCallable, Category="Appearance")
-	UMaterialInstanceDynamic* GetDynamicMaterial();
+	UMG_API UMaterialInstanceDynamic* GetDynamicMaterial();
 
-	void SetFlipForRightToLeftFlowDirection(bool InbFlipForRightToLeftFlowDirection);
+	UMG_API void SetFlipForRightToLeftFlowDirection(bool InbFlipForRightToLeftFlowDirection);
 
-	bool ShouldFlipForRightToLeftFlowDirection() const;
+	UMG_API bool ShouldFlipForRightToLeftFlowDirection() const;
 
 	//~ Begin UWidget Interface
-	virtual void SynchronizeProperties() override;
+	UMG_API virtual void SynchronizeProperties() override;
 	//~ End UWidget Interface
 
 	//~ Begin UVisual Interface
-	virtual void ReleaseSlateResources(bool bReleaseChildren) override;
+	UMG_API virtual void ReleaseSlateResources(bool bReleaseChildren) override;
 	//~ End UVisual Interface
 
 #if WITH_EDITOR
 	//~ Begin UWidget Interface
-	virtual const FText GetPaletteCategory() override;
+	UMG_API virtual const FText GetPaletteCategory() override;
 	//~ End UWidget Interface
 #endif
 
 protected:
 	//~ Begin UWidget Interface
-	virtual TSharedRef<SWidget> RebuildWidget() override;
+	UMG_API virtual TSharedRef<SWidget> RebuildWidget() override;
 	//~ End UWidget Interface
 
 	/** Translates the bound brush data and assigns it to the cached brush used by this widget. */
-	const FSlateBrush* ConvertImage(TAttribute<FSlateBrush> InImageAsset) const;
+	UMG_API const FSlateBrush* ConvertImage(TAttribute<FSlateBrush> InImageAsset) const;
 
 	// Called when we need to stream in content.
-	void RequestAsyncLoad(TSoftObjectPtr<UObject> SoftObject, TFunction<void()>&& Callback);
-	virtual void RequestAsyncLoad(TSoftObjectPtr<UObject> SoftObject, FStreamableDelegate DelegateToCall);
+	UMG_API void RequestAsyncLoad(TSoftObjectPtr<UObject> SoftObject, TFunction<void()>&& Callback);
+	UMG_API virtual void RequestAsyncLoad(TSoftObjectPtr<UObject> SoftObject, FStreamableDelegate DelegateToCall);
 
 	// Called when we need to abort the texture being streamed in.
-	virtual void CancelImageStreaming();
+	UMG_API virtual void CancelImageStreaming();
 
 	// Called when the image streaming starts, after the other one was cancelled.
-	virtual void OnImageStreamingStarted(TSoftObjectPtr<UObject> SoftObject);
+	UMG_API virtual void OnImageStreamingStarted(TSoftObjectPtr<UObject> SoftObject);
 
 	// Called when the image streaming completes.
-	virtual void OnImageStreamingComplete(TSoftObjectPtr<UObject> LoadedSoftObject);
+	UMG_API virtual void OnImageStreamingComplete(TSoftObjectPtr<UObject> LoadedSoftObject);
 
 	//
-	FReply HandleMouseButtonDown(const FGeometry& Geometry, const FPointerEvent& MouseEvent);
+	UMG_API FReply HandleMouseButtonDown(const FGeometry& Geometry, const FPointerEvent& MouseEvent);
 
 #if WITH_ACCESSIBILITY
-	virtual TSharedPtr<SWidget> GetAccessibleWidget() const override;
+	UMG_API virtual TSharedPtr<SWidget> GetAccessibleWidget() const override;
 #endif
 
 protected:
diff --git a/Engine/Source/Runtime/Windows/AudioMixerXAudio2/AudioMixerXAudio2.Build.cs b/Engine/Source/Runtime/Windows/AudioMixerXAudio2/AudioMixerXAudio2.Build.cs
--- a/Engine/Source/Runtime/Windows/AudioMixerXAudio2/AudioMixerXAudio2.Build.cs
+++ b/Engine/Source/Runtime/Windows/AudioMixerXAudio2/AudioMixerXAudio2.Build.cs
@@ -28,6 +28,7 @@ public class AudioMixerXAudio2 : ModuleRules
         PrivateDependencyModuleNames.AddRange(
 			new string[] {
 					"Core",
+					"CoreUObject",
 					"AudioMixer",
 					"AudioMixerCore"
                 }
