[Engine]

diff --git a/Engine/Source/Runtime/Experimental/Animation/Constraints/Public/ConstraintChannel.h b/Engine/Source/Runtime/Experimental/Animation/Constraints/Public/ConstraintChannel.h
--- a/Engine/Source/Runtime/Experimental/Animation/Constraints/Public/ConstraintChannel.h
+++ b/Engine/Source/Runtime/Experimental/Animation/Constraints/Public/ConstraintChannel.h
@@ -19,7 +19,7 @@
 #include "ConstraintChannel.generated.h"
 
 USTRUCT()
-struct CONSTRAINTS_API FMovieSceneConstraintChannel : public FMovieSceneBoolChannel
+struct FMovieSceneConstraintChannel : public FMovieSceneBoolChannel
 {
 	GENERATED_BODY()
 
@@ -52,7 +52,7 @@ struct CONSTRAINTS_API FMovieSceneConstraintChannel : public FMovieSceneBoolChan
 	 * @param OutValue   A value to receive the result
 	 * @return true if the channel was evaluated successfully, false otherwise
 	 */
-	virtual bool Evaluate(FFrameTime InTime, bool& OutValue) const override;
+	CONSTRAINTS_API virtual bool Evaluate(FFrameTime InTime, bool& OutValue) const override;
 
 #if WITH_EDITOR
 	using ExtraLabelFunction = TFunction< FString() >;
@@ -93,7 +93,7 @@ struct TMovieSceneChannelTraits<FMovieSceneConstraintChannel> : TMovieSceneChann
 
 
 USTRUCT()
-struct CONSTRAINTS_API FConstraintAndActiveChannel
+struct FConstraintAndActiveChannel
 {
 	GENERATED_USTRUCT_BODY()
 
@@ -110,4 +110,4 @@ struct CONSTRAINTS_API FConstraintAndActiveChannel
 
 	UPROPERTY()
 	TObjectPtr<UTickableConstraint> ConstraintCopyToSpawn;
-};
\ No newline at end of file
+};
diff --git a/Engine/Source/Runtime/Experimental/Animation/Constraints/Public/ConstraintsActor.h b/Engine/Source/Runtime/Experimental/Animation/Constraints/Public/ConstraintsActor.h
--- a/Engine/Source/Runtime/Experimental/Animation/Constraints/Public/ConstraintsActor.h
+++ b/Engine/Source/Runtime/Experimental/Animation/Constraints/Public/ConstraintsActor.h
@@ -12,32 +12,32 @@
 class UConstraintsManager;
 class UObject;
 
-UCLASS(NotPlaceable)
-class CONSTRAINTS_API AConstraintsActor : public AActor
+UCLASS(NotPlaceable, MinimalAPI)
+class AConstraintsActor : public AActor
 {
 	GENERATED_BODY()
 
 public:
 	// Sets default values for this actor's properties
-	AConstraintsActor(const FObjectInitializer& ObjectInitializer);
-	virtual ~AConstraintsActor();
+	CONSTRAINTS_API AConstraintsActor(const FObjectInitializer& ObjectInitializer);
+	CONSTRAINTS_API virtual ~AConstraintsActor();
 
-	virtual void BeginDestroy() override;
-	virtual void Destroyed() override;
+	CONSTRAINTS_API virtual void BeginDestroy() override;
+	CONSTRAINTS_API virtual void Destroyed() override;
 
-	virtual void PostRegisterAllComponents() override;
+	CONSTRAINTS_API virtual void PostRegisterAllComponents() override;
 
 #if WITH_EDITOR
-	virtual void PostEditUndo() override;
+	CONSTRAINTS_API virtual void PostEditUndo() override;
 #endif
 	
 protected:
 	// Called when the game starts or when spawned
-	virtual void BeginPlay() override;
+	CONSTRAINTS_API virtual void BeginPlay() override;
 
 public:
 	// Called every frame
-	virtual void Tick(float DeltaTime) override;
+	CONSTRAINTS_API virtual void Tick(float DeltaTime) override;
 
 	UPROPERTY()
 	TObjectPtr<UConstraintsManager> ConstraintsManager = nullptr;
diff --git a/Engine/Source/Runtime/Experimental/Animation/Constraints/Public/ConstraintsManager.h b/Engine/Source/Runtime/Experimental/Animation/Constraints/Public/ConstraintsManager.h
--- a/Engine/Source/Runtime/Experimental/Animation/Constraints/Public/ConstraintsManager.h
+++ b/Engine/Source/Runtime/Experimental/Animation/Constraints/Public/ConstraintsManager.h
@@ -19,31 +19,31 @@ class ULevel;
  **/
 
 USTRUCT()
-struct CONSTRAINTS_API FConstraintTickFunction : public FTickFunction
+struct FConstraintTickFunction : public FTickFunction
 {
 	GENERATED_USTRUCT_BODY()
 public:
-	FConstraintTickFunction();
-	~FConstraintTickFunction();
+	CONSTRAINTS_API FConstraintTickFunction();
+	CONSTRAINTS_API ~FConstraintTickFunction();
 
 	/* Begin FTickFunction Interface */
-	virtual void ExecuteTick(
+	CONSTRAINTS_API virtual void ExecuteTick(
 		float DeltaTime,
 		ELevelTick TickType,
 		ENamedThreads::Type CurrentThread,
 		const FGraphEventRef& MyCompletionGraphEvent) override;
 	
-	virtual FString DiagnosticMessage() override;
+	CONSTRAINTS_API virtual FString DiagnosticMessage() override;
 	/* End FTickFunction Interface */
 
 	/** Callable function that represents the actual constraint. **/
 	using ConstraintFunction = TFunction<void()>;
 	
 	/** Register a callable function. **/
-	void RegisterFunction(ConstraintFunction InConstraint);
+	CONSTRAINTS_API void RegisterFunction(ConstraintFunction InConstraint);
 	
 	/** Register a callable function. **/
-	void EvaluateFunctions() const;
+	CONSTRAINTS_API void EvaluateFunctions() const;
 
 	/** Weak ptr to the Constraint holding this tick function. **/
 	TWeakObjectPtr<UTickableConstraint> Constraint;
@@ -66,8 +66,8 @@ struct TStructOpsTypeTraits<FConstraintTickFunction> : public TStructOpsTypeTrai
  * Represents the basic interface of constraint within the constraints manager.
  **/
 
-UCLASS(Abstract, Blueprintable)
-class CONSTRAINTS_API UTickableConstraint : public UObject
+UCLASS(Abstract, Blueprintable, MinimalAPI)
+class UTickableConstraint : public UObject
 {
 	GENERATED_BODY()
 	
@@ -76,43 +76,43 @@ public:
 	virtual ~UTickableConstraint() {}
 
 	/** Returns the actual function that the tick function needs to evaluate. */
-	virtual FConstraintTickFunction::ConstraintFunction GetFunction() const PURE_VIRTUAL(GetFunction, return {};);
+	CONSTRAINTS_API virtual FConstraintTickFunction::ConstraintFunction GetFunction() const PURE_VIRTUAL(GetFunction, return {};);
 
 	/** Evaluates the constraint in a context where it's mot done thru the ConstraintTick's tick function. */
-	virtual void Evaluate(bool bTickHandlesAlso = false) const;
+	CONSTRAINTS_API virtual void Evaluate(bool bTickHandlesAlso = false) const;
 
 	/** Tick function that will be registered and evaluated. */
 	UPROPERTY()
 	FConstraintTickFunction ConstraintTick;
 
 	/** Sets the Active value and enable/disable the tick function. */
-	virtual void SetActive(const bool bIsActive);
+	CONSTRAINTS_API virtual void SetActive(const bool bIsActive);
 	
 	/** Get whether or not it's fully active, it's set to active and all pieces are set up,e.g. objects really exist that are being constrainted*/
-	virtual bool IsFullyActive() const;
+	CONSTRAINTS_API virtual bool IsFullyActive() const;
 	/** If true it contains objects bound to an external system, like sequencer so we don't do certain things, like remove constraints when they don't resolve*/
-	virtual bool HasBoundObjects() const PURE_VIRTUAL(HasBoundObjects, return false;);
+	CONSTRAINTS_API virtual bool HasBoundObjects() const PURE_VIRTUAL(HasBoundObjects, return false;);
 	/** Resolve the bound objects so that any object it references are resovled and correctly set up*/
-	virtual void ResolveBoundObjects(FMovieSceneSequenceID LocalSequenceID, IMovieScenePlayer& Player, UObject* SubObject = nullptr)  PURE_VIRTUAL(ResolveBoundObjects);
+	CONSTRAINTS_API virtual void ResolveBoundObjects(FMovieSceneSequenceID LocalSequenceID, IMovieScenePlayer& Player, UObject* SubObject = nullptr)  PURE_VIRTUAL(ResolveBoundObjects);
 
 	/** @todo document */
-	virtual uint32 GetTargetHash() const PURE_VIRTUAL(GetTargetHash, return 0;);
+	CONSTRAINTS_API virtual uint32 GetTargetHash() const PURE_VIRTUAL(GetTargetHash, return 0;);
 	/** Whether or not this object references this object. If it HasBoundObjects you should ResolveBoundObjects if you expect this to be up to date */
-	virtual bool ReferencesObject(TWeakObjectPtr<UObject> InObject) const PURE_VIRTUAL(ReferencesObject, return false;);
+	CONSTRAINTS_API virtual bool ReferencesObject(TWeakObjectPtr<UObject> InObject) const PURE_VIRTUAL(ReferencesObject, return false;);
 	/** Create duplicate with new Outer*/
-	virtual UTickableConstraint* Duplicate(UObject* NewOuter) const;
+	CONSTRAINTS_API virtual UTickableConstraint* Duplicate(UObject* NewOuter) const;
 
 #if WITH_EDITOR
 	/** Returns the constraint's label used for UI. */
-	virtual FString GetLabel() const;
-	virtual FString GetFullLabel() const;
+	CONSTRAINTS_API virtual FString GetLabel() const;
+	CONSTRAINTS_API virtual FString GetFullLabel() const;
 
 	/** Returns the constraint's type label used for UI. */
-	virtual FString GetTypeLabel() const;
+	CONSTRAINTS_API virtual FString GetTypeLabel() const;
 
 	// UObject interface
-	virtual void PostEditChangeProperty(FPropertyChangedEvent& PropertyChangedEvent) override;
-	virtual void PostEditUndo() override;
+	CONSTRAINTS_API virtual void PostEditChangeProperty(FPropertyChangedEvent& PropertyChangedEvent) override;
+	CONSTRAINTS_API virtual void PostEditUndo() override;
 	// End of UObject interface
 #endif
 	/** @todo documentation. */
@@ -126,8 +126,8 @@ public:
  * This object gathers the different constraints of the level and is held by the ConstraintsActor (unique in the level)
  **/
 
-UCLASS(BLUEPRINTABLE)
-class CONSTRAINTS_API UConstraintsManager : public UObject
+UCLASS(BLUEPRINTABLE, MinimalAPI)
+class UConstraintsManager : public UObject
 {
 	GENERATED_BODY()
 public:
@@ -139,27 +139,27 @@ public:
 	DECLARE_DYNAMIC_MULTICAST_SPARSE_DELEGATE_ThreeParams(FOnConstraintRemoved, UConstraintsManager, OnConstraintRemoved_BP, UConstraintsManager*, Mananger, UTickableConstraint*, Constraint, bool, bDoNotCompensate);
 
 	
-	UConstraintsManager();
-	virtual ~UConstraintsManager();
+	CONSTRAINTS_API UConstraintsManager();
+	CONSTRAINTS_API virtual ~UConstraintsManager();
 	
 	//UObjects
-	virtual void PostLoad() override;
+	CONSTRAINTS_API virtual void PostLoad() override;
 	/** Get the existing Constraints Manager if existing or create a new one. */
-	static UConstraintsManager* Get(UWorld* InWorld);
+	static CONSTRAINTS_API UConstraintsManager* Get(UWorld* InWorld);
 
 	/** Find the existing Constraints Manager. */
-	static UConstraintsManager* Find(const UWorld* InWorld);
+	static CONSTRAINTS_API UConstraintsManager* Find(const UWorld* InWorld);
 
 	/** @todo document */
-	void Init(UWorld* InWorld);
+	CONSTRAINTS_API void Init(UWorld* InWorld);
 	
 	/* Set tick dependencies between two constraints. */
-	void SetConstraintDependencies(
+	CONSTRAINTS_API void SetConstraintDependencies(
 		FConstraintTickFunction* InFunctionToTickBefore,
 		FConstraintTickFunction* InFunctionToTickAfter);
 
 	/** @todo document */
-	void Clear(UWorld* World);
+	CONSTRAINTS_API void Clear(UWorld* World);
 
 
 	/** BP Delegate fired when constraints are added*/
@@ -175,13 +175,13 @@ private:
 	/** @todo document */
 	FDelegateHandle OnActorDestroyedHandle;
 	
-	void OnActorDestroyed(AActor* InActor);
+	CONSTRAINTS_API void OnActorDestroyed(AActor* InActor);
 
-	void RegisterDelegates(UWorld* World);
-	void UnregisterDelegates(UWorld* World);
+	CONSTRAINTS_API void RegisterDelegates(UWorld* World);
+	CONSTRAINTS_API void UnregisterDelegates(UWorld* World);
 	
 	/** @todo document */
-	void Dump() const;
+	CONSTRAINTS_API void Dump() const;
 
 	/** All of the constraints*/
 	UPROPERTY()
@@ -206,68 +206,68 @@ enum class EConstraintsManagerNotifyType
 
 DECLARE_MULTICAST_DELEGATE_TwoParams(FConstraintsManagerNotifyDelegate, EConstraintsManagerNotifyType /* type */, UObject* /* subject */);
 
-class CONSTRAINTS_API FConstraintsManagerController
+class FConstraintsManagerController
 {
 public:
 	/** Static to control if we should remove constraints or not, may not want to do some when
 	compensating since we may be deleting spawnables while doing so*/
-	static bool bDoNotRemoveConstraint;
+	static CONSTRAINTS_API bool bDoNotRemoveConstraint;
 	/** Get and Set Controller acive in this world*/
-	static FConstraintsManagerController& Get(UWorld* InWorld);
+	static CONSTRAINTS_API FConstraintsManagerController& Get(UWorld* InWorld);
 
 	/** Allocates a new constraint with the constraints manager as the owner. */
 	template< typename TConstraint >
 	TConstraint* AllocateConstraintT(const FName& InBaseName) const;
 
 	/** Add this constraint to the manager */
-	bool AddConstraint(UTickableConstraint* InConstraint) const;
+	CONSTRAINTS_API bool AddConstraint(UTickableConstraint* InConstraint) const;
 	
 	/** Make a copy of this constraint and add it to the manager, returns the copy mode if it was added*/
-	UTickableConstraint* AddConstraintFromCopy(UTickableConstraint* CopyOfConstraint) const;
+	CONSTRAINTS_API UTickableConstraint* AddConstraintFromCopy(UTickableConstraint* CopyOfConstraint) const;
 
 	/** Get the index of the given constraint's name. */
-	int32 GetConstraintIndex(const FName& InConstraintName) const;
+	CONSTRAINTS_API int32 GetConstraintIndex(const FName& InConstraintName) const;
 	
 	/** Remove the constraint by name. */
-	bool RemoveConstraint(const FName& InConstraintName, bool bDoNotCompensate = false) const;
+	CONSTRAINTS_API bool RemoveConstraint(const FName& InConstraintName, bool bDoNotCompensate = false) const;
 
 	/** Remove the constraint at the given index. */
-	bool RemoveConstraint(const int32 InConstraintIndex, bool bDoNotCompensate = false) const;
+	CONSTRAINTS_API bool RemoveConstraint(const int32 InConstraintIndex, bool bDoNotCompensate = false) const;
 
 	/** Remove constraint by ptr*/
-	bool RemoveConstraint(UTickableConstraint* InConstraint, bool bDoNotCompensate = false) const;
+	CONSTRAINTS_API bool RemoveConstraint(UTickableConstraint* InConstraint, bool bDoNotCompensate = false) const;
 
 	/** Returns the constraint based on it's name within the manager's constraints array. */
-	UTickableConstraint* GetConstraint(const FName& InConstraintName) const;
+	CONSTRAINTS_API UTickableConstraint* GetConstraint(const FName& InConstraintName) const;
 
 	/** Returns the constraint based on it's index within the manager's constraints array. */
-	UTickableConstraint* GetConstraint(const int32 InConstraintIndex) const;
+	CONSTRAINTS_API UTickableConstraint* GetConstraint(const int32 InConstraintIndex) const;
 	
 	/** Get read-only access to the array of constraints. */
-	const TArray< TObjectPtr<UTickableConstraint> >& GetConstraintsArray() const;
+	CONSTRAINTS_API const TArray< TObjectPtr<UTickableConstraint> >& GetConstraintsArray() const;
 
 	/** Returns manager's constraints array (sorted if needed). */
-	TArray< TObjectPtr<UTickableConstraint> > GetAllConstraints(const bool bSorted = false) const;
+	CONSTRAINTS_API TArray< TObjectPtr<UTickableConstraint> > GetAllConstraints(const bool bSorted = false) const;
 
 	/** Returns a filtered constraints array checking if the predicate for each element is true. */
 	template <typename Predicate>
 	TArray< TObjectPtr<UTickableConstraint> > GetConstraintsByPredicate(Predicate Pred, const bool bSorted = false) const;
 
 	/** Get parent constraints of the specified child. If bSorted is true, then the constraints will be sorted by dependency. */
-	TArray< TObjectPtr<UTickableConstraint> > GetParentConstraints(
+	CONSTRAINTS_API TArray< TObjectPtr<UTickableConstraint> > GetParentConstraints(
 		const uint32 InTargetHash,
 		const bool bSorted = false) const;
 
 	/** Set dependencies between two constraints. */
-	void SetConstraintsDependencies(
+	CONSTRAINTS_API void SetConstraintsDependencies(
 		const FName& InNameToTickBefore,
 		const FName& InNameToTickAfter) const;
 
 	/** Go through each constraint in order and evaluate and tick them*/
-	void EvaluateAllConstraints() const;
+	CONSTRAINTS_API void EvaluateAllConstraints() const;
 
 	/** when PIEing/Simulating it's possible that the constraint isn't in the active manager but still lives*/
-	bool DoesExistInAnyWorld(UTickableConstraint* InConstraint);
+	CONSTRAINTS_API bool DoesExistInAnyWorld(UTickableConstraint* InConstraint);
 
 private:
 	/** Delegeate that's fired when a scene component is constrained, this is needed to make sure things like gizmo's get updated after the constraint tick happens*/
@@ -278,13 +278,13 @@ private:
 	FConstraintsManagerNotifyDelegate NotifyDelegate;
 	
 	/** Find the existing Constraints Manager in World or create a new one. */
-	UConstraintsManager* GetManager() const;
+	CONSTRAINTS_API UConstraintsManager* GetManager() const;
 	
 	/** Find the existing Constraints Manager in World. */
-	UConstraintsManager* FindManager() const;
+	CONSTRAINTS_API UConstraintsManager* FindManager() const;
 
 	/** Destroy the ConstraintsManager from the World. */
-	void DestroyManager() const;
+	CONSTRAINTS_API void DestroyManager() const;
 
 	/** The World that holds the ConstraintsManagerActor. */
 	UWorld* World = nullptr;
@@ -297,5 +297,5 @@ public:
 	FConstraintsManagerNotifyDelegate& GetNotifyDelegate() { return NotifyDelegate; }
 
 	/** Notify from changes in the constraints manager. */
-	void Notify(EConstraintsManagerNotifyType InNotifyType, UObject* InObject) const;
-};
\ No newline at end of file
+	CONSTRAINTS_API void Notify(EConstraintsManagerNotifyType InNotifyType, UObject* InObject) const;
+};
diff --git a/Engine/Source/Runtime/Experimental/Animation/Constraints/Public/ConstraintsScripting.h b/Engine/Source/Runtime/Experimental/Animation/Constraints/Public/ConstraintsScripting.h
--- a/Engine/Source/Runtime/Experimental/Animation/Constraints/Public/ConstraintsScripting.h
+++ b/Engine/Source/Runtime/Experimental/Animation/Constraints/Public/ConstraintsScripting.h
@@ -17,8 +17,8 @@ class UTickableTransformConstraint;
 /**
 * This is a set of helper functions to access various parts of the Sequencer and Control Rig API via Python and Blueprints.
 */
-UCLASS(meta = (Transient, ScriptName = "ConstraintsScriptingLibrary"))
-class CONSTRAINTS_API UConstraintsScriptingLibrary : public UBlueprintFunctionLibrary
+UCLASS(meta = (Transient, ScriptName = "ConstraintsScriptingLibrary"), MinimalAPI)
+class UConstraintsScriptingLibrary : public UBlueprintFunctionLibrary
 {
 
 public:
@@ -34,7 +34,7 @@ public:
 	* @return Returns the mananger
 	*/
 	UFUNCTION(BlueprintCallable, Category = "Editor Scripting | Sequencer Tools | Control Rig | Constraints")
-	static UConstraintsManager* GetManager(UWorld* InWorld);
+	static CONSTRAINTS_API UConstraintsManager* GetManager(UWorld* InWorld);
 
 	/**
 	* Create the transformable handle that deals with getting and setting transforms on this scene component
@@ -44,7 +44,7 @@ public:
 	* @return returns the handle for this scene component
 	**/
 	UFUNCTION(BlueprintCallable, Category = "Editor Scripting | Sequencer Tools | Control Rig | Constraints")
-	static UTransformableComponentHandle* CreateTransformableComponentHandle(UWorld* InWorld, USceneComponent* InSceneComponent, const FName& InSocketName);
+	static CONSTRAINTS_API UTransformableComponentHandle* CreateTransformableComponentHandle(UWorld* InWorld, USceneComponent* InSceneComponent, const FName& InSocketName);
 
 	/**
 	* Create the transformable handle that deals with getting and setting transforms on this object
@@ -54,7 +54,7 @@ public:
 	* @return returns the handle for this scene component
 	**/
 	UFUNCTION(BlueprintCallable, Category = "Editor Scripting | Sequencer Tools | Control Rig | Constraints")
-	static UTransformableHandle* CreateTransformableHandle(UWorld* InWorld, UObject* InObject, const FName& InAttachmentName = NAME_None);
+	static CONSTRAINTS_API UTransformableHandle* CreateTransformableHandle(UWorld* InWorld, UObject* InObject, const FName& InAttachmentName = NAME_None);
 
 	/**
 	* Create Constraint based on the specified type.
@@ -63,7 +63,7 @@ public:
 	* @return return The constraint object
 	*/
 	UFUNCTION(BlueprintCallable, Category = "Editor Scripting | Sequencer Tools | Control Rig | Constraints")
-	static UTickableTransformConstraint* CreateFromType(UWorld* InWorld, const ETransformConstraintType InType);
+	static CONSTRAINTS_API UTickableTransformConstraint* CreateFromType(UWorld* InWorld, const ETransformConstraintType InType);
 
 	/**
 	* Add Constraint to the system using the incoming parent and child handles with the specified type.
@@ -74,7 +74,7 @@ public:
 	* @return return If constraint added correctly
 	*/
 	UFUNCTION(BlueprintCallable, Category = "Editor Scripting | Sequencer Tools | Control Rig | Constraints")
-	static bool AddConstraint(UWorld* InWorld, UTransformableHandle* InParentHandle, UTransformableHandle* InChildHandle, UTickableTransformConstraint *InConstraint,
+	static CONSTRAINTS_API bool AddConstraint(UWorld* InWorld, UTransformableHandle* InParentHandle, UTransformableHandle* InChildHandle, UTickableTransformConstraint *InConstraint,
 			const bool bMaintainOffset);
 
 
@@ -83,7 +83,7 @@ public:
 	@return Copy of the constraints in the level
 	*/
 	UFUNCTION(BlueprintCallable, Category = "Editor Scripting | Sequencer Tools | Control Rig | Constraints")
-	static TArray<UTickableConstraint*> GetConstraintsArray(UWorld* InWorld);
+	static CONSTRAINTS_API TArray<UTickableConstraint*> GetConstraintsArray(UWorld* InWorld);
 
 
 	/**
@@ -93,7 +93,7 @@ public:
 	* @return return If constraint removed correctly
 	*/
 	UFUNCTION(BlueprintCallable, Category = "Editor Scripting | Sequencer Tools | Control Rig | Constraints")
-	static bool RemoveThisConstraint(UWorld* InWorld, UTickableConstraint* InTickableConstraint);
+	static CONSTRAINTS_API bool RemoveThisConstraint(UWorld* InWorld, UTickableConstraint* InTickableConstraint);
 
 	/**
 	* Remove constraint at specified index
@@ -102,6 +102,6 @@ public:
 	* @return return If constraint removed correctly
 	*/
 	UFUNCTION(BlueprintCallable, Category = "Editor Scripting | Sequencer Tools | Control Rig | Constraints")
-	static bool RemoveConstraint(UWorld* InWorld, int32 InIndex);
+	static CONSTRAINTS_API bool RemoveConstraint(UWorld* InWorld, int32 InIndex);
 
-};
\ No newline at end of file
+};
diff --git a/Engine/Source/Runtime/Experimental/Animation/Constraints/Public/TransformConstraint.h b/Engine/Source/Runtime/Experimental/Animation/Constraints/Public/TransformConstraint.h
--- a/Engine/Source/Runtime/Experimental/Animation/Constraints/Public/TransformConstraint.h
+++ b/Engine/Source/Runtime/Experimental/Animation/Constraints/Public/TransformConstraint.h
@@ -21,45 +21,45 @@ enum class EMovieSceneTransformChannel : uint32;
  * UTickableTransformConstraint
  **/
 
-UCLASS(Abstract, Blueprintable)
-class CONSTRAINTS_API UTickableTransformConstraint : public UTickableConstraint
+UCLASS(Abstract, Blueprintable, MinimalAPI)
+class UTickableTransformConstraint : public UTickableConstraint
 {
 	GENERATED_BODY()
 
 public:
 
 	/** Sets up the constraint so that the initial offset is set and dependencies and handles managed. */
-	void Setup();
+	CONSTRAINTS_API void Setup();
 	
 	/** UObjects overrides. */
-	virtual void PostLoad() override;
-	virtual void PostDuplicate(bool bDuplicateForPIE) override;
+	CONSTRAINTS_API virtual void PostLoad() override;
+	CONSTRAINTS_API virtual void PostDuplicate(bool bDuplicateForPIE) override;
 
 	/** Returns the target hash value (i.e. the child handle's hash). */
-	virtual uint32 GetTargetHash() const override;
+	CONSTRAINTS_API virtual uint32 GetTargetHash() const override;
 	
 	/** Test whether an InObject is referenced by that constraint. (i.e. is it's parent or child). */
-	virtual bool ReferencesObject(TWeakObjectPtr<UObject> InObject) const override;
+	CONSTRAINTS_API virtual bool ReferencesObject(TWeakObjectPtr<UObject> InObject) const override;
 	/** If true it contains objects bound to an external system, like sequencer so we don't do certain things, like remove constraints when they don't resolve*/
-	virtual bool HasBoundObjects() const override;
+	CONSTRAINTS_API virtual bool HasBoundObjects() const override;
 	
 	/** Resolve the bound objects so that any object it references are resovled and correctly set up*/
-	virtual void ResolveBoundObjects(FMovieSceneSequenceID LocalSequenceID, IMovieScenePlayer& Player,UObject* SubObject) override;
+	CONSTRAINTS_API virtual void ResolveBoundObjects(FMovieSceneSequenceID LocalSequenceID, IMovieScenePlayer& Player,UObject* SubObject) override;
 
 	/** If Active and the handles and targets are valid*/
-	virtual bool IsFullyActive() const override;
+	CONSTRAINTS_API virtual bool IsFullyActive() const override;
 
 	/** If that constraint needs to be handled by the compensation system. */
-	virtual bool NeedsCompensation() const;
+	CONSTRAINTS_API virtual bool NeedsCompensation() const;
 
 	/** Create duplicate with new Outer*/
-	virtual UTickableConstraint* Duplicate(UObject* NewOuter) const override;
+	CONSTRAINTS_API virtual UTickableConstraint* Duplicate(UObject* NewOuter) const override;
 
 	/** Override the evaluate so we can tick our handles*/
-	virtual void Evaluate(bool bTickHandlesAlso = false) const override;
+	CONSTRAINTS_API virtual void Evaluate(bool bTickHandlesAlso = false) const override;
 
 	/** Sets the Active value and enable/disable the tick function. */
-	virtual void SetActive(const bool bIsActive) override;
+	CONSTRAINTS_API virtual void SetActive(const bool bIsActive) override;
 	
 	/** The transformable handle representing the parent of that constraint. */
 	UPROPERTY(BlueprintReadWrite, Category = "Handle")
@@ -84,88 +84,88 @@ public:
 	bool bDynamicOffset = false;
 
 	/** Returns the constraint type (Position, Parent, Aim...). */
-	int64 GetType() const;
+	CONSTRAINTS_API int64 GetType() const;
 
 	/** Get the current child's global transform. */
-	FTransform GetChildGlobalTransform() const;
+	CONSTRAINTS_API FTransform GetChildGlobalTransform() const;
 	
 	/** Get the current child's local transform. */
-	FTransform GetChildLocalTransform() const;
+	CONSTRAINTS_API FTransform GetChildLocalTransform() const;
 
 	/** Get the current parent's global transform. */
-	FTransform GetParentGlobalTransform() const;
+	CONSTRAINTS_API FTransform GetParentGlobalTransform() const;
 	
 	/** Get the current parent's local transform. */
-	FTransform GetParentLocalTransform() const;
+	CONSTRAINTS_API FTransform GetParentLocalTransform() const;
 
 	/** Returns the channels to key based on the constraint's type. */
-	EMovieSceneTransformChannel GetChannelsToKey() const;
+	CONSTRAINTS_API EMovieSceneTransformChannel GetChannelsToKey() const;
 
 	/**
 	 * Manages changes on the child/parent transformable handle. This can be used to update internal data (i.e. offset)
 	 * when transform changes outside of that system and need to trigger changes within the constraint itself.   
 	*/
-	virtual void OnHandleModified(UTransformableHandle* InHandle, EHandleEvent InEvent);
+	CONSTRAINTS_API virtual void OnHandleModified(UTransformableHandle* InHandle, EHandleEvent InEvent);
 
 	/** Returns the handles tick function (ensuring it lives in the same world). */
-	FTickFunction* GetChildHandleTickFunction() const;
-	FTickFunction* GetParentHandleTickFunction() const;
+	CONSTRAINTS_API FTickFunction* GetChildHandleTickFunction() const;
+	CONSTRAINTS_API FTickFunction* GetParentHandleTickFunction() const;
 
 	/**
 	* Sets up dependencies with the first primary prerequisite available if the parent does not tick.   
 	*/
-	void EnsurePrimaryDependency();
+	CONSTRAINTS_API void EnsurePrimaryDependency();
 	
 protected:
 
 	/** Registers/Unregisters useful delegates for both child and parent handles. */
-	void UnregisterDelegates() const;
-	void RegisterDelegates();
+	CONSTRAINTS_API void UnregisterDelegates() const;
+	CONSTRAINTS_API void RegisterDelegates();
 
 	
 	/**
 	 * Computes the initial offset that is needed to keep the child's global transform unchanged when creating the
 	 * constraint. This can be called whenever necessary to update the current offset.
 	*/
-	virtual void ComputeOffset() PURE_VIRTUAL(ComputeOffset, return;);
+	CONSTRAINTS_API virtual void ComputeOffset() PURE_VIRTUAL(ComputeOffset, return;);
 	
 	/**
 	 * Sets up dependencies between the parent, the constraint and the child using their respective tick functions.
 	 * It creates a dependency graph between them so that they tick in the right order when evaluated.   
 	*/
-	void SetupDependencies();
+	CONSTRAINTS_API void SetupDependencies();
 
 	/** Set the current child's global transform. */
-	void SetChildGlobalTransform(const FTransform& InGlobal) const;
+	CONSTRAINTS_API void SetChildGlobalTransform(const FTransform& InGlobal) const;
 	
 	/** Set the current child's local transform. */
-	void SetChildLocalTransform(const FTransform& InLocal) const;
+	CONSTRAINTS_API void SetChildLocalTransform(const FTransform& InLocal) const;
 
 	/** Defines the constraint's type (Position, Parent, Aim...). */
 	UPROPERTY()
 	ETransformConstraintType Type = ETransformConstraintType::Parent;
 
 	/** Handle active state modification if needed. */
-	void OnActiveStateChanged() const;
+	CONSTRAINTS_API void OnActiveStateChanged() const;
 
 	/** Returns the handle's tick function (ensuring it lives in the same world). */
-	FTickFunction* GetHandleTickFunction(const TObjectPtr<UTransformableHandle>& InHandle) const;
+	CONSTRAINTS_API FTickFunction* GetHandleTickFunction(const TObjectPtr<UTransformableHandle>& InHandle) const;
 
 	/** (Re-)Registers the constraint function and (re-)binds the required delegates*/
-	void InitConstraint();
+	CONSTRAINTS_API void InitConstraint();
 
 #if WITH_EDITOR
 public:
 	/** Returns the constraint's label used for UI. */
-	virtual FString GetLabel() const override;
-	virtual FString GetFullLabel() const override;
+	CONSTRAINTS_API virtual FString GetLabel() const override;
+	CONSTRAINTS_API virtual FString GetFullLabel() const override;
 
 	/** Returns the constraint's type label used for UI. */
-	virtual FString GetTypeLabel() const override;
+	CONSTRAINTS_API virtual FString GetTypeLabel() const override;
 
 	// UObject interface
-	virtual void PostEditChangeProperty(FPropertyChangedEvent& PropertyChangedEvent) override;
-	virtual void PostEditUndo() override;
+	CONSTRAINTS_API virtual void PostEditChangeProperty(FPropertyChangedEvent& PropertyChangedEvent) override;
+	CONSTRAINTS_API virtual void PostEditUndo() override;
 	// End of UObject interface
 
 	/** Returns a delegate that can be used to monitor for property changes. This can be used to monitor constraints changes
@@ -173,10 +173,10 @@ public:
 	 * Note that bScaling is currently the only property change we monitor but this can be used for other properties.
 	 */
 	DECLARE_MULTICAST_DELEGATE_TwoParams(FOnConstraintChanged, UTickableTransformConstraint*, const FPropertyChangedEvent&);
-	static FOnConstraintChanged& GetOnConstraintChanged();
+	static CONSTRAINTS_API FOnConstraintChanged& GetOnConstraintChanged();
 
 protected:
-	static FOnConstraintChanged OnConstraintChanged;
+	static CONSTRAINTS_API FOnConstraintChanged OnConstraintChanged;
 #endif
 };
 
@@ -184,19 +184,19 @@ protected:
  * UTickableTranslationConstraint
  **/
 
-UCLASS(Blueprintable)
-class CONSTRAINTS_API UTickableTranslationConstraint : public UTickableTransformConstraint
+UCLASS(Blueprintable, MinimalAPI)
+class UTickableTranslationConstraint : public UTickableTransformConstraint
 {
 	GENERATED_BODY()
 
 public:
-	UTickableTranslationConstraint();
+	CONSTRAINTS_API UTickableTranslationConstraint();
 
 	/** Returns the position constraint function that the tick function will evaluate. */
-	virtual FConstraintTickFunction::ConstraintFunction GetFunction() const override;
+	CONSTRAINTS_API virtual FConstraintTickFunction::ConstraintFunction GetFunction() const override;
 
 	/** Updates the dynamic offset based on external child's transform changes. */
-	virtual void OnHandleModified(UTransformableHandle* InHandle, EHandleEvent InEvent) override;
+	CONSTRAINTS_API virtual void OnHandleModified(UTransformableHandle* InHandle, EHandleEvent InEvent) override;
 
 protected:
 
@@ -206,10 +206,10 @@ protected:
 		uint32 CachedInputHash = 0;
 	};
 	mutable FDynamicCache Cache;
-	uint32 CalculateInputHash() const;
+	CONSTRAINTS_API uint32 CalculateInputHash() const;
 	
 	/** Computes the child's local translation offset in the parent space. */
-	virtual void ComputeOffset() override;
+	CONSTRAINTS_API virtual void ComputeOffset() override;
 
 public:
 	/** Defines the local child's translation offset in the parent space. */
@@ -224,7 +224,7 @@ protected:
 #if WITH_EDITOR
 public:
 	// UObject interface
-	virtual void PostEditChangeProperty(FPropertyChangedEvent& PropertyChangedEvent) override;
+	CONSTRAINTS_API virtual void PostEditChangeProperty(FPropertyChangedEvent& PropertyChangedEvent) override;
 	// End of UObject interface
 #endif
 };
@@ -233,19 +233,19 @@ public:
  * UTickableRotationConstraint
  **/
 
-UCLASS(Blueprintable)
-class CONSTRAINTS_API UTickableRotationConstraint : public UTickableTransformConstraint
+UCLASS(Blueprintable, MinimalAPI)
+class UTickableRotationConstraint : public UTickableTransformConstraint
 {
 	GENERATED_BODY()
 
 public:
-	UTickableRotationConstraint();
+	CONSTRAINTS_API UTickableRotationConstraint();
 
 	/** Returns the rotation constraint function that the tick function will evaluate. */
-	virtual FConstraintTickFunction::ConstraintFunction GetFunction() const override;
+	CONSTRAINTS_API virtual FConstraintTickFunction::ConstraintFunction GetFunction() const override;
 
 	/** Updates the dynamic offset based on external child's transform changes. */
-	virtual void OnHandleModified(UTransformableHandle* InHandle, EHandleEvent InEvent) override;
+	CONSTRAINTS_API virtual void OnHandleModified(UTransformableHandle* InHandle, EHandleEvent InEvent) override;
 
 protected:
 
@@ -255,10 +255,10 @@ protected:
 		uint32 CachedInputHash = 0;
 	};
 	mutable FDynamicCache Cache;
-	uint32 CalculateInputHash() const;
+	CONSTRAINTS_API uint32 CalculateInputHash() const;
 	
 	/** Computes the child's local rotation offset in the parent space. */
-	virtual void ComputeOffset() override;
+	CONSTRAINTS_API virtual void ComputeOffset() override;
 
 public:
 	/** Defines the local child's rotation offset in the parent space. */
@@ -273,7 +273,7 @@ protected:
 #if WITH_EDITOR
 public:
 	// UObject interface
-	virtual void PostEditChangeProperty(FPropertyChangedEvent& PropertyChangedEvent) override;
+	CONSTRAINTS_API virtual void PostEditChangeProperty(FPropertyChangedEvent& PropertyChangedEvent) override;
 	// End of UObject interface
 #endif
 };
@@ -282,23 +282,23 @@ public:
  * UTickableScaleConstraint
  **/
 
-UCLASS(Blueprintable)
-class CONSTRAINTS_API UTickableScaleConstraint : public UTickableTransformConstraint
+UCLASS(Blueprintable, MinimalAPI)
+class UTickableScaleConstraint : public UTickableTransformConstraint
 {
 	GENERATED_BODY()
 
 public:
-	UTickableScaleConstraint();
+	CONSTRAINTS_API UTickableScaleConstraint();
 	
 	/** Returns the scale constraint function that the tick function will evaluate. */
-	virtual FConstraintTickFunction::ConstraintFunction GetFunction() const override;
+	CONSTRAINTS_API virtual FConstraintTickFunction::ConstraintFunction GetFunction() const override;
 
 	/** Updates the dynamic offset based on external child's transform changes. */
-	virtual void OnHandleModified(UTransformableHandle* InHandle, EHandleEvent InEvent) override;
+	CONSTRAINTS_API virtual void OnHandleModified(UTransformableHandle* InHandle, EHandleEvent InEvent) override;
 	
 protected:
 	/** Computes the child's local scale offset in the parent space. */
-	virtual void ComputeOffset() override;
+	CONSTRAINTS_API virtual void ComputeOffset() override;
 
 	/** Cache data structure to store last child local/global transform. */
 	struct FDynamicCache
@@ -306,7 +306,7 @@ protected:
 		uint32 CachedInputHash = 0;
 	};
 	mutable FDynamicCache Cache;
-	uint32 CalculateInputHash() const;
+	CONSTRAINTS_API uint32 CalculateInputHash() const;
 
 public:
 	/** Defines the local child's scale offset in the parent space. */
@@ -321,7 +321,7 @@ protected:
 #if WITH_EDITOR
 public:
 	// UObject interface
-	virtual void PostEditChangeProperty(FPropertyChangedEvent& PropertyChangedEvent) override;
+	CONSTRAINTS_API virtual void PostEditChangeProperty(FPropertyChangedEvent& PropertyChangedEvent) override;
 	// End of UObject interface
 #endif
 };
@@ -330,16 +330,16 @@ public:
  * UTickableParentConstraint
  **/
 
-UCLASS(Blueprintable)
-class CONSTRAINTS_API UTickableParentConstraint : public UTickableTransformConstraint
+UCLASS(Blueprintable, MinimalAPI)
+class UTickableParentConstraint : public UTickableTransformConstraint
 {
 	GENERATED_BODY()
 
 public:
-	UTickableParentConstraint();
+	CONSTRAINTS_API UTickableParentConstraint();
 	
 	/** Returns the transform constraint function that the tick function will evaluate. */
-	virtual FConstraintTickFunction::ConstraintFunction GetFunction() const override;
+	CONSTRAINTS_API virtual FConstraintTickFunction::ConstraintFunction GetFunction() const override;
 
 	bool IsScalingEnabled() const
 	{
@@ -354,7 +354,7 @@ public:
 	static FName GetScalingPropertyName() { return GET_MEMBER_NAME_CHECKED(UTickableParentConstraint, bScaling); }
 
 	/** Updates the dynamic offset based on external child's transform changes. */
-	virtual void OnHandleModified(UTransformableHandle* InHandle, EHandleEvent InEvent) override;
+	CONSTRAINTS_API virtual void OnHandleModified(UTransformableHandle* InHandle, EHandleEvent InEvent) override;
 
 protected:
 	/** Cache data structure to store last child local/global transform. */
@@ -363,10 +363,10 @@ protected:
 		uint32 CachedInputHash = 0;
 	};
 	mutable FDynamicCache Cache;
-	uint32 CalculateInputHash() const;
+	CONSTRAINTS_API uint32 CalculateInputHash() const;
 
 	/** Computes the child's local transform offset in the parent space. */
-	virtual void ComputeOffset() override;
+	CONSTRAINTS_API virtual void ComputeOffset() override;
 
 public:
 	/** Defines the local child's transform offset in the parent space. */
@@ -385,7 +385,7 @@ protected:
 #if WITH_EDITOR
 public:
 	// UObject interface
-	virtual void PostEditChangeProperty(FPropertyChangedEvent& PropertyChangedEvent) override;
+	CONSTRAINTS_API virtual void PostEditChangeProperty(FPropertyChangedEvent& PropertyChangedEvent) override;
 	// End of UObject interface
 #endif
 };
@@ -394,25 +394,25 @@ public:
  * UTickableLookAtConstraint
  **/
 
-UCLASS(Blueprintable)
-class CONSTRAINTS_API UTickableLookAtConstraint : public UTickableTransformConstraint
+UCLASS(Blueprintable, MinimalAPI)
+class UTickableLookAtConstraint : public UTickableTransformConstraint
 {
 	GENERATED_BODY()
 
 public:
-	UTickableLookAtConstraint();
+	CONSTRAINTS_API UTickableLookAtConstraint();
 	
 	/** Returns the look at constraint function that the tick function will evaluate. */
-	virtual FConstraintTickFunction::ConstraintFunction GetFunction() const override;
+	CONSTRAINTS_API virtual FConstraintTickFunction::ConstraintFunction GetFunction() const override;
 
 	/** If that constraint needs to be handled by the compensation system. */
-	virtual bool NeedsCompensation() const override;
+	CONSTRAINTS_API virtual bool NeedsCompensation() const override;
 
 protected:
 	/**
      * Computes the initial axis that is needed to keep the child's orientation unchanged when creating the constraint.
     */
-	virtual void ComputeOffset() override;
+	CONSTRAINTS_API virtual void ComputeOffset() override;
 
 	/** Defines the aiming axis. */
 	UPROPERTY(EditAnywhere, BlueprintReadWrite, Category="Axis")
@@ -426,7 +426,7 @@ private:
 #if WITH_EDITOR
 public:
 	// UObject interface
-	virtual void PostEditChangeProperty(FPropertyChangedEvent& PropertyChangedEvent) override;
+	CONSTRAINTS_API virtual void PostEditChangeProperty(FPropertyChangedEvent& PropertyChangedEvent) override;
 	// End of UObject interface
 #endif
 };
@@ -435,27 +435,27 @@ public:
  * TransformConstraintUtils
  **/
 
-struct CONSTRAINTS_API FTransformConstraintUtils
+struct FTransformConstraintUtils
 {
 	/** Fills a sorted constraint array that InChild actor is the child of. */
-	static void GetParentConstraints(
+	static CONSTRAINTS_API void GetParentConstraints(
 		UWorld* World,
 		const AActor* InChild,
 		TArray< TObjectPtr<UTickableConstraint> >& OutConstraints);
 
 	/** Create a handle for the scene component.*/
-	static UTransformableComponentHandle* CreateHandleForSceneComponent(
+	static CONSTRAINTS_API UTransformableComponentHandle* CreateHandleForSceneComponent(
 		USceneComponent* InSceneComponent,
 		const FName& InSocketName,
 		UObject* Outer);
 
 	/** Creates a new transform constraint based on the InType. */
-	static UTickableTransformConstraint* CreateFromType(
+	static CONSTRAINTS_API UTickableTransformConstraint* CreateFromType(
 		UWorld* InWorld,
 		const ETransformConstraintType InType);
 
 	/** Creates respective handles and creates a new InType transform constraint. */	
-	static UTickableTransformConstraint* CreateAndAddFromObjects(
+	static CONSTRAINTS_API UTickableTransformConstraint* CreateAndAddFromObjects(
 		UWorld* InWorld,
 		UObject* InParent, const FName& InParentSocketName,
 		UObject* InChild, const FName& InChildSocketName,
@@ -463,7 +463,7 @@ struct CONSTRAINTS_API FTransformConstraintUtils
 		const bool bMaintainOffset = true);
 
 	/** Registers a new transform constraint within the constraints manager. */	
-	static bool AddConstraint(
+	static CONSTRAINTS_API bool AddConstraint(
 		UWorld* InWorld,
 		UTransformableHandle* InParentHandle,
 		UTransformableHandle* InChildHandle,
@@ -471,23 +471,23 @@ struct CONSTRAINTS_API FTransformConstraintUtils
 		const bool bMaintainOffset = true);
 
 	/** Computes the relative transform between both transform based on the constraint's InType. */
-	static FTransform ComputeRelativeTransform(
+	static CONSTRAINTS_API FTransform ComputeRelativeTransform(
 		const FTransform& InChildLocal,
 		const FTransform& InChildWorld,
 		const FTransform& InSpaceWorld,
 		const UTickableTransformConstraint* InConstraint);
 
 	/** Computes the current constraint space local transform. */
-	static TOptional<FTransform> GetRelativeTransform(UWorld* InWorld, const uint32 InHandleHash);
-	static TOptional<FTransform> GetConstraintsRelativeTransform(
+	static CONSTRAINTS_API TOptional<FTransform> GetRelativeTransform(UWorld* InWorld, const uint32 InHandleHash);
+	static CONSTRAINTS_API TOptional<FTransform> GetConstraintsRelativeTransform(
 		const TArray< TObjectPtr<UTickableConstraint> >& InConstraints,
 		const FTransform& InChildLocal, const FTransform& InChildWorld);
 
 	/** Get the last active constraint that has dynamic offset. */
-	static int32 GetLastActiveConstraintIndex(const TArray< TObjectPtr<UTickableConstraint> >& InConstraints);
+	static CONSTRAINTS_API int32 GetLastActiveConstraintIndex(const TArray< TObjectPtr<UTickableConstraint> >& InConstraints);
 
 	/** Fills a constraint array that InParentHandle is the parent of. */
-	static void GetChildrenConstraints(
+	static CONSTRAINTS_API void GetChildrenConstraints(
 		UWorld* World,
 		const UTransformableHandle* InParentHandle,
 		TArray< TObjectPtr<UTickableConstraint> >& OutConstraints);
diff --git a/Engine/Source/Runtime/Experimental/Animation/Constraints/Public/TransformableHandle.h b/Engine/Source/Runtime/Experimental/Animation/Constraints/Public/TransformableHandle.h
--- a/Engine/Source/Runtime/Experimental/Animation/Constraints/Public/TransformableHandle.h
+++ b/Engine/Source/Runtime/Experimental/Animation/Constraints/Public/TransformableHandle.h
@@ -35,8 +35,8 @@ enum class EHandleEvent : uint8
  * UTransformableHandle
  */
 
-UCLASS(Abstract, Blueprintable)
-class CONSTRAINTS_API UTransformableHandle : public UObject 
+UCLASS(Abstract, Blueprintable, MinimalAPI)
+class UTransformableHandle : public UObject 
 {
 	GENERATED_BODY()
 	
@@ -44,41 +44,41 @@ public:
 
 	DECLARE_EVENT_TwoParams(UTransformableHandle, FHandleModifiedEvent, UTransformableHandle*, EHandleEvent);
 	
-	virtual ~UTransformableHandle();
+	CONSTRAINTS_API virtual ~UTransformableHandle();
 	
 	/** Sanity check to ensure the handle is safe to use. */
-	virtual bool IsValid() const PURE_VIRTUAL(IsValid, return false;);
+	CONSTRAINTS_API virtual bool IsValid() const PURE_VIRTUAL(IsValid, return false;);
 	
 	/** Sets the global transform of the underlying transformable object. */
-	virtual void SetGlobalTransform(const FTransform& InGlobal) const PURE_VIRTUAL(SetGlobalTransform, );
+	CONSTRAINTS_API virtual void SetGlobalTransform(const FTransform& InGlobal) const PURE_VIRTUAL(SetGlobalTransform, );
 	/** Sets the local transform of the underlying transformable object in it's parent space. */
-	virtual void SetLocalTransform(const FTransform& InLocal) const PURE_VIRTUAL(SetLocalTransform, );
+	CONSTRAINTS_API virtual void SetLocalTransform(const FTransform& InLocal) const PURE_VIRTUAL(SetLocalTransform, );
 	/** Gets the global transform of the underlying transformable object. */
-	virtual FTransform GetGlobalTransform() const PURE_VIRTUAL(GetGlobalTransform, return FTransform::Identity;);
+	CONSTRAINTS_API virtual FTransform GetGlobalTransform() const PURE_VIRTUAL(GetGlobalTransform, return FTransform::Identity;);
 	/** Gets the local transform of the underlying transformable object in it's parent space. */
-	virtual FTransform GetLocalTransform() const PURE_VIRTUAL(GetLocalTransform, return FTransform::Identity;);
+	CONSTRAINTS_API virtual FTransform GetLocalTransform() const PURE_VIRTUAL(GetLocalTransform, return FTransform::Identity;);
 
 	/** If true it contains objects bound to an external system, like sequencer so we don't do certain things, like remove constraints when they don't resolve*/
-	virtual bool HasBoundObjects() const;
+	CONSTRAINTS_API virtual bool HasBoundObjects() const;
 
 	/** Resolve the bound objects so that any object it references are resovled and correctly set up*/
-	virtual void ResolveBoundObjects(FMovieSceneSequenceID LocalSequenceID, IMovieScenePlayer& Player, UObject* SubObject = nullptr) PURE_VIRTUAL(ResolveBoundObjects);
+	CONSTRAINTS_API virtual void ResolveBoundObjects(FMovieSceneSequenceID LocalSequenceID, IMovieScenePlayer& Player, UObject* SubObject = nullptr) PURE_VIRTUAL(ResolveBoundObjects);
 
 	/** Make a duplicate of myself with this outer*/
-	virtual UTransformableHandle* Duplicate(UObject* NewOuter) const PURE_VIRTUAL(Duplicate, return nullptr;);
+	CONSTRAINTS_API virtual UTransformableHandle* Duplicate(UObject* NewOuter) const PURE_VIRTUAL(Duplicate, return nullptr;);
 
 
 	/** Fix up Binding in case it has changed*/
-	void OnBindingIDsUpdated(const TMap<UE::MovieScene::FFixedObjectBindingID, UE::MovieScene::FFixedObjectBindingID>& OldFixedToNewFixedMap, FMovieSceneSequenceID LocalSequenceID, const FMovieSceneSequenceHierarchy* Hierarchy, IMovieScenePlayer& Player);
+	CONSTRAINTS_API void OnBindingIDsUpdated(const TMap<UE::MovieScene::FFixedObjectBindingID, UE::MovieScene::FFixedObjectBindingID>& OldFixedToNewFixedMap, FMovieSceneSequenceID LocalSequenceID, const FMovieSceneSequenceHierarchy* Hierarchy, IMovieScenePlayer& Player);
 	
 	/** Perform any special ticking needed for this handle, by default it does nothing, todo need to see if we need to tick control rig also*/
 	virtual void TickForBaking() {};
 	/** Get the array of float channels for the specified section*/
-	virtual TArrayView<FMovieSceneFloatChannel*>  GetFloatChannels(const UMovieSceneSection* InSection) const PURE_VIRTUAL(GetFloatChannels, return TArrayView<FMovieSceneFloatChannel*>(); );
+	CONSTRAINTS_API virtual TArrayView<FMovieSceneFloatChannel*>  GetFloatChannels(const UMovieSceneSection* InSection) const PURE_VIRTUAL(GetFloatChannels, return TArrayView<FMovieSceneFloatChannel*>(); );
 	/** Get the array of double channels for the specified section*/
-	virtual TArrayView<FMovieSceneDoubleChannel*>  GetDoubleChannels(const UMovieSceneSection * InSection) const PURE_VIRTUAL(GetDoubleChannels, return TArrayView<FMovieSceneDoubleChannel*>(); );
+	CONSTRAINTS_API virtual TArrayView<FMovieSceneDoubleChannel*>  GetDoubleChannels(const UMovieSceneSection * InSection) const PURE_VIRTUAL(GetDoubleChannels, return TArrayView<FMovieSceneDoubleChannel*>(); );
 	/** Add Transform Keys at the specified times*/
-	virtual bool AddTransformKeys(const TArray<FFrameNumber>& InFrames,
+	CONSTRAINTS_API virtual bool AddTransformKeys(const TArray<FFrameNumber>& InFrames,
 		const TArray<FTransform>& InTransforms,
 		const EMovieSceneTransformChannel& InChannels,
 		const FFrameRate& InTickResolution,
@@ -89,30 +89,30 @@ public:
 	 * Returns the target object containing the tick function (returned in GetTickFunction).
 	 * See FTickFunction::AddPrerequisite for details.
 	 **/
-	virtual UObject* GetPrerequisiteObject() const PURE_VIRTUAL(GetPrerequisiteObject, return nullptr;);
+	CONSTRAINTS_API virtual UObject* GetPrerequisiteObject() const PURE_VIRTUAL(GetPrerequisiteObject, return nullptr;);
 	/**
 	 * Returns the tick function of the underlying transformable object.
 	 * This is used to set dependencies with the constraint.
 	**/
-	virtual FTickFunction* GetTickFunction() const PURE_VIRTUAL(GetTickFunction, return nullptr;);
+	CONSTRAINTS_API virtual FTickFunction* GetTickFunction() const PURE_VIRTUAL(GetTickFunction, return nullptr;);
 
 	/** Generates a hash value of the underlying transformable object. */
-	virtual uint32 GetHash() const PURE_VIRTUAL(GetHash, return 0;);
+	CONSTRAINTS_API virtual uint32 GetHash() const PURE_VIRTUAL(GetHash, return 0;);
 
 	/** Returns the underlying targeted object. */
-	virtual TWeakObjectPtr<UObject> GetTarget() const PURE_VIRTUAL(GetTarget, return nullptr;);
+	CONSTRAINTS_API virtual TWeakObjectPtr<UObject> GetTarget() const PURE_VIRTUAL(GetTarget, return nullptr;);
 
 	/** Check for direct dependencies with InOther. */
-	virtual bool HasDirectDependencyWith(const UTransformableHandle& InOther) const PURE_VIRTUAL(HasDirectDependencyWith, return false;);
+	CONSTRAINTS_API virtual bool HasDirectDependencyWith(const UTransformableHandle& InOther) const PURE_VIRTUAL(HasDirectDependencyWith, return false;);
 
 	/** Look for a possible tick function that can be used as a prerequisite. */
-	virtual FTickPrerequisite GetPrimaryPrerequisite() const PURE_VIRTUAL(GetPrimaryPrerequisite, return FTickPrerequisite(););
+	CONSTRAINTS_API virtual FTickPrerequisite GetPrimaryPrerequisite() const PURE_VIRTUAL(GetPrimaryPrerequisite, return FTickPrerequisite(););
 	
-	FHandleModifiedEvent& HandleModified();
+	CONSTRAINTS_API FHandleModifiedEvent& HandleModified();
 
 #if WITH_EDITOR
-	virtual FString GetLabel() const PURE_VIRTUAL(GetLabel, return FString(););
-	virtual FString GetFullLabel() const PURE_VIRTUAL(GetFullLabel, return FString(););
+	CONSTRAINTS_API virtual FString GetLabel() const PURE_VIRTUAL(GetLabel, return FString(););
+	CONSTRAINTS_API virtual FString GetFullLabel() const PURE_VIRTUAL(GetFullLabel, return FString(););
 #endif
 
 	//possible bindingID
@@ -126,53 +126,53 @@ protected:
  * UTransformableComponentHandle
  */
 
-UCLASS(Blueprintable)
-class CONSTRAINTS_API UTransformableComponentHandle : public UTransformableHandle 
+UCLASS(Blueprintable, MinimalAPI)
+class UTransformableComponentHandle : public UTransformableHandle 
 {
 	GENERATED_BODY()
 	
 public:
 	
-	virtual ~UTransformableComponentHandle();
+	CONSTRAINTS_API virtual ~UTransformableComponentHandle();
 
-	virtual void PostLoad() override;
+	CONSTRAINTS_API virtual void PostLoad() override;
 	
 	/** Sanity check to ensure that Component. */
-	virtual bool IsValid() const override;
+	CONSTRAINTS_API virtual bool IsValid() const override;
 	
 	/** Sets the global transform of Component. */
-	virtual void SetGlobalTransform(const FTransform& InGlobal) const override;
+	CONSTRAINTS_API virtual void SetGlobalTransform(const FTransform& InGlobal) const override;
 	/** Sets the local transform of Component in it's attachment. */
-	virtual void SetLocalTransform(const FTransform& InLocal) const override;
+	CONSTRAINTS_API virtual void SetLocalTransform(const FTransform& InLocal) const override;
 	/** Gets the global transform of Component. */
-	virtual FTransform GetGlobalTransform() const override;
+	CONSTRAINTS_API virtual FTransform GetGlobalTransform() const override;
 	/** Gets the local transform of Component in it's attachment. */
-	virtual FTransform GetLocalTransform() const override;
+	CONSTRAINTS_API virtual FTransform GetLocalTransform() const override;
 	/** Tick the component*/
-	virtual void TickForBaking() override;
+	CONSTRAINTS_API virtual void TickForBaking() override;
 	/** Returns the target object containing the tick function (e.i. Component). */
-	virtual UObject* GetPrerequisiteObject() const override;
+	CONSTRAINTS_API virtual UObject* GetPrerequisiteObject() const override;
 	/** Returns Component's tick function. */
-	virtual FTickFunction* GetTickFunction() const override;
+	CONSTRAINTS_API virtual FTickFunction* GetTickFunction() const override;
 
 	/** Generates a hash value of Component. */
-	virtual uint32 GetHash() const override;
+	CONSTRAINTS_API virtual uint32 GetHash() const override;
 
 	/** Returns the underlying targeted object. */
-	virtual TWeakObjectPtr<UObject> GetTarget() const override;
+	CONSTRAINTS_API virtual TWeakObjectPtr<UObject> GetTarget() const override;
 
 	/** Check for direct dependencies (ie hierarchy) with InOther. */
-	virtual bool HasDirectDependencyWith(const UTransformableHandle& InOther) const override;
+	CONSTRAINTS_API virtual bool HasDirectDependencyWith(const UTransformableHandle& InOther) const override;
 
 	/** Look for a possible tick function that can be used as a prerequisite. */
-	virtual FTickPrerequisite GetPrimaryPrerequisite() const override;
+	CONSTRAINTS_API virtual FTickPrerequisite GetPrimaryPrerequisite() const override;
 
 	/** Get the array of float channels for the specified section*/
-	virtual TArrayView<FMovieSceneFloatChannel*>  GetFloatChannels(const UMovieSceneSection* InSection) const override;
+	CONSTRAINTS_API virtual TArrayView<FMovieSceneFloatChannel*>  GetFloatChannels(const UMovieSceneSection* InSection) const override;
 	/** Get the array of double channels for the specified section*/
-	virtual TArrayView<FMovieSceneDoubleChannel*>  GetDoubleChannels(const UMovieSceneSection* InSection) const override;
+	CONSTRAINTS_API virtual TArrayView<FMovieSceneDoubleChannel*>  GetDoubleChannels(const UMovieSceneSection* InSection) const override;
 	/** Add Transform Keys at the specified times*/
-	virtual bool AddTransformKeys(const TArray<FFrameNumber>& InFrames,
+	CONSTRAINTS_API virtual bool AddTransformKeys(const TArray<FFrameNumber>& InFrames,
 		const TArray<FTransform>& InTransforms,
 		const EMovieSceneTransformChannel& InChannels,
 		const FFrameRate& InTickResolution,
@@ -180,15 +180,15 @@ public:
 		const bool bLocal = true) const override;
 
 	/** Resolve the bound objects so that any object it references are resovled and correctly set up*/
-	virtual void ResolveBoundObjects(FMovieSceneSequenceID LocalSequenceID, IMovieScenePlayer& Player, UObject* SubObject = nullptr) override;
+	CONSTRAINTS_API virtual void ResolveBoundObjects(FMovieSceneSequenceID LocalSequenceID, IMovieScenePlayer& Player, UObject* SubObject = nullptr) override;
 
 	/** Make a duplicate of myself with this outer*/
-	virtual UTransformableHandle* Duplicate(UObject* NewOuter) const override;
+	CONSTRAINTS_API virtual UTransformableHandle* Duplicate(UObject* NewOuter) const override;
 
 #if WITH_EDITOR
 	/** Returns labels used for UI. */
-	virtual FString GetLabel() const override;
-	virtual FString GetFullLabel() const override;
+	CONSTRAINTS_API virtual FString GetLabel() const override;
+	CONSTRAINTS_API virtual FString GetFullLabel() const override;
 #endif
 	
 	/** The Component that this handle is pointing at. */
@@ -200,12 +200,12 @@ public:
 	FName SocketName = NAME_None;
 
 	/** Registers/Unregisters useful delegates to track changes in the Component's transform. */
-	void UnregisterDelegates() const;
-	void RegisterDelegates();
+	CONSTRAINTS_API void UnregisterDelegates() const;
+	CONSTRAINTS_API void RegisterDelegates();
 	
 #if WITH_EDITOR
-	void OnActorMoving(AActor* InActor);
-	void OnPostPropertyChanged(UObject* InObject, FPropertyChangedEvent& InPropertyChangedEvent);
-	void OnObjectsReplaced(const TMap<UObject*, UObject*>& InOldToNewInstances);
+	CONSTRAINTS_API void OnActorMoving(AActor* InActor);
+	CONSTRAINTS_API void OnPostPropertyChanged(UObject* InObject, FPropertyChangedEvent& InPropertyChangedEvent);
+	CONSTRAINTS_API void OnObjectsReplaced(const TMap<UObject*, UObject*>& InOldToNewInstances);
 #endif
-};
\ No newline at end of file
+};
diff --git a/Engine/Source/Runtime/Experimental/Animation/Constraints/Public/TransformableRegistry.h b/Engine/Source/Runtime/Experimental/Animation/Constraints/Public/TransformableRegistry.h
--- a/Engine/Source/Runtime/Experimental/Animation/Constraints/Public/TransformableRegistry.h
+++ b/Engine/Source/Runtime/Experimental/Animation/Constraints/Public/TransformableRegistry.h
@@ -15,7 +15,7 @@ class UTransformableHandle;
  * Can be used to register customization transformable handle factory for specific objects. 
  */
 
-class CONSTRAINTS_API FTransformableRegistry
+class FTransformableRegistry
 {
 public:
 
@@ -24,26 +24,26 @@ public:
 	/** Defines the function that return a hash value from the object and the attachment name. */
 	using GetHashFuncT = TFunction<uint32(const UObject*, const FName&)>;
 	
-	~FTransformableRegistry();
+	CONSTRAINTS_API ~FTransformableRegistry();
 
 	/** Get registry singleton */
-	static FTransformableRegistry& Get();
+	static CONSTRAINTS_API FTransformableRegistry& Get();
 
 	/** Registers a InClass as transformable object. */
-	void Register(UClass* InClass, CreateHandleFuncT InHandleFunc, GetHashFuncT InHashFunc);
+	CONSTRAINTS_API void Register(UClass* InClass, CreateHandleFuncT InHandleFunc, GetHashFuncT InHashFunc);
 
 	/** Get the associated hash function for that specific class. */
-	GetHashFuncT GetHashFunction(const UClass* InClass) const;
+	CONSTRAINTS_API GetHashFuncT GetHashFunction(const UClass* InClass) const;
 	/** Get the associated handle creation function for that specific class. */
-	CreateHandleFuncT GetCreateFunction(const UClass* InClass) const;
+	CONSTRAINTS_API CreateHandleFuncT GetCreateFunction(const UClass* InClass) const;
 
 protected:
 	friend class FConstraintsModule;
 
 	/** Registers the basic transformable objects. (called when starting the constraints module). */
-	static void RegisterBaseObjects();
+	static CONSTRAINTS_API void RegisterBaseObjects();
 	/** Registers the basic transformable objects. (called when shutting down the constraints module). */
-	static void UnregisterAllObjects();
+	static CONSTRAINTS_API void UnregisterAllObjects();
 	
 private:
 	FTransformableRegistry() = default;
@@ -56,7 +56,7 @@ private:
 	};
 
 	/** Find the associated info for that specific class. */
-	const FTransformableInfo* FindInfo(const UClass* InClass) const;
+	CONSTRAINTS_API const FTransformableInfo* FindInfo(const UClass* InClass) const;
 
 	/** List of all registered transformable objects. */
 	TMap<UClass*, FTransformableInfo> Transformables;
diff --git a/Engine/Source/Runtime/Experimental/Dataflow/Core/Public/Dataflow/DataflowConnection.h b/Engine/Source/Runtime/Experimental/Dataflow/Core/Public/Dataflow/DataflowConnection.h
--- a/Engine/Source/Runtime/Experimental/Dataflow/Core/Public/Dataflow/DataflowConnection.h
+++ b/Engine/Source/Runtime/Experimental/Dataflow/Core/Public/Dataflow/DataflowConnection.h
@@ -39,7 +39,7 @@ namespace Dataflow
 // Input Output Base
 //
 USTRUCT()
-struct DATAFLOWCORE_API FDataflowConnection
+struct FDataflowConnection
 {
 	GENERATED_USTRUCT_BODY()
 
@@ -55,14 +55,14 @@ protected:
 
 public:
 	FDataflowConnection() {};
-	FDataflowConnection(Dataflow::FPin::EDirection Direction, FName InType, FName InName, FDataflowNode* OwningNode = nullptr, FProperty* InProperty = nullptr, FGuid InGuid = FGuid::NewGuid());
+	DATAFLOWCORE_API FDataflowConnection(Dataflow::FPin::EDirection Direction, FName InType, FName InName, FDataflowNode* OwningNode = nullptr, FProperty* InProperty = nullptr, FGuid InGuid = FGuid::NewGuid());
 	virtual ~FDataflowConnection() {};
 
 	FDataflowNode* GetOwningNode() { return OwningNode; }
 	const FDataflowNode* GetOwningNode() const { return OwningNode; }
 
 	Dataflow::FPin::EDirection GetDirection() const { return Direction; }
-	int32 GetOffset( ) const;
+	DATAFLOWCORE_API int32 GetOffset( ) const;
 
 	FName GetType() const { return Type; }
 
diff --git a/Engine/Source/Runtime/Experimental/Dataflow/Core/Public/Dataflow/DataflowCoreModule.h b/Engine/Source/Runtime/Experimental/Dataflow/Core/Public/Dataflow/DataflowCoreModule.h
--- a/Engine/Source/Runtime/Experimental/Dataflow/Core/Public/Dataflow/DataflowCoreModule.h
+++ b/Engine/Source/Runtime/Experimental/Dataflow/Core/Public/Dataflow/DataflowCoreModule.h
@@ -11,16 +11,16 @@
 /**
  * The public interface to this module
  */
-class DATAFLOWCORE_API IDataflowCoreModule : public IModuleInterface
+class IDataflowCoreModule : public IModuleInterface
 {
 
 public:
 
 	// IModuleInterface interface
 
-	virtual void StartupModule() override;
+	DATAFLOWCORE_API virtual void StartupModule() override;
 
-	virtual void ShutdownModule() override;
+	DATAFLOWCORE_API virtual void ShutdownModule() override;
 
 	virtual bool SupportsDynamicReloading() override
 	{
diff --git a/Engine/Source/Runtime/Experimental/Dataflow/Core/Public/Dataflow/DataflowGraph.h b/Engine/Source/Runtime/Experimental/Dataflow/Core/Public/Dataflow/DataflowGraph.h
--- a/Engine/Source/Runtime/Experimental/Dataflow/Core/Public/Dataflow/DataflowGraph.h
+++ b/Engine/Source/Runtime/Experimental/Dataflow/Core/Public/Dataflow/DataflowGraph.h
@@ -45,7 +45,7 @@ namespace Dataflow
 	//
 	//
 	//
-	class DATAFLOWCORE_API FGraph
+	class FGraph
 	{
 
 		FGuid  Guid;
@@ -54,7 +54,7 @@ namespace Dataflow
 		TArray< FLink > Connections;
 		TSet< FName > DisabledNodes;
 	public:
-		FGraph(FGuid InGuid = FGuid::NewGuid());
+		DATAFLOWCORE_API FGraph(FGuid InGuid = FGuid::NewGuid());
 		virtual ~FGraph() {}
 
 		const TArray< TSharedPtr<FDataflowNode> >& GetTerminalNodes() const { return TerminalNodes; }
@@ -144,16 +144,16 @@ namespace Dataflow
 			return TSharedPtr<FDataflowNode>(nullptr);
 		}
 
-		void RemoveNode(TSharedPtr<FDataflowNode> Node);
+		DATAFLOWCORE_API void RemoveNode(TSharedPtr<FDataflowNode> Node);
 
-		void ClearConnections(FDataflowConnection* ConnectionBase);
-		void ClearConnections(FDataflowInput* Input);
-		void ClearConnections(FDataflowOutput* Output);
+		DATAFLOWCORE_API void ClearConnections(FDataflowConnection* ConnectionBase);
+		DATAFLOWCORE_API void ClearConnections(FDataflowInput* Input);
+		DATAFLOWCORE_API void ClearConnections(FDataflowOutput* Output);
 
-		void Connect(FDataflowOutput* OutputConnection, FDataflowInput* InputConnection);
-		void Disconnect(FDataflowOutput* OutputConnection, FDataflowInput* InputConnection);
+		DATAFLOWCORE_API void Connect(FDataflowOutput* OutputConnection, FDataflowInput* InputConnection);
+		DATAFLOWCORE_API void Disconnect(FDataflowOutput* OutputConnection, FDataflowInput* InputConnection);
 
-		virtual void Serialize(FArchive& Ar);
+		DATAFLOWCORE_API virtual void Serialize(FArchive& Ar);
 		const TSet<FName>& GetDisabledNodes() const { return DisabledNodes; }
 
 	};
diff --git a/Engine/Source/Runtime/Experimental/Dataflow/Core/Public/Dataflow/DataflowInputOutput.h b/Engine/Source/Runtime/Experimental/Dataflow/Core/Public/Dataflow/DataflowInputOutput.h
--- a/Engine/Source/Runtime/Experimental/Dataflow/Core/Public/Dataflow/DataflowInputOutput.h
+++ b/Engine/Source/Runtime/Experimental/Dataflow/Core/Public/Dataflow/DataflowInputOutput.h
@@ -20,7 +20,7 @@ struct FDataflowOutput;
 //
 namespace Dataflow
 {
-	struct DATAFLOWCORE_API FInputParameters {
+	struct FInputParameters {
 		FInputParameters(FName InType = FName(""), FName InName = FName(""), FDataflowNode * InOwner = nullptr, FProperty * InProperty = nullptr)
 			: Type(InType)
 			, Name(InName)
@@ -69,7 +69,7 @@ public:
 //
 namespace Dataflow
 {
-	struct DATAFLOWCORE_API FOutputParameters
+	struct FOutputParameters
 	{
 		FOutputParameters(FName InType = FName(""), FName InName = FName(""), FDataflowNode* InOwner = nullptr, FProperty* InProperty = nullptr)
 			: Type(InType)
@@ -84,7 +84,7 @@ namespace Dataflow
 	};
 }
 USTRUCT()
-struct DATAFLOWCORE_API FDataflowOutput : public FDataflowConnection
+struct FDataflowOutput : public FDataflowConnection
 {
 	GENERATED_USTRUCT_BODY()
 
@@ -95,21 +95,21 @@ struct DATAFLOWCORE_API FDataflowOutput : public FDataflowConnection
 	size_t PassthroughOffsetAddress = INDEX_NONE;
 
 public:
-	static FDataflowOutput NoOpOutput;
+	static DATAFLOWCORE_API FDataflowOutput NoOpOutput;
 	
 	mutable TSharedPtr<FCriticalSection> OutputLock;
 	
-	FDataflowOutput(const Dataflow::FOutputParameters& Param = {}, FGuid InGuid = FGuid::NewGuid());
+	DATAFLOWCORE_API FDataflowOutput(const Dataflow::FOutputParameters& Param = {}, FGuid InGuid = FGuid::NewGuid());
 
-	TArray<FDataflowInput*>& GetConnections();
-	const TArray<FDataflowInput*>& GetConnections() const;
+	DATAFLOWCORE_API TArray<FDataflowInput*>& GetConnections();
+	DATAFLOWCORE_API const TArray<FDataflowInput*>& GetConnections() const;
 
-	virtual TArray<FDataflowInput*> GetConnectedInputs();
-	virtual const TArray<const FDataflowInput*> GetConnectedInputs() const;
+	DATAFLOWCORE_API virtual TArray<FDataflowInput*> GetConnectedInputs();
+	DATAFLOWCORE_API virtual const TArray<const FDataflowInput*> GetConnectedInputs() const;
 
-	virtual bool AddConnection(FDataflowConnection* InOutput) override;
+	DATAFLOWCORE_API virtual bool AddConnection(FDataflowConnection* InOutput) override;
 
-	virtual bool RemoveConnection(FDataflowConnection* InInput) override;
+	DATAFLOWCORE_API virtual bool RemoveConnection(FDataflowConnection* InInput) override;
 
 	virtual FORCEINLINE void SetPassthroughOffsetAddress(const size_t InPassthroughOffsetAddress)
 	{
@@ -149,7 +149,7 @@ public:
 		return Default;
 	}
 
-	bool EvaluateImpl(Dataflow::FContext& Context) const;
+	DATAFLOWCORE_API bool EvaluateImpl(Dataflow::FContext& Context) const;
 	
 	template<class T>
 	bool Evaluate(Dataflow::FContext& Context) const;
@@ -157,7 +157,7 @@ public:
 	template<class T>
 	TFuture<bool> EvaluateParallel(Dataflow::FContext& Context) const;
 
-	virtual void Invalidate(const Dataflow::FTimestamp& ModifiedTimestamp = Dataflow::FTimestamp::Current()) override;
+	DATAFLOWCORE_API virtual void Invalidate(const Dataflow::FTimestamp& ModifiedTimestamp = Dataflow::FTimestamp::Current()) override;
 
 };
  
@@ -213,4 +213,4 @@ template<class T>
 TFuture<bool> FDataflowOutput::EvaluateParallel(Dataflow::FContext& Context) const
 {
 	return Async(EAsyncExecution::TaskGraph, [&]() -> bool { return this->Evaluate<T>(Context); });
-}
\ No newline at end of file
+}
diff --git a/Engine/Source/Runtime/Experimental/Dataflow/Core/Public/Dataflow/DataflowNode.h b/Engine/Source/Runtime/Experimental/Dataflow/Core/Public/Dataflow/DataflowNode.h
--- a/Engine/Source/Runtime/Experimental/Dataflow/Core/Public/Dataflow/DataflowNode.h
+++ b/Engine/Source/Runtime/Experimental/Dataflow/Core/Public/Dataflow/DataflowNode.h
@@ -16,7 +16,7 @@ struct FDataflowOutput;
 class UScriptStruct;
 
 namespace Dataflow {
-	struct DATAFLOWCORE_API FNodeParameters {
+	struct FNodeParameters {
 		FName Name;
 	};
 	class FGraph;
@@ -32,7 +32,7 @@ namespace Dataflow {
 *        evaluation. 
 */
 USTRUCT()
-struct DATAFLOWCORE_API FDataflowNode
+struct FDataflowNode
 {
 	GENERATED_USTRUCT_BODY()
 
@@ -75,19 +75,19 @@ struct DATAFLOWCORE_API FDataflowNode
 	virtual FName GetDisplayName() const { return ""; }
 	virtual FName GetCategory() const { return ""; }
 	virtual FString GetTags() const { return ""; }
-	virtual FString GetToolTip();
-	FString GetPinToolTip(const FName& PropertyName);
-	FText GetPinDisplayName(const FName& PropertyName);
-	TArray<FString> GetPinMetaData(const FName& PropertyName);
+	DATAFLOWCORE_API virtual FString GetToolTip();
+	DATAFLOWCORE_API FString GetPinToolTip(const FName& PropertyName);
+	DATAFLOWCORE_API FText GetPinDisplayName(const FName& PropertyName);
+	DATAFLOWCORE_API TArray<FString> GetPinMetaData(const FName& PropertyName);
 	virtual TArray<Dataflow::FRenderingParameter> GetRenderParameters() const { return GetRenderParametersImpl(); }
 	// Copy node property values from another node
-	void CopyNodeProperties(const TSharedPtr<FDataflowNode> CopyFromDataflowNode);
+	DATAFLOWCORE_API void CopyNodeProperties(const TSharedPtr<FDataflowNode> CopyFromDataflowNode);
 
 	//
 	// Connections
 	//
 
-	TArray<Dataflow::FPin> GetPins() const;
+	DATAFLOWCORE_API TArray<Dataflow::FPin> GetPins() const;
 
 	/** Override this function to add the AddOptionPin functionality to the node's context menu. */
 	virtual Dataflow::FPin AddPin() { return { Dataflow::FPin::EDirection::NONE, NAME_None, NAME_None }; }
@@ -98,32 +98,32 @@ struct DATAFLOWCORE_API FDataflowNode
 	/** Override this function to add the RemoveOPtionPin functionality to the node's context menu. */
 	virtual bool CanRemovePin() const { return false; }
 
-	virtual void AddInput(FDataflowInput* InPtr);
-	TArray< FDataflowInput* > GetInputs() const;
-	void ClearInputs();
+	DATAFLOWCORE_API virtual void AddInput(FDataflowInput* InPtr);
+	DATAFLOWCORE_API TArray< FDataflowInput* > GetInputs() const;
+	DATAFLOWCORE_API void ClearInputs();
 
-	FDataflowInput* FindInput(FName Name);
-	FDataflowInput* FindInput(void* Reference);
-	const FDataflowInput* FindInput(const void* Reference) const;
+	DATAFLOWCORE_API FDataflowInput* FindInput(FName Name);
+	DATAFLOWCORE_API FDataflowInput* FindInput(void* Reference);
+	DATAFLOWCORE_API const FDataflowInput* FindInput(const void* Reference) const;
 
 
-	virtual void AddOutput(FDataflowOutput* InPtr);
-	int NumOutputs() const;
-	TArray< FDataflowOutput* > GetOutputs() const;
-	void ClearOutputs();
+	DATAFLOWCORE_API virtual void AddOutput(FDataflowOutput* InPtr);
+	DATAFLOWCORE_API int NumOutputs() const;
+	DATAFLOWCORE_API TArray< FDataflowOutput* > GetOutputs() const;
+	DATAFLOWCORE_API void ClearOutputs();
 
-	FDataflowOutput* FindOutput(FName Name);
-	FDataflowOutput* FindOutput(void* Reference);
-	const FDataflowOutput* FindOutput(FName Name) const;
-	const FDataflowOutput* FindOutput(const void* Reference) const;
+	DATAFLOWCORE_API FDataflowOutput* FindOutput(FName Name);
+	DATAFLOWCORE_API FDataflowOutput* FindOutput(void* Reference);
+	DATAFLOWCORE_API const FDataflowOutput* FindOutput(FName Name) const;
+	DATAFLOWCORE_API const FDataflowOutput* FindOutput(const void* Reference) const;
 
-	static const FName DataflowInput;
-	static const FName DataflowOutput;
-	static const FName DataflowPassthrough;
-	static const FName DataflowIntrinsic;
+	static DATAFLOWCORE_API const FName DataflowInput;
+	static DATAFLOWCORE_API const FName DataflowOutput;
+	static DATAFLOWCORE_API const FName DataflowPassthrough;
+	static DATAFLOWCORE_API const FName DataflowIntrinsic;
 
-	static const FLinearColor DefaultNodeTitleColor;
-	static const FLinearColor DefaultNodeBodyTintColor;
+	static DATAFLOWCORE_API const FLinearColor DefaultNodeTitleColor;
+	static DATAFLOWCORE_API const FLinearColor DefaultNodeBodyTintColor;
 
 	/** Override this method to provide custom serialization for this node. */
 	virtual void Serialize(FArchive& Ar) {}
@@ -137,10 +137,10 @@ struct DATAFLOWCORE_API FDataflowNode
 	virtual const UScriptStruct* TypedScriptStruct() const { return nullptr; }
 
 	/** Register the Input and Outputs after the creation in the factory */
-	void RegisterInputConnection(const void* Property);
-	void RegisterOutputConnection(const void* Property, const void* Passthrough = nullptr);
+	DATAFLOWCORE_API void RegisterInputConnection(const void* Property);
+	DATAFLOWCORE_API void RegisterOutputConnection(const void* Property, const void* Passthrough = nullptr);
 	/** Unregister the input connection if one exists matching this property, and then invalidate the graph. */
-	void UnregisterInputConnection(const void* Property);
+	DATAFLOWCORE_API void UnregisterInputConnection(const void* Property);
 
 	//
 	// Evaluation
@@ -220,11 +220,11 @@ struct DATAFLOWCORE_API FDataflowNode
 		return (FindInput(Reference)->GetConnection() != nullptr);
 	}
 
-	void Invalidate(const Dataflow::FTimestamp& ModifiedTimestamp = Dataflow::FTimestamp::Current());
+	DATAFLOWCORE_API void Invalidate(const Dataflow::FTimestamp& ModifiedTimestamp = Dataflow::FTimestamp::Current());
 
 	virtual void OnInvalidate() {}
 
-	virtual bool ValidateConnections();
+	DATAFLOWCORE_API virtual bool ValidateConnections();
 
 	bool HasValidConnections() const { return bHasValidConnections; }
 
diff --git a/Engine/Source/Runtime/Experimental/Dataflow/Core/Public/Dataflow/DataflowNodeColorsRegistry.h b/Engine/Source/Runtime/Experimental/Dataflow/Core/Public/Dataflow/DataflowNodeColorsRegistry.h
--- a/Engine/Source/Runtime/Experimental/Dataflow/Core/Public/Dataflow/DataflowNodeColorsRegistry.h
+++ b/Engine/Source/Runtime/Experimental/Dataflow/Core/Public/Dataflow/DataflowNodeColorsRegistry.h
@@ -17,20 +17,20 @@ namespace Dataflow
 	//
 	//
 	//
-	class DATAFLOWCORE_API FNodeColorsRegistry
+	class FNodeColorsRegistry
 	{
 	public:
-		static FNodeColorsRegistry& Get();
-		static void TearDown();
+		static DATAFLOWCORE_API FNodeColorsRegistry& Get();
+		static DATAFLOWCORE_API void TearDown();
 
-		void RegisterNodeColors(const FName& Category, const FNodeColors& NodeColors);
-		FLinearColor GetNodeTitleColor(const FName& Category);
-		FLinearColor GetNodeBodyTintColor(const FName& Category);
-		void NodeColorsChangedInSettings(const FNodeColorsMap& NodeColorsMap);
+		DATAFLOWCORE_API void RegisterNodeColors(const FName& Category, const FNodeColors& NodeColors);
+		DATAFLOWCORE_API FLinearColor GetNodeTitleColor(const FName& Category);
+		DATAFLOWCORE_API FLinearColor GetNodeBodyTintColor(const FName& Category);
+		DATAFLOWCORE_API void NodeColorsChangedInSettings(const FNodeColorsMap& NodeColorsMap);
 
 	private:
-		FNodeColorsRegistry();
-		~FNodeColorsRegistry();
+		DATAFLOWCORE_API FNodeColorsRegistry();
+		DATAFLOWCORE_API ~FNodeColorsRegistry();
 
 		TMap<FName, FNodeColors > ColorsMap;					// [Category] -> Colors
 		FDelegateHandle DataflowSettingsChangedDelegateHandle;
diff --git a/Engine/Source/Runtime/Experimental/Dataflow/Core/Public/Dataflow/DataflowNodeFactory.h b/Engine/Source/Runtime/Experimental/Dataflow/Core/Public/Dataflow/DataflowNodeFactory.h
--- a/Engine/Source/Runtime/Experimental/Dataflow/Core/Public/Dataflow/DataflowNodeFactory.h
+++ b/Engine/Source/Runtime/Experimental/Dataflow/Core/Public/Dataflow/DataflowNodeFactory.h
@@ -13,13 +13,13 @@ struct FDataflowConnection;
 namespace Dataflow
 {
 
-	struct DATAFLOWCORE_API FNewNodeParameters {
+	struct FNewNodeParameters {
 		FGuid Guid;
 		FName Type;
 		FName Name;
 	};
 
-	struct DATAFLOWCORE_API FFactoryParameters {
+	struct FFactoryParameters {
 		FFactoryParameters() {}
 		FFactoryParameters(FName InTypeName, FName InDisplayName, FName InCategory, FString InTags, FString InToolTip)
 		: TypeName(InTypeName), DisplayName(InDisplayName), Category(InCategory), Tags(InTags), ToolTip(InToolTip) {}
diff --git a/Engine/Source/Runtime/Experimental/Dataflow/Core/Public/Dataflow/DataflowNodeParameters.h b/Engine/Source/Runtime/Experimental/Dataflow/Core/Public/Dataflow/DataflowNodeParameters.h
--- a/Engine/Source/Runtime/Experimental/Dataflow/Core/Public/Dataflow/DataflowNodeParameters.h
+++ b/Engine/Source/Runtime/Experimental/Dataflow/Core/Public/Dataflow/DataflowNodeParameters.h
@@ -17,7 +17,7 @@ struct FDataflowConnection;
 
 namespace Dataflow
 {
-	struct DATAFLOWCORE_API FTimestamp
+	struct FTimestamp
 	{
 		typedef uint64 Type;
 		Type Value = Type(0);
@@ -27,11 +27,11 @@ namespace Dataflow
 		bool operator<(const FTimestamp& InTimestamp) const { return Value < InTimestamp.Value; }
 		bool IsInvalid() { return Value == Invalid; }
 
-		static Type Current();
-		static Type Invalid; // 0
+		static DATAFLOWCORE_API Type Current();
+		static DATAFLOWCORE_API Type Invalid; // 0
 	};
 
-	struct DATAFLOWCORE_API FRenderingParameter {
+	struct FRenderingParameter {
 		FRenderingParameter() {}
 		FRenderingParameter(FName InTypeName, const TArray<FName>& InOutputs)
 			: Type(InTypeName), Outputs(InOutputs) {}
@@ -83,7 +83,7 @@ namespace Dataflow
 		// @todo(dataflow) make an API for FContextCache
 	};
 
-	class DATAFLOWCORE_API FContext
+	class FContext
 	{
 	protected:
 		FContext(FContext&&) = default;
@@ -163,9 +163,9 @@ namespace Dataflow
 		virtual void Evaluate(const FDataflowNode* Node, const FDataflowOutput* Output) = 0;
 		virtual bool Evaluate(const FDataflowOutput& Connection) = 0;
 
-		void PushToCallstack(const FDataflowConnection* Connection);
-		void PopFromCallstack(const FDataflowConnection* Connection);
-		bool IsInCallstack(const FDataflowConnection* Connection) const;
+		DATAFLOWCORE_API void PushToCallstack(const FDataflowConnection* Connection);
+		DATAFLOWCORE_API void PopFromCallstack(const FDataflowConnection* Connection);
+		DATAFLOWCORE_API bool IsInCallstack(const FDataflowConnection* Connection) const;
 
 	private:
 #if DATAFLOW_EDITOR_EVALUATION
@@ -173,11 +173,11 @@ namespace Dataflow
 #endif
 	};
 
-	struct DATAFLOWCORE_API FContextScopedCallstack
+	struct FContextScopedCallstack
 	{
 	public:
-		FContextScopedCallstack(FContext& InContext, const FDataflowConnection* InConnection);
-		~FContextScopedCallstack();
+		DATAFLOWCORE_API FContextScopedCallstack(FContext& InContext, const FDataflowConnection* InConnection);
+		DATAFLOWCORE_API ~FContextScopedCallstack();
 
 		bool IsLoopDetected() const { return bLoopDetected; }
 
@@ -193,7 +193,7 @@ namespace Dataflow
 	virtual bool IsA(FName InType) const override { return InType==StaticType() || Super::IsA(InType); }	\
 	virtual FName GetType() const override { return StaticType(); }
 
-	class DATAFLOWCORE_API FContextSingle : public FContext
+	class FContextSingle : public FContext
 	{
 		FContextCache DataStore;
 
@@ -224,11 +224,11 @@ namespace Dataflow
 			return DataStore.IsEmpty();
 		}
 
-		virtual void Evaluate(const FDataflowNode* Node, const FDataflowOutput* Output) override;
-		virtual bool Evaluate(const FDataflowOutput& Connection) override;
+		DATAFLOWCORE_API virtual void Evaluate(const FDataflowNode* Node, const FDataflowOutput* Output) override;
+		DATAFLOWCORE_API virtual bool Evaluate(const FDataflowOutput& Connection) override;
 	};
 	
-	class DATAFLOWCORE_API FContextThreaded : public FContext
+	class FContextThreaded : public FContext
 	{
 		FContextCache DataStore;
 		TSharedPtr<FCriticalSection> CacheLock;
@@ -269,8 +269,8 @@ namespace Dataflow
 			return DataStore.IsEmpty();
 		}
 
-		virtual void Evaluate(const FDataflowNode* Node, const FDataflowOutput* Output) override;
-		virtual bool Evaluate(const FDataflowOutput& Connection) override;
+		DATAFLOWCORE_API virtual void Evaluate(const FDataflowNode* Node, const FDataflowOutput* Output) override;
+		DATAFLOWCORE_API virtual bool Evaluate(const FDataflowOutput& Connection) override;
 	};
 
 }
diff --git a/Engine/Source/Runtime/Experimental/Dataflow/Core/Public/Dataflow/DataflowSelection.h b/Engine/Source/Runtime/Experimental/Dataflow/Core/Public/Dataflow/DataflowSelection.h
--- a/Engine/Source/Runtime/Experimental/Dataflow/Core/Public/Dataflow/DataflowSelection.h
+++ b/Engine/Source/Runtime/Experimental/Dataflow/Core/Public/Dataflow/DataflowSelection.h
@@ -5,36 +5,36 @@
 #include "DataflowSelection.generated.h"
 
 USTRUCT()
-struct DATAFLOWCORE_API FDataflowSelection
+struct FDataflowSelection
 {
 	GENERATED_USTRUCT_BODY()
 
-	void Initialize(int32 NumBits, bool Value);
-	void Initialize(const FDataflowSelection& Other);
+	DATAFLOWCORE_API void Initialize(int32 NumBits, bool Value);
+	DATAFLOWCORE_API void Initialize(const FDataflowSelection& Other);
 	int32 Num() const { return SelectionArray.Num(); }
-	int32 NumSelected() const;
-	bool AnySelected() const;
+	DATAFLOWCORE_API int32 NumSelected() const;
+	DATAFLOWCORE_API bool AnySelected() const;
 	bool IsSelected(int32 Idx) const { return SelectionArray[Idx]; }
 	void SetSelected(int32 Idx) { SelectionArray[Idx] = true; }
 	void SetNotSelected(int32 Idx) { SelectionArray[Idx] = false; }
-	void AsArray(TArray<int32>& SelectionArr) const;
-	TArray<int32> AsArray() const;
+	DATAFLOWCORE_API void AsArray(TArray<int32>& SelectionArr) const;
+	DATAFLOWCORE_API TArray<int32> AsArray() const;
 	
 	/**
 	* Sets the selection from a sparse array (it only contains the indices of the selected items)
 	*/
-	void SetFromArray(const TArray<int32>& SelectionArr);
+	DATAFLOWCORE_API void SetFromArray(const TArray<int32>& SelectionArr);
 	
 	/**
 	* Sets the selection from a dense array (it contains a true/false element for every item)
 	* (for example from the "Internal" attr from FacesGroup)
 	*/
-	void SetFromArray(const TArray<bool>& SelectionArr);
-	void AND(const FDataflowSelection& Other, FDataflowSelection& Result) const;
-	void OR(const FDataflowSelection& Other, FDataflowSelection& Result) const;
-	void XOR(const FDataflowSelection& Other, FDataflowSelection& Result) const;
+	DATAFLOWCORE_API void SetFromArray(const TArray<bool>& SelectionArr);
+	DATAFLOWCORE_API void AND(const FDataflowSelection& Other, FDataflowSelection& Result) const;
+	DATAFLOWCORE_API void OR(const FDataflowSelection& Other, FDataflowSelection& Result) const;
+	DATAFLOWCORE_API void XOR(const FDataflowSelection& Other, FDataflowSelection& Result) const;
 	void Invert() { SelectionArray.BitwiseNOT(); }
-	void SetWithMask(const bool Value, const FDataflowSelection& Mask);
+	DATAFLOWCORE_API void SetWithMask(const bool Value, const FDataflowSelection& Mask);
 	const TBitArray<>& GetBitArray() const { return SelectionArray; };
 
 private:
@@ -43,21 +43,21 @@ private:
 };
 
 USTRUCT()
-struct DATAFLOWCORE_API FDataflowTransformSelection : public FDataflowSelection
+struct FDataflowTransformSelection : public FDataflowSelection
 {
 	GENERATED_USTRUCT_BODY()
 };
 
 
 USTRUCT()
-struct DATAFLOWCORE_API FDataflowVertexSelection : public FDataflowSelection
+struct FDataflowVertexSelection : public FDataflowSelection
 {
 	GENERATED_USTRUCT_BODY()
 };
 
 
 USTRUCT()
-struct DATAFLOWCORE_API FDataflowFaceSelection : public FDataflowSelection
+struct FDataflowFaceSelection : public FDataflowSelection
 {
 	GENERATED_USTRUCT_BODY()
 };
diff --git a/Engine/Source/Runtime/Experimental/Dataflow/Core/Public/Dataflow/DataflowSettings.h b/Engine/Source/Runtime/Experimental/Dataflow/Core/Public/Dataflow/DataflowSettings.h
--- a/Engine/Source/Runtime/Experimental/Dataflow/Core/Public/Dataflow/DataflowSettings.h
+++ b/Engine/Source/Runtime/Experimental/Dataflow/Core/Public/Dataflow/DataflowSettings.h
@@ -26,8 +26,8 @@ struct FNodeColors
 	FLinearColor NodeBodyTintColor = FLinearColor(0.f, 0.f, 0.f);
 };
 
-UCLASS(config = Editor, defaultconfig, meta = (DisplayName = "Dataflow"))
-class DATAFLOWCORE_API UDataflowSettings : public UDeveloperSettings
+UCLASS(config = Editor, defaultconfig, meta = (DisplayName = "Dataflow"), MinimalAPI)
+class UDataflowSettings : public UDeveloperSettings
 {
 	GENERATED_UCLASS_BODY()
 
@@ -51,19 +51,19 @@ class DATAFLOWCORE_API UDataflowSettings : public UDeveloperSettings
 	TMap<FName, FNodeColors> NodeColorsMap;
 
 	// Begin UDeveloperSettings Interface
-	virtual FName GetCategoryName() const override;
+	DATAFLOWCORE_API virtual FName GetCategoryName() const override;
 
-	FNodeColors RegisterColors(const FName& Category, const FNodeColors& Colors);
+	DATAFLOWCORE_API FNodeColors RegisterColors(const FName& Category, const FNodeColors& Colors);
 
 	const TMap<FName, FNodeColors>& GetNodeColorsMap() { return NodeColorsMap; }
 
 #if WITH_EDITOR
-	virtual FText GetSectionText() const override;
+	DATAFLOWCORE_API virtual FText GetSectionText() const override;
 #endif
 	// END UDeveloperSettings Interface
 
 #if WITH_EDITOR
-	virtual void PostEditChangeProperty(FPropertyChangedEvent& PropertyChangedEvent) override;
+	DATAFLOWCORE_API virtual void PostEditChangeProperty(FPropertyChangedEvent& PropertyChangedEvent) override;
 #endif
 
 public:
diff --git a/Engine/Source/Runtime/Experimental/Dataflow/Core/Public/Dataflow/DataflowTerminalNode.h b/Engine/Source/Runtime/Experimental/Dataflow/Core/Public/Dataflow/DataflowTerminalNode.h
--- a/Engine/Source/Runtime/Experimental/Dataflow/Core/Public/Dataflow/DataflowTerminalNode.h
+++ b/Engine/Source/Runtime/Experimental/Dataflow/Core/Public/Dataflow/DataflowTerminalNode.h
@@ -24,7 +24,7 @@ struct FDataflowOutput;
 *       evaluation graph. 
 */
 USTRUCT()
-struct DATAFLOWCORE_API FDataflowTerminalNode : public FDataflowNode
+struct FDataflowTerminalNode : public FDataflowNode
 {
 	GENERATED_USTRUCT_BODY()
 
diff --git a/Engine/Source/Runtime/Experimental/Dataflow/Engine/Public/Dataflow/DataflowBlueprintLibrary.h b/Engine/Source/Runtime/Experimental/Dataflow/Engine/Public/Dataflow/DataflowBlueprintLibrary.h
--- a/Engine/Source/Runtime/Experimental/Dataflow/Engine/Public/Dataflow/DataflowBlueprintLibrary.h
+++ b/Engine/Source/Runtime/Experimental/Dataflow/Engine/Public/Dataflow/DataflowBlueprintLibrary.h
@@ -7,8 +7,8 @@
 
 #include "DataflowBlueprintLibrary.generated.h"
 
-UCLASS()
-class DATAFLOWENGINE_API UDataflowBlueprintLibrary : public UBlueprintFunctionLibrary
+UCLASS(MinimalAPI)
+class UDataflowBlueprintLibrary : public UBlueprintFunctionLibrary
 {
 	GENERATED_BODY()
 
@@ -17,5 +17,5 @@ public:
 	* Find a specific terminal node by name evaluate it using a specific UObject
 	*/
 	UFUNCTION(BlueprintCallable, Category = "Dataflow", meta = (Keywords = "Dataflow graph"))
-	static void EvaluateTerminalNodeByName(UDataflow* Dataflow, FName TerminalNodeName, UObject* ResultAsset);
-};
\ No newline at end of file
+	static DATAFLOWENGINE_API void EvaluateTerminalNodeByName(UDataflow* Dataflow, FName TerminalNodeName, UObject* ResultAsset);
+};
diff --git a/Engine/Source/Runtime/Experimental/Dataflow/Engine/Public/Dataflow/DataflowEdNode.h b/Engine/Source/Runtime/Experimental/Dataflow/Engine/Public/Dataflow/DataflowEdNode.h
--- a/Engine/Source/Runtime/Experimental/Dataflow/Engine/Public/Dataflow/DataflowEdNode.h
+++ b/Engine/Source/Runtime/Experimental/Dataflow/Engine/Public/Dataflow/DataflowEdNode.h
@@ -23,8 +23,8 @@ class UObject;
 namespace Dataflow { class FGraph; class FRenderingParameters; }
 namespace GeometryCollection::Facades { class FRenderingFacade; }
 
-UCLASS()
-class DATAFLOWENGINE_API UDataflowEdNode : public UEdGraphNode
+UCLASS(MinimalAPI)
+class UDataflowEdNode : public UEdGraphNode
 {
 	GENERATED_UCLASS_BODY()
 
@@ -34,23 +34,23 @@ class DATAFLOWENGINE_API UDataflowEdNode : public UEdGraphNode
 public:
 
 	// UEdGraphNode interface
-	virtual void AllocateDefaultPins();
-	virtual FText GetNodeTitle(ENodeTitleType::Type TitleType) const;
+	DATAFLOWENGINE_API virtual void AllocateDefaultPins();
+	DATAFLOWENGINE_API virtual FText GetNodeTitle(ENodeTitleType::Type TitleType) const;
 #if WITH_EDITOR && !(UE_BUILD_SHIPPING || UE_BUILD_TEST)
-	virtual void PinConnectionListChanged(UEdGraphPin* Pin) override;
+	DATAFLOWENGINE_API virtual void PinConnectionListChanged(UEdGraphPin* Pin) override;
 #endif // WITH_EDITOR && !UE_BUILD_SHIPPING
 #if WITH_EDITOR
-	virtual FLinearColor GetNodeTitleColor() const override;
-	virtual FLinearColor GetNodeBodyTintColor() const override;
-	virtual FText GetTooltipText() const override;
-	virtual void GetPinHoverText(const UEdGraphPin& Pin, FString& HoverTextOut) const override;
-	virtual FText GetPinDisplayName(const UEdGraphPin* Pin) const override;
-	virtual void AutowireNewNode(UEdGraphPin* FromPin) override;
+	DATAFLOWENGINE_API virtual FLinearColor GetNodeTitleColor() const override;
+	DATAFLOWENGINE_API virtual FLinearColor GetNodeBodyTintColor() const override;
+	DATAFLOWENGINE_API virtual FText GetTooltipText() const override;
+	DATAFLOWENGINE_API virtual void GetPinHoverText(const UEdGraphPin& Pin, FString& HoverTextOut) const override;
+	DATAFLOWENGINE_API virtual FText GetPinDisplayName(const UEdGraphPin* Pin) const override;
+	DATAFLOWENGINE_API virtual void AutowireNewNode(UEdGraphPin* FromPin) override;
 #endif // WITH_EDITOR
 	// End of UEdGraphNode interface
 
 	// UObject interface
-	void Serialize(FArchive& Ar);
+	DATAFLOWENGINE_API void Serialize(FArchive& Ar);
 	// End UObject interface
 
 	bool IsBound() { return DataflowGraph && DataflowNodeGuid.IsValid(); }
@@ -59,26 +59,26 @@ public:
 	TSharedPtr<const Dataflow::FGraph> GetDataflowGraph() const { return DataflowGraph; }
 	void SetDataflowGraph(TSharedPtr<Dataflow::FGraph> InDataflowGraph) { DataflowGraph = InDataflowGraph; }
 
-	void UpdatePinsFromDataflowNode();
+	DATAFLOWENGINE_API void UpdatePinsFromDataflowNode();
 
 	FGuid GetDataflowNodeGuid() const { return DataflowNodeGuid; }
 	void SetDataflowNodeGuid(FGuid InGuid) { DataflowNodeGuid = InGuid; }
 
-	TSharedPtr<FDataflowNode> GetDataflowNode();
-	TSharedPtr<const FDataflowNode> GetDataflowNode() const;
+	DATAFLOWENGINE_API TSharedPtr<FDataflowNode> GetDataflowNode();
+	DATAFLOWENGINE_API TSharedPtr<const FDataflowNode> GetDataflowNode() const;
 
 	/** Add a new option pin if the underlying Dataflow node AddPin member is overriden. */
-	void AddOptionPin();
+	DATAFLOWENGINE_API void AddOptionPin();
 	/** Remove an option pin if the underlying Dataflow node RemovePin member is overriden. */
-	void RemoveOptionPin();
+	DATAFLOWENGINE_API void RemoveOptionPin();
 
 	//
 	// Node Rendering
 	//
-	void SetAssetRender(bool bInRender);
+	DATAFLOWENGINE_API void SetAssetRender(bool bInRender);
 	bool DoAssetRender() { return bRenderInAssetEditor; }
-	TArray<Dataflow::FRenderingParameter> GetRenderParameters() const;
-	virtual bool Render(GeometryCollection::Facades::FRenderingFacade& RenderData, TSharedPtr<Dataflow::FContext> Context) const;
+	DATAFLOWENGINE_API TArray<Dataflow::FRenderingParameter> GetRenderParameters() const;
+	DATAFLOWENGINE_API virtual bool Render(GeometryCollection::Facades::FRenderingFacade& RenderData, TSharedPtr<Dataflow::FContext> Context) const;
 
 	UPROPERTY()
 	bool bRenderInAssetEditor = false;
diff --git a/Engine/Source/Runtime/Experimental/Dataflow/Engine/Public/Dataflow/DataflowEngineTypes.h b/Engine/Source/Runtime/Experimental/Dataflow/Engine/Public/Dataflow/DataflowEngineTypes.h
--- a/Engine/Source/Runtime/Experimental/Dataflow/Engine/Public/Dataflow/DataflowEngineTypes.h
+++ b/Engine/Source/Runtime/Experimental/Dataflow/Engine/Public/Dataflow/DataflowEngineTypes.h
@@ -9,7 +9,7 @@
 #include "DataflowEngineTypes.generated.h"
 
 USTRUCT(BlueprintType)
-struct DATAFLOWENGINE_API FStringValuePair
+struct FStringValuePair
 {
 	GENERATED_BODY()
 
@@ -22,4 +22,4 @@ struct DATAFLOWENGINE_API FStringValuePair
 	bool operator ==(const FStringValuePair& Other) const {
 		return Key.Equals(Other.Key) && Value.Equals(Other.Value);
 	}
-};
\ No newline at end of file
+};
diff --git a/Engine/Source/Runtime/Experimental/Dataflow/Engine/Public/Dataflow/DataflowObject.h b/Engine/Source/Runtime/Experimental/Dataflow/Engine/Public/Dataflow/DataflowObject.h
--- a/Engine/Source/Runtime/Experimental/Dataflow/Engine/Public/Dataflow/DataflowObject.h
+++ b/Engine/Source/Runtime/Experimental/Dataflow/Engine/Public/Dataflow/DataflowObject.h
@@ -29,7 +29,7 @@ namespace Dataflow { class FGraph; }
 *     dynamic collection
 *
 */
-class DATAFLOWENGINE_API FDataflowAssetEdit
+class FDataflowAssetEdit
 {
 public:
 	typedef TFunctionRef<void()> FPostEditFunctionCallback;
@@ -38,10 +38,10 @@ public:
 	/**
 	 * @param UDataflow				The "FAsset" to edit
 	 */
-	FDataflowAssetEdit(UDataflow *InAsset, FPostEditFunctionCallback InCallable);
-	~FDataflowAssetEdit();
+	DATAFLOWENGINE_API FDataflowAssetEdit(UDataflow *InAsset, FPostEditFunctionCallback InCallable);
+	DATAFLOWENGINE_API ~FDataflowAssetEdit();
 
-	Dataflow::FGraph* GetGraph();
+	DATAFLOWENGINE_API Dataflow::FGraph* GetGraph();
 
 private:
 	FPostEditFunctionCallback PostEditCallback;
@@ -54,24 +54,24 @@ private:
 * UObject wrapper for the Dataflow::FGraph
 *
 */
-UCLASS(BlueprintType, customconstructor)
-class DATAFLOWENGINE_API UDataflow : public UEdGraph
+UCLASS(BlueprintType, customconstructor, MinimalAPI)
+class UDataflow : public UEdGraph
 {
 	GENERATED_UCLASS_BODY()
 
 	Dataflow::FTimestamp LastModifiedRenderTarget = Dataflow::FTimestamp::Invalid; 
 	TArray<const UDataflowEdNode*> RenderTargets; // Not Serialized
 	TSharedPtr<Dataflow::FGraph, ESPMode::ThreadSafe> Dataflow;
-	void PostEditCallback();
+	DATAFLOWENGINE_API void PostEditCallback();
 
 public:
-	UDataflow(const FObjectInitializer& ObjectInitializer = FObjectInitializer::Get());
+	DATAFLOWENGINE_API UDataflow(const FObjectInitializer& ObjectInitializer = FObjectInitializer::Get());
 
 	/** 
 	* Find all the node of a speific type and evaluate them using a specific UObject
 	*/
 	UE_DEPRECATED(5.1, "Use Blueprint library version of the function")
-	void EvaluateTerminalNodeByName(FName NodeName, UObject* Asset);
+	DATAFLOWENGINE_API void EvaluateTerminalNodeByName(FName NodeName, UObject* Asset);
 
 	virtual bool IsEditorOnly() const { return true; }
 
@@ -90,12 +90,12 @@ public:
 public:
 	/** UObject Interface */
 #if WITH_EDITOR
-	virtual void PostEditChangeProperty(struct FPropertyChangedEvent& PropertyChangedEvent) override;
+	DATAFLOWENGINE_API virtual void PostEditChangeProperty(struct FPropertyChangedEvent& PropertyChangedEvent) override;
 #endif
-	virtual void PostLoad() override;
+	DATAFLOWENGINE_API virtual void PostLoad() override;
 	/** End UObject Interface */
 
-	void Serialize(FArchive& Ar);
+	DATAFLOWENGINE_API void Serialize(FArchive& Ar);
 
 	/** Accessors for internal geometry collection */
 	TSharedPtr<const Dataflow::FGraph, ESPMode::ThreadSafe> GetDataflow() const { return Dataflow; }
@@ -111,8 +111,8 @@ public:
 	//
 	// Render Targets
 	//
-	void AddRenderTarget(UDataflowEdNode*);
-	void RemoveRenderTarget(UDataflowEdNode*);
+	DATAFLOWENGINE_API void AddRenderTarget(UDataflowEdNode*);
+	DATAFLOWENGINE_API void RemoveRenderTarget(UDataflowEdNode*);
 	const TArray<const UDataflowEdNode*>& GetRenderTargets() const { return RenderTargets; }
 	const Dataflow::FTimestamp& GetRenderingTimestamp() const { return LastModifiedRenderTarget; }
 
diff --git a/Engine/Source/Runtime/Experimental/Dataflow/Engine/Public/Dataflow/DataflowObjectInterface.h b/Engine/Source/Runtime/Experimental/Dataflow/Engine/Public/Dataflow/DataflowObjectInterface.h
--- a/Engine/Source/Runtime/Experimental/Dataflow/Engine/Public/Dataflow/DataflowObjectInterface.h
+++ b/Engine/Source/Runtime/Experimental/Dataflow/Engine/Public/Dataflow/DataflowObjectInterface.h
@@ -13,7 +13,7 @@ class UObject;
 namespace Dataflow
 {
 	template<class Base = FContextSingle>
-	class DATAFLOWENGINE_API TEngineContext : public Base
+	class TEngineContext : public Base
 	{
 	public:
 		DATAFLOW_CONTEXT_INTERNAL(Base, TEngineContext);
@@ -34,4 +34,4 @@ namespace Dataflow
 	typedef TEngineContext<FContextSingle> FEngineContext;
 	typedef TEngineContext<FContextThreaded> FEngineContextThreaded;
 
-}
\ No newline at end of file
+}
diff --git a/Engine/Source/Runtime/Experimental/Dataflow/Engine/Public/Dataflow/DataflowRenderingFactory.h b/Engine/Source/Runtime/Experimental/Dataflow/Engine/Public/Dataflow/DataflowRenderingFactory.h
--- a/Engine/Source/Runtime/Experimental/Dataflow/Engine/Public/Dataflow/DataflowRenderingFactory.h
+++ b/Engine/Source/Runtime/Experimental/Dataflow/Engine/Public/Dataflow/DataflowRenderingFactory.h
@@ -13,7 +13,7 @@ namespace Dataflow
 {
 	class FContext;
 
-	struct DATAFLOWENGINE_API FGraphRenderingState {
+	struct FGraphRenderingState {
 		FGraphRenderingState(const FDataflowNode* InNode, const FRenderingParameter& InParameters, Dataflow::FContext& InContext)
 			: Node(InNode)
 			, RenderType(InParameters.Type)
diff --git a/Engine/Source/Runtime/Experimental/FieldSystem/Source/FieldSystemEngine/Public/Field/FieldSystemActor.h b/Engine/Source/Runtime/Experimental/FieldSystem/Source/FieldSystemEngine/Public/Field/FieldSystemActor.h
--- a/Engine/Source/Runtime/Experimental/FieldSystem/Source/FieldSystemEngine/Public/Field/FieldSystemActor.h
+++ b/Engine/Source/Runtime/Experimental/FieldSystem/Source/FieldSystemEngine/Public/Field/FieldSystemActor.h
@@ -17,8 +17,8 @@
 
 class UFieldSystemComponent;
 
-UCLASS(meta=(ChildCanTick))
-class FIELDSYSTEMENGINE_API AFieldSystemActor: public AActor
+UCLASS(meta=(ChildCanTick), MinimalAPI)
+class AFieldSystemActor: public AActor
 {
 	GENERATED_UCLASS_BODY()
 
@@ -28,5 +28,5 @@ public:
 	UPROPERTY(VisibleAnywhere, BlueprintReadOnly, Category = Field, meta = (ExposeFunctionCategories = "Components|FieldSystem", AllowPrivateAccess = "true"))
 	TObjectPtr<UFieldSystemComponent> FieldSystemComponent;
 	UFieldSystemComponent* GetFieldSystemComponent() const { return FieldSystemComponent; }
-	virtual void OnConstruction(const FTransform& Transform) override;
+	FIELDSYSTEMENGINE_API virtual void OnConstruction(const FTransform& Transform) override;
 };
diff --git a/Engine/Source/Runtime/Experimental/FieldSystem/Source/FieldSystemEngine/Public/Field/FieldSystemAsset.h b/Engine/Source/Runtime/Experimental/FieldSystem/Source/FieldSystemEngine/Public/Field/FieldSystemAsset.h
--- a/Engine/Source/Runtime/Experimental/FieldSystem/Source/FieldSystemEngine/Public/Field/FieldSystemAsset.h
+++ b/Engine/Source/Runtime/Experimental/FieldSystem/Source/FieldSystemEngine/Public/Field/FieldSystemAsset.h
@@ -16,8 +16,8 @@
 class FArchive;
 
 
-UCLASS()
-class FIELDSYSTEMENGINE_API UFieldSystem : public UObject
+UCLASS(MinimalAPI)
+class UFieldSystem : public UObject
 {
 	GENERATED_BODY()
 
@@ -28,7 +28,7 @@ public:
 
 	void Reset() { Commands.Empty(); }
 
-	void Serialize(FArchive& Ar);
+	FIELDSYSTEMENGINE_API void Serialize(FArchive& Ar);
 
 	TArray< FFieldSystemCommand > Commands;
 
