[Engine]

diff --git a/Engine/Source/Runtime/Experimental/FieldSystem/Source/FieldSystemEngine/Public/Field/FieldSystemComponent.h b/Engine/Source/Runtime/Experimental/FieldSystem/Source/FieldSystemEngine/Public/Field/FieldSystemComponent.h
--- a/Engine/Source/Runtime/Experimental/FieldSystem/Source/FieldSystemEngine/Public/Field/FieldSystemComponent.h
+++ b/Engine/Source/Runtime/Experimental/FieldSystem/Source/FieldSystemEngine/Public/Field/FieldSystemComponent.h
@@ -33,8 +33,8 @@ class FChaosSolversModule;
 /**
 *	FieldSystemComponent
 */
-UCLASS(meta = (BlueprintSpawnableComponent))
-class FIELDSYSTEMENGINE_API UFieldSystemComponent : public UPrimitiveComponent
+UCLASS(meta = (BlueprintSpawnableComponent), MinimalAPI)
+class UFieldSystemComponent : public UPrimitiveComponent
 {
 	GENERATED_UCLASS_BODY()
 		friend class FFieldSystemEditorCommands;
@@ -45,7 +45,7 @@ public:
 	//~ Begin USceneComponent Interface.
 
 	//~ Begin UPrimitiveComponent Interface.
-	virtual FPrimitiveSceneProxy* CreateSceneProxy() override;
+	FIELDSYSTEMENGINE_API virtual FPrimitiveSceneProxy* CreateSceneProxy() override;
 	//~ End UPrimitiveComponent Interface.
 
 	/** Set the field system asset @todo(remove the field system, we dont need the asset */
@@ -93,7 +93,7 @@ public:
 	*
 	*/
 	UFUNCTION(BlueprintCallable, Category = "Field", DisplayName = "Apply Uniform Force", meta = (UnsafeDuringActorConstruction = "true"))
-	void ApplyLinearForce(UPARAM(DisplayName = "Enable Field") bool Enabled,
+	FIELDSYSTEMENGINE_API void ApplyLinearForce(UPARAM(DisplayName = "Enable Field") bool Enabled,
 			UPARAM(DisplayName = "Uniform Direction") FVector Direction,
 			UPARAM(DisplayName = "Field Magnitude") float Magnitude);
 
@@ -108,7 +108,7 @@ public:
 	*
 	*/
 	UFUNCTION(BlueprintCallable, Category = "Field", DisplayName = "Set Dynamic State", meta = (UnsafeDuringActorConstruction = "true"))
-	void ApplyStayDynamicField(UPARAM(DisplayName = "Enable Field") bool Enabled,
+	FIELDSYSTEMENGINE_API void ApplyStayDynamicField(UPARAM(DisplayName = "Enable Field") bool Enabled,
 			UPARAM(DisplayName = "Center Position") FVector Position,
 			UPARAM(DisplayName = "Field Radius") float Radius);
 
@@ -123,7 +123,7 @@ public:
 	*
 	*/
 	UFUNCTION(BlueprintCallable, Category = "Field", DisplayName = "Apply Radial Force", meta = (UnsafeDuringActorConstruction = "true"))
-	void ApplyRadialForce(UPARAM(DisplayName = "Enable Field") bool Enabled,
+	FIELDSYSTEMENGINE_API void ApplyRadialForce(UPARAM(DisplayName = "Enable Field") bool Enabled,
 			UPARAM(DisplayName = "Center Position") FVector Position,
 			UPARAM(DisplayName = "Field Magnitude") float Magnitude);
 
@@ -140,7 +140,7 @@ public:
 	*
 	*/
 	UFUNCTION(BlueprintCallable, Category = "Field", DisplayName = "Falloff Radial Force", meta = (UnsafeDuringActorConstruction = "true"))
-	void ApplyRadialVectorFalloffForce(UPARAM(DisplayName = "Enable Field") bool Enabled,
+	FIELDSYSTEMENGINE_API void ApplyRadialVectorFalloffForce(UPARAM(DisplayName = "Enable Field") bool Enabled,
 			UPARAM(DisplayName = "Center Position") FVector Position,
 			UPARAM(DisplayName = "Falloff Radius") float Radius,
 			UPARAM(DisplayName = "Field Magnitude") float Magnitude);
@@ -159,7 +159,7 @@ public:
 	*
 	*/
 	UFUNCTION(BlueprintCallable, Category = "Field", DisplayName = "Falloff Uniform Force", meta = (UnsafeDuringActorConstruction = "true"))
-	void ApplyUniformVectorFalloffForce(UPARAM(DisplayName = "Enable Field") bool Enabled,
+	FIELDSYSTEMENGINE_API void ApplyUniformVectorFalloffForce(UPARAM(DisplayName = "Enable Field") bool Enabled,
 			UPARAM(DisplayName = "Center Position") FVector Position,
 			UPARAM(DisplayName = "Uniform Direction") FVector Direction,
 			UPARAM(DisplayName = "Falloff Radius") float Radius,
@@ -179,7 +179,7 @@ public:
 	*
 	*/
 	UFUNCTION(BlueprintCallable, Category = "Field", DisplayName = "Apply External Strain", meta = (UnsafeDuringActorConstruction = "true"))
-	void ApplyStrainField(UPARAM(DisplayName = "Enable Field") bool Enabled,
+	FIELDSYSTEMENGINE_API void ApplyStrainField(UPARAM(DisplayName = "Enable Field") bool Enabled,
 			UPARAM(DisplayName = "Center Position") FVector Position,
 			UPARAM(DisplayName = "Falloff Radius") float Radius,
 			UPARAM(DisplayName = "Field Magnitude") float Magnitude,
@@ -199,7 +199,7 @@ public:
 	*
 	*/
 	UFUNCTION(BlueprintCallable, Category = "Field", DisplayName = "Add Transient Field", meta = (UnsafeDuringActorConstruction = "true"))
-	void ApplyPhysicsField(UPARAM(DisplayName = "Enable Field") bool Enabled,
+	FIELDSYSTEMENGINE_API void ApplyPhysicsField(UPARAM(DisplayName = "Enable Field") bool Enabled,
 			UPARAM(DisplayName = "Physics Type") EFieldPhysicsType Target,
 			UPARAM(DisplayName = "Meta Data") UFieldSystemMetaData* MetaData,
 			UPARAM(DisplayName = "Field Node") UFieldNodeBase* Field);
@@ -223,7 +223,7 @@ public:
 	*
 	*/
 	UFUNCTION(BlueprintCallable, Category = "Field", DisplayName = "Add Persistent Field", meta = (UnsafeDuringActorConstruction = "true"))
-	void AddPersistentField(UPARAM(DisplayName = "Enable Field") bool Enabled,
+	FIELDSYSTEMENGINE_API void AddPersistentField(UPARAM(DisplayName = "Enable Field") bool Enabled,
 			UPARAM(DisplayName = "Physics Type") EFieldPhysicsType Target,
 			UPARAM(DisplayName = "Meta Data")  UFieldSystemMetaData* MetaData,
 			UPARAM(DisplayName = "Field Node") UFieldNodeBase* Field);
@@ -234,7 +234,7 @@ public:
 	*
 	*/
 	UFUNCTION(BlueprintCallable, Category = "Field", DisplayName = "Remove Persistent Fields", meta = (UnsafeDuringActorConstruction = "true"))
-	void RemovePersistentFields();
+	FIELDSYSTEMENGINE_API void RemovePersistentFields();
 
 	//
 	// Blueprint construction field interface
@@ -256,7 +256,7 @@ public:
 	*/
 
 	UFUNCTION(BlueprintCallable, Category = "Field", DisplayName = "Add Construction Field")
-	void AddFieldCommand(UPARAM(DisplayName = "Enable Field") bool Enabled,
+	FIELDSYSTEMENGINE_API void AddFieldCommand(UPARAM(DisplayName = "Enable Field") bool Enabled,
 			UPARAM(DisplayName = "Physics Type") EFieldPhysicsType Target,
 			UPARAM(DisplayName = "Meta Data") UFieldSystemMetaData* MetaData,
 			UPARAM(DisplayName = "Field Node") UFieldNodeBase* Field);
@@ -268,7 +268,7 @@ public:
 	*/
 
 	UFUNCTION(BlueprintCallable, Category = "Field", DisplayName = "Remove Construction Fields")
-	void ResetFieldSystem();
+	FIELDSYSTEMENGINE_API void ResetFieldSystem();
 
 	/** Get all the construction fields*/
 	const TArray< FFieldSystemCommand >& GetConstructionFields() const { return SetupConstructionFields; }
@@ -276,25 +276,25 @@ public:
 protected:
 
 	/** Get ell ethe supported physics scenes */
-	TSet<FPhysScene_Chaos*> GetPhysicsScenes() const;
+	FIELDSYSTEMENGINE_API TSet<FPhysScene_Chaos*> GetPhysicsScenes() const;
 
 	/** Get ell the supported physics solvers */
-	TArray<Chaos::FPhysicsSolverBase*> GetPhysicsSolvers() const;
+	FIELDSYSTEMENGINE_API TArray<Chaos::FPhysicsSolverBase*> GetPhysicsSolvers() const;
 
 	/** Build a physics field command and dispatch it */
-	void BuildFieldCommand(bool Enabled, EFieldPhysicsType Target, UFieldSystemMetaData* MetaData, UFieldNodeBase* Field, const bool IsTransient);
+	FIELDSYSTEMENGINE_API void BuildFieldCommand(bool Enabled, EFieldPhysicsType Target, UFieldSystemMetaData* MetaData, UFieldNodeBase* Field, const bool IsTransient);
 
 	/** Dispatch the field command to chaos/world */
-	void DispatchFieldCommand(const FFieldSystemCommand& InCommand, const bool IsTransient);
+	FIELDSYSTEMENGINE_API void DispatchFieldCommand(const FFieldSystemCommand& InCommand, const bool IsTransient);
 
 	/** Remove the persistent commands from chaos/world  */
-	void ClearFieldCommands();
+	FIELDSYSTEMENGINE_API void ClearFieldCommands();
 
 	//~ Begin UActorComponent Interface.
-	virtual void OnCreatePhysicsState() override;
-	virtual void OnDestroyPhysicsState() override;
-	virtual bool ShouldCreatePhysicsState() const override;
-	virtual bool HasValidPhysicsState() const override;
+	FIELDSYSTEMENGINE_API virtual void OnCreatePhysicsState() override;
+	FIELDSYSTEMENGINE_API virtual void OnDestroyPhysicsState() override;
+	FIELDSYSTEMENGINE_API virtual bool ShouldCreatePhysicsState() const override;
+	FIELDSYSTEMENGINE_API virtual bool HasValidPhysicsState() const override;
 	//~ End UActorComponent Interface.
 
 	/** Chaos module linked to that component */
diff --git a/Engine/Source/Runtime/Experimental/FieldSystem/Source/FieldSystemEngine/Public/Field/FieldSystemObjects.h b/Engine/Source/Runtime/Experimental/FieldSystem/Source/FieldSystemEngine/Public/Field/FieldSystemObjects.h
--- a/Engine/Source/Runtime/Experimental/FieldSystem/Source/FieldSystemEngine/Public/Field/FieldSystemObjects.h
+++ b/Engine/Source/Runtime/Experimental/FieldSystem/Source/FieldSystemEngine/Public/Field/FieldSystemObjects.h
@@ -19,8 +19,8 @@
 * Context :
 *   Contexts are used to pass extra data into the field evaluation.
 */
-UCLASS()
-class FIELDSYSTEMENGINE_API UFieldSystemMetaData : public UActorComponent
+UCLASS(MinimalAPI)
+class UFieldSystemMetaData : public UActorComponent
 {
 	GENERATED_BODY()
 
@@ -33,8 +33,8 @@ public:
 /*
 * UFieldSystemMetaDataIteration : Not used anymore, just hiding it right now but will probably be deprecated if not used in the future
 */
-UCLASS()
-class FIELDSYSTEMENGINE_API UFieldSystemMetaDataIteration : public UFieldSystemMetaData
+UCLASS(MinimalAPI)
+class UFieldSystemMetaDataIteration : public UFieldSystemMetaData
 {
 	GENERATED_BODY()
 
@@ -42,14 +42,14 @@ public:
 	virtual ~UFieldSystemMetaDataIteration() {}
 	virtual FFieldSystemMetaData::EMetaType Type() const override { return  FFieldSystemMetaData::EMetaType::ECommandData_Iteration; }
 
-	virtual FFieldSystemMetaData* NewMetaData() const override;
+	FIELDSYSTEMENGINE_API virtual FFieldSystemMetaData* NewMetaData() const override;
 
 	/**
 	 * Set the number of iteration type
 	 * @param    Iterations Number of iterations (WIP)
 	 */
 	UFUNCTION(BlueprintPure, Category = "Field", meta = (Iterations = "1"))
-	UFieldSystemMetaDataIteration* SetMetaDataIteration(UPARAM(DisplayName = "Iteration Count") int Iterations);
+	FIELDSYSTEMENGINE_API UFieldSystemMetaDataIteration* SetMetaDataIteration(UPARAM(DisplayName = "Iteration Count") int Iterations);
 
 	/** Number of iterations (WIP) */
 	UPROPERTY(BlueprintReadWrite, EditAnywhere, Category = "Field", DisplayName = "Iteration Count")
@@ -59,22 +59,22 @@ public:
 /*
 * UFieldSystemMetaDataProcessingResolution
 */
-UCLASS(ClassGroup = "Field", meta = (BlueprintSpawnableComponent, ToolTip = "Control the set of particles on which the field will be applied"), ShowCategories = ("Field"), DisplayName = "MetaDataResolution")
-class FIELDSYSTEMENGINE_API UFieldSystemMetaDataProcessingResolution : public UFieldSystemMetaData
+UCLASS(ClassGroup = "Field", meta = (BlueprintSpawnableComponent, ToolTip = "Control the set of particles on which the field will be applied"), ShowCategories = ("Field"), DisplayName = "MetaDataResolution", MinimalAPI)
+class UFieldSystemMetaDataProcessingResolution : public UFieldSystemMetaData
 {
 	GENERATED_BODY()
 
 public:
 	virtual ~UFieldSystemMetaDataProcessingResolution() {}
 	virtual FFieldSystemMetaData::EMetaType Type() const override { return  FFieldSystemMetaData::EMetaType::ECommandData_ProcessingResolution; }
-	virtual FFieldSystemMetaData* NewMetaData() const override;
+	FIELDSYSTEMENGINE_API virtual FFieldSystemMetaData* NewMetaData() const override;
 
 	/**
 	 * Set the processing resolution type
 	 * @param    ResolutionType Type of processing resolution used to select the particles on which the field will be applied
 	 */
 	UFUNCTION(BlueprintPure, Category = "Field", DisplayName = "Set Meta Data Resolution")
-	UFieldSystemMetaDataProcessingResolution* SetMetaDataaProcessingResolutionType(EFieldResolutionType ResolutionType);
+	FIELDSYSTEMENGINE_API UFieldSystemMetaDataProcessingResolution* SetMetaDataaProcessingResolutionType(EFieldResolutionType ResolutionType);
 
 	/** Precessing resolution type used to select the particles on which the field will be applied */
 	UPROPERTY(EditAnywhere, BlueprintReadWrite, Category = "Field")
@@ -84,15 +84,15 @@ public:
 /*
 * UFieldSystemMetaDataFilter
 */
-UCLASS(ClassGroup = "Field", meta = (BlueprintSpawnableComponent, ToolTip = "Filter the particles on which the field will be applied"), ShowCategories = ("Field"), DisplayName = "MetaDataFilter")
-class FIELDSYSTEMENGINE_API UFieldSystemMetaDataFilter : public UFieldSystemMetaData
+UCLASS(ClassGroup = "Field", meta = (BlueprintSpawnableComponent, ToolTip = "Filter the particles on which the field will be applied"), ShowCategories = ("Field"), DisplayName = "MetaDataFilter", MinimalAPI)
+class UFieldSystemMetaDataFilter : public UFieldSystemMetaData
 {
 	GENERATED_BODY()
 
 public:
 	virtual ~UFieldSystemMetaDataFilter() {}
 	virtual FFieldSystemMetaData::EMetaType Type() const override { return  FFieldSystemMetaData::EMetaType::ECommandData_Filter; }
-	virtual FFieldSystemMetaData* NewMetaData() const override;
+	FIELDSYSTEMENGINE_API virtual FFieldSystemMetaData* NewMetaData() const override;
 
 	/**
 	 * Set the particles filter type
@@ -101,7 +101,7 @@ public:
 	 * @param    PositionType Position type used to compute the samples positions
 	 */
 	UFUNCTION(BlueprintPure, Category = "Field", DisplayName = "Set Meta Data Filter")
-	UFieldSystemMetaDataFilter* SetMetaDataFilterType(UPARAM(DisplayName = "State Type") EFieldFilterType FilterType, 
+	FIELDSYSTEMENGINE_API UFieldSystemMetaDataFilter* SetMetaDataFilterType(UPARAM(DisplayName = "State Type") EFieldFilterType FilterType, 
 													  UPARAM(DisplayName = "Object Type") EFieldObjectType ObjectType,
 													  UPARAM(DisplayName = "Position Type") EFieldPositionType PositionType );
 
@@ -121,8 +121,8 @@ public:
 /**
 * Field Evaluation
 */
-UCLASS()
-class FIELDSYSTEMENGINE_API UFieldNodeBase : public UActorComponent
+UCLASS(MinimalAPI)
+class UFieldNodeBase : public UActorComponent
 {
 	GENERATED_BODY()
 
@@ -137,8 +137,8 @@ public:
 /**
 * FieldNodeInt
 */
-UCLASS()
-class FIELDSYSTEMENGINE_API UFieldNodeInt : public UFieldNodeBase
+UCLASS(MinimalAPI)
+class UFieldNodeInt : public UFieldNodeBase
 {
 	GENERATED_BODY()
 
@@ -150,8 +150,8 @@ public:
 /**
 * FieldNodeFloat
 */
-UCLASS()
-class FIELDSYSTEMENGINE_API UFieldNodeFloat : public UFieldNodeBase
+UCLASS(MinimalAPI)
+class UFieldNodeFloat : public UFieldNodeBase
 {
 	GENERATED_BODY()
 
@@ -163,8 +163,8 @@ public:
 /**
 * FieldNodeVector
 */
-UCLASS()
-class FIELDSYSTEMENGINE_API UFieldNodeVector : public UFieldNodeBase
+UCLASS(MinimalAPI)
+class UFieldNodeVector : public UFieldNodeBase
 {
 	GENERATED_BODY()
 public:
@@ -176,8 +176,8 @@ public:
 /**
 * UUniformInteger
 **/
-UCLASS(ClassGroup = "Field", meta = (BlueprintSpawnableComponent, ToolTip = "Set a uniform integer value independently of the sample position. The output is equal to magnitude"), ShowCategories = ("Field"))
-class FIELDSYSTEMENGINE_API UUniformInteger : public UFieldNodeInt
+UCLASS(ClassGroup = "Field", meta = (BlueprintSpawnableComponent, ToolTip = "Set a uniform integer value independently of the sample position. The output is equal to magnitude"), ShowCategories = ("Field"), MinimalAPI)
+class UUniformInteger : public UFieldNodeInt
 {
 	GENERATED_BODY()
 public:
@@ -189,14 +189,14 @@ public:
 
 	virtual ~UUniformInteger() {}
 
-	virtual FFieldNodeBase* NewEvaluationGraph(TArray<const UFieldNodeBase*>& Nodes) const override;
+	FIELDSYSTEMENGINE_API virtual FFieldNodeBase* NewEvaluationGraph(TArray<const UFieldNodeBase*>& Nodes) const override;
 
 	/**
 	 * Set a uniform integer value independently of the sample position. The output is equal to magnitude
 	 * @param    Magnitude The field output will be equal the magnitude
 	 */
 	UFUNCTION(BlueprintPure, Category = "Field", meta = (Magnitude = "0"))
-	UUniformInteger* SetUniformInteger(UPARAM(DisplayName = "Field Magnitude") int32 Magnitude);
+	FIELDSYSTEMENGINE_API UUniformInteger* SetUniformInteger(UPARAM(DisplayName = "Field Magnitude") int32 Magnitude);
 
 	/** The field output will be equal the magnitude */
 	UPROPERTY(BlueprintReadWrite, EditAnywhere, Category = "Field", DisplayName = "Field Magnitude")
@@ -207,8 +207,8 @@ public:
 /**
 * URadialIntMask
 */
-UCLASS(ClassGroup = "Field", meta = (BlueprintSpawnableComponent, ToolTip = "This function first defines a radial integer field equal to Interior-value inside a sphere / Exterior-value outside. This field will be used alongside the particle input value and the mask condition to compute the particle output value."), ShowCategories = ("Field"))
-class FIELDSYSTEMENGINE_API URadialIntMask : public UFieldNodeInt
+UCLASS(ClassGroup = "Field", meta = (BlueprintSpawnableComponent, ToolTip = "This function first defines a radial integer field equal to Interior-value inside a sphere / Exterior-value outside. This field will be used alongside the particle input value and the mask condition to compute the particle output value."), ShowCategories = ("Field"), MinimalAPI)
+class URadialIntMask : public UFieldNodeInt
 {
 	GENERATED_BODY()
 public:
@@ -223,7 +223,7 @@ public:
 	{}
 	virtual ~URadialIntMask() {}
 
-	virtual FFieldNodeBase* NewEvaluationGraph(TArray<const UFieldNodeBase*>& Nodes) const override;
+	FIELDSYSTEMENGINE_API virtual FFieldNodeBase* NewEvaluationGraph(TArray<const UFieldNodeBase*>& Nodes) const override;
 
 	/**
 	* This function first defines a radial integer field equal to Interior-value inside a sphere / Exterior-value outside. This field will be used alongside the particle input value and the mask condition to compute the particle output value.
@@ -237,7 +237,7 @@ public:
 	 * @param    SetMaskConditionIn If the mask condition is set to always the output value will be the intermediate one. If set to not interior/exterior the output value will be the intermediate one if the input is different from the interior/exterior value
 	 */
 	UFUNCTION(BlueprintPure, Category = "Field", meta = (InteriorValue = "1"), DisplayName = "Set Radial Mask")
-	URadialIntMask* SetRadialIntMask(UPARAM(DisplayName = "Mask Radius") float Radius,
+	FIELDSYSTEMENGINE_API URadialIntMask* SetRadialIntMask(UPARAM(DisplayName = "Mask Radius") float Radius,
 			UPARAM(DisplayName = "Center Position") FVector Position,
 			UPARAM(DisplayName = "Interior Value") int32 InteriorValue,
 			UPARAM(DisplayName = "Exterior Value") int32 ExteriorValue,
@@ -268,8 +268,8 @@ public:
 /**
 * UUniformScalar
 **/
-UCLASS(ClassGroup = "Field", meta = (BlueprintSpawnableComponent, ToolTip = "Set a uniform scalar value independently of the sample position. The output is equal to magnitude"), ShowCategories = ("Field"))
-class FIELDSYSTEMENGINE_API UUniformScalar : public UFieldNodeFloat
+UCLASS(ClassGroup = "Field", meta = (BlueprintSpawnableComponent, ToolTip = "Set a uniform scalar value independently of the sample position. The output is equal to magnitude"), ShowCategories = ("Field"), MinimalAPI)
+class UUniformScalar : public UFieldNodeFloat
 {
 	GENERATED_BODY()
 public:
@@ -281,14 +281,14 @@ public:
 
 	virtual ~UUniformScalar() {}
 
-	virtual FFieldNodeBase* NewEvaluationGraph(TArray<const UFieldNodeBase*>& Nodes) const override;
+	FIELDSYSTEMENGINE_API virtual FFieldNodeBase* NewEvaluationGraph(TArray<const UFieldNodeBase*>& Nodes) const override;
 
 	/**
 	 * Set a uniform scalar value independently of the sample position. The output is equal to magnitude
 	 * @param    Magnitude The field output will be equal the magnitude
 	 */
 	UFUNCTION(BlueprintPure, Category = "Field", meta = (Magnitude = "1.0"))
-	UUniformScalar* SetUniformScalar(UPARAM(DisplayName = "Field Magnitude") float Magnitude);
+	FIELDSYSTEMENGINE_API UUniformScalar* SetUniformScalar(UPARAM(DisplayName = "Field Magnitude") float Magnitude);
 
 	/** The field output will be equal the magnitude */
 	UPROPERTY(BlueprintReadWrite, EditAnywhere, Category = "Field", DisplayName = "Field Magnitude")
@@ -299,8 +299,8 @@ public:
 /**
 * UWaveScalar
 **/
-UCLASS(ClassGroup = "Field", meta = (BlueprintSpawnableComponent, ToolTip = "Set a temporal wave scalar value according to the sample distance from the field position."), ShowCategories = ("Field"))
-class FIELDSYSTEMENGINE_API UWaveScalar : public UFieldNodeFloat
+UCLASS(ClassGroup = "Field", meta = (BlueprintSpawnableComponent, ToolTip = "Set a temporal wave scalar value according to the sample distance from the field position."), ShowCategories = ("Field"), MinimalAPI)
+class UWaveScalar : public UFieldNodeFloat
 {
 	GENERATED_BODY()
 public:
@@ -317,7 +317,7 @@ public:
 
 	virtual ~UWaveScalar() {}
 
-	virtual FFieldNodeBase* NewEvaluationGraph(TArray<const UFieldNodeBase*>& Nodes) const override;
+	FIELDSYSTEMENGINE_API virtual FFieldNodeBase* NewEvaluationGraph(TArray<const UFieldNodeBase*>& Nodes) const override;
 
 	/**
 	 * Set a temporal wave scalar value according to the sample distance from the field position.
@@ -329,7 +329,7 @@ public:
 	 * @param    Falloff Type of falloff function used if the falloff function is picked
 	 */
 	UFUNCTION(BlueprintPure, Category = "Field", meta = (Magnitude = "1.0", Wavelength = "1000", Period = "1", HidePin = "Time"), DisplayName = "Set Wave Scalar")
-	UWaveScalar* SetWaveScalar(UPARAM(DisplayName = "Field Magnitude") float Magnitude,
+	FIELDSYSTEMENGINE_API UWaveScalar* SetWaveScalar(UPARAM(DisplayName = "Field Magnitude") float Magnitude,
 			UPARAM(DisplayName = "Center Position") FVector Position,
 			UPARAM(DisplayName = "Wave Length") float Wavelength,
 			UPARAM(DisplayName = "Wave Period") float Period,
@@ -367,8 +367,8 @@ public:
 /**
 * RadialFalloff
 */
-UCLASS(ClassGroup = "Field", meta = (BlueprintSpawnableComponent, ToolTip = "Sphere scalar field that will be defined only within a sphere"), ShowCategories = ("Field"))
-class FIELDSYSTEMENGINE_API URadialFalloff : public UFieldNodeFloat
+UCLASS(ClassGroup = "Field", meta = (BlueprintSpawnableComponent, ToolTip = "Sphere scalar field that will be defined only within a sphere"), ShowCategories = ("Field"), MinimalAPI)
+class URadialFalloff : public UFieldNodeFloat
 {
 	GENERATED_BODY()
 public:
@@ -386,7 +386,7 @@ public:
 
 	virtual ~URadialFalloff() {}
 
-	virtual FFieldNodeBase* NewEvaluationGraph(TArray<const UFieldNodeBase*>& Nodes) const override;
+	FIELDSYSTEMENGINE_API virtual FFieldNodeBase* NewEvaluationGraph(TArray<const UFieldNodeBase*>& Nodes) const override;
 
 	/**
 	 * Sphere scalar field that will be defined only within a sphere
@@ -399,7 +399,7 @@ public:
 	 * @param    Falloff Type of falloff function used if the falloff function is picked
 	 */
 	UFUNCTION(BlueprintPure, Category = "Field", meta = (Magnitude = "1.0", MinRange = "0.0", MaxRange = "1.0"), DisplayName = "Set Radial Falloff")
-	URadialFalloff* SetRadialFalloff(UPARAM(DisplayName = "Field Magnitude") float Magnitude,
+	FIELDSYSTEMENGINE_API URadialFalloff* SetRadialFalloff(UPARAM(DisplayName = "Field Magnitude") float Magnitude,
 			UPARAM(DisplayName = "Min Range") float MinRange,
 			UPARAM(DisplayName = "Max Range") float MaxRange,
 			UPARAM(DisplayName = "Default Value") float Default,
@@ -439,8 +439,8 @@ public:
 /**
 * PlaneFalloff
 */
-UCLASS(ClassGroup = "Field", meta = (BlueprintSpawnableComponent, ToolTip = "Plane scalar field that will be defined only within a distance from a plane"), ShowCategories = ("Field"))
-class FIELDSYSTEMENGINE_API UPlaneFalloff : public UFieldNodeFloat
+UCLASS(ClassGroup = "Field", meta = (BlueprintSpawnableComponent, ToolTip = "Plane scalar field that will be defined only within a distance from a plane"), ShowCategories = ("Field"), MinimalAPI)
+class UPlaneFalloff : public UFieldNodeFloat
 {
 	GENERATED_BODY()
 
@@ -460,7 +460,7 @@ public:
 
 	virtual ~UPlaneFalloff() {}
 
-	virtual FFieldNodeBase* NewEvaluationGraph(TArray<const UFieldNodeBase*>& Nodes) const override;
+	FIELDSYSTEMENGINE_API virtual FFieldNodeBase* NewEvaluationGraph(TArray<const UFieldNodeBase*>& Nodes) const override;
 
 	/**
 	 * Plane scalar field that will be defined only within a distance from a plane
@@ -474,7 +474,7 @@ public:
 	 * @param    Falloff Type of falloff function used to model the evolution of the field from the plane surface to the distance isosurface
 	 */
 	UFUNCTION(BlueprintPure, Category = "Field", meta = (Magnitude = "1.0", MinRange = "0.0", MaxRange = "1.0"))
-	UPlaneFalloff* SetPlaneFalloff(UPARAM(DisplayName = "Field Magnitude") float Magnitude,
+	FIELDSYSTEMENGINE_API UPlaneFalloff* SetPlaneFalloff(UPARAM(DisplayName = "Field Magnitude") float Magnitude,
 			UPARAM(DisplayName = "Min Range") float MinRange,
 			UPARAM(DisplayName = "Max Range") float MaxRange,
 			UPARAM(DisplayName = "Default Value") float Default,
@@ -519,8 +519,8 @@ public:
 /**
 * BoxFalloff
 */
-UCLASS(ClassGroup = "Field", meta = (BlueprintSpawnableComponent, ToolTip = "Box scalar field that will be defined only within a box"), ShowCategories = ("Field"))
-class FIELDSYSTEMENGINE_API UBoxFalloff : public UFieldNodeFloat
+UCLASS(ClassGroup = "Field", meta = (BlueprintSpawnableComponent, ToolTip = "Box scalar field that will be defined only within a box"), ShowCategories = ("Field"), MinimalAPI)
+class UBoxFalloff : public UFieldNodeFloat
 {
 	GENERATED_BODY()
 
@@ -538,7 +538,7 @@ public:
 
 	virtual ~UBoxFalloff() {}
 
-	virtual FFieldNodeBase* NewEvaluationGraph(TArray<const UFieldNodeBase*>& Nodes) const override;
+	FIELDSYSTEMENGINE_API virtual FFieldNodeBase* NewEvaluationGraph(TArray<const UFieldNodeBase*>& Nodes) const override;
 
 	/**
 	 * Box scalar field that will be defined only within a box
@@ -550,7 +550,7 @@ public:
 	 * @param    Falloff Type of falloff function used to model the evolution of the field from the box surface to the sample position
 	 */
 	UFUNCTION(BlueprintPure, Category = "Field", meta = (Magnitude = "1.0", MinRange = "0.0", MaxRange = "1.0"))
-	UBoxFalloff* SetBoxFalloff(UPARAM(DisplayName = "Field Magnitude")  float Magnitude,
+	FIELDSYSTEMENGINE_API UBoxFalloff* SetBoxFalloff(UPARAM(DisplayName = "Field Magnitude")  float Magnitude,
 			UPARAM(DisplayName = "Min Range")  float MinRange,
 			UPARAM(DisplayName = "Max Range")  float MaxRange,
 			UPARAM(DisplayName = "Default Value")  float Default,
@@ -586,8 +586,8 @@ public:
 /**
 * NoiseField
 */
-UCLASS(ClassGroup = "Field", meta = (BlueprintSpawnableComponent, ToolTip = "Defines a perlin noise scalar value if the sample is within a box"), ShowCategories = ("Field"))
-class FIELDSYSTEMENGINE_API UNoiseField : public UFieldNodeFloat
+UCLASS(ClassGroup = "Field", meta = (BlueprintSpawnableComponent, ToolTip = "Defines a perlin noise scalar value if the sample is within a box"), ShowCategories = ("Field"), MinimalAPI)
+class UNoiseField : public UFieldNodeFloat
 {
 	GENERATED_BODY()
 
@@ -602,7 +602,7 @@ public:
 
 	virtual ~UNoiseField() {}
 
-	virtual FFieldNodeBase* NewEvaluationGraph(TArray<const UFieldNodeBase*>& Nodes) const override;
+	FIELDSYSTEMENGINE_API virtual FFieldNodeBase* NewEvaluationGraph(TArray<const UFieldNodeBase*>& Nodes) const override;
 
 	/**
 	 * Defines a perlin noise scalar value if the sample is within a box
@@ -611,7 +611,7 @@ public:
 	 * @param    Transform Transform of the box in which the perlin noise will be defined
 	 */
 	UFUNCTION(BlueprintPure, Category = "Field", meta = (MinRange = "0.0", MaxRange = "1.0"), DisplayName = "Set Noise Field")
-	UNoiseField* SetNoiseField(UPARAM(DisplayName = "Min Range") float MinRange,
+	FIELDSYSTEMENGINE_API UNoiseField* SetNoiseField(UPARAM(DisplayName = "Min Range") float MinRange,
 			UPARAM(DisplayName = "Max Range") float MaxRange,
 			UPARAM(DisplayName = "Noise Transform")  FTransform Transform);
 
@@ -631,8 +631,8 @@ public:
 /**
 * UniformVector
 **/
-UCLASS(ClassGroup = "Field", meta = (BlueprintSpawnableComponent, ToolTip = "Set a uniform vector value independently of the sample position.The output is equal to magnitude * direction"), ShowCategories = ("Field"))
-class FIELDSYSTEMENGINE_API UUniformVector : public UFieldNodeVector
+UCLASS(ClassGroup = "Field", meta = (BlueprintSpawnableComponent, ToolTip = "Set a uniform vector value independently of the sample position.The output is equal to magnitude * direction"), ShowCategories = ("Field"), MinimalAPI)
+class UUniformVector : public UFieldNodeVector
 {
 	GENERATED_BODY()
 public:
@@ -645,7 +645,7 @@ public:
 
 	virtual ~UUniformVector() {}
 
-	virtual FFieldNodeBase* NewEvaluationGraph(TArray<const UFieldNodeBase*>& Nodes) const override;
+	FIELDSYSTEMENGINE_API virtual FFieldNodeBase* NewEvaluationGraph(TArray<const UFieldNodeBase*>& Nodes) const override;
 
 	/**
 	 * Set a uniform vector value independently of the sample position.The output is equal to magnitude * direction
@@ -653,7 +653,7 @@ public:
 	 * @param    Direction Normalized direction of the uniform vector field
 	 */
 	UFUNCTION(BlueprintPure, Category = "Field", meta = (Magnitude = "1.0"))
-	UUniformVector* SetUniformVector(UPARAM(DisplayName = "Field Magnitude") float Magnitude,
+	FIELDSYSTEMENGINE_API UUniformVector* SetUniformVector(UPARAM(DisplayName = "Field Magnitude") float Magnitude,
 			UPARAM(DisplayName = "Uniform Direction") FVector Direction);
 
 	/** Magnitude of the uniform vector field */
@@ -669,8 +669,8 @@ public:
 /**
 * RadialVector
 */
-UCLASS(ClassGroup = "Field", meta = (BlueprintSpawnableComponent, ToolTip = "Set a radial vector value, the direction being the vector from the sample position to the field one. The output is equal to magnitude * direction"), ShowCategories = ("Field"))
-class FIELDSYSTEMENGINE_API URadialVector : public UFieldNodeVector
+UCLASS(ClassGroup = "Field", meta = (BlueprintSpawnableComponent, ToolTip = "Set a radial vector value, the direction being the vector from the sample position to the field one. The output is equal to magnitude * direction"), ShowCategories = ("Field"), MinimalAPI)
+class URadialVector : public UFieldNodeVector
 {
 	GENERATED_BODY()
 public:
@@ -683,7 +683,7 @@ public:
 
 	virtual ~URadialVector() {}
 
-	virtual FFieldNodeBase* NewEvaluationGraph(TArray<const UFieldNodeBase*>& Nodes) const override;
+	FIELDSYSTEMENGINE_API virtual FFieldNodeBase* NewEvaluationGraph(TArray<const UFieldNodeBase*>& Nodes) const override;
 
 	/**
 	 * Set a radial vector value. The direction is the normalized vector from the field position to the sample one. The output is equal to this direction * magnitude.
@@ -691,7 +691,7 @@ public:
 	 * @param    Position Center position of the radial vector field
 	 */
 	UFUNCTION(BlueprintPure, Category = "Field", meta = (Magnitude = "1.0"))
-	URadialVector* SetRadialVector(UPARAM(DisplayName = "Field Magnitude") float Magnitude,
+	FIELDSYSTEMENGINE_API URadialVector* SetRadialVector(UPARAM(DisplayName = "Field Magnitude") float Magnitude,
 								   UPARAM(DisplayName = "Center Position") FVector Position);
 
 	/** Magnitude of the radial vector field */
@@ -706,8 +706,8 @@ public:
 /**
 * URandomVector
 */
-UCLASS(ClassGroup = "Field", meta = (BlueprintSpawnableComponent, ToolTip = "Set a random vector value independently of the sample position. The output is equal to magnitude * random direction "), ShowCategories = ("Field"))
-class FIELDSYSTEMENGINE_API URandomVector : public UFieldNodeVector
+UCLASS(ClassGroup = "Field", meta = (BlueprintSpawnableComponent, ToolTip = "Set a random vector value independently of the sample position. The output is equal to magnitude * random direction "), ShowCategories = ("Field"), MinimalAPI)
+class URandomVector : public UFieldNodeVector
 {
 	GENERATED_BODY()
 public:
@@ -719,14 +719,14 @@ public:
 
 	virtual ~URandomVector() {}
 
-	virtual FFieldNodeBase* NewEvaluationGraph(TArray<const UFieldNodeBase*>& Nodes) const override;
+	FIELDSYSTEMENGINE_API virtual FFieldNodeBase* NewEvaluationGraph(TArray<const UFieldNodeBase*>& Nodes) const override;
 
 	/**
 	 * Set a random vector value independently of the sample position. The output is equal to magnitude * random direction
 	 * @param    Magnitude Magnitude of the random vector field
 	 */
 	UFUNCTION(BlueprintPure, Category = "Field", meta = (Magnitude = "1.0"))
-	URandomVector* SetRandomVector(UPARAM(DisplayName = "Field Magnitude") float Magnitude);
+	FIELDSYSTEMENGINE_API URandomVector* SetRandomVector(UPARAM(DisplayName = "Field Magnitude") float Magnitude);
 
 	/** Magnitude of the random vector field*/
 	UPROPERTY(BlueprintReadWrite, EditAnywhere, Category = "Field", DisplayName = "Field Magnitude")
@@ -737,8 +737,8 @@ public:
 /**
 * UOperatorField
 */
-UCLASS(ClassGroup = "Field", meta = (BlueprintSpawnableComponent, ToolTip = "Compute an operation between 2 incoming fields"), ShowCategories = ("Field"))
-class FIELDSYSTEMENGINE_API UOperatorField : public UFieldNodeBase
+UCLASS(ClassGroup = "Field", meta = (BlueprintSpawnableComponent, ToolTip = "Compute an operation between 2 incoming fields"), ShowCategories = ("Field"), MinimalAPI)
+class UOperatorField : public UFieldNodeBase
 {
 	GENERATED_BODY()
 public:
@@ -753,10 +753,10 @@ public:
 
 	virtual ~UOperatorField() {}
 
-	virtual FFieldNodeBase::EFieldType Type() const;
+	FIELDSYSTEMENGINE_API virtual FFieldNodeBase::EFieldType Type() const;
 	virtual bool ResultsExpector() const override { return true; }
 
-	virtual FFieldNodeBase* NewEvaluationGraph(TArray<const UFieldNodeBase*>& Nodes) const override;
+	FIELDSYSTEMENGINE_API virtual FFieldNodeBase* NewEvaluationGraph(TArray<const UFieldNodeBase*>& Nodes) const override;
 
 	/**
 	 * Compute an operation between 2 incoming fields
@@ -766,7 +766,7 @@ public:
 	 * @param    Operation Type of math operation you want to perform between the 2 fields
 	 */
 	UFUNCTION(BlueprintPure, Category = "Field", meta = (Magnitude = "1.0"))
-	UOperatorField* SetOperatorField(UPARAM(DisplayName = "Field Magnitude") float Magnitude,
+	FIELDSYSTEMENGINE_API UOperatorField* SetOperatorField(UPARAM(DisplayName = "Field Magnitude") float Magnitude,
 			UPARAM(DisplayName = "Left Field") const UFieldNodeBase* LeftField,
 			UPARAM(DisplayName = "Right Field") const UFieldNodeBase* RightField,
 			UPARAM(DisplayName = "Field Operation") EFieldOperationType Operation);
@@ -791,8 +791,8 @@ public:
 /**
 * UToIntegerField
 */
-UCLASS(ClassGroup = "Field", meta = (BlueprintSpawnableComponent, ToolTip = "Convert a scalar field to a integer one"), ShowCategories = ("Field"))
-class FIELDSYSTEMENGINE_API UToIntegerField : public UFieldNodeInt
+UCLASS(ClassGroup = "Field", meta = (BlueprintSpawnableComponent, ToolTip = "Convert a scalar field to a integer one"), ShowCategories = ("Field"), MinimalAPI)
+class UToIntegerField : public UFieldNodeInt
 {
 	GENERATED_BODY()
 public:
@@ -804,14 +804,14 @@ public:
 
 	virtual ~UToIntegerField() {}
 
-	virtual FFieldNodeBase* NewEvaluationGraph(TArray<const UFieldNodeBase*>& Nodes) const override;
+	FIELDSYSTEMENGINE_API virtual FFieldNodeBase* NewEvaluationGraph(TArray<const UFieldNodeBase*>& Nodes) const override;
 
 	/**
 	 * Convert a float field to a integer one
 	 * @param    FloatField Float field to be converted to an an integer one
 	 */
 	UFUNCTION(BlueprintPure, Category = "Field", DisplayName = "Set To Integer Field")
-	UToIntegerField* SetToIntegerField(UPARAM(DisplayName = "Float Field") const UFieldNodeFloat* FloatField);
+	FIELDSYSTEMENGINE_API UToIntegerField* SetToIntegerField(UPARAM(DisplayName = "Float Field") const UFieldNodeFloat* FloatField);
 
 	/** Scalar field to be converted to an an integer one */
 	UPROPERTY(BlueprintReadWrite, EditAnywhere, Category = "Field", DisplayName = "Float Field")
@@ -821,8 +821,8 @@ public:
 /**
 * UToFloatField
 */
-UCLASS(ClassGroup = "Field", meta = (BlueprintSpawnableComponent, ToolTip = "Convert an integer field to a scalar one"), ShowCategories = ("Field"))
-class FIELDSYSTEMENGINE_API UToFloatField : public UFieldNodeFloat
+UCLASS(ClassGroup = "Field", meta = (BlueprintSpawnableComponent, ToolTip = "Convert an integer field to a scalar one"), ShowCategories = ("Field"), MinimalAPI)
+class UToFloatField : public UFieldNodeFloat
 {
 	GENERATED_BODY()
 public:
@@ -834,14 +834,14 @@ public:
 
 	virtual ~UToFloatField() {}
 
-	virtual FFieldNodeBase* NewEvaluationGraph(TArray<const UFieldNodeBase*>& Nodes) const override;
+	FIELDSYSTEMENGINE_API virtual FFieldNodeBase* NewEvaluationGraph(TArray<const UFieldNodeBase*>& Nodes) const override;
 
 	/**
 	 * Convert an integer field to a float one
 	 * @param    IntegerField Integer field to be converted to an a float one
 	 */
 	UFUNCTION(BlueprintPure, Category = "Field", DisplayName = "Set To Float Field")
-	UToFloatField* SetToFloatField(UPARAM(DisplayName = "Integer Field") const UFieldNodeInt* IntegerField);
+	FIELDSYSTEMENGINE_API UToFloatField* SetToFloatField(UPARAM(DisplayName = "Integer Field") const UFieldNodeInt* IntegerField);
 
 	/** Integer field to be converted to an a scalar one */
 	UPROPERTY(BlueprintReadWrite, EditAnywhere, Category = "Field", DisplayName = "Integer Field")
@@ -851,8 +851,8 @@ public:
 /**
 * UCullingField
 */
-UCLASS(ClassGroup = "Field", meta = (BlueprintSpawnableComponent, ToolTip = "Evaluate the input field according to the result of the culling field"), ShowCategories = ("Field"))
-class FIELDSYSTEMENGINE_API UCullingField : public UFieldNodeBase
+UCLASS(ClassGroup = "Field", meta = (BlueprintSpawnableComponent, ToolTip = "Evaluate the input field according to the result of the culling field"), ShowCategories = ("Field"), MinimalAPI)
+class UCullingField : public UFieldNodeBase
 {
 	GENERATED_BODY()
 public:
@@ -866,9 +866,9 @@ public:
 
 	virtual ~UCullingField() {}
 
-	virtual FFieldNodeBase* NewEvaluationGraph(TArray<const UFieldNodeBase*>& Nodes) const override;
+	FIELDSYSTEMENGINE_API virtual FFieldNodeBase* NewEvaluationGraph(TArray<const UFieldNodeBase*>& Nodes) const override;
 
-	virtual FFieldNodeBase::EFieldType Type() const override;
+	FIELDSYSTEMENGINE_API virtual FFieldNodeBase::EFieldType Type() const override;
 
 	/**
 	 * Evaluate the input field according to the result of the culling field.
@@ -878,7 +878,7 @@ public:
 	 * @param    Operation Evaluate the input field if the result of the culling field is equal to 0 (Inside) or different from 0 (Outside).
 	 */
 	UFUNCTION(BlueprintPure, Category = "Field")
-	UCullingField* SetCullingField(UPARAM(DisplayName = "Culling Field") const UFieldNodeBase* Culling,
+	FIELDSYSTEMENGINE_API UCullingField* SetCullingField(UPARAM(DisplayName = "Culling Field") const UFieldNodeBase* Culling,
 			UPARAM(DisplayName = "Input Field") const UFieldNodeBase* Field,
 			UPARAM(DisplayName = "Culling Operation") EFieldCullingOperationType Operation);
 
@@ -898,8 +898,8 @@ public:
 /**
 * UReturnResultsField
 */
-UCLASS(ClassGroup = "Field", meta = (BlueprintSpawnableComponent, ToolTip = "Terminal field of a graph"), ShowCategories = ("Field"))
-class FIELDSYSTEMENGINE_API UReturnResultsTerminal : public UFieldNodeBase
+UCLASS(ClassGroup = "Field", meta = (BlueprintSpawnableComponent, ToolTip = "Terminal field of a graph"), ShowCategories = ("Field"), MinimalAPI)
+class UReturnResultsTerminal : public UFieldNodeBase
 {
 	GENERATED_BODY()
 public:
@@ -909,12 +909,12 @@ public:
 
 	virtual ~UReturnResultsTerminal() {}
 
-	virtual FFieldNodeBase* NewEvaluationGraph(TArray<const UFieldNodeBase*>& Nodes) const override;
+	FIELDSYSTEMENGINE_API virtual FFieldNodeBase* NewEvaluationGraph(TArray<const UFieldNodeBase*>& Nodes) const override;
 	virtual FFieldNodeBase::EFieldType Type() const override { return  FFieldNodeBase::EFieldType::EField_Results; }
 
 	/** Terminal field of a graph */
 	UFUNCTION(BlueprintPure, Category = "Field", DisplayName = "Set Terminal Field")
-	UReturnResultsTerminal* SetReturnResultsTerminal();
+	FIELDSYSTEMENGINE_API UReturnResultsTerminal* SetReturnResultsTerminal();
 
 };
 
@@ -922,7 +922,7 @@ public:
 * Field Commands container that will be stored in the construction script
 */
 USTRUCT(BlueprintType)
-struct FIELDSYSTEMENGINE_API FFieldObjectCommands
+struct FFieldObjectCommands
 {
 	GENERATED_BODY()
 
diff --git a/Engine/Source/Runtime/Experimental/GeometryCollectionEngine/Public/ChaosBlueprint.h b/Engine/Source/Runtime/Experimental/GeometryCollectionEngine/Public/ChaosBlueprint.h
--- a/Engine/Source/Runtime/Experimental/GeometryCollectionEngine/Public/ChaosBlueprint.h
+++ b/Engine/Source/Runtime/Experimental/GeometryCollectionEngine/Public/ChaosBlueprint.h
@@ -39,22 +39,22 @@ DECLARE_DYNAMIC_MULTICAST_DELEGATE_OneParam(FOnChaosRemovalEvents, const TArray<
 
 
 /** Object allowing for retrieving Chaos Destruction data. */
-UCLASS(ClassGroup = (Chaos), hidecategories = (Object, ActorComponent, Physics, Rendering, Mobility, LOD), ShowCategories = Trigger, meta = (BlueprintSpawnableComponent))
-class GEOMETRYCOLLECTIONENGINE_API UChaosDestructionListener : public USceneComponent
+UCLASS(ClassGroup = (Chaos), hidecategories = (Object, ActorComponent, Physics, Rendering, Mobility, LOD), ShowCategories = Trigger, meta = (BlueprintSpawnableComponent), MinimalAPI)
+class UChaosDestructionListener : public USceneComponent
 {
 	GENERATED_UCLASS_BODY()
 		
 public:
 	//~ Begin UObject Interface
 #if WITH_EDITOR
-	virtual void PostEditChangeProperty(struct FPropertyChangedEvent& PropertyChangedEvent) override;
+	GEOMETRYCOLLECTIONENGINE_API virtual void PostEditChangeProperty(struct FPropertyChangedEvent& PropertyChangedEvent) override;
 #endif
 	//~ End UObject Interface
 	
 	//~ Begin UActorComponent interface
-	virtual void BeginPlay() override;
-	virtual void EndPlay(const EEndPlayReason::Type EndPlayReason) override;
-	virtual void TickComponent(float DeltaTime, enum ELevelTick TickType, FActorComponentTickFunction *ThisTickFunction) override;
+	GEOMETRYCOLLECTIONENGINE_API virtual void BeginPlay() override;
+	GEOMETRYCOLLECTIONENGINE_API virtual void EndPlay(const EEndPlayReason::Type EndPlayReason) override;
+	GEOMETRYCOLLECTIONENGINE_API virtual void TickComponent(float DeltaTime, enum ELevelTick TickType, FActorComponentTickFunction *ThisTickFunction) override;
 	//~ End UActorComponent interface
 
 	// Whether or not collision event listening is enabled
@@ -99,55 +99,55 @@ public:
 
 	// Dynamically adds a chaos solver to the listener
 	UFUNCTION(BlueprintCallable, Category = "Destruction Listener", meta = (WorldContext = "WorldContextObject"))
-	void AddChaosSolverActor(AChaosSolverActor* ChaosSolverActor);
+	GEOMETRYCOLLECTIONENGINE_API void AddChaosSolverActor(AChaosSolverActor* ChaosSolverActor);
 	
 	// Dynamically removes a chaos solver from the listener
 	UFUNCTION(BlueprintCallable, Category = "Destruction Listener", meta = (WorldContext = "WorldContextObject"))
-	void RemoveChaosSolverActor(AChaosSolverActor* ChaosSolverActor);
+	GEOMETRYCOLLECTIONENGINE_API void RemoveChaosSolverActor(AChaosSolverActor* ChaosSolverActor);
 
 	// Dynamically adds a chaos solver to the listener
 	UFUNCTION(BlueprintCallable, Category = "Destruction Listener", meta = (WorldContext = "WorldContextObject"))
-	void AddGeometryCollectionActor(AGeometryCollectionActor* GeometryCollectionActor);
+	GEOMETRYCOLLECTIONENGINE_API void AddGeometryCollectionActor(AGeometryCollectionActor* GeometryCollectionActor);
 
 	// Dynamically removes a chaos solver from the listener
 	UFUNCTION(BlueprintCallable, Category = "Destruction Listener", meta = (WorldContext = "WorldContextObject"))
-	void RemoveGeometryCollectionActor(AGeometryCollectionActor* GeometryCollectionActor);
+	GEOMETRYCOLLECTIONENGINE_API void RemoveGeometryCollectionActor(AGeometryCollectionActor* GeometryCollectionActor);
 
 	// Sets collision event request settings dynamically
 	UFUNCTION(BlueprintCallable, Category = "Destruction Listener", meta = (WorldContext = "WorldContextObject"))
-	void SetCollisionEventRequestSettings(const FChaosCollisionEventRequestSettings& InSettings);
+	GEOMETRYCOLLECTIONENGINE_API void SetCollisionEventRequestSettings(const FChaosCollisionEventRequestSettings& InSettings);
 
 	// Sets breaking event request settings dynamically
 	UFUNCTION(BlueprintCallable, Category = "Destruction Listener", meta = (WorldContext = "WorldContextObject"))
-	void SetBreakingEventRequestSettings(const FChaosBreakingEventRequestSettings& InSettings);
+	GEOMETRYCOLLECTIONENGINE_API void SetBreakingEventRequestSettings(const FChaosBreakingEventRequestSettings& InSettings);
 
 	// Sets trailing event request settings dynamically
 	UFUNCTION(BlueprintCallable, Category = "Destruction Listener", meta = (WorldContext = "WorldContextObject"))
-	void SetTrailingEventRequestSettings(const FChaosTrailingEventRequestSettings& InSettings);
+	GEOMETRYCOLLECTIONENGINE_API void SetTrailingEventRequestSettings(const FChaosTrailingEventRequestSettings& InSettings);
 
 	// Sets removal event request settings dynamically
 	UFUNCTION(BlueprintCallable, Category = "Destruction Listener", meta = (WorldContext = "WorldContextObject"))
-	void SetRemovalEventRequestSettings(const FChaosRemovalEventRequestSettings& InSettings);
+	GEOMETRYCOLLECTIONENGINE_API void SetRemovalEventRequestSettings(const FChaosRemovalEventRequestSettings& InSettings);
 
 	// Enables or disables collision event listening
 	UFUNCTION(BlueprintCallable, Category = "Destruction Listener", meta = (WorldContext = "WorldContextObject"))
-	void SetCollisionEventEnabled(bool bIsEnabled);
+	GEOMETRYCOLLECTIONENGINE_API void SetCollisionEventEnabled(bool bIsEnabled);
 
 	// Enables or disables breaking event listening
 	UFUNCTION(BlueprintCallable, Category = "Destruction Listener", meta = (WorldContext = "WorldContextObject"))
-	void SetBreakingEventEnabled(bool bIsEnabled);
+	GEOMETRYCOLLECTIONENGINE_API void SetBreakingEventEnabled(bool bIsEnabled);
 
 	// Enables or disables trailing event listening
 	UFUNCTION(BlueprintCallable, Category = "Destruction Listener", meta = (WorldContext = "WorldContextObject"))
-	void SetTrailingEventEnabled(bool bIsEnabled);
+	GEOMETRYCOLLECTIONENGINE_API void SetTrailingEventEnabled(bool bIsEnabled);
 
 	// Enables or disables removal event listening
 	UFUNCTION(BlueprintCallable, Category = "Destruction Listener", meta = (WorldContext = "WorldContextObject"))
-	void SetRemovalEventEnabled(bool bIsEnabled);
+	GEOMETRYCOLLECTIONENGINE_API void SetRemovalEventEnabled(bool bIsEnabled);
 
 	// Returns if the destruction listener is listening to any events
 	UFUNCTION(BlueprintCallable, Category = "Destruction Listener", meta = (WorldContext = "WorldContextObject"))
-	bool IsEventListening() const;
+	GEOMETRYCOLLECTIONENGINE_API bool IsEventListening() const;
 
 	/** Called when new collision events are available. */
 	UPROPERTY(BlueprintAssignable)
@@ -167,29 +167,29 @@ public:
 
 	// Sorts collision events according to the given sort method	
 	UFUNCTION(BlueprintCallable, Category = "Destruction Listener", meta = (WorldContext = "WorldContextObject"))
-	void SortCollisionEvents(UPARAM(ref) TArray<FChaosCollisionEventData>& CollisionEvents, EChaosCollisionSortMethod SortMethod);
+	GEOMETRYCOLLECTIONENGINE_API void SortCollisionEvents(UPARAM(ref) TArray<FChaosCollisionEventData>& CollisionEvents, EChaosCollisionSortMethod SortMethod);
 
 	// Sorts breaking events according to the given sort method	
 	UFUNCTION(BlueprintCallable, Category = "Destruction Listener", meta = (WorldContext = "WorldContextObject"))
-	void SortBreakingEvents(UPARAM(ref) TArray<FChaosBreakingEventData>& BreakingEvents, EChaosBreakingSortMethod SortMethod);
+	GEOMETRYCOLLECTIONENGINE_API void SortBreakingEvents(UPARAM(ref) TArray<FChaosBreakingEventData>& BreakingEvents, EChaosBreakingSortMethod SortMethod);
 
 	// Sorts trailing events according to the given sort method	
 	UFUNCTION(BlueprintCallable, Category = "Destruction Listener", meta = (WorldContext = "WorldContextObject"))
-	void SortTrailingEvents(UPARAM(ref) TArray<FChaosTrailingEventData>& TrailingEvents, EChaosTrailingSortMethod SortMethod);
+	GEOMETRYCOLLECTIONENGINE_API void SortTrailingEvents(UPARAM(ref) TArray<FChaosTrailingEventData>& TrailingEvents, EChaosTrailingSortMethod SortMethod);
 
 	// Sorts removal events according to the given sort method	
 	UFUNCTION(BlueprintCallable, Category = "Destruction Listener", meta = (WorldContext = "WorldContextObject"))
-	void SortRemovalEvents(UPARAM(ref) TArray<FChaosRemovalEventData>& RemovalEvents, EChaosRemovalSortMethod SortMethod);
+	GEOMETRYCOLLECTIONENGINE_API void SortRemovalEvents(UPARAM(ref) TArray<FChaosRemovalEventData>& RemovalEvents, EChaosRemovalSortMethod SortMethod);
 
 private:
 	// Updates the scene component transform settings
-	void UpdateTransformSettings();
+	GEOMETRYCOLLECTIONENGINE_API void UpdateTransformSettings();
 
 	// Retrieves data from solvers
-	void GetDataFromSolvers();
+	GEOMETRYCOLLECTIONENGINE_API void GetDataFromSolvers();
 
-	void ClearEvents();
-	void UpdateEvents();
+	GEOMETRYCOLLECTIONENGINE_API void ClearEvents();
+	GEOMETRYCOLLECTIONENGINE_API void UpdateEvents();
 
 protected:
 
@@ -254,16 +254,16 @@ public:
 		ChaosRemovalFilter = InRemovalFilter;
 	}
 
-	void RegisterChaosEvents(FPhysScene* Scene);
-	void UnregisterChaosEvents(FPhysScene* Scene);
+	GEOMETRYCOLLECTIONENGINE_API void RegisterChaosEvents(FPhysScene* Scene);
+	GEOMETRYCOLLECTIONENGINE_API void UnregisterChaosEvents(FPhysScene* Scene);
 
-	void RegisterChaosEvents(TSharedPtr<FPhysScene_Chaos> Scene);
-	void UnregisterChaosEvents(TSharedPtr<FPhysScene_Chaos> Scene);
+	GEOMETRYCOLLECTIONENGINE_API void RegisterChaosEvents(TSharedPtr<FPhysScene_Chaos> Scene);
+	GEOMETRYCOLLECTIONENGINE_API void UnregisterChaosEvents(TSharedPtr<FPhysScene_Chaos> Scene);
 
 
 	// Chaos Event Handlers
-	void HandleCollisionEvents(const Chaos::FCollisionEventData& CollisionData);
-	void HandleBreakingEvents(const Chaos::FBreakingEventData& BreakingData);
-	void HandleTrailingEvents(const Chaos::FTrailingEventData& TrailingData);
-	void HandleRemovalEvents(const Chaos::FRemovalEventData& RemovalData);
+	GEOMETRYCOLLECTIONENGINE_API void HandleCollisionEvents(const Chaos::FCollisionEventData& CollisionData);
+	GEOMETRYCOLLECTIONENGINE_API void HandleBreakingEvents(const Chaos::FBreakingEventData& BreakingData);
+	GEOMETRYCOLLECTIONENGINE_API void HandleTrailingEvents(const Chaos::FTrailingEventData& TrailingData);
+	GEOMETRYCOLLECTIONENGINE_API void HandleRemovalEvents(const Chaos::FRemovalEventData& RemovalData);
 };
diff --git a/Engine/Source/Runtime/Experimental/GeometryCollectionEngine/Public/ChaosBreakingEventFilter.h b/Engine/Source/Runtime/Experimental/GeometryCollectionEngine/Public/ChaosBreakingEventFilter.h
--- a/Engine/Source/Runtime/Experimental/GeometryCollectionEngine/Public/ChaosBreakingEventFilter.h
+++ b/Engine/Source/Runtime/Experimental/GeometryCollectionEngine/Public/ChaosBreakingEventFilter.h
@@ -85,14 +85,14 @@ struct FChaosBreakingEventRequestSettings
 	{}
 };
 
-class GEOMETRYCOLLECTIONENGINE_API FChaosBreakingEventFilter : public IChaosEventFilter<Chaos::FBreakingDataArray, TArray<FChaosBreakingEventData>, EChaosBreakingSortMethod>
+class FChaosBreakingEventFilter : public IChaosEventFilter<Chaos::FBreakingDataArray, TArray<FChaosBreakingEventData>, EChaosBreakingSortMethod>
 {
 public:
 	FChaosBreakingEventFilter(FChaosBreakingEventRequestSettings* FilterSettingsIn) : BreakingEventRequestSettings(FilterSettingsIn) {}
 
-	virtual void FilterEvents(const FTransform& ChaosComponentTransform, const Chaos::FBreakingDataArray& RawBreakingDataArray) override;
+	GEOMETRYCOLLECTIONENGINE_API virtual void FilterEvents(const FTransform& ChaosComponentTransform, const Chaos::FBreakingDataArray& RawBreakingDataArray) override;
 
-	virtual void SortEvents(TArray<FChaosBreakingEventData>& InOutBreakingEvents, EChaosBreakingSortMethod SortMethod, const FTransform& InTransform) override;
+	GEOMETRYCOLLECTIONENGINE_API virtual void SortEvents(TArray<FChaosBreakingEventData>& InOutBreakingEvents, EChaosBreakingSortMethod SortMethod, const FTransform& InTransform) override;
 
 private:
 	FChaosBreakingEventFilter() : BreakingEventRequestSettings(nullptr) {}
diff --git a/Engine/Source/Runtime/Experimental/GeometryCollectionEngine/Public/ChaosCollisionEventFilter.h b/Engine/Source/Runtime/Experimental/GeometryCollectionEngine/Public/ChaosCollisionEventFilter.h
--- a/Engine/Source/Runtime/Experimental/GeometryCollectionEngine/Public/ChaosCollisionEventFilter.h
+++ b/Engine/Source/Runtime/Experimental/GeometryCollectionEngine/Public/ChaosCollisionEventFilter.h
@@ -106,14 +106,14 @@ struct FChaosCollisionEventRequestSettings
 	{}
 };
 
-class GEOMETRYCOLLECTIONENGINE_API FChaosCollisionEventFilter : public IChaosEventFilter<Chaos::FCollisionDataArray, TArray<FChaosCollisionEventData>, EChaosCollisionSortMethod>
+class FChaosCollisionEventFilter : public IChaosEventFilter<Chaos::FCollisionDataArray, TArray<FChaosCollisionEventData>, EChaosCollisionSortMethod>
  {
 public:
 
 	FChaosCollisionEventFilter(FChaosCollisionEventRequestSettings* FilterSettingsIn) : CollisionEventRequestSettings(FilterSettingsIn) {}
 
-	virtual void FilterEvents(const FTransform& ChaosComponentTransform, const Chaos::FCollisionDataArray& RawCollisionDataArray) override;
-	virtual void SortEvents(TArray<FChaosCollisionEventData>& InOutCollisionEvents, EChaosCollisionSortMethod SortMethod, const FTransform& InTransform) override;
+	GEOMETRYCOLLECTIONENGINE_API virtual void FilterEvents(const FTransform& ChaosComponentTransform, const Chaos::FCollisionDataArray& RawCollisionDataArray) override;
+	GEOMETRYCOLLECTIONENGINE_API virtual void SortEvents(TArray<FChaosCollisionEventData>& InOutCollisionEvents, EChaosCollisionSortMethod SortMethod, const FTransform& InTransform) override;
 
 private:
 	FChaosCollisionEventFilter() : CollisionEventRequestSettings(nullptr) { check(false);  }
diff --git a/Engine/Source/Runtime/Experimental/GeometryCollectionEngine/Public/ChaosRemovalEventFilter.h b/Engine/Source/Runtime/Experimental/GeometryCollectionEngine/Public/ChaosRemovalEventFilter.h
--- a/Engine/Source/Runtime/Experimental/GeometryCollectionEngine/Public/ChaosRemovalEventFilter.h
+++ b/Engine/Source/Runtime/Experimental/GeometryCollectionEngine/Public/ChaosRemovalEventFilter.h
@@ -74,15 +74,15 @@ struct FChaosRemovalEventRequestSettings
 	}
 };
 
-class GEOMETRYCOLLECTIONENGINE_API FChaosRemovalEventFilter
+class FChaosRemovalEventFilter
 	: public IChaosEventFilter<Chaos::FRemovalDataArray, TArray<FChaosRemovalEventData>, EChaosRemovalSortMethod>
 {
 public:
 	FChaosRemovalEventFilter(FChaosRemovalEventRequestSettings* FilterSettingsIn) : RemovalEventRequestSettings(FilterSettingsIn) {}
 
-	virtual void FilterEvents(const FTransform& ChaosComponentTransform, const Chaos::FRemovalDataArray& RawRemovalDataArray) override;
+	GEOMETRYCOLLECTIONENGINE_API virtual void FilterEvents(const FTransform& ChaosComponentTransform, const Chaos::FRemovalDataArray& RawRemovalDataArray) override;
 
-	virtual void SortEvents(TArray<FChaosRemovalEventData>& InOutRemovalEvents, EChaosRemovalSortMethod SortMethod, const FTransform& InTransform) override;
+	GEOMETRYCOLLECTIONENGINE_API virtual void SortEvents(TArray<FChaosRemovalEventData>& InOutRemovalEvents, EChaosRemovalSortMethod SortMethod, const FTransform& InTransform) override;
 
 private:
 	FChaosRemovalEventFilter() : RemovalEventRequestSettings(nullptr) {}
diff --git a/Engine/Source/Runtime/Experimental/GeometryCollectionEngine/Public/ChaosTrailingEventFilter.h b/Engine/Source/Runtime/Experimental/GeometryCollectionEngine/Public/ChaosTrailingEventFilter.h
--- a/Engine/Source/Runtime/Experimental/GeometryCollectionEngine/Public/ChaosTrailingEventFilter.h
+++ b/Engine/Source/Runtime/Experimental/GeometryCollectionEngine/Public/ChaosTrailingEventFilter.h
@@ -95,15 +95,15 @@ struct FChaosTrailingEventRequestSettings
 	}
 };
 
-class GEOMETRYCOLLECTIONENGINE_API FChaosTrailingEventFilter 
+class FChaosTrailingEventFilter 
 	: public IChaosEventFilter<Chaos::FTrailingDataArray, TArray<FChaosTrailingEventData>, EChaosTrailingSortMethod>
 {
 public:
 	FChaosTrailingEventFilter(FChaosTrailingEventRequestSettings* FilterSettingsIn) : TrailingEventRequestSettings(FilterSettingsIn) {}
 
-	virtual void FilterEvents(const FTransform& ChaosComponentTransform, const Chaos::FTrailingDataArray& RawTrailingDataArray) override;
+	GEOMETRYCOLLECTIONENGINE_API virtual void FilterEvents(const FTransform& ChaosComponentTransform, const Chaos::FTrailingDataArray& RawTrailingDataArray) override;
 
-	virtual void SortEvents(TArray<FChaosTrailingEventData>& InOutTrailingEvents, EChaosTrailingSortMethod SortMethod, const FTransform& InTransform) override;
+	GEOMETRYCOLLECTIONENGINE_API virtual void SortEvents(TArray<FChaosTrailingEventData>& InOutTrailingEvents, EChaosTrailingSortMethod SortMethod, const FTransform& InTransform) override;
 
 private:
 	FChaosTrailingEventFilter() : TrailingEventRequestSettings(nullptr) {}
diff --git a/Engine/Source/Runtime/Experimental/GeometryCollectionEngine/Public/GeometryCollection/GeometryCollectionActor.h b/Engine/Source/Runtime/Experimental/GeometryCollectionEngine/Public/GeometryCollection/GeometryCollectionActor.h
--- a/Engine/Source/Runtime/Experimental/GeometryCollectionEngine/Public/GeometryCollection/GeometryCollectionActor.h
+++ b/Engine/Source/Runtime/Experimental/GeometryCollectionEngine/Public/GeometryCollection/GeometryCollectionActor.h
@@ -15,18 +15,18 @@ class UGeometryCollectionComponent;
 class UGeometryCollectionDebugDrawComponent;
 struct FHitResult;
 
-UCLASS()
-class GEOMETRYCOLLECTIONENGINE_API AGeometryCollectionActor: public AActor
+UCLASS(MinimalAPI)
+class AGeometryCollectionActor: public AActor
 {
 	GENERATED_UCLASS_BODY()
 
 public:
 
 	/* Game state callback */
-	virtual void Tick(float DeltaSeconds) override;
+	GEOMETRYCOLLECTIONENGINE_API virtual void Tick(float DeltaSeconds) override;
 
 #if WITH_EDITOR
-	virtual bool GetReferencedContentObjects(TArray<UObject*>& Objects) const override;
+	GEOMETRYCOLLECTIONENGINE_API virtual bool GetReferencedContentObjects(TArray<UObject*>& Objects) const override;
 #endif
 
 	/* GeometryCollectionComponent */
@@ -40,6 +40,6 @@ public:
 	PRAGMA_ENABLE_DEPRECATION_WARNINGS
 	
 	UFUNCTION(BlueprintCallable, Category = "Physics")
-	bool RaycastSingle(FVector Start, FVector End, FHitResult& OutHit) const;
+	GEOMETRYCOLLECTIONENGINE_API bool RaycastSingle(FVector Start, FVector End, FHitResult& OutHit) const;
 
 };
diff --git a/Engine/Source/Runtime/Experimental/GeometryCollectionEngine/Public/GeometryCollection/GeometryCollectionCache.h b/Engine/Source/Runtime/Experimental/GeometryCollectionEngine/Public/GeometryCollection/GeometryCollectionCache.h
--- a/Engine/Source/Runtime/Experimental/GeometryCollectionEngine/Public/GeometryCollection/GeometryCollectionCache.h
+++ b/Engine/Source/Runtime/Experimental/GeometryCollectionEngine/Public/GeometryCollection/GeometryCollectionCache.h
@@ -11,50 +11,50 @@ class UGeometryCollection;
 
 GEOMETRYCOLLECTIONENGINE_API DECLARE_LOG_CATEGORY_EXTERN(LogGeometryCollectionCache, Log, All);
 
-UCLASS(Experimental)
-class GEOMETRYCOLLECTIONENGINE_API UGeometryCollectionCache : public UObject
+UCLASS(Experimental, MinimalAPI)
+class UGeometryCollectionCache : public UObject
 {
 	GENERATED_BODY()
 
 public:
 
 	/** Tagnames for asset registry tags */
-	static FName TagName_Name;			// Name of the cache
-	static FName TagName_IdGuid;		// ID GUID for the cache - never changes for a given cache
-	static FName TagName_StateGuid;		// State GUID - changes when an edit is made to 
+	static GEOMETRYCOLLECTIONENGINE_API FName TagName_Name;			// Name of the cache
+	static GEOMETRYCOLLECTIONENGINE_API FName TagName_IdGuid;		// ID GUID for the cache - never changes for a given cache
+	static GEOMETRYCOLLECTIONENGINE_API FName TagName_StateGuid;		// State GUID - changes when an edit is made to 
 
 	/**
 	 * Given a raw track with transforms per-particle on each frame record, set to this cache
 	 * and strip out any data we don't need (transform repeats and disabled particles etc.)
 	 */
-	void SetFromRawTrack(const FRecordedTransformTrack& InTrack);
+	GEOMETRYCOLLECTIONENGINE_API void SetFromRawTrack(const FRecordedTransformTrack& InTrack);
 
 	/** Set directly from a track, does not do any data stripping. */
-	void SetFromTrack(const FRecordedTransformTrack& InTrack);
+	GEOMETRYCOLLECTIONENGINE_API void SetFromTrack(const FRecordedTransformTrack& InTrack);
 
 	/** Sets the geometry collection that this cache supports, empties the recorded data in this cache */
-	void SetSupportedCollection(const UGeometryCollection* InCollection);
+	GEOMETRYCOLLECTIONENGINE_API void SetSupportedCollection(const UGeometryCollection* InCollection);
 
 	/** UObject Interface */
-	virtual void GetAssetRegistryTags(TArray<FAssetRegistryTag>& OutTags) const override;
+	GEOMETRYCOLLECTIONENGINE_API virtual void GetAssetRegistryTags(TArray<FAssetRegistryTag>& OutTags) const override;
 	/** End UObject Interface */
 
 	/** Access the recorded tracks */
 	const FRecordedTransformTrack* GetData() const { return &RecordedData; }
 
 	/** Given a collection, create an empty compatible cache for it */
-	static UGeometryCollectionCache* CreateCacheForCollection(const UGeometryCollection* InCollection);
+	static GEOMETRYCOLLECTIONENGINE_API UGeometryCollectionCache* CreateCacheForCollection(const UGeometryCollection* InCollection);
 
 	/** Get the GUID for the state of the supported collection when this cache was last recorded to. */
 	FGuid GetCompatibleStateGuid() const { return CompatibleCollectionState; }
 
 	/** Tests whether a cache is compatible with a provided collection */
-	bool CompatibleWithForRecord(const UGeometryCollection* InCollection);
-	bool CompatibleWithForPlayback(const UGeometryCollection* InCollection);
+	GEOMETRYCOLLECTIONENGINE_API bool CompatibleWithForRecord(const UGeometryCollection* InCollection);
+	GEOMETRYCOLLECTIONENGINE_API bool CompatibleWithForPlayback(const UGeometryCollection* InCollection);
 
 private:
 
-	void ProcessRawRecordedDataInternal(const FRecordedTransformTrack& InTrack);
+	GEOMETRYCOLLECTIONENGINE_API void ProcessRawRecordedDataInternal(const FRecordedTransformTrack& InTrack);
 
 	/** The recorded data from the simulation */
 	UPROPERTY()
@@ -75,7 +75,7 @@ private:
  * Initial purpose is to allow an opaque way to call some editor system to generate new assets
  * but this could be expanded later for runtime recording and playback if the need arises
  */
-class GEOMETRYCOLLECTIONENGINE_API ITargetCacheProvider : public IModularFeature
+class ITargetCacheProvider : public IModularFeature
 {
 public:
 	static FName GetFeatureName() { return "GeometryCollectionTargetCacheProvider"; }
diff --git a/Engine/Source/Runtime/Experimental/GeometryCollectionEngine/Public/GeometryCollection/GeometryCollectionCreationParameters.h b/Engine/Source/Runtime/Experimental/GeometryCollectionEngine/Public/GeometryCollection/GeometryCollectionCreationParameters.h
--- a/Engine/Source/Runtime/Experimental/GeometryCollectionEngine/Public/GeometryCollection/GeometryCollectionCreationParameters.h
+++ b/Engine/Source/Runtime/Experimental/GeometryCollectionEngine/Public/GeometryCollection/GeometryCollectionCreationParameters.h
@@ -6,15 +6,15 @@
 
 class FGeometryCollection;
 
-class GEOMETRYCOLLECTIONENGINE_API FGeometryCollectionCreationParameters
+class FGeometryCollectionCreationParameters
 {
 public:
-	FGeometryCollectionCreationParameters(FGeometryCollection& GeometryCollection, bool ReCalculateNormalsIn = false, bool ReCalculateTangetsIn = false);
-	~FGeometryCollectionCreationParameters();
+	GEOMETRYCOLLECTIONENGINE_API FGeometryCollectionCreationParameters(FGeometryCollection& GeometryCollection, bool ReCalculateNormalsIn = false, bool ReCalculateTangetsIn = false);
+	GEOMETRYCOLLECTIONENGINE_API ~FGeometryCollectionCreationParameters();
 
 private:
 	bool ReCalculateNormals;
 	bool ReCalculateTangents;
 	FGeometryCollection& GeometryCollection;
 
-};
\ No newline at end of file
+};
diff --git a/Engine/Source/Runtime/Experimental/GeometryCollectionEngine/Public/GeometryCollection/GeometryCollectionDamagePropagationData.h b/Engine/Source/Runtime/Experimental/GeometryCollectionEngine/Public/GeometryCollection/GeometryCollectionDamagePropagationData.h
--- a/Engine/Source/Runtime/Experimental/GeometryCollectionEngine/Public/GeometryCollection/GeometryCollectionDamagePropagationData.h
+++ b/Engine/Source/Runtime/Experimental/GeometryCollectionEngine/Public/GeometryCollection/GeometryCollectionDamagePropagationData.h
@@ -5,7 +5,7 @@
 #include "GeometryCollectionDamagePropagationData.generated.h"
 
 USTRUCT(BlueprintType)
-struct GEOMETRYCOLLECTIONENGINE_API FGeometryCollectionDamagePropagationData
+struct FGeometryCollectionDamagePropagationData
 {
 public:
 	GENERATED_BODY()
diff --git a/Engine/Source/Runtime/Experimental/GeometryCollectionEngine/Public/GeometryCollection/GeometryCollectionDebugDrawActor.h b/Engine/Source/Runtime/Experimental/GeometryCollectionEngine/Public/GeometryCollection/GeometryCollectionDebugDrawActor.h
--- a/Engine/Source/Runtime/Experimental/GeometryCollectionEngine/Public/GeometryCollection/GeometryCollectionDebugDrawActor.h
+++ b/Engine/Source/Runtime/Experimental/GeometryCollectionEngine/Public/GeometryCollection/GeometryCollectionDebugDrawActor.h
@@ -82,8 +82,8 @@ struct UE_DEPRECATED(5.0, "Deprecated. Use normal debug draw Chaos Physics comma
 *   spawned by any GeometryDebugDrawComponent that needs it.
 */
 class UE_DEPRECATED(5.0, "Deprecated. Use normal debug draw Chaos Physics commands") AGeometryCollectionDebugDrawActor;
-UCLASS(HideCategories = ("Rendering", "Replication", "Input", "Actor", "Collision", "LOD", "Cooking"))
-class GEOMETRYCOLLECTIONENGINE_API AGeometryCollectionDebugDrawActor : public AActor
+UCLASS(HideCategories = ("Rendering", "Replication", "Input", "Actor", "Collision", "LOD", "Cooking"), MinimalAPI)
+class AGeometryCollectionDebugDrawActor : public AActor
 {
 	GENERATED_UCLASS_BODY()
 public:
diff --git a/Engine/Source/Runtime/Experimental/GeometryCollectionEngine/Public/GeometryCollection/GeometryCollectionDebugDrawComponent.h b/Engine/Source/Runtime/Experimental/GeometryCollectionEngine/Public/GeometryCollection/GeometryCollectionDebugDrawComponent.h
--- a/Engine/Source/Runtime/Experimental/GeometryCollectionEngine/Public/GeometryCollection/GeometryCollectionDebugDrawComponent.h
+++ b/Engine/Source/Runtime/Experimental/GeometryCollectionEngine/Public/GeometryCollection/GeometryCollectionDebugDrawComponent.h
@@ -23,8 +23,8 @@ namespace Chaos { template<class T, int d> class TPBDRigidParticles; }
 *   This component is automatically added to every GeometryCollectionActor.
 */
 class UE_DEPRECATED(5.0, "Deprecated. Use normal debug draw Chaos Physics commands") UGeometryCollectionDebugDrawComponent;
-UCLASS(meta = (BlueprintSpawnableComponent), HideCategories = ("Tags", "Activation", "Cooking", "AssetUserData", "Collision"))
-class GEOMETRYCOLLECTIONENGINE_API UGeometryCollectionDebugDrawComponent : public UActorComponent
+UCLASS(meta = (BlueprintSpawnableComponent), HideCategories = ("Tags", "Activation", "Cooking", "AssetUserData", "Collision"), MinimalAPI)
+class UGeometryCollectionDebugDrawComponent : public UActorComponent
 {
 	GENERATED_UCLASS_BODY()
 
diff --git a/Engine/Source/Runtime/Experimental/GeometryCollectionEngine/Public/GeometryCollection/GeometryCollectionEngineConversion.h b/Engine/Source/Runtime/Experimental/GeometryCollectionEngine/Public/GeometryCollection/GeometryCollectionEngineConversion.h
--- a/Engine/Source/Runtime/Experimental/GeometryCollectionEngine/Public/GeometryCollection/GeometryCollectionEngineConversion.h
+++ b/Engine/Source/Runtime/Experimental/GeometryCollectionEngine/Public/GeometryCollection/GeometryCollectionEngineConversion.h
@@ -23,7 +23,7 @@ typedef TTuple<const USkeletalMesh *, const USkeletalMeshComponent *, FTransform
 /**
  * The public interface to this module
  */
-class GEOMETRYCOLLECTIONENGINE_API FGeometryCollectionEngineConversion
+class FGeometryCollectionEngineConversion
 {
 public:
 
@@ -32,7 +32,7 @@ public:
 	 * @param Materials : Materials fetched from the StaticMeshComponent used to configure this geometry
 	 * @param GeometryCollection  : Collection to append the mesh into.
 	 */
-	static int32 AppendMaterials(const TArray<UMaterialInterface*>& Materials, UGeometryCollection* GeometryCollectionObject, bool bAddInteriorCopy);
+	static GEOMETRYCOLLECTIONENGINE_API int32 AppendMaterials(const TArray<UMaterialInterface*>& Materials, UGeometryCollection* GeometryCollectionObject, bool bAddInteriorCopy);
 
 	/**
 	 * Appends instanced mesh indices 
@@ -41,7 +41,7 @@ public:
 	 * @param StaticMesh  static mesh to add reference to 
 	 * @param Materials materials corresponding to the static mesh instance to get the index from 
 	 */
-	static void AppendAutoInstanceMeshIndices(UGeometryCollection* GeometryCollectionObject, int32 FromTransformIndex, const UStaticMesh* StaticMesh, const TArray<UMaterialInterface*>& Materials);
+	static GEOMETRYCOLLECTIONENGINE_API void AppendAutoInstanceMeshIndices(UGeometryCollection* GeometryCollectionObject, int32 FromTransformIndex, const UStaticMesh* StaticMesh, const TArray<UMaterialInterface*>& Materials);
 
 	/**
 	 * Appends a MeshDescription to a GeometryCollection.
@@ -53,14 +53,14 @@ public:
 	 * @param BodySetup : Optional collision setup to transfer to the geometry collection
 	 * @param bReindexMaterials	: Whether to reindex materials -- if appending multiple meshes, pass false and call ReindexMaterials afterwards
 	 */
-	static void AppendMeshDescription(const FMeshDescription* MeshDescription, const FString& Name, int32 StartMaterialIndex, const FTransform& StaticMeshTransform,
+	static GEOMETRYCOLLECTIONENGINE_API void AppendMeshDescription(const FMeshDescription* MeshDescription, const FString& Name, int32 StartMaterialIndex, const FTransform& StaticMeshTransform,
 		FGeometryCollection* GeometryCollection, UBodySetup* BodySetup = nullptr, bool bReindexMaterials = true, bool bAddInternalMaterials = true, bool bSetInternalFromMaterialIndex = false);
 
 	/**
 	 * Get a HiRes (or LOD 0 if no HiRes available) MeshDescription for the given static mesh, and make sure it includes normals and tangents.
 	 * (i.e., to pass to AppendMeshDescription)
 	 */
-	static FMeshDescription* GetMaxResMeshDescriptionWithNormalsAndTangents(const UStaticMesh* StaticMesh);
+	static GEOMETRYCOLLECTIONENGINE_API FMeshDescription* GetMaxResMeshDescriptionWithNormalsAndTangents(const UStaticMesh* StaticMesh);
 
 	/**
 	 * Appends a static mesh to a GeometryCollectionComponent.
@@ -71,7 +71,7 @@ public:
 	 * @param bReindexMaterials	: Whether to reindex materials -- if appending multiple meshes, pass false and call ReindexMaterials afterwards
 	 * @return true if succeeded in appending a mesh
 	 */
-	static bool AppendStaticMesh(const UStaticMesh* StaticMesh, const TArray<UMaterialInterface*>& Materials, const FTransform& StaticMeshTransform, 
+	static GEOMETRYCOLLECTIONENGINE_API bool AppendStaticMesh(const UStaticMesh* StaticMesh, const TArray<UMaterialInterface*>& Materials, const FTransform& StaticMeshTransform, 
 		UGeometryCollection* GeometryCollectionObject, bool bReindexMaterials = true, bool bAddInternalMaterials = true, bool bSplitComponents = false, bool bSetInternalFromMaterialIndex = false);
 
 	// TODO: consider having the below fn return bool, and have a helper that returns the StartMaterialIndex from UGeometryCollection
@@ -88,7 +88,7 @@ public:
 	 * @param bReindexMaterials	: Whether to reindex materials -- if appending multiple meshes, pass false and call ReindexMaterials afterwards
 	 * @return true if succeeded in appending a mesh
 	 */
-	static bool AppendStaticMesh(const UStaticMesh* StaticMesh, int32 StartMaterialIndex, const FTransform& StaticMeshTransform, 
+	static GEOMETRYCOLLECTIONENGINE_API bool AppendStaticMesh(const UStaticMesh* StaticMesh, int32 StartMaterialIndex, const FTransform& StaticMeshTransform, 
 		FGeometryCollection* GeometryCollection, bool bReindexMaterials = true, bool bAddInternalMaterials = true, bool bSplitComponents = false, bool bSetInternalFromMaterialIndex = false);
 
 	/**
@@ -97,7 +97,7 @@ public:
 	*  @param StaticMeshTransform : Mesh transform.
 	*  @param GeometryCollection  : Collection to append the mesh into.
 	*/
-	static void AppendStaticMesh(const UStaticMesh* StaticMesh, const UStaticMeshComponent *StaticMeshComponent, const FTransform& StaticMeshTransform, UGeometryCollection* GeometryCollection,
+	static GEOMETRYCOLLECTIONENGINE_API void AppendStaticMesh(const UStaticMesh* StaticMesh, const UStaticMeshComponent *StaticMeshComponent, const FTransform& StaticMeshTransform, UGeometryCollection* GeometryCollection,
 		bool bReindexMaterials = true, bool bAddInternalMaterials = true, bool bSplitComponents = false, bool bSetInternalFromMaterialIndex = false);
 
 	/**
@@ -106,7 +106,7 @@ public:
 	*  @param SkeletalMeshTransform : Mesh transform.
 	*  @param GeometryCollection    : Collection to append the mesh into.
 	*/
-	static void AppendSkeletalMesh(const USkeletalMesh* SkeletalMesh, const USkeletalMeshComponent *SkeletalMeshComponent, const FTransform& SkeletalMeshTransform, UGeometryCollection* GeometryCollection, bool bReindexMaterials = true);
+	static GEOMETRYCOLLECTIONENGINE_API void AppendSkeletalMesh(const USkeletalMesh* SkeletalMesh, const USkeletalMeshComponent *SkeletalMeshComponent, const FTransform& SkeletalMeshTransform, UGeometryCollection* GeometryCollection, bool bReindexMaterials = true);
 
 	/**
 	*  Appends a skeletal mesh to a GeometryCollection.
@@ -115,7 +115,7 @@ public:
 	*  @param SkeletalMeshTransform : Mesh transform.
 	*  @param GeometryCollection    : Collection to append the mesh into.
 	*/
-	static bool AppendSkeletalMesh(const USkeletalMesh* SkeletalMesh, int32 MaterialStartIndex, const FTransform& SkeletalMeshTransform, FManagedArrayCollection* InCollection, bool bReindexMaterials = true);
+	static GEOMETRYCOLLECTIONENGINE_API bool AppendSkeletalMesh(const USkeletalMesh* SkeletalMesh, int32 MaterialStartIndex, const FTransform& SkeletalMeshTransform, FManagedArrayCollection* InCollection, bool bReindexMaterials = true);
 
 	/**
 	*  Appends a skeleton mesh to a GeometryCollection.
@@ -123,7 +123,7 @@ public:
 	*  @param SkeletalMeshTransform : Mesh transform.
 	*  @param GeometryCollection    : Collection to append the mesh into.
 	*/
-	static void AppendSkeleton(const USkeleton* Skeleton, const FTransform& SkeletalMeshTransform, FManagedArrayCollection* InCollection);
+	static GEOMETRYCOLLECTIONENGINE_API void AppendSkeleton(const USkeleton* Skeleton, const FTransform& SkeletalMeshTransform, FManagedArrayCollection* InCollection);
 
 	/**
 	*  Appends a skeletal mesh to a GeometryCollectionComponent.
@@ -131,13 +131,13 @@ public:
 	*  @param SkeletalMeshTransform : Mesh transform.
 	*  @param GeometryCollection    : Collection to append the mesh into.
 	*/
-	static int32 AppendSkeletalMeshMaterials(const USkeletalMesh* SkeletalMesh, const USkeletalMeshComponent* SkeletalMeshComponent, UGeometryCollection* GeometryCollectionObject);
+	static GEOMETRYCOLLECTIONENGINE_API int32 AppendSkeletalMeshMaterials(const USkeletalMesh* SkeletalMesh, const USkeletalMeshComponent* SkeletalMeshComponent, UGeometryCollection* GeometryCollectionObject);
 
-	static const FSkeletalMeshLODRenderData* GetSkeletalMeshLOD(const USkeletalMesh* SkeletalMesh, int32 LOD);
+	static GEOMETRYCOLLECTIONENGINE_API const FSkeletalMeshLODRenderData* GetSkeletalMeshLOD(const USkeletalMesh* SkeletalMesh, int32 LOD);
 
-	static int32 AppendGeometryCollectionMaterials(const UGeometryCollection* SourceGeometryCollection, const UGeometryCollectionComponent* GeometryCollectionComponent, UGeometryCollection* TargetGeometryCollectionObject);
+	static GEOMETRYCOLLECTIONENGINE_API int32 AppendGeometryCollectionMaterials(const UGeometryCollection* SourceGeometryCollection, const UGeometryCollectionComponent* GeometryCollectionComponent, UGeometryCollection* TargetGeometryCollectionObject);
 
-	static void AppendGeometryCollectionInstancedMeshes(const UGeometryCollection* SourceGeometryCollectionObject, UGeometryCollection* TargetGeometryCollectionObject, int32 TargetTransformStartIndex);
+	static GEOMETRYCOLLECTIONENGINE_API void AppendGeometryCollectionInstancedMeshes(const UGeometryCollection* SourceGeometryCollectionObject, UGeometryCollection* TargetGeometryCollectionObject, int32 TargetTransformStartIndex);
 
 	/**
 	 * Appends a GeometryCollection to another GeometryCollection.
@@ -148,7 +148,7 @@ public:
 	 * @param bReindexMaterials	: Whether to reindex materials -- if appending multiple meshes, pass false and call ReindexMaterials afterwards
 	 * @return true if succeeded in appending a geometry collection
 	 */
-	static bool AppendGeometryCollection(const FGeometryCollection* SourceGeometryCollection, int32 MaterialStartIndex, const FTransform& GeometryCollectionTransform, FGeometryCollection* TargetGeometryCollection, bool bReindexMaterials = true);
+	static GEOMETRYCOLLECTIONENGINE_API bool AppendGeometryCollection(const FGeometryCollection* SourceGeometryCollection, int32 MaterialStartIndex, const FTransform& GeometryCollectionTransform, FGeometryCollection* TargetGeometryCollection, bool bReindexMaterials = true);
 
 	/**
 	*  Appends a GeometryCollection to a GeometryCollectionComponent.
@@ -157,7 +157,7 @@ public:
 	*  @param GeometryCollectionTransform : GeometryCollection transform.
 	*  @param TargetGeometryCollection  : Collection to append the GeometryCollection into.
 	*/
-	static void AppendGeometryCollection(const UGeometryCollection* SourceGeometryCollection, const TArray<UMaterialInterface*>& Materials, const FTransform& GeometryCollectionTransform, UGeometryCollection* TargetGeometryCollectionObject, bool ReindexMaterials = true);
+	static GEOMETRYCOLLECTIONENGINE_API void AppendGeometryCollection(const UGeometryCollection* SourceGeometryCollection, const TArray<UMaterialInterface*>& Materials, const FTransform& GeometryCollectionTransform, UGeometryCollection* TargetGeometryCollectionObject, bool ReindexMaterials = true);
 
 	/**
 	*  Appends a GeometryCollection to a GeometryCollectionComponent.
@@ -165,7 +165,7 @@ public:
 	*  @param GeometryCollectionTransform : GeometryCollection transform.
 	*  @param TargetGeometryCollection  : Collection to append the GeometryCollection into.
 	*/
-	static void AppendGeometryCollection(const UGeometryCollection* SourceGeometryCollection, const UGeometryCollectionComponent* GeometryCollectionComponent, const FTransform& GeometryCollectionTransform, UGeometryCollection* TargetGeometryCollectionObject, bool ReindexMaterials = true);
+	static GEOMETRYCOLLECTIONENGINE_API void AppendGeometryCollection(const UGeometryCollection* SourceGeometryCollection, const UGeometryCollectionComponent* GeometryCollectionComponent, const FTransform& GeometryCollectionTransform, UGeometryCollection* TargetGeometryCollectionObject, bool ReindexMaterials = true);
 
 	/**
 	*  Appends a GeometryCollectionSource to a GeometryCollection
@@ -173,5 +173,5 @@ public:
 	*  @param GeometryCollectionInOut : GeometryCollection to append to 
 	*  @param MaterialsInOut  : array of materials to append to
 	*/
-	static void AppendGeometryCollectionSource(const FGeometryCollectionSource& GeometryCollectionSource, FGeometryCollection& GeometryCollectionInOut, TArray<UMaterial*>& MaterialsInOut, bool ReindexMaterials = true);
-};
\ No newline at end of file
+	static GEOMETRYCOLLECTIONENGINE_API void AppendGeometryCollectionSource(const FGeometryCollectionSource& GeometryCollectionSource, FGeometryCollection& GeometryCollectionInOut, TArray<UMaterial*>& MaterialsInOut, bool ReindexMaterials = true);
+};
diff --git a/Engine/Source/Runtime/Experimental/GeometryCollectionEngine/Public/GeometryCollection/GeometryCollectionExternalRenderInterface.h b/Engine/Source/Runtime/Experimental/GeometryCollectionEngine/Public/GeometryCollection/GeometryCollectionExternalRenderInterface.h
--- a/Engine/Source/Runtime/Experimental/GeometryCollectionEngine/Public/GeometryCollection/GeometryCollectionExternalRenderInterface.h
+++ b/Engine/Source/Runtime/Experimental/GeometryCollectionEngine/Public/GeometryCollection/GeometryCollectionExternalRenderInterface.h
@@ -9,13 +9,13 @@
 class UGeometryCollection;
 class UGeometryCollectionComponent;
 
-UINTERFACE()
-class GEOMETRYCOLLECTIONENGINE_API UGeometryCollectionExternalRenderInterface : public UInterface
+UINTERFACE(MinimalAPI)
+class UGeometryCollectionExternalRenderInterface : public UInterface
 {
 	GENERATED_BODY()
 };
 
-class GEOMETRYCOLLECTIONENGINE_API IGeometryCollectionExternalRenderInterface
+class IGeometryCollectionExternalRenderInterface
 {
 	GENERATED_BODY()
 
diff --git a/Engine/Source/Runtime/Experimental/GeometryCollectionEngine/Public/GeometryCollection/GeometryCollectionISMPoolActor.h b/Engine/Source/Runtime/Experimental/GeometryCollectionEngine/Public/GeometryCollection/GeometryCollectionISMPoolActor.h
--- a/Engine/Source/Runtime/Experimental/GeometryCollectionEngine/Public/GeometryCollection/GeometryCollectionISMPoolActor.h
+++ b/Engine/Source/Runtime/Experimental/GeometryCollectionEngine/Public/GeometryCollection/GeometryCollectionISMPoolActor.h
@@ -7,8 +7,8 @@
 class UGeometryCollectionISMPoolComponent;
 class UGeometryCollectionISMPoolDebugDrawComponent;
 
-UCLASS(ConversionRoot, ComponentWrapperClass)
-class GEOMETRYCOLLECTIONENGINE_API AGeometryCollectionISMPoolActor : public AActor
+UCLASS(ConversionRoot, ComponentWrapperClass, MinimalAPI)
+class AGeometryCollectionISMPoolActor : public AActor
 {
 	GENERATED_UCLASS_BODY()
 
diff --git a/Engine/Source/Runtime/Experimental/GeometryCollectionEngine/Public/GeometryCollection/GeometryCollectionISMPoolComponent.h b/Engine/Source/Runtime/Experimental/GeometryCollectionEngine/Public/GeometryCollection/GeometryCollectionISMPoolComponent.h
--- a/Engine/Source/Runtime/Experimental/GeometryCollectionEngine/Public/GeometryCollection/GeometryCollectionISMPoolComponent.h
+++ b/Engine/Source/Runtime/Experimental/GeometryCollectionEngine/Public/GeometryCollection/GeometryCollectionISMPoolComponent.h
@@ -287,8 +287,8 @@ struct FGeometryCollectionISMPool
 * UGeometryCollectionISMPoolComponent.
 * Component that manages a pool of ISM components in order to allow multiple client components that use the same meshes to the share ISMs.
 */
-UCLASS(meta = (BlueprintSpawnableComponent))
-class GEOMETRYCOLLECTIONENGINE_API UGeometryCollectionISMPoolComponent: public USceneComponent
+UCLASS(meta = (BlueprintSpawnableComponent), MinimalAPI)
+class UGeometryCollectionISMPoolComponent: public USceneComponent
 {
 	GENERATED_UCLASS_BODY()
 
@@ -297,7 +297,7 @@ public:
 	using FMeshId = int32;
 
 	//~ Begin UActorComponent Interface
-	virtual void GetResourceSizeEx(FResourceSizeEx& CumulativeResourceSize) override;
+	GEOMETRYCOLLECTIONENGINE_API virtual void GetResourceSizeEx(FResourceSizeEx& CumulativeResourceSize) override;
 	//~ End UActorComponent Interface
 
 	/** 
@@ -305,25 +305,25 @@ public:
 	* no resources are created until the meshes are added for this group 
 	* return a mesh group Id used to add and update instances
 	*/
-	FMeshGroupId CreateMeshGroup();
+	GEOMETRYCOLLECTIONENGINE_API FMeshGroupId CreateMeshGroup();
 
 	/** destroy  a mesh group and its associated resources */
-	void DestroyMeshGroup(FMeshGroupId MeshGroupId);
+	GEOMETRYCOLLECTIONENGINE_API void DestroyMeshGroup(FMeshGroupId MeshGroupId);
 
 	/** Add a static mesh for a mesh group */
-	FMeshId AddMeshToGroup(FMeshGroupId MeshGroupId, const FGeometryCollectionStaticMeshInstance& MeshInstance, int32 InstanceCount, TArrayView<const float> CustomDataFloats);
+	GEOMETRYCOLLECTIONENGINE_API FMeshId AddMeshToGroup(FMeshGroupId MeshGroupId, const FGeometryCollectionStaticMeshInstance& MeshInstance, int32 InstanceCount, TArrayView<const float> CustomDataFloats);
 
 	/** Add a static mesh for a mesh group */
-	bool BatchUpdateInstancesTransforms(FMeshGroupId MeshGroupId, FMeshId MeshId, int32 StartInstanceIndex, TArrayView<const FTransform> NewInstancesTransforms, bool bWorldSpace = false, bool bMarkRenderStateDirty = false, bool bTeleport = false);
+	GEOMETRYCOLLECTIONENGINE_API bool BatchUpdateInstancesTransforms(FMeshGroupId MeshGroupId, FMeshId MeshId, int32 StartInstanceIndex, TArrayView<const FTransform> NewInstancesTransforms, bool bWorldSpace = false, bool bMarkRenderStateDirty = false, bool bTeleport = false);
 
 	UE_DEPRECATED(5.3, "BatchUpdateInstancesTransforms Array parameter version is deprecated, use the TArrayView version instead")
-	bool BatchUpdateInstancesTransforms(FMeshGroupId MeshGroupId, FMeshId MeshId, int32 StartInstanceIndex, const TArray<FTransform>& NewInstancesTransforms, bool bWorldSpace = false, bool bMarkRenderStateDirty = false, bool bTeleport = false);
+	GEOMETRYCOLLECTIONENGINE_API bool BatchUpdateInstancesTransforms(FMeshGroupId MeshGroupId, FMeshId MeshId, int32 StartInstanceIndex, const TArray<FTransform>& NewInstancesTransforms, bool bWorldSpace = false, bool bMarkRenderStateDirty = false, bool bTeleport = false);
 
 	/** 
 	 * Preallocate an ISM in the pool. 
 	 * Doing this early for known mesh instance descriptions can reduce the component registration cost of AddMeshToGroup() for newly discovered mesh descriptions.
 	 */
-	void PreallocateMeshInstance(const FGeometryCollectionStaticMeshInstance& MeshInstance);
+	GEOMETRYCOLLECTIONENGINE_API void PreallocateMeshInstance(const FGeometryCollectionStaticMeshInstance& MeshInstance);
 
 private:
 	uint32 NextMeshGroupId = 0;
diff --git a/Engine/Source/Runtime/Experimental/GeometryCollectionEngine/Public/GeometryCollection/GeometryCollectionISMPoolSubSystem.h b/Engine/Source/Runtime/Experimental/GeometryCollectionEngine/Public/GeometryCollection/GeometryCollectionISMPoolSubSystem.h
--- a/Engine/Source/Runtime/Experimental/GeometryCollectionEngine/Public/GeometryCollection/GeometryCollectionISMPoolSubSystem.h
+++ b/Engine/Source/Runtime/Experimental/GeometryCollectionEngine/Public/GeometryCollection/GeometryCollectionISMPoolSubSystem.h
@@ -11,24 +11,24 @@ class AGeometryCollectionISMPoolActor;
  * A subsystem managing ISMPool actors.
  * Used by geometry collection now but repurposed for more general use.
  */
-UCLASS()
-class GEOMETRYCOLLECTIONENGINE_API UGeometryCollectionISMPoolSubSystem : public UWorldSubsystem
+UCLASS(MinimalAPI)
+class UGeometryCollectionISMPoolSubSystem : public UWorldSubsystem
 {
 	GENERATED_BODY()
 
 public:
-	UGeometryCollectionISMPoolSubSystem();
+	GEOMETRYCOLLECTIONENGINE_API UGeometryCollectionISMPoolSubSystem();
 
 	// USubsystem BEGIN
-	virtual void Initialize(FSubsystemCollectionBase& Collection) override;
-	virtual void Deinitialize() override;
+	GEOMETRYCOLLECTIONENGINE_API virtual void Initialize(FSubsystemCollectionBase& Collection) override;
+	GEOMETRYCOLLECTIONENGINE_API virtual void Deinitialize() override;
 	// USubsystem END
 
 	/** Finds or creates an actor. */
-	AGeometryCollectionISMPoolActor* FindISMPoolActor();
+	GEOMETRYCOLLECTIONENGINE_API AGeometryCollectionISMPoolActor* FindISMPoolActor();
 	
 	/** Get all actors managed by the subsystem. */
-	void GetISMPoolActors(TArray<AGeometryCollectionISMPoolActor*>& OutActors) const;
+	GEOMETRYCOLLECTIONENGINE_API void GetISMPoolActors(TArray<AGeometryCollectionISMPoolActor*>& OutActors) const;
 
 protected:
 	/** For now we only use one ISMPool actor per world, but we could extend the system to manage many more and return the right one based on search criteria. */
