[Engine]

diff --git a/Engine/Source/Runtime/Experimental/GeometryCollectionEngine/Public/GeometryCollection/GeometryCollectionObject.h b/Engine/Source/Runtime/Experimental/GeometryCollectionEngine/Public/GeometryCollection/GeometryCollectionObject.h
--- a/Engine/Source/Runtime/Experimental/GeometryCollectionEngine/Public/GeometryCollection/GeometryCollectionObject.h
+++ b/Engine/Source/Runtime/Experimental/GeometryCollectionEngine/Public/GeometryCollection/GeometryCollectionObject.h
@@ -26,7 +26,7 @@ class UMaterialInterface;
 class UPhysicalMaterial;
 
 USTRUCT(BlueprintType)
-struct GEOMETRYCOLLECTIONENGINE_API FGeometryCollectionSource
+struct FGeometryCollectionSource
 {
 	GENERATED_BODY()
 
@@ -66,7 +66,7 @@ struct GEOMETRYCOLLECTIONENGINE_API FGeometryCollectionSource
 };
 
 USTRUCT(BlueprintType)
-struct GEOMETRYCOLLECTIONENGINE_API FGeometryCollectionAutoInstanceMesh
+struct FGeometryCollectionAutoInstanceMesh
 {
 	GENERATED_BODY()
 
@@ -87,13 +87,13 @@ struct GEOMETRYCOLLECTIONENGINE_API FGeometryCollectionAutoInstanceMesh
 	UPROPERTY(VisibleAnywhere, Category = "AutoInstance")
 	TArray<float> CustomData;
 
-	int32 GetNumDataPerInstance() const;
+	GEOMETRYCOLLECTIONENGINE_API int32 GetNumDataPerInstance() const;
 
-	bool operator ==(const FGeometryCollectionAutoInstanceMesh& Other) const;
+	GEOMETRYCOLLECTIONENGINE_API bool operator ==(const FGeometryCollectionAutoInstanceMesh& Other) const;
 };
 
 USTRUCT(BlueprintType)
-struct GEOMETRYCOLLECTIONENGINE_API FGeometryCollectionEmbeddedExemplar
+struct FGeometryCollectionEmbeddedExemplar
 {
 	GENERATED_BODY()
 
@@ -125,11 +125,11 @@ struct GEOMETRYCOLLECTIONENGINE_API FGeometryCollectionEmbeddedExemplar
 };
 
 USTRUCT()
-struct GEOMETRYCOLLECTIONENGINE_API FGeometryCollectionLevelSetData
+struct FGeometryCollectionLevelSetData
 {
 	GENERATED_BODY()
 
-	FGeometryCollectionLevelSetData();
+	GEOMETRYCOLLECTIONENGINE_API FGeometryCollectionLevelSetData();
 
 	/*
 	*  Resolution on the smallest axes for the level set. (def: 5)
@@ -158,11 +158,11 @@ struct GEOMETRYCOLLECTIONENGINE_API FGeometryCollectionLevelSetData
 
 
 USTRUCT()
-struct GEOMETRYCOLLECTIONENGINE_API FGeometryCollectionCollisionParticleData
+struct FGeometryCollectionCollisionParticleData
 {
 	GENERATED_BODY()
 
-	FGeometryCollectionCollisionParticleData();
+	GEOMETRYCOLLECTIONENGINE_API FGeometryCollectionCollisionParticleData();
 
 	/**
 	 * Number of particles on the triangulated surface to use for collisions.
@@ -180,11 +180,11 @@ struct GEOMETRYCOLLECTIONENGINE_API FGeometryCollectionCollisionParticleData
 
 
 USTRUCT()
-struct GEOMETRYCOLLECTIONENGINE_API FGeometryCollectionCollisionTypeData
+struct FGeometryCollectionCollisionTypeData
 {
 	GENERATED_BODY()
 
-	FGeometryCollectionCollisionTypeData();
+	GEOMETRYCOLLECTIONENGINE_API FGeometryCollectionCollisionTypeData();
 
 	/*
 	*  CollisionType defines how to initialize the rigid collision structures.
@@ -228,11 +228,11 @@ struct GEOMETRYCOLLECTIONENGINE_API FGeometryCollectionCollisionTypeData
 
 
 USTRUCT()
-struct GEOMETRYCOLLECTIONENGINE_API FGeometryCollectionSizeSpecificData
+struct FGeometryCollectionSizeSpecificData
 {
 	GENERATED_BODY()
 
-	FGeometryCollectionSizeSpecificData();
+	GEOMETRYCOLLECTIONENGINE_API FGeometryCollectionSizeSpecificData();
 
 	/** The max size these settings apply to*/
 	UPROPERTY(EditAnywhere, Category = "Collisions")
@@ -306,9 +306,9 @@ struct GEOMETRYCOLLECTIONENGINE_API FGeometryCollectionSizeSpecificData
 	UPROPERTY(EditAnywhere, Category = "Collisions")
 	int32 DamageThreshold;
 
-	bool Serialize(FArchive& Ar);
+	GEOMETRYCOLLECTIONENGINE_API bool Serialize(FArchive& Ar);
 #if WITH_EDITORONLY_DATA
-	void PostSerialize(const FArchive& Ar);
+	GEOMETRYCOLLECTIONENGINE_API void PostSerialize(const FArchive& Ar);
 #endif
 };
 
@@ -328,7 +328,7 @@ struct TStructOpsTypeTraits<FGeometryCollectionSizeSpecificData> : public TStruc
 
 
 USTRUCT(BlueprintType)
-struct GEOMETRYCOLLECTIONENGINE_API FGeometryCollectionProxyMeshData
+struct FGeometryCollectionProxyMeshData
 {
 	GENERATED_BODY()
 
@@ -337,7 +337,7 @@ struct GEOMETRYCOLLECTIONENGINE_API FGeometryCollectionProxyMeshData
 };
 
 USTRUCT()
-struct GEOMETRYCOLLECTIONENGINE_API FGeometryCollectionRenderResourceSizeInfo
+struct FGeometryCollectionRenderResourceSizeInfo
 {
 	GENERATED_BODY();
 
@@ -356,31 +356,31 @@ struct GEOMETRYCOLLECTIONENGINE_API FGeometryCollectionRenderResourceSizeInfo
 * UObject wrapper for the FGeometryCollection
 *
 */
-UCLASS(BlueprintType, customconstructor)
-class GEOMETRYCOLLECTIONENGINE_API UGeometryCollection : public UObject, public IInterface_AssetUserData
+UCLASS(BlueprintType, customconstructor, MinimalAPI)
+class UGeometryCollection : public UObject, public IInterface_AssetUserData
 {
 	GENERATED_UCLASS_BODY()
 
 public:
-	UGeometryCollection(const FObjectInitializer& ObjectInitializer = FObjectInitializer::Get());
+	GEOMETRYCOLLECTIONENGINE_API UGeometryCollection(const FObjectInitializer& ObjectInitializer = FObjectInitializer::Get());
 
 	/** UObject Interface */
 #if WITH_EDITOR
-	virtual void PostEditChangeProperty(struct FPropertyChangedEvent& PropertyChangedEvent) override;
-	virtual bool Modify(bool bAlwaysMarkDirty = true) override;
+	GEOMETRYCOLLECTIONENGINE_API virtual void PostEditChangeProperty(struct FPropertyChangedEvent& PropertyChangedEvent) override;
+	GEOMETRYCOLLECTIONENGINE_API virtual bool Modify(bool bAlwaysMarkDirty = true) override;
 #endif
-	virtual void PostInitProperties() override;
-	virtual void PostLoad() override;
-	virtual void BeginDestroy() override;
+	GEOMETRYCOLLECTIONENGINE_API virtual void PostInitProperties() override;
+	GEOMETRYCOLLECTIONENGINE_API virtual void PostLoad() override;
+	GEOMETRYCOLLECTIONENGINE_API virtual void BeginDestroy() override;
 	/** End UObject Interface */
 
-	void Serialize(FArchive& Ar);
+	GEOMETRYCOLLECTIONENGINE_API void Serialize(FArchive& Ar);
 #if WITH_EDITORONLY_DATA
-	void PostSerialize(const FArchive& Ar);
+	GEOMETRYCOLLECTIONENGINE_API void PostSerialize(const FArchive& Ar);
 #endif
 
 #if WITH_EDITOR
-	void EnsureDataIsCooked(bool bInitResources, bool bIsTransacting, bool bIsPersistant, bool bAllowCopyFromDDC = true);
+	GEOMETRYCOLLECTIONENGINE_API void EnsureDataIsCooked(bool bInitResources, bool bIsTransacting, bool bIsPersistant, bool bAllowCopyFromDDC = true);
 #endif
 
 	/** Accessors for internal geometry collection */
@@ -389,69 +389,69 @@ public:
 	const TSharedPtr<FGeometryCollection, ESPMode::ThreadSafe> GetGeometryCollection() const { return GeometryCollection; }
 
 	/** Return collection to initial (ie. empty) state. */
-	void Reset();
+	GEOMETRYCOLLECTIONENGINE_API void Reset();
 
 	/** Reset the collection from another set of attributes and materials. */
-	void ResetFrom(const FManagedArrayCollection& InCollection, const TArray<UMaterial*>& InMaterials, bool bHasInternalMaterials);
+	GEOMETRYCOLLECTIONENGINE_API void ResetFrom(const FManagedArrayCollection& InCollection, const TArray<UMaterial*>& InMaterials, bool bHasInternalMaterials);
 	
-	int32 AppendGeometry(const UGeometryCollection & Element, bool ReindexAllMaterials = false, const FTransform& TransformRoot = FTransform::Identity);
-	int32 NumElements(const FName& Group) const;
-	void RemoveElements(const FName& Group, const TArray<int32>& SortedDeletionList);
+	GEOMETRYCOLLECTIONENGINE_API int32 AppendGeometry(const UGeometryCollection & Element, bool ReindexAllMaterials = false, const FTransform& TransformRoot = FTransform::Identity);
+	GEOMETRYCOLLECTIONENGINE_API int32 NumElements(const FName& Group) const;
+	GEOMETRYCOLLECTIONENGINE_API void RemoveElements(const FName& Group, const TArray<int32>& SortedDeletionList);
 
 	/** Has data for static mesh rendering. */
-	bool HasMeshData() const;
+	GEOMETRYCOLLECTIONENGINE_API bool HasMeshData() const;
 	/** Has data for nanite rendering. */
-	bool HasNaniteData() const;
+	GEOMETRYCOLLECTIONENGINE_API bool HasNaniteData() const;
 
-	uint32 GetNaniteResourceID() const;
-	uint32 GetNaniteHierarchyOffset() const;
-	uint32 GetNaniteHierarchyOffset(int32 GeometryIndex, bool bFlattened = false) const;
+	GEOMETRYCOLLECTIONENGINE_API uint32 GetNaniteResourceID() const;
+	GEOMETRYCOLLECTIONENGINE_API uint32 GetNaniteHierarchyOffset() const;
+	GEOMETRYCOLLECTIONENGINE_API uint32 GetNaniteHierarchyOffset(int32 GeometryIndex, bool bFlattened = false) const;
 
 	/** ReindexMaterialSections */
-	void ReindexMaterialSections();
+	GEOMETRYCOLLECTIONENGINE_API void ReindexMaterialSections();
 
 	/** appends the standard materials to this UObject */
-	void InitializeMaterials(bool bHasLegacyInternalMaterialsPairs = false);
+	GEOMETRYCOLLECTIONENGINE_API void InitializeMaterials(bool bHasLegacyInternalMaterialsPairs = false);
 
 	/** Add a material to the materials array and update the selected bone material to be at the end of the array */
-	int32 AddNewMaterialSlot(bool bCopyLastMaterial = true);
+	GEOMETRYCOLLECTIONENGINE_API int32 AddNewMaterialSlot(bool bCopyLastMaterial = true);
 
 	/** Remove a material from the materials array, keeping the selected bone material at the end of the array. Returns false if materials could not be removed (e.g. because there were too few). */
-	bool RemoveLastMaterialSlot();
+	GEOMETRYCOLLECTIONENGINE_API bool RemoveLastMaterialSlot();
 
 
 	/** Returns true if there is anything to render */
-	bool HasVisibleGeometry() const;
+	GEOMETRYCOLLECTIONENGINE_API bool HasVisibleGeometry() const;
 
 	/** Invalidates this collection signaling a structural change and renders any previously recorded caches unable to play with this collection */
-	void InvalidateCollection();
+	GEOMETRYCOLLECTIONENGINE_API void InvalidateCollection();
 
 	/** Check to see if Simulation Data requires regeneration */
-	bool IsSimulationDataDirty() const;
+	GEOMETRYCOLLECTIONENGINE_API bool IsSimulationDataDirty() const;
 
 	/** Attach a Static Mesh exemplar for embedded geometry, if that mesh has not already been attached. Return the exemplar index. */
-	int32 AttachEmbeddedGeometryExemplar(const UStaticMesh* Exemplar);
+	GEOMETRYCOLLECTIONENGINE_API int32 AttachEmbeddedGeometryExemplar(const UStaticMesh* Exemplar);
 
 	/** Remove embedded geometry exemplars with indices matching the sorted removal list. */
-	void RemoveExemplars(const TArray<int32>& SortedRemovalIndices);
+	GEOMETRYCOLLECTIONENGINE_API void RemoveExemplars(const TArray<int32>& SortedRemovalIndices);
 
 	/** find or add a auto instance mesh and return its index */
-	const FGeometryCollectionAutoInstanceMesh& GetAutoInstanceMesh(int32 AutoInstanceMeshIndex) const;
+	GEOMETRYCOLLECTIONENGINE_API const FGeometryCollectionAutoInstanceMesh& GetAutoInstanceMesh(int32 AutoInstanceMeshIndex) const;
 
 	/**  find or add a auto instance mesh from another one and return its index */
-	int32 FindOrAddAutoInstanceMesh(const FGeometryCollectionAutoInstanceMesh& AutoInstanceMesh);
+	GEOMETRYCOLLECTIONENGINE_API int32 FindOrAddAutoInstanceMesh(const FGeometryCollectionAutoInstanceMesh& AutoInstanceMesh);
 
 	/** find or add a auto instance mesh from a mesh and alist of material and return its index */
-	int32 FindOrAddAutoInstanceMesh(const UStaticMesh* StaticMesh, const TArray<UMaterialInterface*>& Materials);
+	GEOMETRYCOLLECTIONENGINE_API int32 FindOrAddAutoInstanceMesh(const UStaticMesh* StaticMesh, const TArray<UMaterialInterface*>& Materials);
 
 	/** Produce a deep copy of GeometryCollection member, stripped of data unecessary for gameplay. */
-	TSharedPtr<FGeometryCollection, ESPMode::ThreadSafe> GenerateMinimalGeometryCollection() const;
+	GEOMETRYCOLLECTIONENGINE_API TSharedPtr<FGeometryCollection, ESPMode::ThreadSafe> GenerateMinimalGeometryCollection() const;
 
 	/** copy a collection and remove geometry from it */
-	static TSharedPtr<FGeometryCollection, ESPMode::ThreadSafe> CopyCollectionAndRemoveGeometry(const TSharedPtr<const FGeometryCollection, ESPMode::ThreadSafe>& CollectionToCopy);
+	static GEOMETRYCOLLECTIONENGINE_API TSharedPtr<FGeometryCollection, ESPMode::ThreadSafe> CopyCollectionAndRemoveGeometry(const TSharedPtr<const FGeometryCollection, ESPMode::ThreadSafe>& CollectionToCopy);
 
 	/** get the size of the render data resources associated with this collection */
-	FGeometryCollectionRenderResourceSizeInfo GetRenderResourceSizeInfo() const;
+	GEOMETRYCOLLECTIONENGINE_API FGeometryCollectionRenderResourceSizeInfo GetRenderResourceSizeInfo() const;
 
 #if WITH_EDITOR
 	/** If this flag is set, we only regenerate simulation data when requested via CreateSimulationData() */
@@ -462,37 +462,37 @@ public:
 	* Note : this does not check if the simulation data is drty or not and will cause a load from the DDC
 	* use CreateSimulationDataIfNeeded() for avoiding extra runtime cost 
 	*/
-	void CreateSimulationData();
+	GEOMETRYCOLLECTIONENGINE_API void CreateSimulationData();
 
 	/** Create the simulation data ( calls CreateSimulationData) only if the simulation data is dirty */
-	void CreateSimulationDataIfNeeded();
+	GEOMETRYCOLLECTIONENGINE_API void CreateSimulationDataIfNeeded();
 
 	/** Rebuild the render data. */
-	void RebuildRenderData();
+	GEOMETRYCOLLECTIONENGINE_API void RebuildRenderData();
 
 	/** Propogate render state dirty to components */
-	void PropagateMarkDirtyToComponents() const;
+	GEOMETRYCOLLECTIONENGINE_API void PropagateMarkDirtyToComponents() const;
 
 	/** Propogate the fact that transform have been updated to all components */
-	void PropagateTransformUpdateToComponents() const;
+	GEOMETRYCOLLECTIONENGINE_API void PropagateTransformUpdateToComponents() const;
 #endif
 
-	void InitResources();
-	void ReleaseResources();
+	GEOMETRYCOLLECTIONENGINE_API void InitResources();
+	GEOMETRYCOLLECTIONENGINE_API void ReleaseResources();
 
 	/** Fills params struct with parameters used for precomputing content. */
-	void GetSharedSimulationParams(FSharedSimulationParameters& OutParams) const;
+	GEOMETRYCOLLECTIONENGINE_API void GetSharedSimulationParams(FSharedSimulationParameters& OutParams) const;
 
 	/**
 	* Get Mass or density as set by the asset 
 	* Mass is return in Kg and Density is returned in Kg/Cm3
 	* @param bOutIsDensity  is set to true by the function if the returned value is to be treated as a density
 	*/
-	float GetMassOrDensity(bool& bOutIsDensity) const;
+	GEOMETRYCOLLECTIONENGINE_API float GetMassOrDensity(bool& bOutIsDensity) const;
 
 	/** Accessors for the two guids used to identify this collection */
-	FGuid GetIdGuid() const;
-	FGuid GetStateGuid() const;
+	GEOMETRYCOLLECTIONENGINE_API FGuid GetIdGuid() const;
+	GEOMETRYCOLLECTIONENGINE_API FGuid GetStateGuid() const;
 
 	/** Get the cached root index */
 	int32 GetRootIndex() const { return RootIndex; }
@@ -582,7 +582,7 @@ public:
 	TArray<FGeometryCollectionAutoInstanceMesh> AutoInstanceMeshes;
 
 	UFUNCTION(BlueprintCallable, Category = "Nanite")
-	void SetEnableNanite(bool bValue);
+	GEOMETRYCOLLECTIONENGINE_API void SetEnableNanite(bool bValue);
 
 	/**
 	 * Enable support for Nanite.
@@ -591,7 +591,7 @@ public:
 	bool EnableNanite;
 
 	UFUNCTION(BlueprintCallable, Category = "Rendering")
-	void SetConvertVertexColorsToSRGB(bool bValue);
+	GEOMETRYCOLLECTIONENGINE_API void SetConvertVertexColorsToSRGB(bool bValue);
 
 	/**
 	 * Convert vertex colors to sRGB for rendering. Exposed to avoid changing vertex color rendering for legacy assets; should typically be true for new geometry collections.
@@ -727,10 +727,10 @@ public:
 	UPROPERTY(EditAnywhere, Category = "Collisions")
 	TArray<FGeometryCollectionSizeSpecificData> SizeSpecificData;
 
-	int GetDefaultSizeSpecificDataIndex() const;
-	FGeometryCollectionSizeSpecificData& GetDefaultSizeSpecificData();
-	const FGeometryCollectionSizeSpecificData& GetDefaultSizeSpecificData() const;
-	static FGeometryCollectionSizeSpecificData GeometryCollectionSizeSpecificDataDefaults();
+	GEOMETRYCOLLECTIONENGINE_API int GetDefaultSizeSpecificDataIndex() const;
+	GEOMETRYCOLLECTIONENGINE_API FGeometryCollectionSizeSpecificData& GetDefaultSizeSpecificData();
+	GEOMETRYCOLLECTIONENGINE_API const FGeometryCollectionSizeSpecificData& GetDefaultSizeSpecificData() const;
+	static GEOMETRYCOLLECTIONENGINE_API FGeometryCollectionSizeSpecificData GeometryCollectionSizeSpecificDataDefaults();
 
 	/**
 	* Enable remove pieces on fracture
@@ -752,7 +752,7 @@ public:
 	}
 
 	/** Returns the asset path for the automatically populated selected material. */
-	static const TCHAR* GetSelectedMaterialPath();
+	static GEOMETRYCOLLECTIONENGINE_API const TCHAR* GetSelectedMaterialPath();
 
 #if WITH_EDITORONLY_DATA
 	/** Importing data and options used for this geometry collection */
@@ -767,20 +767,20 @@ public:
 	/*
 	* Update the convex geometry on the collection.
 	*/
-	void UpdateConvexGeometry();
-	void UpdateConvexGeometryIfMissing();
+	GEOMETRYCOLLECTIONENGINE_API void UpdateConvexGeometry();
+	GEOMETRYCOLLECTIONENGINE_API void UpdateConvexGeometryIfMissing();
 	
 	/*
 	 * Update properties that depend on the geometry and clustering: Proximity, Convex Hulls, Volume and Size data.
 	 */
-	void UpdateGeometryDependentProperties();
+	GEOMETRYCOLLECTIONENGINE_API void UpdateGeometryDependentProperties();
 
 
 	//~ Begin IInterface_AssetUserData Interface
-	virtual void AddAssetUserData(UAssetUserData* InUserData) override;
-	virtual void RemoveUserDataOfClass(TSubclassOf<UAssetUserData> InUserDataClass) override;
-	virtual UAssetUserData* GetAssetUserDataOfClass(TSubclassOf<UAssetUserData> InUserDataClass) override;
-	virtual const TArray<UAssetUserData*>* GetAssetUserDataArray() const override;
+	GEOMETRYCOLLECTIONENGINE_API virtual void AddAssetUserData(UAssetUserData* InUserData) override;
+	GEOMETRYCOLLECTIONENGINE_API virtual void RemoveUserDataOfClass(TSubclassOf<UAssetUserData> InUserDataClass) override;
+	GEOMETRYCOLLECTIONENGINE_API virtual UAssetUserData* GetAssetUserDataOfClass(TSubclassOf<UAssetUserData> InUserDataClass) override;
+	GEOMETRYCOLLECTIONENGINE_API virtual const TArray<UAssetUserData*>* GetAssetUserDataArray() const override;
 	//~ End IInterface_AssetUserData Interface
 
 	//
@@ -797,21 +797,21 @@ public:
 
 private:
 #if WITH_EDITOR
-	void CreateSimulationDataImp(bool bCopyFromDDC);
-	void CreateRenderDataImp(bool bCopyFromDDC);
+	GEOMETRYCOLLECTIONENGINE_API void CreateSimulationDataImp(bool bCopyFromDDC);
+	GEOMETRYCOLLECTIONENGINE_API void CreateRenderDataImp(bool bCopyFromDDC);
 #endif
 
 	/*
 	* Used to transfer deprecated properties to the size specific structures during serialization
 	* and to add back the default size specific data when deleted.
 	*/
-	void ValidateSizeSpecificDataDefaults();
+	GEOMETRYCOLLECTIONENGINE_API void ValidateSizeSpecificDataDefaults();
 
 	// update cached root index using the current hierarchy setup
-	void UpdateRootIndex();
+	GEOMETRYCOLLECTIONENGINE_API void UpdateRootIndex();
 
 	// fill instanced mesh instance count from geometry collection data if not done yet 
-	void FillAutoInstanceMeshesInstancesIfNeeded();
+	GEOMETRYCOLLECTIONENGINE_API void FillAutoInstanceMeshesInstancesIfNeeded();
 
 private:
 	/** Guid created on construction of this collection. It should be used to uniquely identify this collection */
diff --git a/Engine/Source/Runtime/Experimental/GeometryCollectionEngine/Public/GeometryCollection/GeometryCollectionRenderLevelSetActor.h b/Engine/Source/Runtime/Experimental/GeometryCollectionEngine/Public/GeometryCollection/GeometryCollectionRenderLevelSetActor.h
--- a/Engine/Source/Runtime/Experimental/GeometryCollectionEngine/Public/GeometryCollection/GeometryCollectionRenderLevelSetActor.h
+++ b/Engine/Source/Runtime/Experimental/GeometryCollectionEngine/Public/GeometryCollection/GeometryCollectionRenderLevelSetActor.h
@@ -24,13 +24,13 @@
 *    render level set actors isn't currently supported very well,
 *    but could be improved in the future
 */
-UCLASS()
-class GEOMETRYCOLLECTIONENGINE_API AGeometryCollectionRenderLevelSetActor : public AActor
+UCLASS(MinimalAPI)
+class AGeometryCollectionRenderLevelSetActor : public AActor
 {
 	GENERATED_UCLASS_BODY()
 public:
 
-	static int InstanceCount;
+	static GEOMETRYCOLLECTIONENGINE_API int InstanceCount;
 
 	// Volume texture to fill	
 	UPROPERTY(EditAnywhere, BlueprintReadWrite, Category = "Volume")
@@ -61,20 +61,20 @@ public:
 	bool RenderVolumeBoundingBox;
 
 	// Find/create the level set renderer singleton actor as required. Return whether the found or created actor.
-	static AGeometryCollectionRenderLevelSetActor* FindOrCreate(UWorld* World);
+	static GEOMETRYCOLLECTIONENGINE_API AGeometryCollectionRenderLevelSetActor* FindOrCreate(UWorld* World);
 
 	// Load a new level set to render
-	bool SetLevelSetToRender(const Chaos::FLevelSet &LevelSet, const FTransform &LocalToWorld);
+	GEOMETRYCOLLECTIONENGINE_API bool SetLevelSetToRender(const Chaos::FLevelSet &LevelSet, const FTransform &LocalToWorld);
 
 	// Sync level set transform to the render material
-	void SyncLevelSetTransform(const FTransform &LocalToWorld);
+	GEOMETRYCOLLECTIONENGINE_API void SyncLevelSetTransform(const FTransform &LocalToWorld);
 
 	// Some initialization happens in here 
-	virtual void BeginPlay() override;
+	GEOMETRYCOLLECTIONENGINE_API virtual void BeginPlay() override;
 
 #if WITH_EDITOR
 	// Allowed for live updates to parameters from inside the editor when ejected
-	virtual void PostEditChangeProperty(struct FPropertyChangedEvent& e) override;
+	GEOMETRYCOLLECTIONENGINE_API virtual void PostEditChangeProperty(struct FPropertyChangedEvent& e) override;
 #endif
 
 	// set and sync enabled
@@ -98,11 +98,11 @@ private:
 	UMaterialInstanceDynamic* DynRayMarchMaterial;
 
 	// Synchronizes the state of this actor class with the post process render material
-	void SyncMaterialParameters();
+	GEOMETRYCOLLECTIONENGINE_API void SyncMaterialParameters();
 
 	// Private for now since step size mult might not be super useful due to the current
 	// rendering algorithms employed in the shaders
 	// @todo: expose this in a meaningful way in the future if necessary
 	// UPROPERTY(EditAnywhere, BlueprintReadWrite, Category = "Rendering")
 	float StepSizeMult;
-};
\ No newline at end of file
+};
diff --git a/Engine/Source/Runtime/Experimental/GeometryCollectionEngine/Public/GeometryCollection/PhysicsAssetSimulation.h b/Engine/Source/Runtime/Experimental/GeometryCollectionEngine/Public/GeometryCollection/PhysicsAssetSimulation.h
--- a/Engine/Source/Runtime/Experimental/GeometryCollectionEngine/Public/GeometryCollection/PhysicsAssetSimulation.h
+++ b/Engine/Source/Runtime/Experimental/GeometryCollectionEngine/Public/GeometryCollection/PhysicsAssetSimulation.h
@@ -14,8 +14,8 @@ class UPhysicsAsset;
 class USkeletalMeshComponent;
 
 
-struct GEOMETRYCOLLECTIONENGINE_API FPhysicsAssetSimulationUtil
+struct FPhysicsAssetSimulationUtil
 {
-	static void BuildParams(const UObject* Caller, const AActor* OwningActor, const USkeletalMeshComponent* SkelMeshComponent, const UPhysicsAsset* PhysicsAsset, FSkeletalMeshPhysicsProxyParams& OutParams);
-	static bool UpdateAnimState(const UObject* Caller, const AActor* OwningActor, const USkeletalMeshComponent* SkelMeshComponent, const float Dt, FSkeletalMeshPhysicsProxyParams& OutParams);
+	static GEOMETRYCOLLECTIONENGINE_API void BuildParams(const UObject* Caller, const AActor* OwningActor, const USkeletalMeshComponent* SkelMeshComponent, const UPhysicsAsset* PhysicsAsset, FSkeletalMeshPhysicsProxyParams& OutParams);
+	static GEOMETRYCOLLECTIONENGINE_API bool UpdateAnimState(const UObject* Caller, const AActor* OwningActor, const USkeletalMeshComponent* SkelMeshComponent, const float Dt, FSkeletalMeshPhysicsProxyParams& OutParams);
 };
diff --git a/Engine/Source/Runtime/Experimental/Iris/Core/Private/Iris/ReplicationSystem/ReplicationConnections.h b/Engine/Source/Runtime/Experimental/Iris/Core/Private/Iris/ReplicationSystem/ReplicationConnections.h
--- a/Engine/Source/Runtime/Experimental/Iris/Core/Private/Iris/ReplicationSystem/ReplicationConnections.h
+++ b/Engine/Source/Runtime/Experimental/Iris/Core/Private/Iris/ReplicationSystem/ReplicationConnections.h
@@ -66,7 +66,7 @@ public:
 		return ConnectionId < GetMaxConnectionCount() && ValidConnections.GetBit(ConnectionId);
 	}
 
-	IRISCORE_API void AddConnection(uint32 ConnectionId)
+	void AddConnection(uint32 ConnectionId)
 	{
 		check(ValidConnections.GetBit(ConnectionId) == false);
 		ValidConnections.SetBit(ConnectionId);
diff --git a/Engine/Source/Runtime/Experimental/Iris/Core/Private/Iris/ReplicationSystem/ReplicationProtocolManager.h b/Engine/Source/Runtime/Experimental/Iris/Core/Private/Iris/ReplicationSystem/ReplicationProtocolManager.h
--- a/Engine/Source/Runtime/Experimental/Iris/Core/Private/Iris/ReplicationSystem/ReplicationProtocolManager.h
+++ b/Engine/Source/Runtime/Experimental/Iris/Core/Private/Iris/ReplicationSystem/ReplicationProtocolManager.h
@@ -23,7 +23,7 @@ public:
 	IRISCORE_API const FReplicationProtocol* CreateReplicationProtocol(const UObject* ArchetypeOrCDOUsedAsKey, const FReplicationProtocolIdentifier ProtocolId, const FReplicationFragments& Fragments, const TCHAR* DebugName, bool bVerifyId = false);
 
 	/* Create protocol from registered fragment data with provided Id, verification is optional */
-	IRISCORE_API const FReplicationProtocol* CreateReplicationProtocol(const FReplicationProtocolIdentifier ProtocolId, const FReplicationFragments& Fragments, const TCHAR* DebugName, bool bVerifyId = false)
+	const FReplicationProtocol* CreateReplicationProtocol(const FReplicationProtocolIdentifier ProtocolId, const FReplicationFragments& Fragments, const TCHAR* DebugName, bool bVerifyId = false)
 	{
 		return CreateReplicationProtocol(nullptr, ProtocolId, Fragments, DebugName, bVerifyId);
 	}
diff --git a/Engine/Source/Runtime/Experimental/Iris/Core/Public/Iris/DataStream/DataStreamManager.h b/Engine/Source/Runtime/Experimental/Iris/Core/Public/Iris/DataStream/DataStreamManager.h
--- a/Engine/Source/Runtime/Experimental/Iris/Core/Public/Iris/DataStream/DataStreamManager.h
+++ b/Engine/Source/Runtime/Experimental/Iris/Core/Public/Iris/DataStream/DataStreamManager.h
@@ -107,7 +107,7 @@ public:
 	UE::Net::Private::FNetExports& GetNetExports();
 
 private:
-	IRISCORE_API UDataStreamManager();
+	UDataStreamManager();
 
 	static void AddReferencedObjects(UObject* Object, FReferenceCollector& Collector);
 
diff --git a/Engine/Source/Runtime/Experimental/Iris/Core/Public/Iris/ReplicationState/IrisFastArraySerializer.h b/Engine/Source/Runtime/Experimental/Iris/Core/Public/Iris/ReplicationState/IrisFastArraySerializer.h
--- a/Engine/Source/Runtime/Experimental/Iris/Core/Public/Iris/ReplicationState/IrisFastArraySerializer.h
+++ b/Engine/Source/Runtime/Experimental/Iris/Core/Public/Iris/ReplicationState/IrisFastArraySerializer.h
@@ -44,7 +44,7 @@ struct FIrisFastArraySerializer : public FFastArraySerializer
 	IRISCORE_API FIrisFastArraySerializer& operator=(FIrisFastArraySerializer&& Other);
 	
 	/** Override MarkItemDirty in order to mark object as dirty in the DirtyNetObjectTracker */
-	IRISCORE_API void MarkItemDirty(FFastArraySerializerItem & Item)
+	void MarkItemDirty(FFastArraySerializerItem & Item)
 	{ 
 		FFastArraySerializer::MarkItemDirty(Item);
 		if (ReplicationStateHeader.IsBound())
@@ -55,7 +55,7 @@ struct FIrisFastArraySerializer : public FFastArraySerializer
 	}
 	
 	/** Override MarkArrayDirty in order to mark object as dirty in the DirtyNetObjectTracker */
-	IRISCORE_API void MarkArrayDirty()
+	void MarkArrayDirty()
 	{
 		FFastArraySerializer::MarkArrayDirty();
 		if (ReplicationStateHeader.IsBound())
diff --git a/Engine/Source/Runtime/Experimental/Iris/Core/Public/Iris/ReplicationSystem/Prioritization/SphereNetObjectPrioritizer.h b/Engine/Source/Runtime/Experimental/Iris/Core/Public/Iris/ReplicationSystem/Prioritization/SphereNetObjectPrioritizer.h
--- a/Engine/Source/Runtime/Experimental/Iris/Core/Public/Iris/ReplicationSystem/Prioritization/SphereNetObjectPrioritizer.h
+++ b/Engine/Source/Runtime/Experimental/Iris/Core/Public/Iris/ReplicationSystem/Prioritization/SphereNetObjectPrioritizer.h
@@ -55,7 +55,7 @@ protected:
 		VectorRegister PriorityFactor;
 	};
 
-	struct IRISCORE_API FBatchParams
+	struct FBatchParams
 	{
 		FPriorityCalculationConstants PriorityCalculationConstants;
 		UE::Net::FReplicationView View;
diff --git a/Engine/Source/Runtime/Experimental/Iris/Core/Public/Iris/Serialization/NetBitStreamReader.h b/Engine/Source/Runtime/Experimental/Iris/Core/Public/Iris/Serialization/NetBitStreamReader.h
--- a/Engine/Source/Runtime/Experimental/Iris/Core/Public/Iris/Serialization/NetBitStreamReader.h
+++ b/Engine/Source/Runtime/Experimental/Iris/Core/Public/Iris/Serialization/NetBitStreamReader.h
@@ -31,7 +31,7 @@ public:
 	 * Reads a bool from the stream and returns the value,
 	 * A failed read will always return false and stream will be marked as overflown
 	 */
-	IRISCORE_API bool ReadBool() { return ReadBits(1) & 1U; }
+	bool ReadBool() { return ReadBits(1) & 1U; }
 
 	/**
 	 * Reads BitCount bits and stores them in Dst, starting from bit offset 0. The bits will be stored
diff --git a/Engine/Source/Runtime/Experimental/Voronoi/Public/Voronoi/Voronoi.h b/Engine/Source/Runtime/Experimental/Voronoi/Public/Voronoi/Voronoi.h
--- a/Engine/Source/Runtime/Experimental/Voronoi/Public/Voronoi/Voronoi.h
+++ b/Engine/Source/Runtime/Experimental/Voronoi/Public/Voronoi/Voronoi.h
@@ -35,7 +35,7 @@ class FVoronoiDiagram;
 
 // The Voronoi Compute Helper caches information to help compute cells quickly
 // To safely parallelize voronoi diagram construction, create a separate compute helper per-thread
-class VORONOI_API FVoronoiComputeHelper
+class FVoronoiComputeHelper
 {
 	TPimplPtr<voro::voro_compute<voro::container>> Compute;
 public:
@@ -46,19 +46,19 @@ public:
 	FVoronoiComputeHelper(FVoronoiComputeHelper&& Other) = default;
 	FVoronoiComputeHelper& operator=(FVoronoiComputeHelper&& Other) = default;
 
-	void Init();
+	VORONOI_API void Init();
 
 	friend class FVoronoiDiagram;
 };
 
-class VORONOI_API FVoronoiDiagram
+class FVoronoiDiagram
 {
 	int32 NumSites;
 	TPimplPtr<voro::container> Container;
 	FBox Bounds;
 
 public:
-	const static int MinDefaultSitesPerThread;
+	VORONOI_API const static int MinDefaultSitesPerThread;
 
 	/**
 	 * Create a Voronoi diagram with the given sites, in a box bounding containing those sites
@@ -68,7 +68,7 @@ public:
 	 * @param SquaredDistSkipPtThreshold	A safety threshold to avoid creating invalid cells: sites that are within this distance of an already-added site will not be added
 	 *										(If you know there will be no duplicate sites, can set to zero for faster perf.)
 	 */
-	FVoronoiDiagram(const TArrayView<const FVector>& Sites, double ExtraBoundingSpace, double SquaredDistSkipPtThreshold = UE_KINDA_SMALL_NUMBER);
+	VORONOI_API FVoronoiDiagram(const TArrayView<const FVector>& Sites, double ExtraBoundingSpace, double SquaredDistSkipPtThreshold = UE_KINDA_SMALL_NUMBER);
 	/**
 	 * Create a Voronoi diagram with the given sites, in a given bounding box
 	 * 
@@ -78,7 +78,7 @@ public:
 	 * @param SquaredDistSkipPtThreshold	A safety threshold to avoid creating invalid cells: sites that are within this distance of an already-added site will not be added.
 	 *										(If you know there will be no duplicate sites, can set to zero for faster perf.)
 	 */
-	FVoronoiDiagram(const TArrayView<const FVector>& Sites, const FBox &Bounds, double ExtraBoundingSpace, double SquaredDistSkipPtThreshold = UE_KINDA_SMALL_NUMBER);
+	VORONOI_API FVoronoiDiagram(const TArrayView<const FVector>& Sites, const FBox &Bounds, double ExtraBoundingSpace, double SquaredDistSkipPtThreshold = UE_KINDA_SMALL_NUMBER);
 
 	/**
 	 * Make a container for a Voronoi diagram without any points.  Call AddSites() after to finish creating the diagram.
@@ -87,7 +87,7 @@ public:
 	 * @param Bounds				Bounding box within which to compute the Voronoi diagram
 	 * @param ExtraBoundingSpace	Amount to expand the bounding box before computing the Voronoi diagram
 	 */
-	FVoronoiDiagram(int32 ExpectedNumSites, const FBox& Bounds, double ExtraBoundingSpace);
+	VORONOI_API FVoronoiDiagram(int32 ExpectedNumSites, const FBox& Bounds, double ExtraBoundingSpace);
 
 	/**
 	 * Make an empty diagram, to be filled in with a subsequent call to "Initialize()"
@@ -109,26 +109,26 @@ public:
 	 * @param SquaredDistSkipPtThreshold	A safety threshold to avoid creating invalid cells: sites that are within this distance of an already-added site will not be added.
 	 *										(If you know there will be no duplicate sites, can set to zero for faster perf.)
 	 */
-	void Initialize(const TArrayView<const FVector>& Sites, const FBox& Bounds, double ExtraBoundingSpace, double SquaredDistSkipPtThreshold = UE_KINDA_SMALL_NUMBER);
+	VORONOI_API void Initialize(const TArrayView<const FVector>& Sites, const FBox& Bounds, double ExtraBoundingSpace, double SquaredDistSkipPtThreshold = UE_KINDA_SMALL_NUMBER);
 
-	~FVoronoiDiagram();
+	VORONOI_API ~FVoronoiDiagram();
 
-	static FBox GetBounds(const TArrayView<const FVector>& Sites, double ExtraBoundingSpace = UE_DOUBLE_KINDA_SMALL_NUMBER);
+	static VORONOI_API FBox GetBounds(const TArrayView<const FVector>& Sites, double ExtraBoundingSpace = UE_DOUBLE_KINDA_SMALL_NUMBER);
 
 	int32 Num() const
 	{
 		return NumSites;
 	}
 
-	void AddSites(const TArrayView<const FVector>& Sites, double SquaredDistSkipPtThreshold = 0.0f);
+	VORONOI_API void AddSites(const TArrayView<const FVector>& Sites, double SquaredDistSkipPtThreshold = 0.0f);
 
-	void ComputeAllCellsSerial(TArray<FVoronoiCellInfo>& AllCells);
-	void ComputeAllCells(TArray<FVoronoiCellInfo>& AllCells, int32 ApproxSitesPerThread = -1);
+	VORONOI_API void ComputeAllCellsSerial(TArray<FVoronoiCellInfo>& AllCells);
+	VORONOI_API void ComputeAllCells(TArray<FVoronoiCellInfo>& AllCells, int32 ApproxSitesPerThread = -1);
 
-	void ComputeAllNeighbors(TArray<TArray<int32>>& AllNeighbors, bool bExcludeBounds = true, int32 ApproxSitesPerThread = -1);
+	VORONOI_API void ComputeAllNeighbors(TArray<TArray<int32>>& AllNeighbors, bool bExcludeBounds = true, int32 ApproxSitesPerThread = -1);
 
-	void ComputeCellEdgesSerial(TArray<TTuple<FVector, FVector>>& Edges, TArray<int32>& CellMember);
-	void ComputeCellEdges(TArray<TTuple<FVector, FVector>>& Edges, TArray<int32>& CellMember, int32 ApproxSitesPerThread = -1);
+	VORONOI_API void ComputeCellEdgesSerial(TArray<TTuple<FVector, FVector>>& Edges, TArray<int32>& CellMember);
+	VORONOI_API void ComputeCellEdges(TArray<TTuple<FVector, FVector>>& Edges, TArray<int32>& CellMember, int32 ApproxSitesPerThread = -1);
 
 	/**
 	 * Find the closest Voronoi cell to a given position
@@ -138,7 +138,7 @@ public:
 	 * @param OutFoundSite		The position of the closest Voronoi site, if any
 	 * @return					The id of the Voronoi cell containing the given position, or -1 if position is outside diagram
 	 */
-	int32 FindCell(const FVector& Pos, FVoronoiComputeHelper& ComputeHelper, FVector& OutFoundSite) const;
+	VORONOI_API int32 FindCell(const FVector& Pos, FVoronoiComputeHelper& ComputeHelper, FVector& OutFoundSite) const;
 
 	int32 FindCell(const FVector& Pos, FVoronoiComputeHelper& ComputeHelper) const
 	{
@@ -146,7 +146,7 @@ public:
 		return FindCell(Pos, ComputeHelper, OutFoundSite);
 	}
 
-	FVoronoiComputeHelper GetComputeHelper() const;
+	VORONOI_API FVoronoiComputeHelper GetComputeHelper() const;
 
 	int32 ApproxSitesPerThreadWithDefault(int32 ApproxSitesPerThreadIn)
 	{
@@ -158,7 +158,7 @@ public:
 	}
 
 private:
-	TArray<int32> GetParallelBlockRanges(int32 ApproxSitesPerThread);
+	VORONOI_API TArray<int32> GetParallelBlockRanges(int32 ApproxSitesPerThread);
 };
 
 /**
@@ -170,7 +170,7 @@ private:
  * 
  * Note: Queries are only thread safe if each thread has its own FVoronoiComputeHelper
  */
-class VORONOI_API FVoronoiDiagramField
+class FVoronoiDiagramField
 {
 	TArray<FVector> Sites;
 	TArray<TArray<int32>> Neighbors;
diff --git a/Engine/Source/Runtime/ExternalRPCRegistry/Public/ExternalRpcRegistry.h b/Engine/Source/Runtime/ExternalRPCRegistry/Public/ExternalRpcRegistry.h
--- a/Engine/Source/Runtime/ExternalRPCRegistry/Public/ExternalRpcRegistry.h
+++ b/Engine/Source/Runtime/ExternalRPCRegistry/Public/ExternalRpcRegistry.h
@@ -105,61 +105,61 @@ public:
  * This class is designed to be a singleton that handles registry, maintenance, and cleanup of any REST endpoints exposed on the process 
  * for use in communicating with the process externally. 
  */
-UCLASS()
-class EXTERNALRPCREGISTRY_API UExternalRpcRegistry : public UObject
+UCLASS(MinimalAPI)
+class UExternalRpcRegistry : public UObject
 {
 	GENERATED_BODY()
 protected:
-	static UExternalRpcRegistry * ObjectInstance;
+	static EXTERNALRPCREGISTRY_API UExternalRpcRegistry * ObjectInstance;
 	TMap<FName, FExternalRouteDesc> RegisteredRoutes;
 	TArray<FString> ActiveRpcCategories;
 public:
-	static UExternalRpcRegistry * GetInstance();
+	static EXTERNALRPCREGISTRY_API UExternalRpcRegistry * GetInstance();
 
-	~UExternalRpcRegistry();
+	EXTERNALRPCREGISTRY_API ~UExternalRpcRegistry();
 
 	int PortToUse = 11223;
 
 	/**
 	* Check if this Rpc is from a category that is meant to be enabled.
 	*/
-	bool IsActiveRpcCategory(FString InCategory);
+	EXTERNALRPCREGISTRY_API bool IsActiveRpcCategory(FString InCategory);
 
 	/**
 	 * Try to get a route registered under given friendly name. Returns false if could not be found.
 	 */
-	bool GetRegisteredRoute(FName RouteName, FExternalRouteInfo& OutRouteInfo);
+	EXTERNALRPCREGISTRY_API bool GetRegisteredRoute(FName RouteName, FExternalRouteInfo& OutRouteInfo);
 
-	void RegisterNewRoute(FExternalRouteInfo InRouteInfo, const FHttpRequestHandler& Handler, bool bOverrideIfBound = false);
+	EXTERNALRPCREGISTRY_API void RegisterNewRoute(FExternalRouteInfo InRouteInfo, const FHttpRequestHandler& Handler, bool bOverrideIfBound = false);
 
 	/**
 	 * Register a new route.
 	 * Will override existing routes if option is set, otherwise will error and fail to bind.
 	 */
-	void RegisterNewRouteWithArguments(FName RouteName, const FHttpPath& HttpPath, const EHttpServerRequestVerbs& RequestVerbs, const FHttpRequestHandler& Handler, TArray<FExternalRpcArgumentDesc*> InArguments, bool bOverrideIfBound = false, bool bIsAlwaysOn = false, FString OptionalCategory = TEXT("Unknown"), FString OptionalContentType = TEXT(""));
+	EXTERNALRPCREGISTRY_API void RegisterNewRouteWithArguments(FName RouteName, const FHttpPath& HttpPath, const EHttpServerRequestVerbs& RequestVerbs, const FHttpRequestHandler& Handler, TArray<FExternalRpcArgumentDesc*> InArguments, bool bOverrideIfBound = false, bool bIsAlwaysOn = false, FString OptionalCategory = TEXT("Unknown"), FString OptionalContentType = TEXT(""));
 
 	/**
 	* Deprecated way to register a new route.
 	* Will override existing routes if option is set, otherwise will error and fail to bind.
 	 */
 	UE_DEPRECATED(5.0, "RegisterNewRoute is deprecated when needing to add arguments, please use RegisterNewRouteWithArguments instead.")
-	void RegisterNewRoute(FName RouteName, const FHttpPath& HttpPath, const EHttpServerRequestVerbs& RequestVerbs, const FHttpRequestHandler& Handler, bool bOverrideIfBound = false, bool bIsAlwaysOn = false, FString OptionalCategory = TEXT("Unknown"), FString OptionalContentType = TEXT(""), FString OptionalExpectedFormat = TEXT(""));
+	EXTERNALRPCREGISTRY_API void RegisterNewRoute(FName RouteName, const FHttpPath& HttpPath, const EHttpServerRequestVerbs& RequestVerbs, const FHttpRequestHandler& Handler, bool bOverrideIfBound = false, bool bIsAlwaysOn = false, FString OptionalCategory = TEXT("Unknown"), FString OptionalContentType = TEXT(""), FString OptionalExpectedFormat = TEXT(""));
 
 	/**
 	 * Clean up all routes - generally called as part of the destructor to make sure we don't have any function pointers dangling around.
 	 */
-	void CleanUpAllRoutes();
+	EXTERNALRPCREGISTRY_API void CleanUpAllRoutes();
 
 	/**
 	 * Clean up a route.
 	 * Can be set to fail if trying to unbind an unbound route.
 	 */
-	void CleanUpRoute(FName RouteName, bool bFailIfUnbound = false);
+	EXTERNALRPCREGISTRY_API void CleanUpRoute(FName RouteName, bool bFailIfUnbound = false);
 
 	/**
 	 * Default Route Listing http call. Spits out all registered routes and describes them via a REST API call.
 	 * Always registered at /listrpcs GET by default
 	 */
-	bool HttpListOpenRoutes(const FHttpServerRequest& Request, const FHttpResultCallback& OnComplete);
+	EXTERNALRPCREGISTRY_API bool HttpListOpenRoutes(const FHttpServerRequest& Request, const FHttpResultCallback& OnComplete);
 
-};
\ No newline at end of file
+};
diff --git a/Engine/Source/Runtime/FieldNotification/Public/FieldNotificationDelegate.h b/Engine/Source/Runtime/FieldNotification/Public/FieldNotificationDelegate.h
--- a/Engine/Source/Runtime/FieldNotification/Public/FieldNotificationDelegate.h
+++ b/Engine/Source/Runtime/FieldNotification/Public/FieldNotificationDelegate.h
@@ -22,7 +22,7 @@ class UObject;
 namespace UE::FieldNotification
 {
 	
-class FIELDNOTIFICATION_API FFieldMulticastDelegate
+class FFieldMulticastDelegate
 {
 public:
 	using FDelegate = INotifyFieldValueChanged::FFieldValueChangedDelegate;
@@ -70,8 +70,8 @@ private:
 	using InvocationListType = TArray<FInvocationElement>;
 
 public:
-	FDelegateHandle Add(const UObject* InObject, FFieldId InFieldId, FDelegate InNewDelegate);
-	FDelegateHandle Add(const UObject* InObject, FFieldId InFieldId, const FDynamicDelegate& InDynamicDelegate);
+	FIELDNOTIFICATION_API FDelegateHandle Add(const UObject* InObject, FFieldId InFieldId, FDelegate InNewDelegate);
+	FIELDNOTIFICATION_API FDelegateHandle Add(const UObject* InObject, FFieldId InFieldId, const FDynamicDelegate& InDynamicDelegate);
 
 	struct FRemoveResult
 	{
@@ -80,28 +80,28 @@ public:
 		const UObject* Object = nullptr;
 		FFieldId FieldId;
 	};
-	FRemoveResult Remove(FDelegateHandle InDelegate);
-	FRemoveResult Remove(const FDynamicDelegate& InDynamicDelegate);
+	FIELDNOTIFICATION_API FRemoveResult Remove(FDelegateHandle InDelegate);
+	FIELDNOTIFICATION_API FRemoveResult Remove(const FDynamicDelegate& InDynamicDelegate);
 
 	struct FRemoveFromResult
 	{
 		bool bRemoved = false;
 		bool bHasOtherBoundDelegates = false;
 	};
-	FRemoveFromResult RemoveFrom(const UObject* InObject, FFieldId InFieldId, FDelegateHandle InDelegate);
-	FRemoveFromResult RemoveFrom(const UObject* InObject, FFieldId InFieldId, const FDynamicDelegate& InDynamicDelegate);
+	FIELDNOTIFICATION_API FRemoveFromResult RemoveFrom(const UObject* InObject, FFieldId InFieldId, FDelegateHandle InDelegate);
+	FIELDNOTIFICATION_API FRemoveFromResult RemoveFrom(const UObject* InObject, FFieldId InFieldId, const FDynamicDelegate& InDynamicDelegate);
 
 	struct FRemoveAllResult
 	{
 		int32 RemoveCount = 0;
 		TBitArray<> HasFields;
 	};
-	FRemoveAllResult RemoveAll(const UObject* InObject, const void* InUserObject);
-	FRemoveAllResult RemoveAll(const UObject* InObject, FFieldId InFieldId, const void* InUserObject);
+	FIELDNOTIFICATION_API FRemoveAllResult RemoveAll(const UObject* InObject, const void* InUserObject);
+	FIELDNOTIFICATION_API FRemoveAllResult RemoveAll(const UObject* InObject, FFieldId InFieldId, const void* InUserObject);
 
-	void Broadcast(UObject* InObject, FFieldId InFieldId);
+	FIELDNOTIFICATION_API void Broadcast(UObject* InObject, FFieldId InFieldId);
 
-	void Reset();
+	FIELDNOTIFICATION_API void Reset();
 
 	struct FDelegateView
 	{
@@ -113,7 +113,7 @@ public:
 		const UObject* BindingObject;
 		FName BindingFunctionName;
 	};
-	TArray<FDelegateView> GetView() const;
+	FIELDNOTIFICATION_API TArray<FDelegateView> GetView() const;
 
 private:
 	int32 LowerBound(FFieldId InFieldId) const
@@ -153,4 +153,4 @@ private:
 	uint16 AddedEmplaceAt = std::numeric_limits<uint16>::max();
 };
 
-} //namespace
\ No newline at end of file
+} //namespace
diff --git a/Engine/Source/Runtime/Foliage/Public/FoliageInstancedStaticMeshComponent.h b/Engine/Source/Runtime/Foliage/Public/FoliageInstancedStaticMeshComponent.h
--- a/Engine/Source/Runtime/Foliage/Public/FoliageInstancedStaticMeshComponent.h
+++ b/Engine/Source/Runtime/Foliage/Public/FoliageInstancedStaticMeshComponent.h
@@ -16,8 +16,8 @@ class AController;
 DECLARE_DYNAMIC_MULTICAST_DELEGATE_SevenParams(FInstancePointDamageSignature, int32, InstanceIndex, float, Damage, class AController*, InstigatedBy, FVector, HitLocation, FVector, ShotFromDirection, const class UDamageType*, DamageType, AActor*, DamageCauser);
 DECLARE_DYNAMIC_MULTICAST_DELEGATE_SevenParams(FInstanceRadialDamageSignature, const TArray<int32>&, Instances, const TArray<float>&, Damages, class AController*, InstigatedBy, FVector, Origin, float, MaxRadius, const class UDamageType*, DamageType, AActor*, DamageCauser);
 
-UCLASS(ClassGroup = Foliage, Blueprintable)
-class FOLIAGE_API UFoliageInstancedStaticMeshComponent : public UHierarchicalInstancedStaticMeshComponent
+UCLASS(ClassGroup = Foliage, Blueprintable, MinimalAPI)
+class UFoliageInstancedStaticMeshComponent : public UHierarchicalInstancedStaticMeshComponent
 {
 	GENERATED_UCLASS_BODY()
 public:
@@ -36,11 +36,11 @@ public:
 	uint64 FoliageHiddenEditorViews;
 #endif// WITH_EDITOR
 
-	virtual void ReceiveComponentDamage(float DamageAmount, FDamageEvent const& DamageEvent, AController* EventInstigator, AActor* DamageCauser) override;
+	FOLIAGE_API virtual void ReceiveComponentDamage(float DamageAmount, FDamageEvent const& DamageEvent, AController* EventInstigator, AActor* DamageCauser) override;
 
 	//USceneComponent Interface
 #if WITH_EDITOR
-	virtual uint64 GetHiddenEditorViews() const override;
+	FOLIAGE_API virtual uint64 GetHiddenEditorViews() const override;
 #endif// WITH_EDITOR
 
 	/** Used by procedural generation to link generated component with its creator */
diff --git a/Engine/Source/Runtime/Foliage/Public/FoliageStatistics.h b/Engine/Source/Runtime/Foliage/Public/FoliageStatistics.h
--- a/Engine/Source/Runtime/Foliage/Public/FoliageStatistics.h
+++ b/Engine/Source/Runtime/Foliage/Public/FoliageStatistics.h
@@ -9,8 +9,8 @@
 
 class UStaticMesh;
 
-UCLASS()
-class FOLIAGE_API UFoliageStatistics : public UBlueprintFunctionLibrary
+UCLASS(MinimalAPI)
+class UFoliageStatistics : public UBlueprintFunctionLibrary
 {
 	GENERATED_UCLASS_BODY()
 
@@ -24,7 +24,7 @@ class FOLIAGE_API UFoliageStatistics : public UBlueprintFunctionLibrary
 	* return number of foliage instances with their mesh set to Mesh that overlap the sphere
 	*/
 	UFUNCTION(BlueprintCallable, Category = "Foliage", meta = (WorldContext = "WorldContextObject", UnsafeDuringActorConstruction = "true"))
-	static int32 FoliageOverlappingSphereCount(UObject* WorldContextObject, const UStaticMesh* StaticMesh, FVector CenterPosition, float Radius);
+	static FOLIAGE_API int32 FoliageOverlappingSphereCount(UObject* WorldContextObject, const UStaticMesh* StaticMesh, FVector CenterPosition, float Radius);
 
 	/** 
 	 *	Gets the number of instances overlapping a provided box
@@ -32,7 +32,7 @@ class FOLIAGE_API UFoliageStatistics : public UBlueprintFunctionLibrary
 	 *	@param Box Box to overlap
 	 */
 	UFUNCTION(BlueprintCallable, Category = "Foliage", meta = (WorldContext = "WorldContextObject", UnsafeDuringActorConstruction = "true"))
-	static int32 FoliageOverlappingBoxCount(UObject* WorldContextObject, const UStaticMesh* StaticMesh, FBox Box);
+	static FOLIAGE_API int32 FoliageOverlappingBoxCount(UObject* WorldContextObject, const UStaticMesh* StaticMesh, FBox Box);
 
 	/** 
 	 *  Get the transform of every instance overlapping the provided FBox
@@ -41,7 +41,7 @@ class FOLIAGE_API UFoliageStatistics : public UBlueprintFunctionLibrary
 	 *	@param OutTransforms Array to populate with transforms
 	 */
 	UFUNCTION(BlueprintCallable, Category = "Foliage", meta = (WorldContext = "WorldContextObject", UnsafeDuringActorConstruction = "true"))
-	static void FoliageOverlappingBoxTransforms(UObject* WorldContextObject, const UStaticMesh* StaticMesh, FBox Box, TArray<FTransform>& OutTransforms);
+	static FOLIAGE_API void FoliageOverlappingBoxTransforms(UObject* WorldContextObject, const UStaticMesh* StaticMesh, FBox Box, TArray<FTransform>& OutTransforms);
 
 	/** 
 	 *  DEBUG FUNCTION: This is not fast, use only for debug/development.
@@ -50,6 +50,6 @@ class FOLIAGE_API UFoliageStatistics : public UBlueprintFunctionLibrary
 	 *  @param	Radius			The radius of the sphere.
 	 *	@param OutMeshCounts	Map of Meshes to instance counts
 	 */
-	static void FoliageOverlappingMeshCounts_Debug(UObject* WorldContextObject, FVector CenterPosition, float Radius, TMap<UStaticMesh*, int32>& OutMeshCounts);
+	static FOLIAGE_API void FoliageOverlappingMeshCounts_Debug(UObject* WorldContextObject, FVector CenterPosition, float Radius, TMap<UStaticMesh*, int32>& OutMeshCounts);
 };
 
diff --git a/Engine/Source/Runtime/Foliage/Public/FoliageType.h b/Engine/Source/Runtime/Foliage/Public/FoliageType.h
--- a/Engine/Source/Runtime/Foliage/Public/FoliageType.h
+++ b/Engine/Source/Runtime/Foliage/Public/FoliageType.h
@@ -407,11 +407,11 @@ public:
 	uint32 IsSelected:1;
 #endif
 public:
-	FOLIAGE_API float GetRadius(bool bInSingleInstanceMode) const { return bInSingleInstanceMode && bSingleInstanceModeOverrideRadius ? SingleInstanceModeRadius : Radius; }
+	float GetRadius(bool bInSingleInstanceMode) const { return bInSingleInstanceMode && bSingleInstanceModeOverrideRadius ? SingleInstanceModeRadius : Radius; }
 
 	// PROCEDURAL
 
-	FOLIAGE_API float GetSeedDensitySquared() const { return InitialSeedDensity * InitialSeedDensity; }
+	float GetSeedDensitySquared() const { return InitialSeedDensity * InitialSeedDensity; }
 	FOLIAGE_API float GetMaxRadius() const;
 	FOLIAGE_API float GetScaleForAge(const float Age) const;
 	FOLIAGE_API float GetInitAge(FRandomStream& RandomStream) const;
diff --git a/Engine/Source/Runtime/Foliage/Public/FoliageTypeObject.h b/Engine/Source/Runtime/Foliage/Public/FoliageTypeObject.h
--- a/Engine/Source/Runtime/Foliage/Public/FoliageTypeObject.h
+++ b/Engine/Source/Runtime/Foliage/Public/FoliageTypeObject.h
@@ -11,7 +11,7 @@
 
 /** A wrapper struct used to allow the use of either FoliageType assets or FoliageType blueprint classes */
 USTRUCT()
-struct FOLIAGE_API FFoliageTypeObject
+struct FFoliageTypeObject
 {
 	GENERATED_USTRUCT_BODY()
 
@@ -19,25 +19,25 @@ struct FOLIAGE_API FFoliageTypeObject
 		: FoliageTypeObject(nullptr), TypeInstance(nullptr), bIsAsset(false) {}
 
 	/** Refreshes the type instance based on the assigned object. Intended to be called after some change is made. */
-	void RefreshInstance();
+	FOLIAGE_API void RefreshInstance();
 
 	/** Gets the instance of this foliage type. Creates the instance if needed. */
-	const UFoliageType* GetInstance();
+	FOLIAGE_API const UFoliageType* GetInstance();
 
 	/** Gets the instance of this foliage type. */
-	const UFoliageType* GetInstance() const;
+	FOLIAGE_API const UFoliageType* GetInstance() const;
 
 	/** @return Whether this would return a valid instance */
-	bool ContainsValidInstance() const;
+	FOLIAGE_API bool ContainsValidInstance() const;
 
 	/** @return Whether any foliage type is assigned at all */
-	bool HasFoliageType() const;
+	FOLIAGE_API bool HasFoliageType() const;
 
-	bool IsDirty() const;
-	void SetClean();
+	FOLIAGE_API bool IsDirty() const;
+	FOLIAGE_API void SetClean();
 
 #if WITH_EDITORONLY_DATA
-	void PostSerialize(const FArchive& Ar);
+	FOLIAGE_API void PostSerialize(const FArchive& Ar);
 #endif
 
 private:
@@ -66,4 +66,4 @@ struct TStructOpsTypeTraits<FFoliageTypeObject> : public TStructOpsTypeTraitsBas
 		WithPostSerialize = true,
 	};
 };
-#endif
\ No newline at end of file
+#endif
diff --git a/Engine/Source/Runtime/Foliage/Public/GrassInstancedStaticMeshComponent.h b/Engine/Source/Runtime/Foliage/Public/GrassInstancedStaticMeshComponent.h
--- a/Engine/Source/Runtime/Foliage/Public/GrassInstancedStaticMeshComponent.h
+++ b/Engine/Source/Runtime/Foliage/Public/GrassInstancedStaticMeshComponent.h
@@ -6,16 +6,16 @@
 #include "Components/HierarchicalInstancedStaticMeshComponent.h"
 #include "GrassInstancedStaticMeshComponent.generated.h"
 
-UCLASS(ClassGroup = Foliage)
-class FOLIAGE_API UGrassInstancedStaticMeshComponent : public UHierarchicalInstancedStaticMeshComponent
+UCLASS(ClassGroup = Foliage, MinimalAPI)
+class UGrassInstancedStaticMeshComponent : public UHierarchicalInstancedStaticMeshComponent
 {
 	GENERATED_UCLASS_BODY()
 
 public:
-	static void BuildTreeAnyThread(TArray<FMatrix>& InstanceTransforms, TArray<float>& InstanceCustomDataFloats, int32 NumCustomDataFloats, const FBox& MeshBox, TArray<FClusterNode>& OutClusterTree, TArray<int32>& OutSortedInstances, TArray<int32>& OutInstanceReorderTable, int32& OutOcclusionLayerNum, int32 MaxInstancesPerLeaf, bool InGenerateInstanceScalingRange);
+	static FOLIAGE_API void BuildTreeAnyThread(TArray<FMatrix>& InstanceTransforms, TArray<float>& InstanceCustomDataFloats, int32 NumCustomDataFloats, const FBox& MeshBox, TArray<FClusterNode>& OutClusterTree, TArray<int32>& OutSortedInstances, TArray<int32>& OutInstanceReorderTable, int32& OutOcclusionLayerNum, int32 MaxInstancesPerLeaf, bool InGenerateInstanceScalingRange);
 	
-	void AcceptPrebuiltTree(TArray<FInstancedStaticMeshInstanceData>& InInstanceData, TArray<FClusterNode>& InClusterTree, int32 InOcclusionLayerNumNodes, int32 InNumBuiltRenderInstances);
+	FOLIAGE_API void AcceptPrebuiltTree(TArray<FInstancedStaticMeshInstanceData>& InInstanceData, TArray<FClusterNode>& InClusterTree, int32 InOcclusionLayerNumNodes, int32 InNumBuiltRenderInstances);
 
-	bool SupportsWorldPositionOffsetVelocity() const override;
+	FOLIAGE_API bool SupportsWorldPositionOffsetVelocity() const override;
 };
 
diff --git a/Engine/Source/Runtime/Foliage/Public/InstancedFoliage.h b/Engine/Source/Runtime/Foliage/Public/InstancedFoliage.h
--- a/Engine/Source/Runtime/Foliage/Public/InstancedFoliage.h
+++ b/Engine/Source/Runtime/Foliage/Public/InstancedFoliage.h
@@ -366,7 +366,7 @@ struct FFoliageInfo
 	FOLIAGE_API void AddToBaseHash(int32 InstanceIdx);
 	FOLIAGE_API void RemoveFromBaseHash(int32 InstanceIdx);
 	FOLIAGE_API void RecomputeHash();
-	FOLIAGE_API bool ShouldAttachToBaseComponent() const { return Implementation->ShouldAttachToBaseComponent(); }
+	bool ShouldAttachToBaseComponent() const { return Implementation->ShouldAttachToBaseComponent(); }
 
 	// For debugging. Validate state after editing.
 	void CheckValid();
@@ -497,7 +497,7 @@ struct FDesiredFoliageInstance
 
 #if WITH_EDITOR
 // Struct to hold potential instances we've sampled
-struct FOLIAGE_API FPotentialInstance
+struct FPotentialInstance
 {
 	FVector HitLocation;
 	FVector HitNormal;
@@ -505,7 +505,7 @@ struct FOLIAGE_API FPotentialInstance
 	float HitWeight;
 	FDesiredFoliageInstance DesiredInstance;
 
-	FPotentialInstance(FVector InHitLocation, FVector InHitNormal, UPrimitiveComponent* InHitComponent, float InHitWeight, const FDesiredFoliageInstance& InDesiredInstance = FDesiredFoliageInstance());
-	bool PlaceInstance(const UWorld* InWorld, const UFoliageType* Settings, FFoliageInstance& Inst, bool bSkipCollision = false);
+	FOLIAGE_API FPotentialInstance(FVector InHitLocation, FVector InHitNormal, UPrimitiveComponent* InHitComponent, float InHitWeight, const FDesiredFoliageInstance& InDesiredInstance = FDesiredFoliageInstance());
+	FOLIAGE_API bool PlaceInstance(const UWorld* InWorld, const UFoliageType* Settings, FFoliageInstance& Inst, bool bSkipCollision = false);
 };
 #endif
diff --git a/Engine/Source/Runtime/Foliage/Public/InstancedFoliageActor.h b/Engine/Source/Runtime/Foliage/Public/InstancedFoliageActor.h
--- a/Engine/Source/Runtime/Foliage/Public/InstancedFoliageActor.h
+++ b/Engine/Source/Runtime/Foliage/Public/InstancedFoliageActor.h
@@ -44,7 +44,7 @@ private:
 
 public:
 	FOLIAGE_API bool ForEachFoliageInfo(TFunctionRef<bool(UFoliageType* FoliageType, FFoliageInfo& FoliageInfo)> InOperation);
-	FOLIAGE_API const TMap<UFoliageType*, TUniqueObj<FFoliageInfo>>& GetFoliageInfos() const { return ObjectPtrDecay(FoliageInfos); }
+	const TMap<UFoliageType*, TUniqueObj<FFoliageInfo>>& GetFoliageInfos() const { return ObjectPtrDecay(FoliageInfos); }
 	FOLIAGE_API TUniqueObj<FFoliageInfo>& AddFoliageInfo(UFoliageType* FoliageType);
 	FOLIAGE_API TUniqueObj<FFoliageInfo>& AddFoliageInfo(UFoliageType* FoliageType, TUniqueObj<FFoliageInfo>&& FoliageInfo);
 	FOLIAGE_API bool RemoveFoliageInfoAndCopyValue(UFoliageType* FoliageType, TUniqueObj<FFoliageInfo>& OutFoliageInfo);
diff --git a/Engine/Source/Runtime/Foliage/Public/ProceduralFoliageComponent.h b/Engine/Source/Runtime/Foliage/Public/ProceduralFoliageComponent.h
--- a/Engine/Source/Runtime/Foliage/Public/ProceduralFoliageComponent.h
+++ b/Engine/Source/Runtime/Foliage/Public/ProceduralFoliageComponent.h
@@ -38,8 +38,8 @@ struct FTileLayout
 	FVector::FReal HalfHeight;
 };
 
-UCLASS(BlueprintType)
-class FOLIAGE_API UProceduralFoliageComponent : public UActorComponent
+UCLASS(BlueprintType, MinimalAPI)
+class UProceduralFoliageComponent : public UActorComponent
 {
 	GENERATED_UCLASS_BODY()
 
@@ -97,54 +97,54 @@ class FOLIAGE_API UProceduralFoliageComponent : public UActorComponent
 	};
 				
 	// UObject interface
-	virtual void PostEditImport() override;
+	FOLIAGE_API virtual void PostEditImport() override;
 
 	/**
 	 * Returns a params struct based on this component properties. 
 	 */
-	FGenerateProceduralContentParams GetGenerateProceduralContentParams() const;
+	FOLIAGE_API FGenerateProceduralContentParams GetGenerateProceduralContentParams() const;
 
 	/**
 	 * Runs the procedural foliage simulation, removes the old result, creates instances with the new result
 	 * @return True if the simulation succeeded
 	 */ 
-	bool ResimulateProceduralFoliage(TFunctionRef<void(const TArray<FDesiredFoliageInstance>&)> AddInstancesFunc);
+	FOLIAGE_API bool ResimulateProceduralFoliage(TFunctionRef<void(const TArray<FDesiredFoliageInstance>&)> AddInstancesFunc);
 
 	/** 
 	 * Runs the procedural foliage simulation to generate a list of desired instances to spawn.
 	 * @return True if the simulation succeeded
 	 */
-	bool GenerateProceduralContent(TArray<FDesiredFoliageInstance>& OutInstances);
-	static bool GenerateProceduralContent(const FGenerateProceduralContentParams& InParams, TArray<FDesiredFoliageInstance>& OutInstances);
+	FOLIAGE_API bool GenerateProceduralContent(TArray<FDesiredFoliageInstance>& OutInstances);
+	static FOLIAGE_API bool GenerateProceduralContent(const FGenerateProceduralContentParams& InParams, TArray<FDesiredFoliageInstance>& OutInstances);
 			
 	/** Removes all spawned foliage instances in the level that were spawned by this component */
-	void RemoveProceduralContent(bool bInRebuildTree = true);
-	static void RemoveProceduralContent(UWorld* InWorld, const FGuid& InProceduralGuid, bool bInRebuildTree, TSet<AInstancedFoliageActor*>& OutModifiedActors);
+	FOLIAGE_API void RemoveProceduralContent(bool bInRebuildTree = true);
+	static FOLIAGE_API void RemoveProceduralContent(UWorld* InWorld, const FGuid& InProceduralGuid, bool bInRebuildTree, TSet<AInstancedFoliageActor*>& OutModifiedActors);
 
 	/** @return True if any foliage instances in the level were spawned by this component */
-	bool HasSpawnedAnyInstances();
+	FOLIAGE_API bool HasSpawnedAnyInstances();
 	
 	/** @return The position in world space of the bottom-left corner of the bottom-left-most active tile */
-	FVector GetWorldPosition() const;
-	static FVector GetWorldPosition(const FGenerateProceduralContentParams& Param);
+	FOLIAGE_API FVector GetWorldPosition() const;
+	static FOLIAGE_API FVector GetWorldPosition(const FGenerateProceduralContentParams& Param);
 
 	/** @return The bounds of area encompassed by the simulation */
-	virtual FBox GetBounds() const;
+	FOLIAGE_API virtual FBox GetBounds() const;
 
 	/** @return The body instanced used for bounds checking */
-	FBodyInstance* GetBoundsBodyInstance() const;
+	FOLIAGE_API FBodyInstance* GetBoundsBodyInstance() const;
 
 	/** Determines the basic layout of the tiles used in the simulation */
-	void GetTileLayout(FTileLayout& OutTileLayout) const;
-	static void GetTileLayout(const FGenerateProceduralContentParams& Params, FTileLayout& OutTileLayout);
+	FOLIAGE_API void GetTileLayout(FTileLayout& OutTileLayout) const;
+	static FOLIAGE_API void GetTileLayout(const FGenerateProceduralContentParams& Params, FTileLayout& OutTileLayout);
 
 	void SetSpawningVolume(AVolume* InSpawningVolume) { SpawningVolume = InSpawningVolume; }
 	const FGuid& GetProceduralGuid() const { return ProceduralGuid; }
 
 #if WITH_EDITOR
-	void LoadSimulatedRegion();
-	void UnloadSimulatedRegion();
-	bool IsSimulatedRegionLoaded();
+	FOLIAGE_API void LoadSimulatedRegion();
+	FOLIAGE_API void UnloadSimulatedRegion();
+	FOLIAGE_API bool IsSimulatedRegionLoaded();
 #endif
 
 private:
diff --git a/Engine/Source/Runtime/Foliage/Public/ProceduralFoliageInstance.h b/Engine/Source/Runtime/Foliage/Public/ProceduralFoliageInstance.h
--- a/Engine/Source/Runtime/Foliage/Public/ProceduralFoliageInstance.h
+++ b/Engine/Source/Runtime/Foliage/Public/ProceduralFoliageInstance.h
@@ -55,21 +55,21 @@ struct FProceduralFoliageOverlap
 };
 
 USTRUCT(BlueprintType)
-struct FOLIAGE_API FProceduralFoliageInstance
+struct FProceduralFoliageInstance
 {
 public:
 	GENERATED_USTRUCT_BODY()
-	FProceduralFoliageInstance();
+	FOLIAGE_API FProceduralFoliageInstance();
 
-	static FProceduralFoliageInstance* Domination(FProceduralFoliageInstance* A, FProceduralFoliageInstance* B, ESimulationOverlap::Type OverlapType);
+	static FOLIAGE_API FProceduralFoliageInstance* Domination(FProceduralFoliageInstance* A, FProceduralFoliageInstance* B, ESimulationOverlap::Type OverlapType);
 
-	float GetMaxRadius() const;
-	float GetShadeRadius() const;
-	float GetCollisionRadius() const;
+	FOLIAGE_API float GetMaxRadius() const;
+	FOLIAGE_API float GetShadeRadius() const;
+	FOLIAGE_API float GetCollisionRadius() const;
 
 	bool IsAlive() const { return bAlive; }
 
-	void TerminateInstance();
+	FOLIAGE_API void TerminateInstance();
 
 public:
 	UPROPERTY()
diff --git a/Engine/Source/Runtime/Foliage/Public/ProceduralFoliageSpawner.h b/Engine/Source/Runtime/Foliage/Public/ProceduralFoliageSpawner.h
--- a/Engine/Source/Runtime/Foliage/Public/ProceduralFoliageSpawner.h
+++ b/Engine/Source/Runtime/Foliage/Public/ProceduralFoliageSpawner.h
@@ -13,8 +13,8 @@
 
 class UProceduralFoliageTile;
 
-UCLASS(BlueprintType, Blueprintable)
-class FOLIAGE_API UProceduralFoliageSpawner : public UObject
+UCLASS(BlueprintType, Blueprintable, MinimalAPI)
+class UProceduralFoliageSpawner : public UObject
 {
 	GENERATED_UCLASS_BODY()
 
@@ -54,10 +54,10 @@ private:
 
 public:
 	UFUNCTION(BlueprintCallable, Category = ProceduralFoliageSimulation)
-	void Simulate(int32 NumSteps = -1);
-	void Empty();
+	FOLIAGE_API void Simulate(int32 NumSteps = -1);
+	FOLIAGE_API void Empty();
 
-	int32 GetRandomNumber();
+	FOLIAGE_API int32 GetRandomNumber();
 
 	const TArray<FFoliageTypeObject>& GetFoliageTypes() const { return FoliageTypes; }
 
@@ -66,20 +66,20 @@ public:
 	const TArray<TSoftObjectPtr<UMaterialInterface>>& GetFoliageTerrainMaterials() const { return OverrideFoliageTerrainMaterials; }
 
 	/** Returns the instances that need to spawn for a given min,max region */
-	void GetInstancesToSpawn(TArray<FProceduralFoliageInstance>& OutInstances, const FVector& Min = FVector(-FLT_MAX, -FLT_MAX, -FLT_MAX), const FVector& Max = FVector(FLT_MAX, FLT_MAX, FLT_MAX) ) const;
+	FOLIAGE_API void GetInstancesToSpawn(TArray<FProceduralFoliageInstance>& OutInstances, const FVector& Min = FVector(-FLT_MAX, -FLT_MAX, -FLT_MAX), const FVector& Max = FVector(FLT_MAX, FLT_MAX, FLT_MAX) ) const;
 
-	virtual void Serialize(FArchive& Ar);
+	FOLIAGE_API virtual void Serialize(FArchive& Ar);
 
 	/** Takes a tile index and returns a random tile associated with that index. */
-	const UProceduralFoliageTile* GetRandomTile(int32 X, int32 Y);
+	FOLIAGE_API const UProceduralFoliageTile* GetRandomTile(int32 X, int32 Y);
 
 	/** Creates a temporary empty tile with the appropriate settings created for it. */
-	UProceduralFoliageTile* CreateTempTile();
+	FOLIAGE_API UProceduralFoliageTile* CreateTempTile();
 
 private:
-	void CreateProceduralFoliageInstances();
+	FOLIAGE_API void CreateProceduralFoliageInstances();
 
-	void SetClean();
+	FOLIAGE_API void SetClean();
 private:
 	TArray<TWeakObjectPtr<UProceduralFoliageTile>> PrecomputedTiles;
 
diff --git a/Engine/Source/Runtime/Foliage/Public/ProceduralFoliageTile.h b/Engine/Source/Runtime/Foliage/Public/ProceduralFoliageTile.h
--- a/Engine/Source/Runtime/Foliage/Public/ProceduralFoliageTile.h
+++ b/Engine/Source/Runtime/Foliage/Public/ProceduralFoliageTile.h
@@ -26,8 +26,8 @@ struct FBodyInstance;
  *	Note also that, barring any core changes to the ordering of TSet, foliage generation is deterministic 
  *	(i.e. given the same inputs, the result of the simulation will always be the same).
  */
-UCLASS()
-class FOLIAGE_API UProceduralFoliageTile : public UObject
+UCLASS(MinimalAPI)
+class UProceduralFoliageTile : public UObject
 {
 	GENERATED_UCLASS_BODY()
 
@@ -37,12 +37,12 @@ class FOLIAGE_API UProceduralFoliageTile : public UObject
 	 *	@param RandomSeed For seeding (in the RNG sense) the random dispersion of seeds (in the foliage sense)
 	 *	@param MaxNumSteps The maximum number of steps to run the simulation
 	 */
-	void Simulate(const UProceduralFoliageSpawner* InFoliageSpawner, const int32 RandomSeed, const int32 MaxNumSteps, const int32 InLastCancel);
+	FOLIAGE_API void Simulate(const UProceduralFoliageSpawner* InFoliageSpawner, const int32 RandomSeed, const int32 MaxNumSteps, const int32 InLastCancel);
 	
 	/**
 	 *	Extracts information on each instance that should be spawned based on the simulation.
 	 */
-	void ExtractDesiredInstances(TArray<FDesiredFoliageInstance>& OutDesiredInstances, const FTransform& WorldTM, const FVector2D& ActorVolumeLocation, FVector::FReal ActorVolumeMaxExtent, const FGuid& ProceduralGuid, const FVector::FReal HalfHeight, const FBodyInstance* VolumeBodyInstance, bool bEmptyTileInfo = true);
+	FOLIAGE_API void ExtractDesiredInstances(TArray<FDesiredFoliageInstance>& OutDesiredInstances, const FTransform& WorldTM, const FVector2D& ActorVolumeLocation, FVector::FReal ActorVolumeMaxExtent, const FGuid& ProceduralGuid, const FVector::FReal HalfHeight, const FBodyInstance* VolumeBodyInstance, bool bEmptyTileInfo = true);
 	
 	/**
 	 *	Copies instances within this tile to another
@@ -51,26 +51,26 @@ class FOLIAGE_API UProceduralFoliageTile : public UObject
 	 *	@param RelativeTM The transform of the destination tile relative to this tile
 	 *	@param Overlap The amount by which the tiles are overlapping
 	 */
-	void CopyInstancesToTile(UProceduralFoliageTile* ToTile, const FBox2D& LocalAABB, const FTransform& RelativeTM, const float Overlap) const;
+	FOLIAGE_API void CopyInstancesToTile(UProceduralFoliageTile* ToTile, const FBox2D& LocalAABB, const FTransform& RelativeTM, const float Overlap) const;
 
 	/**
 	 *	Removes a single instance from the tile
 	 *	@param Inst The instance to remove
 	 */
-	void RemoveInstance(FProceduralFoliageInstance* Inst);
+	FOLIAGE_API void RemoveInstance(FProceduralFoliageInstance* Inst);
 
 	/** Removes all instances from the tile */
-	void RemoveInstances();
+	FOLIAGE_API void RemoveInstances();
 
 	/** Empty arrays and sets */
-	void Empty();
+	FOLIAGE_API void Empty();
 
 	/** Initializes the tile for simulation. Called automatically by Simulate, so only bother to call if the tile won't actually simulate. */
-	void InitSimulation(const UProceduralFoliageSpawner* InFoliageSpawner, const int32 InRandomSeed);
+	FOLIAGE_API void InitSimulation(const UProceduralFoliageSpawner* InFoliageSpawner, const int32 InRandomSeed);
 
 	// UObject interface
-	virtual void GetResourceSizeEx(FResourceSizeEx& CumulativeResourceSize) override;
-	virtual void BeginDestroy() override;
+	FOLIAGE_API virtual void GetResourceSizeEx(FResourceSizeEx& CumulativeResourceSize) override;
+	FOLIAGE_API virtual void BeginDestroy() override;
 
 private:
 	
diff --git a/Engine/Source/Runtime/Foliage/Public/ProceduralFoliageVolume.h b/Engine/Source/Runtime/Foliage/Public/ProceduralFoliageVolume.h
--- a/Engine/Source/Runtime/Foliage/Public/ProceduralFoliageVolume.h
+++ b/Engine/Source/Runtime/Foliage/Public/ProceduralFoliageVolume.h
@@ -11,8 +11,8 @@
 class UProceduralFoliageComponent;
 class FLoaderAdapterActor;
 
-UCLASS()
-class FOLIAGE_API AProceduralFoliageVolume: public AVolume, public IWorldPartitionActorLoaderInterface
+UCLASS(MinimalAPI)
+class AProceduralFoliageVolume: public AVolume, public IWorldPartitionActorLoaderInterface
 {
 	GENERATED_UCLASS_BODY()
 
@@ -21,18 +21,18 @@ class FOLIAGE_API AProceduralFoliageVolume: public AVolume, public IWorldPartiti
 
 #if WITH_EDITOR
 	//~ Begin AActor Interface
-	virtual void PostRegisterAllComponents();
+	FOLIAGE_API virtual void PostRegisterAllComponents();
 	//~ End AActor Interface
 
 	//~ Begin UObject Interface
-	virtual void BeginDestroy() override;
-	virtual void PostEditImport() override;
+	FOLIAGE_API virtual void BeginDestroy() override;
+	FOLIAGE_API virtual void PostEditImport() override;
 	//~ End UObject Interface
 
-	virtual bool GetReferencedContentObjects(TArray<UObject*>& Objects) const override;
+	FOLIAGE_API virtual bool GetReferencedContentObjects(TArray<UObject*>& Objects) const override;
 
 	//~ Begin IWorldPartitionActorLoaderInterface interface
-	virtual ILoaderAdapter* GetLoaderAdapter() override;
+	FOLIAGE_API virtual ILoaderAdapter* GetLoaderAdapter() override;
 	//~ End IWorldPartitionActorLoaderInterface interface
 
 private:
diff --git a/Engine/Source/Runtime/GameplayDebugger/Public/GameplayDebuggerAddonBase.h b/Engine/Source/Runtime/GameplayDebugger/Public/GameplayDebuggerAddonBase.h
--- a/Engine/Source/Runtime/GameplayDebugger/Public/GameplayDebuggerAddonBase.h
+++ b/Engine/Source/Runtime/GameplayDebugger/Public/GameplayDebuggerAddonBase.h
@@ -10,7 +10,7 @@ class AActor;
 class APlayerController;
 class AGameplayDebuggerCategoryReplicator;
 
-class GAMEPLAYDEBUGGER_API FGameplayDebuggerAddonBase : public TSharedFromThis<FGameplayDebuggerAddonBase>
+class FGameplayDebuggerAddonBase : public TSharedFromThis<FGameplayDebuggerAddonBase>
 {
 public:
 
@@ -18,27 +18,27 @@ public:
 
 	int32 GetNumInputHandlers() const { return InputHandlers.Num(); }
 	FGameplayDebuggerInputHandler& GetInputHandler(int32 HandlerId) { return InputHandlers[HandlerId]; }
-	FString GetInputHandlerDescription(int32 HandlerId) const;
-	UWorld* GetWorldFromReplicator() const;
+	GAMEPLAYDEBUGGER_API FString GetInputHandlerDescription(int32 HandlerId) const;
+	GAMEPLAYDEBUGGER_API UWorld* GetWorldFromReplicator() const;
 	/** Returns the first non-null world associated with, in order: OwnerPC, DebugActor, AGameplayDebuggerCategoryReplicator owner */
-	UWorld* GetDataWorld(const APlayerController* OwnerPC, const AActor* DebugActor) const;
+	GAMEPLAYDEBUGGER_API UWorld* GetDataWorld(const APlayerController* OwnerPC, const AActor* DebugActor) const;
 
 	/** [ALL] called when gameplay debugger is activated */
-	virtual void OnGameplayDebuggerActivated();
+	GAMEPLAYDEBUGGER_API virtual void OnGameplayDebuggerActivated();
 
 	/** [ALL] called when gameplay debugger is deactivated */
-	virtual void OnGameplayDebuggerDeactivated();
+	GAMEPLAYDEBUGGER_API virtual void OnGameplayDebuggerDeactivated();
 
 	/** check if simulate in editor mode is active */
-	static bool IsSimulateInEditor();
+	static GAMEPLAYDEBUGGER_API bool IsSimulateInEditor();
 
 protected:
 
 	/** tries to find selected actor in local world */
-	AActor* FindLocalDebugActor() const;
+	GAMEPLAYDEBUGGER_API AActor* FindLocalDebugActor() const;
 
 	/** returns replicator actor */
-	AGameplayDebuggerCategoryReplicator* GetReplicator() const;
+	GAMEPLAYDEBUGGER_API AGameplayDebuggerCategoryReplicator* GetReplicator() const;
 
 	/** creates new key binding handler: single key press */
 	template<class UserClass>
diff --git a/Engine/Source/Runtime/GameplayDebugger/Public/GameplayDebuggerCategory.h b/Engine/Source/Runtime/GameplayDebugger/Public/GameplayDebuggerCategory.h
--- a/Engine/Source/Runtime/GameplayDebugger/Public/GameplayDebuggerCategory.h
+++ b/Engine/Source/Runtime/GameplayDebugger/Public/GameplayDebuggerCategory.h
@@ -45,33 +45,33 @@ struct FDebugDrawDelegateHelper;
 /**
  * Single category of visual debugger tool
  */
-class GAMEPLAYDEBUGGER_API FGameplayDebuggerCategory : public FGameplayDebuggerAddonBase
+class FGameplayDebuggerCategory : public FGameplayDebuggerAddonBase
 {
 public:
 
-	FGameplayDebuggerCategory();
-	virtual ~FGameplayDebuggerCategory();
+	GAMEPLAYDEBUGGER_API FGameplayDebuggerCategory();
+	GAMEPLAYDEBUGGER_API virtual ~FGameplayDebuggerCategory();
 
 	/** [AUTH] gather data for replication */
-	virtual void CollectData(APlayerController* OwnerPC, AActor* DebugActor);
+	GAMEPLAYDEBUGGER_API virtual void CollectData(APlayerController* OwnerPC, AActor* DebugActor);
 
 	/** [LOCAL] draw collected data */
-	virtual void DrawData(APlayerController* OwnerPC, FGameplayDebuggerCanvasContext& CanvasContext);
+	GAMEPLAYDEBUGGER_API virtual void DrawData(APlayerController* OwnerPC, FGameplayDebuggerCanvasContext& CanvasContext);
 
 	/** [LOCAL] creates a scene proxy for more advanced debug rendering */
-	virtual FDebugRenderSceneProxy* CreateDebugSceneProxy(const UPrimitiveComponent* InComponent, FDebugDrawDelegateHelper*& OutDelegateHelper);
+	GAMEPLAYDEBUGGER_API virtual FDebugRenderSceneProxy* CreateDebugSceneProxy(const UPrimitiveComponent* InComponent, FDebugDrawDelegateHelper*& OutDelegateHelper);
 
 	/** [LOCAL] called after successful replication of entire data pack to client */
-	virtual void OnDataPackReplicated(int32 DataPackId);
+	GAMEPLAYDEBUGGER_API virtual void OnDataPackReplicated(int32 DataPackId);
 
 	/** [AUTH] adds line of text tagged with {color} to replicated data */
-	void AddTextLine(const FString& TextLine);
+	GAMEPLAYDEBUGGER_API void AddTextLine(const FString& TextLine);
 
 	/** [AUTH] adds shape to replicated data */
-	void AddShape(const FGameplayDebuggerShape& Shape);
+	GAMEPLAYDEBUGGER_API void AddShape(const FGameplayDebuggerShape& Shape);
 
 	/** [LOCAL] draw category */
-	void DrawCategory(APlayerController* OwnerPC, FGameplayDebuggerCanvasContext& CanvasContext);
+	GAMEPLAYDEBUGGER_API void DrawCategory(APlayerController* OwnerPC, FGameplayDebuggerCanvasContext& CanvasContext);
 
 	/** [LOCAL] check if category should be drawn */
 	bool ShouldDrawCategory(bool bHasDebugActor) const { return IsCategoryEnabled() && (!bShowOnlyWithDebugActor || bHasDebugActor); }
@@ -115,22 +115,22 @@ protected:
 	 * @note Method is expected to always return valid outputs if valid controller is provided. Otherwise it
 	 * depends if the view information has been replicated before the methods gets called.
 	 */
-	bool GetViewPoint(const APlayerController* OwnerPC, FVector& OutViewLocation, FVector& OutViewDirection) const;
+	GAMEPLAYDEBUGGER_API bool GetViewPoint(const APlayerController* OwnerPC, FVector& OutViewLocation, FVector& OutViewDirection) const;
 
 	/**
 	 * Indicates if a given location is within a vision cone built from provided view location and direction based
 	 * on MaxViewDistance and MaxViewAngle from GameplayDebuggerUserSettings
 	 */
-	static bool IsLocationInViewCone(const FVector& ViewLocation, const FVector& ViewDirection, const FVector& TargetLocation);
+	static GAMEPLAYDEBUGGER_API bool IsLocationInViewCone(const FVector& ViewLocation, const FVector& ViewDirection, const FVector& TargetLocation);
 
 	/** [AUTH] marks data pack as needing replication */
-	void MarkDataPackDirty(int32 DataPackId);
+	GAMEPLAYDEBUGGER_API void MarkDataPackDirty(int32 DataPackId);
 
 	/** [LOCAL] requests new scene proxy */
-	void MarkRenderStateDirty();
+	GAMEPLAYDEBUGGER_API void MarkRenderStateDirty();
 
 	/** [LOCAL] preferred view flag for creating scene proxy */
-	FString GetSceneProxyViewFlag() const;
+	GAMEPLAYDEBUGGER_API FString GetSceneProxyViewFlag() const;
 
 	/** [ALL] sets up DataPack replication, needs address of property holding data, DataPack's struct must define Serialize(FArchive& Ar) function
 	 *  returns DataPackId
@@ -152,7 +152,7 @@ protected:
 	void ForceImmediateCollect() { LastCollectDataTime = -MAX_dbl; }
 
 	/** [ALL] Clears out accumulated replicated data. */
-	void ResetReplicatedData();
+	GAMEPLAYDEBUGGER_API void ResetReplicatedData();
 
 	/** update interval, 0 = each tick */
 	float CollectDataInterval;
diff --git a/Engine/Source/Runtime/GameplayDebugger/Public/GameplayDebuggerCategoryReplicator.h b/Engine/Source/Runtime/GameplayDebugger/Public/GameplayDebuggerCategoryReplicator.h
--- a/Engine/Source/Runtime/GameplayDebugger/Public/GameplayDebuggerCategoryReplicator.h
+++ b/Engine/Source/Runtime/GameplayDebugger/Public/GameplayDebuggerCategoryReplicator.h
@@ -116,60 +116,60 @@ struct FGameplayDebuggerVisLogSync
 
 DECLARE_MULTICAST_DELEGATE_TwoParams(FNotifyGameplayDebuggerOwnerChange, AGameplayDebuggerCategoryReplicator*, APlayerController* /** Old Owner */);
 
-UCLASS(NotBlueprintable, NotBlueprintType, notplaceable, noteditinlinenew, hidedropdown, Transient)
-class GAMEPLAYDEBUGGER_API AGameplayDebuggerCategoryReplicator : public AActor
+UCLASS(NotBlueprintable, NotBlueprintType, notplaceable, noteditinlinenew, hidedropdown, Transient, MinimalAPI)
+class AGameplayDebuggerCategoryReplicator : public AActor
 {
 	GENERATED_UCLASS_BODY()
 
-	virtual class UNetConnection* GetNetConnection() const override;
-	virtual bool IsNetRelevantFor(const AActor* RealViewer, const AActor* ViewTarget, const FVector& SrcLocation) const override;
+	GAMEPLAYDEBUGGER_API virtual class UNetConnection* GetNetConnection() const override;
+	GAMEPLAYDEBUGGER_API virtual bool IsNetRelevantFor(const AActor* RealViewer, const AActor* ViewTarget, const FVector& SrcLocation) const override;
 
 protected:
-	virtual void BeginPlay() override;
+	GAMEPLAYDEBUGGER_API virtual void BeginPlay() override;
 
 #if UE_WITH_IRIS
-	virtual void BeginReplication() override;
+	GAMEPLAYDEBUGGER_API virtual void BeginReplication() override;
 
-	virtual const AActor* GetNetOwner() const override;
+	GAMEPLAYDEBUGGER_API virtual const AActor* GetNetOwner() const override;
 #endif
 
 public:
-	virtual void EndPlay(const EEndPlayReason::Type EndPlayReason) override;
-	virtual void TickActor(float DeltaTime, enum ELevelTick TickType, FActorTickFunction& ThisTickFunction) override;
-	virtual void PostNetReceive() override;
+	GAMEPLAYDEBUGGER_API virtual void EndPlay(const EEndPlayReason::Type EndPlayReason) override;
+	GAMEPLAYDEBUGGER_API virtual void TickActor(float DeltaTime, enum ELevelTick TickType, FActorTickFunction& ThisTickFunction) override;
+	GAMEPLAYDEBUGGER_API virtual void PostNetReceive() override;
 
-	virtual void PreReplication(IRepChangedPropertyTracker& ChangedPropertyTracker) override;
+	GAMEPLAYDEBUGGER_API virtual void PreReplication(IRepChangedPropertyTracker& ChangedPropertyTracker) override;
 
 	/** [AUTH] set new owner */
-	void SetReplicatorOwner(APlayerController* InOwnerPC);
-	static FNotifyGameplayDebuggerOwnerChange NotifyDebuggerOwnerChange;
+	GAMEPLAYDEBUGGER_API void SetReplicatorOwner(APlayerController* InOwnerPC);
+	static GAMEPLAYDEBUGGER_API FNotifyGameplayDebuggerOwnerChange NotifyDebuggerOwnerChange;
 
 	/** [ALL] set replicator state */
-	void SetEnabled(bool bEnable);
+	GAMEPLAYDEBUGGER_API void SetEnabled(bool bEnable);
 
 	/** [ALL] set category state */
-	void SetCategoryEnabled(int32 CategoryId, bool bEnable);
+	GAMEPLAYDEBUGGER_API void SetCategoryEnabled(int32 CategoryId, bool bEnable);
 
 	/** [ALL] set actor for debugging */
-	void SetDebugActor(AActor* Actor, bool bSelectInEditor = false);
+	GAMEPLAYDEBUGGER_API void SetDebugActor(AActor* Actor, bool bSelectInEditor = false);
 
 	/** [ALL] set view matrix that should be used for culling and picking instead of the associated controller view point */
-	void SetViewPoint(const FVector& InViewLocation, const FVector& InViewDirection);
+	GAMEPLAYDEBUGGER_API void SetViewPoint(const FVector& InViewLocation, const FVector& InViewDirection);
 
 	/** Indicates that a view matrix has been set for culling and picking instead of the associated controller view point */
 	bool IsViewPointSet() const { return ViewLocation.IsSet() && ViewDirection.IsSet(); }
 
 	/** [ALL] reset view matrix that should be used for culling and picking to revert on the associated controller view point */
-	void ResetViewPoint();
+	GAMEPLAYDEBUGGER_API void ResetViewPoint();
 
 	/** [ALL] send input event to category */
-	void SendCategoryInputEvent(int32 CategoryId, int32 HandlerId);
+	GAMEPLAYDEBUGGER_API void SendCategoryInputEvent(int32 CategoryId, int32 HandlerId);
 
 	/** [ALL] send input event to extension */
-	void SendExtensionInputEvent(int32 ExtensionId, int32 HandlerId);
+	GAMEPLAYDEBUGGER_API void SendExtensionInputEvent(int32 ExtensionId, int32 HandlerId);
 
 	/** [AUTH] starts data collection */
-	void CollectCategoryData(bool bForce = false);
+	GAMEPLAYDEBUGGER_API void CollectCategoryData(bool bForce = false);
 
 	/** get current debug actor */
 	AActor* GetDebugActor() const { return DebugActor.Actor.Get(); }
@@ -181,7 +181,7 @@ public:
 	int16 GetDebugActorCounter() const { return DebugActor.SyncCounter; }
 
 	/** get view point information */
-	bool GetViewPoint(FVector& OutViewLocation, FVector& OutViewDirection) const;
+	GAMEPLAYDEBUGGER_API bool GetViewPoint(FVector& OutViewLocation, FVector& OutViewDirection) const;
 
 	const FGameplayDebuggerVisLogSync& GetVisLogSyncData() const { return VisLogSync; }
 
@@ -195,7 +195,7 @@ public:
 	bool IsEditorWorldReplicator() const { return bIsEditorWorldReplicator; }
 
 	/** get category state */
-	bool IsCategoryEnabled(int32 CategoryId) const;
+	GAMEPLAYDEBUGGER_API bool IsCategoryEnabled(int32 CategoryId) const;
 
 	/** check if debug actor was selected */
 	bool HasDebugActor() const { return DebugActor.ActorName != NAME_None; }
@@ -216,12 +216,12 @@ public:
 	bool IsLocal() const { return bIsLocal; }
 
 #if WITH_EDITOR
-	void InitForEditor();
+	GAMEPLAYDEBUGGER_API void InitForEditor();
 #endif // WITH_EDITOR
 
 protected:
 	UFUNCTION()
-	void OnRep_ReplicatedData();
+	GAMEPLAYDEBUGGER_API void OnRep_ReplicatedData();
 
 
 	friend FGameplayDebuggerNetPack;
@@ -261,50 +261,50 @@ protected:
 	uint32 bSendDataPacksUsingRPC : 1;
 
 	/** notify about changes in known category set */
-	void OnCategoriesChanged();
+	GAMEPLAYDEBUGGER_API void OnCategoriesChanged();
 
 	/** notify about changes in known extensions set */
-	void OnExtensionsChanged();
+	GAMEPLAYDEBUGGER_API void OnExtensionsChanged();
 
 	/** send notifies to all categories about current tool state */
-	void NotifyCategoriesToolState(bool bIsActive);
+	GAMEPLAYDEBUGGER_API void NotifyCategoriesToolState(bool bIsActive);
 
 	/** send notifies to all categories about current tool state */
-	void NotifyExtensionsToolState(bool bIsActive);
+	GAMEPLAYDEBUGGER_API void NotifyExtensionsToolState(bool bIsActive);
 
 	UFUNCTION(Server, Reliable, WithValidation, meta = (CallInEditor = "true"))
-	void ServerSetEnabled(bool bEnable);
+	GAMEPLAYDEBUGGER_API void ServerSetEnabled(bool bEnable);
 
 	UFUNCTION(Server, Reliable, WithValidation, meta = (CallInEditor = "true"))
-	void ServerSetDebugActor(AActor* Actor, bool bSelectInEditor);
+	GAMEPLAYDEBUGGER_API void ServerSetDebugActor(AActor* Actor, bool bSelectInEditor);
 
 	UFUNCTION(Server, Reliable, WithValidation, meta = (CallInEditor = "true"))
-	void ServerSetViewPoint(const FVector& InViewLocation, const FVector& InViewDirection);
+	GAMEPLAYDEBUGGER_API void ServerSetViewPoint(const FVector& InViewLocation, const FVector& InViewDirection);
 
 	UFUNCTION(Server, Reliable, WithValidation, meta = (CallInEditor = "true"))
-	void ServerResetViewPoint();
+	GAMEPLAYDEBUGGER_API void ServerResetViewPoint();
 
 	UFUNCTION(Server, Reliable, WithValidation, meta = (CallInEditor = "true"))
-	void ServerSetCategoryEnabled(int32 CategoryId, bool bEnable);
+	GAMEPLAYDEBUGGER_API void ServerSetCategoryEnabled(int32 CategoryId, bool bEnable);
 
 	/** helper function for replicating input for category handlers */
 	UFUNCTION(Server, Reliable, WithValidation, meta = (CallInEditor = "true"))
-	void ServerSendCategoryInputEvent(int32 CategoryId, int32 HandlerId);
+	GAMEPLAYDEBUGGER_API void ServerSendCategoryInputEvent(int32 CategoryId, int32 HandlerId);
 
 	/** helper function for replicating input for extension handlers */
 	UFUNCTION(Server, Reliable, WithValidation, meta = (CallInEditor = "true"))
-	void ServerSendExtensionInputEvent(int32 ExtensionId, int32 HandlerId);
+	GAMEPLAYDEBUGGER_API void ServerSendExtensionInputEvent(int32 ExtensionId, int32 HandlerId);
 
 	/** helper function for optionally sending DataPackPackets as RPC's */
 	UFUNCTION(Client, Reliable)
-	void ClientDataPackPacket(const FGameplayDebuggerDataPackRPCParams& Params);
+	GAMEPLAYDEBUGGER_API void ClientDataPackPacket(const FGameplayDebuggerDataPackRPCParams& Params);
 	
 	/** [LOCAL] notify from CategoryData replication */
-	void OnReceivedDataPackPacket(int32 CategoryId, int32 DataPackId, const FGameplayDebuggerDataPack& DataPacket);
+	GAMEPLAYDEBUGGER_API void OnReceivedDataPackPacket(int32 CategoryId, int32 DataPackId, const FGameplayDebuggerDataPack& DataPacket);
 
 	/** called both from BeginPlay and InitForEditor to setup instance's internal */
-	void Init();
+	GAMEPLAYDEBUGGER_API void Init();
 
 private:
-	void SendDataPackPacket(FName CategoryName, int32 DataPackIdx, FGameplayDebuggerDataPack& DataPack);
+	GAMEPLAYDEBUGGER_API void SendDataPackPacket(FName CategoryName, int32 DataPackIdx, FGameplayDebuggerDataPack& DataPack);
 };
diff --git a/Engine/Source/Runtime/GameplayDebugger/Public/GameplayDebuggerConfig.h b/Engine/Source/Runtime/GameplayDebugger/Public/GameplayDebuggerConfig.h
--- a/Engine/Source/Runtime/GameplayDebugger/Public/GameplayDebuggerConfig.h
+++ b/Engine/Source/Runtime/GameplayDebugger/Public/GameplayDebuggerConfig.h
@@ -109,8 +109,8 @@ struct FGameplayDebuggerExtensionConfig
 	FGameplayDebuggerExtensionConfig() : UseExtension(EGameplayDebuggerOverrideMode::UseDefault) {}
 };
 
-UCLASS(config = Engine, defaultconfig)
-class GAMEPLAYDEBUGGER_API UGameplayDebuggerConfig : public UObject
+UCLASS(config = Engine, defaultconfig, MinimalAPI)
+class UGameplayDebuggerConfig : public UObject
 {
 	GENERATED_UCLASS_BODY()
 
@@ -193,23 +193,23 @@ class GAMEPLAYDEBUGGER_API UGameplayDebuggerConfig : public UObject
 	TArray<FGameplayDebuggerExtensionConfig> Extensions;
 
 	/** updates entry in Categories array and modifies category creation params */
-	void UpdateCategoryConfig(const FName CategoryName, int32& SlotIdx, uint8& CategoryState);
+	GAMEPLAYDEBUGGER_API void UpdateCategoryConfig(const FName CategoryName, int32& SlotIdx, uint8& CategoryState);
 
 	/** updates entry in Categories array and modifies input binding params */
-	void UpdateCategoryInputConfig(const FName CategoryName, const FName InputName, FName& KeyName, FGameplayDebuggerInputModifier& KeyModifier);
+	GAMEPLAYDEBUGGER_API void UpdateCategoryInputConfig(const FName CategoryName, const FName InputName, FName& KeyName, FGameplayDebuggerInputModifier& KeyModifier);
 
 	/** updates entry in Extensions array and modifies extension creation params */
-	void UpdateExtensionConfig(const FName ExtensionName, uint8& UseExtension);
+	GAMEPLAYDEBUGGER_API void UpdateExtensionConfig(const FName ExtensionName, uint8& UseExtension);
 
 	/** updates entry in Categories array and modifies input binding params */
-	void UpdateExtensionInputConfig(const FName ExtensionName, const FName InputName, FName& KeyName, FGameplayDebuggerInputModifier& KeyModifier);
+	GAMEPLAYDEBUGGER_API void UpdateExtensionInputConfig(const FName ExtensionName, const FName InputName, FName& KeyName, FGameplayDebuggerInputModifier& KeyModifier);
 
 	/** remove all category and extension data from unknown sources (outdated entries) */
-	void RemoveUnknownConfigs();
+	GAMEPLAYDEBUGGER_API void RemoveUnknownConfigs();
 
-	virtual void Serialize(FArchive& Ar) override;
+	GAMEPLAYDEBUGGER_API virtual void Serialize(FArchive& Ar) override;
 #if WITH_EDITOR
-	virtual void PostEditChangeChainProperty(struct FPropertyChangedChainEvent& PropertyChangedEvent) override;
+	GAMEPLAYDEBUGGER_API virtual void PostEditChangeChainProperty(struct FPropertyChangedChainEvent& PropertyChangedEvent) override;
 #endif
 
 private:
@@ -222,8 +222,8 @@ private:
 };
 
 
-UCLASS(config = EditorPerProjectUserSettings, meta = (DisplayName = "Gameplay Debugger"))
-class GAMEPLAYDEBUGGER_API UGameplayDebuggerUserSettings : public UDeveloperSettings
+UCLASS(config = EditorPerProjectUserSettings, meta = (DisplayName = "Gameplay Debugger"), MinimalAPI)
+class UGameplayDebuggerUserSettings : public UDeveloperSettings
 {
 	GENERATED_BODY()
 protected:
@@ -231,7 +231,7 @@ protected:
 
 public:
 	static int32 GetFontSize() { return GetDefault<UGameplayDebuggerUserSettings>()->FontSize; }
-	static void SetFontSize(const int32 InFontSize);
+	static GAMEPLAYDEBUGGER_API void SetFontSize(const int32 InFontSize);
 
 	/** Controls whether GameplayDebugger will be available in pure editor mode.
 	 *  @Note that you need to reload the map for the changes to this property to take effect */
diff --git a/Engine/Source/Runtime/GameplayDebugger/Public/GameplayDebuggerExtension.h b/Engine/Source/Runtime/GameplayDebugger/Public/GameplayDebuggerExtension.h
--- a/Engine/Source/Runtime/GameplayDebugger/Public/GameplayDebuggerExtension.h
+++ b/Engine/Source/Runtime/GameplayDebugger/Public/GameplayDebuggerExtension.h
@@ -22,28 +22,28 @@
 
 class APlayerController;
 
-class GAMEPLAYDEBUGGER_API FGameplayDebuggerExtension : public FGameplayDebuggerAddonBase
+class FGameplayDebuggerExtension : public FGameplayDebuggerAddonBase
 {
 public:
 
 	virtual ~FGameplayDebuggerExtension() {}
-	virtual void OnGameplayDebuggerActivated() override;
-	virtual void OnGameplayDebuggerDeactivated() override;
+	GAMEPLAYDEBUGGER_API virtual void OnGameplayDebuggerActivated() override;
+	GAMEPLAYDEBUGGER_API virtual void OnGameplayDebuggerDeactivated() override;
 
 	/** [LOCAL] description for gameplay debugger's header row, newline character is ignored */
-	virtual FString GetDescription() const;
+	GAMEPLAYDEBUGGER_API virtual FString GetDescription() const;
 
 	/** [LOCAL] called when added to debugger tool or tool is activated */
-	virtual void OnActivated();
+	GAMEPLAYDEBUGGER_API virtual void OnActivated();
 
 	/** [LOCAL] called when removed from debugger tool or tool is deactivated */
-	virtual void OnDeactivated();
+	GAMEPLAYDEBUGGER_API virtual void OnDeactivated();
 
 	/** check if extension is created for local player */
-	bool IsLocal() const;
+	GAMEPLAYDEBUGGER_API bool IsLocal() const;
 
 protected:
 
 	/** get player controller owning gameplay debugger tool */
-	APlayerController* GetPlayerController() const;
+	GAMEPLAYDEBUGGER_API APlayerController* GetPlayerController() const;
 };
diff --git a/Engine/Source/Runtime/GameplayDebugger/Public/GameplayDebuggerPlayerManager.h b/Engine/Source/Runtime/GameplayDebugger/Public/GameplayDebuggerPlayerManager.h
--- a/Engine/Source/Runtime/GameplayDebugger/Public/GameplayDebuggerPlayerManager.h
+++ b/Engine/Source/Runtime/GameplayDebugger/Public/GameplayDebuggerPlayerManager.h
@@ -29,53 +29,53 @@ struct FGameplayDebuggerPlayerData
 	TObjectPtr<AGameplayDebuggerCategoryReplicator> Replicator = nullptr;
 };
 
-UCLASS(NotBlueprintable, NotBlueprintType, notplaceable, noteditinlinenew, hidedropdown, Transient)
-class GAMEPLAYDEBUGGER_API AGameplayDebuggerPlayerManager : public AActor, public FTickableGameObject
+UCLASS(NotBlueprintable, NotBlueprintType, notplaceable, noteditinlinenew, hidedropdown, Transient, MinimalAPI)
+class AGameplayDebuggerPlayerManager : public AActor, public FTickableGameObject
 {
 	GENERATED_UCLASS_BODY()
 
-	virtual void TickActor(float DeltaTime, enum ELevelTick TickType, FActorTickFunction& ThisTickFunction) override;
+	GAMEPLAYDEBUGGER_API virtual void TickActor(float DeltaTime, enum ELevelTick TickType, FActorTickFunction& ThisTickFunction) override;
 
 	// we're ticking only the manager only when in editor
 	// FTickableGameObject begin
 	virtual UWorld* GetTickableGameObjectWorld() const override { return GetWorld(); }
-	virtual void Tick(float DeltaTime) override;
-	virtual ETickableTickType GetTickableTickType() const override;
-	virtual TStatId GetStatId() const override;
+	GAMEPLAYDEBUGGER_API virtual void Tick(float DeltaTime) override;
+	GAMEPLAYDEBUGGER_API virtual ETickableTickType GetTickableTickType() const override;
+	GAMEPLAYDEBUGGER_API virtual TStatId GetStatId() const override;
 	virtual bool IsTickableInEditor() const override { return true; }
 	/** results in this instance not being ticked in any of the game worlds */
 	virtual bool IsTickable() const override { return bEditorTimeTick; }
 	// FTickableGameObject end
 
-	virtual void PostInitProperties() override;
+	GAMEPLAYDEBUGGER_API virtual void PostInitProperties() override;
 
 protected:
-	virtual void BeginPlay() override;
+	GAMEPLAYDEBUGGER_API virtual void BeginPlay() override;
 
-	void OnReplayScrubTeardown(UWorld* InWorld);
-	void OnGameModeLogout(AGameModeBase* GameMode, AController* Exiting);
+	GAMEPLAYDEBUGGER_API void OnReplayScrubTeardown(UWorld* InWorld);
+	GAMEPLAYDEBUGGER_API void OnGameModeLogout(AGameModeBase* GameMode, AController* Exiting);
 
 public:
-	virtual void EndPlay(const EEndPlayReason::Type Reason) override;
-	void Init();
+	GAMEPLAYDEBUGGER_API virtual void EndPlay(const EEndPlayReason::Type Reason) override;
+	GAMEPLAYDEBUGGER_API void Init();
 		
-	void UpdateAuthReplicators();
-	void RegisterReplicator(AGameplayDebuggerCategoryReplicator& Replicator);
-	void RefreshInputBindings(AGameplayDebuggerCategoryReplicator& Replicator);
+	GAMEPLAYDEBUGGER_API void UpdateAuthReplicators();
+	GAMEPLAYDEBUGGER_API void RegisterReplicator(AGameplayDebuggerCategoryReplicator& Replicator);
+	GAMEPLAYDEBUGGER_API void RefreshInputBindings(AGameplayDebuggerCategoryReplicator& Replicator);
 
-	AGameplayDebuggerCategoryReplicator* GetReplicator(const APlayerController& OwnerPC) const;
-	UInputComponent* GetInputComponent(const APlayerController& OwnerPC) const;
-	UGameplayDebuggerLocalController* GetLocalController(const APlayerController& OwnerPC) const;
+	GAMEPLAYDEBUGGER_API AGameplayDebuggerCategoryReplicator* GetReplicator(const APlayerController& OwnerPC) const;
+	GAMEPLAYDEBUGGER_API UInputComponent* GetInputComponent(const APlayerController& OwnerPC) const;
+	GAMEPLAYDEBUGGER_API UGameplayDebuggerLocalController* GetLocalController(const APlayerController& OwnerPC) const;
 #if WITH_EDITOR
 	UGameplayDebuggerLocalController* GetEditorController() const { return EditorWorldData.Controller; }
 #endif // WITH_EDITOR
 	
-	const FGameplayDebuggerPlayerData* GetPlayerData(const APlayerController& OwnerPC) const;
+	GAMEPLAYDEBUGGER_API const FGameplayDebuggerPlayerData* GetPlayerData(const APlayerController& OwnerPC) const;
 
-	static AGameplayDebuggerPlayerManager& GetCurrent(UWorld* World);
+	static GAMEPLAYDEBUGGER_API AGameplayDebuggerPlayerManager& GetCurrent(UWorld* World);
 
 	/** extracts view location and direction from a player controller that can be used for picking */
-	static void GetViewPoint(const APlayerController& OwnerPC, FVector& OutViewLocation, FVector& OutViewDirection);
+	static GAMEPLAYDEBUGGER_API void GetViewPoint(const APlayerController& OwnerPC, FVector& OutViewLocation, FVector& OutViewDirection);
 
 protected:
 
