Consolidated duplicated stats in FileIoCache

#rnx

diff --git a/Engine/Source/Runtime/Experimental/IoStoreOnDemand/Private/FileIoCache.cpp b/Engine/Source/Runtime/Experimental/IoStoreOnDemand/Private/FileIoCache.cpp
--- a/Engine/Source/Runtime/Experimental/IoStoreOnDemand/Private/FileIoCache.cpp
+++ b/Engine/Source/Runtime/Experimental/IoStoreOnDemand/Private/FileIoCache.cpp
@@ -1,8 +1,7 @@
 // Copyright Epic Games, Inc. All Rights Reserved.
 
-// {{{1 misc
-
 #include "FileIoCache.h"
+#include "Statistics.h"
 
 #include "Containers/IntrusiveDoubleLinkedList.h"
 #include "GenericPlatform/GenericPlatformFile.h"
@@ -25,98 +24,11 @@
 
 DEFINE_LOG_CATEGORY(LogIoCache);
 
-// }}}
-
 namespace UE::IO::Private
 {
 
 using namespace UE::Tasks;
 
-// {{{1 stats
-
-////////////////////////////////////////////////////////////////////////////////
-#define BUILD_DATUM_X(x) \
-	x(CachedBytes) \
-	x(ErrorCount) \
-	x(GetCount) \
-	x(PendingBytes) \
-	x(PutCount) \
-	x(PutExistingCount) \
-	x(PutRejectCount) \
-	x(ReadBytes)
-
-#if !COUNTERSTRACE_ENABLED
-namespace FCountersTrace
-{
-	struct FCounterInt
-	{
-		FCounterInt(...) {}
-		void Set(int64 i) { V = i; }
-		void Add(int64 d) { V += d; }
-		int64 Get() const { return V;}
-		std::atomic<int64> V = 0;
-	};
-	using FCounterAtomicInt = FCounterInt;
-}
-#endif 
-
-class FStatistics
-{
-public:
-	enum class EDatum : int32
-	{
-#define DATUM_X(x) x,
-		BUILD_DATUM_X(DATUM_X)
-#undef DATUM_X
-		_Count,
-	};
-
-	void Add(EDatum Which, int64 Delta);
-	void Set(EDatum Which, int64 Value);
-	int64 Get(EDatum Which) const;
-
-private:
-	struct FDatum
-	{
-		FDatum(const TCHAR* Name);
-		FCountersTrace::FCounterAtomicInt Value;
-	};
-
-	FDatum Datums[int(EDatum::_Count)] = {
-#define DATUM_X(x) FDatum{ TEXT("Ias/FileCache/") TEXT(PREPROCESSOR_TO_STRING(x)) },
-		BUILD_DATUM_X(DATUM_X)
-#undef DATUM_X
-	};
-};
-
-static FStatistics GStatistics;
-
-FStatistics::FDatum::FDatum(const TCHAR* Name)
-: Value(TraceCounterNameType_Static, Name, TraceCounterDisplayHint_None)
-{
-	FCountersTrace::FCounterAtomicInt(TEXT(""), TraceCounterDisplayHint_None).Get();
-	FCountersTrace::FCounterInt(TEXT(""), TraceCounterDisplayHint_None).Get();
-}
-
-void FStatistics::Add(EDatum Which, int64 Delta)
-{
-	Datums[int(Which)].Value.Add(Delta);
-}
-
-void FStatistics::Set(EDatum Which, int64 Value)
-{
-	Datums[int(Which)].Value.Set(Value);
-}
-
-int64 FStatistics::Get(EDatum Which) const
-{
-	return Datums[int(Which)].Value.Get();
-}
-
-
-
-// {{{1 toc-file
-
 ////////////////////////////////////////////////////////////////////////////////
 class FCacheFileToc
 {
@@ -225,8 +137,6 @@ FIoStatus FCacheFileToc::Save(const FString& FilePath, const uint64 CursorPos)
 
 
 
-// {{{1 cache-map
-
 ////////////////////////////////////////////////////////////////////////////////
 enum class ECacheEntryState : uint8
 {
@@ -313,20 +223,21 @@ bool FCacheMap::Get(const FIoHash& Key, FCacheEntry& OutEntry) const
 
 bool FCacheMap::InsertPending(FIoHash Key, FIoBuffer& Data, bool& bAdded)
 {
-	check(Data.GetSize() > 0);
+	uint64 DataSize = Data.GetSize();
+	check(DataSize > 0);
 
 	FScopeLock _(&Cs);
 
 	bAdded = false;
-	if (TotalPendingBytes + Data.GetSize() > MaxPendingBytes)
+	if (TotalPendingBytes + DataSize > MaxPendingBytes)
 	{
-		GStatistics.Add(FStatistics::EDatum::PutRejectCount, 1);
+		FOnDemandIoBackendStats::Get()->OnCachePutReject(DataSize);
 		return false;
 	}
 
 	if (Lookup.Contains(Key))
 	{
-		GStatistics.Add(FStatistics::EDatum::PutExistingCount, 1);
+		FOnDemandIoBackendStats::Get()->OnCachePutExisting(DataSize);
 		return true;
 	}
 
@@ -338,10 +249,10 @@ bool FCacheMap::InsertPending(FIoHash Key, FIoBuffer& Data, bool& bAdded)
 
 	Pending.AddTail(Entry.Get());
 
-	TotalPendingBytes += Data.GetSize();
+	TotalPendingBytes += DataSize;
 
-	GStatistics.Add(FStatistics::EDatum::PendingBytes, Data.GetSize());
-	GStatistics.Add(FStatistics::EDatum::PutCount, 1);
+	FOnDemandIoBackendStats::Get()->OnCachePut();
+	FOnDemandIoBackendStats::Get()->OnCachePendingBytes(TotalPendingBytes);
 
 	return bAdded = true;
 }
@@ -373,7 +284,7 @@ int32 FCacheMap::RemovePending(FCacheEntryList& OutPending, uint32 MaxSize)
 	}
 
 	TotalPendingBytes -= ReturnSize;
-	GStatistics.Set(FStatistics::EDatum::PendingBytes, TotalPendingBytes);
+	FOnDemandIoBackendStats::Get()->OnCachePendingBytes(TotalPendingBytes);
 
 	return ReturnSize;
 }
@@ -398,7 +309,7 @@ void FCacheMap::InsertPersisted(FCacheEntryList&& InPersisted, const uint64 Curs
 
 	Persisted.AddTail(MoveTemp(InPersisted));
 	TotalPersistedBytes += PersistedBytes;
-	GStatistics.Add(FStatistics::EDatum::CachedBytes, PersistedBytes);
+	FOnDemandIoBackendStats::Get()->OnCachePersistedBytes(PersistedBytes);
 }
 
 void FCacheMap::RemovePersisted(const uint64 RequiredSize)
@@ -422,7 +333,7 @@ void FCacheMap::RemovePersisted(const uint64 RequiredSize)
 	}
 
 	TotalPersistedBytes -= RemovedBytes;
-	GStatistics.Add(FStatistics::EDatum::CachedBytes, -int64(RemovedBytes));
+	FOnDemandIoBackendStats::Get()->OnCachePersistedBytes(TotalPersistedBytes);
 }
 
 FIoStatus FCacheMap::Load(const FString& FilePath, uint64& OutCursorPos)
@@ -456,7 +367,7 @@ FIoStatus FCacheMap::Load(const FString& FilePath, uint64& OutCursorPos)
 		check(ExpectedCursosPos == OutCursorPos);
 	}
 
-	GStatistics.Set(FStatistics::EDatum::CachedBytes, TotalPersistedBytes);
+	FOnDemandIoBackendStats::Get()->OnCachePersistedBytes(TotalPersistedBytes);
 
 	return FIoStatus(EIoErrorCode::Ok);
 }
@@ -481,8 +392,6 @@ FIoStatus FCacheMap::Save(const FString& FilePath, const uint64 CursorPos)
 
 
 
-// {{{1 governor
-
 ////////////////////////////////////////////////////////////////////////////////
 class FGovernorExternal
 {
@@ -598,8 +507,6 @@ void FGovernorInternal::Return(uint32 LeftOver)
 
 
 
-// {{{1 file-impl
-
 ////////////////////////////////////////////////////////////////////////////////
 class FFileIoCache final
 	: public FRunnable
@@ -685,7 +592,7 @@ TTask<TIoStatusOr<FIoBuffer>> FFileIoCache::Get(const FIoHash& Key, const FIoRea
 				const uint64 ReadSize = FMath::Min(Options.GetSize(), Entry.Data.GetSize());
 				FIoBuffer Buffer = Options.GetTargetVa() ? FIoBuffer(FIoBuffer::Wrap, Options.GetTargetVa(), ReadSize) : FIoBuffer(ReadSize);
 				Buffer.GetMutableView().CopyFrom(Entry.Data.GetView().RightChop(ReadOffset));
-				GStatistics.Add(FStatistics::EDatum::GetCount, 1);
+				FOnDemandIoBackendStats::Get()->OnCacheGet(Entry.Data.GetSize());
 
 				return TIoStatusOr<FIoBuffer>(Buffer);
 			}
@@ -713,12 +620,11 @@ TTask<TIoStatusOr<FIoBuffer>> FFileIoCache::Get(const FIoHash& Key, const FIoRea
 					const FIoHash Hash = FIoHash::HashBuffer(Buffer.GetView());
 					if (Hash == ExpectedHash) 
 					{
-						GStatistics.Add(FStatistics::EDatum::GetCount, 1);
-						GStatistics.Add(FStatistics::EDatum::ReadBytes, ReadSize);
+						FOnDemandIoBackendStats::Get()->OnCacheGet(ReadSize);
 						return TIoStatusOr<FIoBuffer>(Buffer);
 					}
 
-					GStatistics.Add(FStatistics::EDatum::ErrorCount, 1);
+					FOnDemandIoBackendStats::Get()->OnCacheError();
 					UE_LOG(LogIoCache, Verbose, TEXT("Read chunk failed, hash mismatch, Key='%s', Hash='%s', ExpectedHash='%s', File='%s', Offset='%llu', Size='%llu'"),
 						*LexToString(Entry.Key), *LexToString(Hash), *LexToString(ExpectedHash), *CacheFilePath, ReadOffset, ReadSize);
 
@@ -726,7 +632,7 @@ TTask<TIoStatusOr<FIoBuffer>> FFileIoCache::Get(const FIoHash& Key, const FIoRea
 				}
 				else
 				{
-					GStatistics.Add(FStatistics::EDatum::ErrorCount, 1);
+					FOnDemandIoBackendStats::Get()->OnCacheError();
 					UE_LOG(LogIoCache, Warning, TEXT("Read chunk failed, unable to open cache file '%s' for reading"), *CacheFilePath);
 					return TIoStatusOr<FIoBuffer>(FIoStatus(EIoErrorCode::ReadError));
 				}
@@ -922,13 +828,9 @@ void FFileIoCache::FileWriterThreadInner()
 	Entries.Reset();
 }
 
-// }}}
-
 } // namespace UE::IO::Private
 
 TUniquePtr<IIoCache> MakeFileIoCache(const FFileIoCacheConfig& Config)
 {
 	return MakeUnique<UE::IO::Private::FFileIoCache>(Config);
 }
-
-// vim: foldlevel=1
diff --git a/Engine/Source/Runtime/Experimental/IoStoreOnDemand/Private/OnDemandIoDispatcherBackend.cpp b/Engine/Source/Runtime/Experimental/IoStoreOnDemand/Private/OnDemandIoDispatcherBackend.cpp
--- a/Engine/Source/Runtime/Experimental/IoStoreOnDemand/Private/OnDemandIoDispatcherBackend.cpp
+++ b/Engine/Source/Runtime/Experimental/IoStoreOnDemand/Private/OnDemandIoDispatcherBackend.cpp
@@ -986,18 +986,7 @@ void FOnDemandIoBackend::CompleteRequest(FChunkRequest* ChunkRequest)
 
 	if (bCanCache && !ChunkRequest->bCached && Chunk.GetSize() > 0)
 	{
-		if (FIoStatus Status = Cache->Put(ChunkRequest->Params.ChunkKey, Chunk); Status.IsOk())
-		{
-			Stats.OnCachePut(Chunk.GetSize());
-		}
-		else
-		{
-			Stats.OnCacheReject(Chunk.GetSize());
-		}
-	}
-	else if (ChunkRequest->bCached)
-	{
-		Stats.OnCacheHit(Chunk.GetSize());
+		Cache->Put(ChunkRequest->Params.ChunkKey, Chunk);
 	}
 
 	Stats.OnChunkRequestRelease();
diff --git a/Engine/Source/Runtime/Experimental/IoStoreOnDemand/Private/Statistics.cpp b/Engine/Source/Runtime/Experimental/IoStoreOnDemand/Private/Statistics.cpp
--- a/Engine/Source/Runtime/Experimental/IoStoreOnDemand/Private/Statistics.cpp
+++ b/Engine/Source/Runtime/Experimental/IoStoreOnDemand/Private/Statistics.cpp
@@ -41,12 +41,14 @@ FCounterAtomicInt	GIoRequestsFailed(TEXT("Ias/IoRequestsFailed"), TraceCounterDi
 FCounterInt			GReadRequestsCreated(TEXT("Ias/ReadRequestsCreated"), TraceCounterDisplayHint_None);
 FCounterAtomicInt	GReadRequestsRemoved(TEXT("Ias/ReadRequestsRemoved"), TraceCounterDisplayHint_None);
 // cache stats
-FCounterAtomicInt	GCacheHits(TEXT("Ias/CacheHits"), TraceCounterDisplayHint_None);
-FCounterAtomicInt	GCacheHitsSize(TEXT("Ias/Size/CacheHitsSize"), TraceCounterDisplayHint_Memory);
-FCounterAtomicInt	GCachePuts(TEXT("Ias/CachePuts"), TraceCounterDisplayHint_None);
-FCounterAtomicInt	GCachePutsSize(TEXT("Ias/Size/CachePutsSize"), TraceCounterDisplayHint_Memory);
-FCounterAtomicInt	GCacheRejects(TEXT("Ias/CacheRejects"), TraceCounterDisplayHint_None);
-FCounterAtomicInt	GCacheRejectsSize(TEXT("Ias/Size/CacheRejectsSize"), TraceCounterDisplayHint_Memory);
+FCounterAtomicInt	GCacheErrorCount(TEXT("Ias/CacheErrorCount"), TraceCounterDisplayHint_None);
+FCounterAtomicInt	GCacheGetCount(TEXT("Ias/CacheGetCount"), TraceCounterDisplayHint_None);
+FCounterAtomicInt	GCachePutCount(TEXT("Ias/CachePutCount"), TraceCounterDisplayHint_None);
+FCounterAtomicInt	GCachePutExistingCount(TEXT("Ias/CachePutExistingCount"), TraceCounterDisplayHint_None);
+FCounterAtomicInt	GCachePutRejectCount(TEXT("Ias/CachePutRejectCount"), TraceCounterDisplayHint_None);
+FCounterAtomicInt	GCacheCachedBytes(TEXT("Ias/CacheCachedBytes"), TraceCounterDisplayHint_Memory);
+FCounterAtomicInt	GCachePendingBytes(TEXT("Ias/CachePendingBytes"), TraceCounterDisplayHint_Memory);
+FCounterAtomicInt	GCacheReadBytes(TEXT("Ias/CacheReadBytes"), TraceCounterDisplayHint_Memory);
 // http stats
 FCounterInt			GHttpRequestsCompleted(TEXT("Ias/HttpRequestsCompleted"), TraceCounterDisplayHint_None);
 FCounterInt			GHttpRequestsFailed(TEXT("Ias/HttpRequestsFailed"), TraceCounterDisplayHint_None);
@@ -67,12 +69,14 @@ CSV_DEFINE_STAT(Ias, FrameIoRequestsFailed);
 CSV_DEFINE_STAT(Ias, FrameReadRequestsCreated);
 CSV_DEFINE_STAT(Ias, FrameReadRequestsRemoved);
 // cache stats
-CSV_DEFINE_STAT(Ias, FrameCacheHits);
-CSV_DEFINE_STAT(Ias, FrameCacheHitsSize);
-CSV_DEFINE_STAT(Ias, FrameCachePuts);
-CSV_DEFINE_STAT(Ias, FrameCachePutsSize);
-CSV_DEFINE_STAT(Ias, FrameCacheRejects);
-CSV_DEFINE_STAT(Ias, FrameCacheRejectsSize);
+CSV_DEFINE_STAT(Ias, FrameCacheErrorCount);
+CSV_DEFINE_STAT(Ias, FrameCacheGetCount);
+CSV_DEFINE_STAT(Ias, FrameCachePutCount);
+CSV_DEFINE_STAT(Ias, FrameCachePutExistingCount);
+CSV_DEFINE_STAT(Ias, FrameCachePutRejectCount);
+CSV_DEFINE_STAT(Ias, FrameCacheCachedBytes);
+CSV_DEFINE_STAT(Ias, FrameCachePendingBytes);
+CSV_DEFINE_STAT(Ias, FrameCacheReadBytes);
 // http stats
 CSV_DEFINE_STAT(Ias, FrameHttpRequestsCompleted);
 CSV_DEFINE_STAT(Ias, FrameHttpRequestsFailed);
@@ -80,6 +84,24 @@ CSV_DEFINE_STAT(Ias, FrameHttpRequestsPending);
 CSV_DEFINE_STAT(Ias, FrameHttpRequestsInflight);
 CSV_DEFINE_STAT(Ias, FrameHttpRequestsCompletedSize);
 
+static FOnDemandIoBackendStats* GStatistics = nullptr;
+
+FOnDemandIoBackendStats::FOnDemandIoBackendStats()
+{
+	check(GStatistics == nullptr);
+	GStatistics = this;
+}
+
+FOnDemandIoBackendStats::~FOnDemandIoBackendStats()
+{
+	GStatistics = nullptr;
+}
+
+FOnDemandIoBackendStats* FOnDemandIoBackendStats::Get()
+{
+	return GStatistics;
+}
+
 void FOnDemandIoBackendStats::OnIoRequestEnqueue()
 {
 	GIoRequestsMade.Add(1);
@@ -119,31 +141,49 @@ void FOnDemandIoBackendStats::OnChunkRequestRelease()
 	CSV_CUSTOM_STAT_DEFINED(FrameReadRequestsRemoved, int32(GReadRequestsRemoved.Get()), ECsvCustomStatOp::Set);
 }
 
-void FOnDemandIoBackendStats::OnCacheHit(uint64 InSize)
+
+void FOnDemandIoBackendStats::OnCacheError()
 {
-	GCacheHits.Add(1);
-	CSV_CUSTOM_STAT_DEFINED(FrameCacheHits, int32(GCacheHits.Get()), ECsvCustomStatOp::Set);
+	GCacheErrorCount.Add(1);
+	CSV_CUSTOM_STAT_DEFINED(FrameCacheErrorCount, int32(GCacheErrorCount.Get()), ECsvCustomStatOp::Set);
+}
+
+void FOnDemandIoBackendStats::OnCacheGet(uint64 DataSize)
+{
+	GCacheGetCount.Add(1);
+	CSV_CUSTOM_STAT_DEFINED(FrameCacheGetCount, int32(GCacheGetCount.Get()), ECsvCustomStatOp::Set);
+	GCacheReadBytes.Add(DataSize);
+	CSV_CUSTOM_STAT_DEFINED(FrameCacheReadBytes, BytesToApproxKB(GCacheReadBytes.Get()), ECsvCustomStatOp::Set);
+}
 
-	GCacheHitsSize.Add(InSize);
-	CSV_CUSTOM_STAT_DEFINED(FrameCacheHitsSize, BytesToApproxKB(GCacheHitsSize.Get()), ECsvCustomStatOp::Set);
+void FOnDemandIoBackendStats::OnCachePut()
+{
+	GCachePutCount.Add(1);
+	CSV_CUSTOM_STAT_DEFINED(FrameCachePutCount, int32(GCachePutCount.Get()), ECsvCustomStatOp::Set);
 }
 
-void FOnDemandIoBackendStats::OnCachePut(uint64 InSize)
+void FOnDemandIoBackendStats::OnCachePutExisting(uint64 /*DataSize*/)
 {
-	GCachePuts.Add(1);
-	CSV_CUSTOM_STAT_DEFINED(FrameCachePuts, int32(GCachePuts.Get()), ECsvCustomStatOp::Set);
+	GCachePutExistingCount.Add(1);
+	CSV_CUSTOM_STAT_DEFINED(FrameCachePutExistingCount, int32(GCachePutExistingCount.Get()), ECsvCustomStatOp::Set);
+}
 
-	GCachePutsSize.Add(InSize);
-	CSV_CUSTOM_STAT_DEFINED(FrameCachePutsSize, BytesToApproxKB(GCachePutsSize.Get()), ECsvCustomStatOp::Set);
+void FOnDemandIoBackendStats::OnCachePutReject(uint64 /*DataSize*/)
+{
+	GCachePutRejectCount.Add(1);
+	CSV_CUSTOM_STAT_DEFINED(FrameCachePutRejectCount, int32(GCachePutRejectCount.Get()), ECsvCustomStatOp::Set);
 }
 
-void FOnDemandIoBackendStats::OnCacheReject(uint64 InSize)
+void FOnDemandIoBackendStats::OnCachePendingBytes(uint64 TotalSize)
 {
-	GCacheRejects.Add(1);
-	CSV_CUSTOM_STAT_DEFINED(FrameCacheRejects, int32(GCacheRejects.Get()), ECsvCustomStatOp::Set);
+	GCachePendingBytes.Set(TotalSize);
+	CSV_CUSTOM_STAT_DEFINED(FrameCachePendingBytes, BytesToApproxKB(GCachePendingBytes.Get()), ECsvCustomStatOp::Set);
+}
 
-	GCacheRejectsSize.Add(InSize);
-	CSV_CUSTOM_STAT_DEFINED(FrameCacheRejectsSize, BytesToApproxKB(GCacheRejectsSize.Get()), ECsvCustomStatOp::Set);
+void FOnDemandIoBackendStats::OnCachePersistedBytes(uint64 TotalSize)
+{
+	GCacheCachedBytes.Set(TotalSize);
+	CSV_CUSTOM_STAT_DEFINED(FrameCacheCachedBytes, BytesToApproxKB(GCacheCachedBytes.Get()), ECsvCustomStatOp::Set);
 }
 
 void FOnDemandIoBackendStats::OnHttpRequestEnqueue()
diff --git a/Engine/Source/Runtime/Experimental/IoStoreOnDemand/Private/Statistics.h b/Engine/Source/Runtime/Experimental/IoStoreOnDemand/Private/Statistics.h
--- a/Engine/Source/Runtime/Experimental/IoStoreOnDemand/Private/Statistics.h
+++ b/Engine/Source/Runtime/Experimental/IoStoreOnDemand/Private/Statistics.h
@@ -17,23 +17,32 @@ namespace UE::IO::Private
 {
 
 #if IAS_WITH_STATISTICS
-#	define IAS_STATISTICS_IMPL() ;
+#	define IAS_STATISTICS_IMPL(...) ;
 #else
-#	define IAS_STATISTICS_IMPL() {}
+#	define IAS_STATISTICS_IMPL(...) { return __VA_ARGS__; }
 #endif
 
 class FOnDemandIoBackendStats
 {
 public:
+	static FOnDemandIoBackendStats* Get() IAS_STATISTICS_IMPL(nullptr)
+	FOnDemandIoBackendStats() IAS_STATISTICS_IMPL()
+	~FOnDemandIoBackendStats() IAS_STATISTICS_IMPL()
 	void OnIoRequestEnqueue() IAS_STATISTICS_IMPL()
 	void OnIoRequestComplete(uint64 RequestSize) IAS_STATISTICS_IMPL()
 	void OnIoRequestCancel() IAS_STATISTICS_IMPL()
 	void OnIoRequestFail() IAS_STATISTICS_IMPL()
 	void OnChunkRequestCreate() IAS_STATISTICS_IMPL()
 	void OnChunkRequestRelease() IAS_STATISTICS_IMPL()
-	void OnCacheHit(uint64 InSize) IAS_STATISTICS_IMPL()
-	void OnCachePut(uint64 InSize) IAS_STATISTICS_IMPL()
-	void OnCacheReject(uint64 InSize) IAS_STATISTICS_IMPL()
+
+	void OnCacheError() IAS_STATISTICS_IMPL()
+	void OnCacheGet(uint64 DataSize) IAS_STATISTICS_IMPL()
+	void OnCachePut() IAS_STATISTICS_IMPL()
+	void OnCachePutExisting(uint64 DataSize) IAS_STATISTICS_IMPL()
+	void OnCachePutReject(uint64 DataSize) IAS_STATISTICS_IMPL()
+	void OnCachePendingBytes(uint64 TotalSize) IAS_STATISTICS_IMPL()
+	void OnCachePersistedBytes(uint64 TotalSize) IAS_STATISTICS_IMPL()
+
 	void OnHttpRequestEnqueue() IAS_STATISTICS_IMPL()
 	void OnHttpRequestDequeue() IAS_STATISTICS_IMPL()
 	void OnHttpRequestComplete(uint64 InSize) IAS_STATISTICS_IMPL()
