IAS: Fix fluctuating CachedBytes stat

#rnx

diff --git a/Engine/Source/Runtime/Experimental/IoStoreOnDemand/Private/FileIoCache.cpp b/Engine/Source/Runtime/Experimental/IoStoreOnDemand/Private/FileIoCache.cpp
--- a/Engine/Source/Runtime/Experimental/IoStoreOnDemand/Private/FileIoCache.cpp
+++ b/Engine/Source/Runtime/Experimental/IoStoreOnDemand/Private/FileIoCache.cpp
@@ -309,7 +309,7 @@ void FCacheMap::InsertPersisted(FCacheEntryList&& InPersisted, const uint64 Curs
 
 	Persisted.AddTail(MoveTemp(InPersisted));
 	TotalPersistedBytes += PersistedBytes;
-	FOnDemandIoBackendStats::Get()->OnCachePersistedBytes(PersistedBytes);
+	FOnDemandIoBackendStats::Get()->OnCachePersistedBytes(TotalPersistedBytes);
 }
 
 void FCacheMap::RemovePersisted(const uint64 RequiredSize)
