Force save the IAS's cache table each write.

#rnx

diff --git a/Engine/Source/Runtime/Experimental/IoStoreOnDemand/Private/FileIoCache.cpp b/Engine/Source/Runtime/Experimental/IoStoreOnDemand/Private/FileIoCache.cpp
--- a/Engine/Source/Runtime/Experimental/IoStoreOnDemand/Private/FileIoCache.cpp
+++ b/Engine/Source/Runtime/Experimental/IoStoreOnDemand/Private/FileIoCache.cpp
@@ -22,6 +22,8 @@
 
 #include <atomic>
 
+#define IAS_FORCE_TOC_SAVE
+
 DEFINE_LOG_CATEGORY(LogIoCache);
 
 namespace UE::IO::Private
@@ -663,12 +665,12 @@ void FFileIoCache::Initialize()
 	WriterThread.Reset(FRunnableThread::Create(this, TEXT("Ias.FileCache"), 0, TPri_BelowNormal));
 
 	const FString CacheDir = FPaths::ProjectPersistentDownloadDir() / TEXT("IoCache");
-	const FString CacheTocPath = CacheDir / TEXT("cache.utoc");
 	CacheFilePath = CacheDir / TEXT("cache.ucas");
 	WriteCursorPos = 0;
 
 	IFileManager& FileMgr = IFileManager::Get();
 
+	FString CacheTocPath = CacheFilePath + TEXT(".toc");
 	if (CacheConfig.DropCache)
 	{
 		FileMgr.Delete(*CacheTocPath);
@@ -733,7 +735,7 @@ void FFileIoCache::Shutdown()
 	TickWriterEvent->Trigger();
 	WriterThread->Kill();
 
-	const FString CacheTocPath = FPaths::ProjectPersistentDownloadDir() / TEXT("IoCache") / TEXT("cache.utoc");
+	FString CacheTocPath = CacheFilePath + TEXT(".toc");
 	UE_LOG(LogIoCache, Log, TEXT("Saving TOC '%s'"), *CacheTocPath);
 	CacheMap.Save(CacheTocPath, WriteCursorPos);
 
@@ -825,6 +827,14 @@ void FFileIoCache::FileWriterThreadInner()
 
 	CacheMap.RemovePersisted(PendingSize);
 
+#if defined(IAS_FORCE_TOC_SAVE)
+	{
+		TRACE_CPUPROFILER_EVENT_SCOPE(FFileIoCache::TocSave);
+		FString CacheTocPath = CacheFilePath + TEXT(".toc");
+		CacheMap.Save(CacheTocPath, WriteCursorPos);
+	}
+#endif
+
 	Entries.Reset();
 }
 
