IAS: Add back stat for CacheRejectBytes

#rnx

diff --git a/Engine/Source/Runtime/Experimental/IoStoreOnDemand/Private/Statistics.cpp b/Engine/Source/Runtime/Experimental/IoStoreOnDemand/Private/Statistics.cpp
--- a/Engine/Source/Runtime/Experimental/IoStoreOnDemand/Private/Statistics.cpp
+++ b/Engine/Source/Runtime/Experimental/IoStoreOnDemand/Private/Statistics.cpp
@@ -49,6 +49,7 @@ FCounterAtomicInt	GCachePutRejectCount(TEXT("Ias/CachePutRejectCount"), TraceCou
 FCounterAtomicInt	GCacheCachedBytes(TEXT("Ias/CacheCachedBytes"), TraceCounterDisplayHint_Memory);
 FCounterAtomicInt	GCachePendingBytes(TEXT("Ias/CachePendingBytes"), TraceCounterDisplayHint_Memory);
 FCounterAtomicInt	GCacheReadBytes(TEXT("Ias/CacheReadBytes"), TraceCounterDisplayHint_Memory);
+FCounterAtomicInt	GCacheRejectBytes(TEXT("Ias/CachePutRejectBytes"), TraceCounterDisplayHint_Memory);
 // http stats
 FCounterInt			GHttpRequestsCompleted(TEXT("Ias/HttpRequestsCompleted"), TraceCounterDisplayHint_None);
 FCounterInt			GHttpRequestsFailed(TEXT("Ias/HttpRequestsFailed"), TraceCounterDisplayHint_None);
@@ -77,6 +78,7 @@ CSV_DEFINE_STAT(Ias, FrameCachePutRejectCount);
 CSV_DEFINE_STAT(Ias, FrameCacheCachedBytes);
 CSV_DEFINE_STAT(Ias, FrameCachePendingBytes);
 CSV_DEFINE_STAT(Ias, FrameCacheReadBytes);
+CSV_DEFINE_STAT(Ias, FrameCacheRejectBytes);
 // http stats
 CSV_DEFINE_STAT(Ias, FrameHttpRequestsCompleted);
 CSV_DEFINE_STAT(Ias, FrameHttpRequestsFailed);
@@ -168,10 +170,12 @@ void FOnDemandIoBackendStats::OnCachePutExisting(uint64 /*DataSize*/)
 	CSV_CUSTOM_STAT_DEFINED(FrameCachePutExistingCount, int32(GCachePutExistingCount.Get()), ECsvCustomStatOp::Set);
 }
 
-void FOnDemandIoBackendStats::OnCachePutReject(uint64 /*DataSize*/)
+void FOnDemandIoBackendStats::OnCachePutReject(uint64 DataSize)
 {
 	GCachePutRejectCount.Add(1);
 	CSV_CUSTOM_STAT_DEFINED(FrameCachePutRejectCount, int32(GCachePutRejectCount.Get()), ECsvCustomStatOp::Set);
+	GCacheRejectBytes.Add(DataSize);
+	CSV_CUSTOM_STAT_DEFINED(FrameCacheRejectBytes, BytesToApproxKB(GCacheRejectBytes.Get()), ECsvCustomStatOp::Set);
 }
 
 void FOnDemandIoBackendStats::OnCachePendingBytes(uint64 TotalSize)
