Fix some physics mover behaviour

diff --git a/Engine/Source/Runtime/Experimental/Chaos/Private/Chaos/Character/CharacterGroundConstraintContainer.cpp b/Engine/Source/Runtime/Experimental/Chaos/Private/Chaos/Character/CharacterGroundConstraintContainer.cpp
--- a/Engine/Source/Runtime/Experimental/Chaos/Private/Chaos/Character/CharacterGroundConstraintContainer.cpp
+++ b/Engine/Source/Runtime/Experimental/Chaos/Private/Chaos/Character/CharacterGroundConstraintContainer.cpp
@@ -71,7 +71,10 @@ namespace Chaos
 	{
 		if (Constraint != nullptr)
 		{
-			Constraint->CharacterParticle->RemoveConstraintHandle(Constraint);
+			if (Constraint->CharacterParticle != nullptr)
+			{
+				Constraint->CharacterParticle->RemoveConstraintHandle(Constraint);
+			}
 			if (Constraint->GroundParticle != nullptr)
 			{
 				Constraint->GroundParticle->RemoveConstraintHandle(Constraint);
diff --git a/Engine/Source/Runtime/Experimental/Chaos/Private/Chaos/Character/CharacterGroundConstraintSolver.cpp b/Engine/Source/Runtime/Experimental/Chaos/Private/Chaos/Character/CharacterGroundConstraintSolver.cpp
--- a/Engine/Source/Runtime/Experimental/Chaos/Private/Chaos/Character/CharacterGroundConstraintSolver.cpp
+++ b/Engine/Source/Runtime/Experimental/Chaos/Private/Chaos/Character/CharacterGroundConstraintSolver.cpp
@@ -183,6 +183,9 @@ namespace Chaos
 			}
 
 			// If the slope is too steep adjust the target to not point up the slope
+			// and set the radial force to zero
+			const FReal DtSq = Dt * Dt;
+
 			FReal DP = FVec3::DotProduct(Settings.VerticalAxis, GroundNormal);
 			if (DP <= Data.CosMaxWalkableSlopeAngle)
 			{
@@ -193,6 +196,11 @@ namespace Chaos
 				{
 					TargetDeltaPos -= UpMotion * FVec3(UpSlope);
 				}
+				ConstraintData.RadialImpulseLimit = FSolverReal(Settings.FrictionForceLimit * DtSq);
+			}
+			else
+			{
+				ConstraintData.RadialImpulseLimit = FSolverReal(Settings.RadialForceLimit * DtSq);
 			}
 
 			const FVec3 MotionTarget = Character_X + TargetDeltaPos;
@@ -201,12 +209,11 @@ namespace Chaos
 
 			const FMatrix33 WorldSpaceCharacterInvI = Utilities::ComputeWorldSpaceInertia(BodyData.CharacterBody.Q(), BodyData.CharacterBody.InvILocal());
 			const FReal CharacterInvI = Settings.VerticalAxis.Dot(WorldSpaceCharacterInvI * Settings.VerticalAxis);
-			const FReal DtSq = Dt * Dt;
+			
 
 			// Write constraint data
 			ConstraintData.AngularImpulseLimit = FSolverReal(Settings.TwistTorqueLimit * DtSq);
 			ConstraintData.AngularSwingImpulseLimit = FSolverReal(Settings.SwingTorqueLimit * DtSq);
-			ConstraintData.RadialImpulseLimit = FSolverReal(Settings.RadialForceLimit * DtSq);
 			ConstraintData.VerticalAxis = FSolverVec3(Settings.VerticalAxis);
 			ConstraintData.AssumedOnGroundHeight = FSolverReal(Settings.AssumedOnGroundHeight);
 			ConstraintData.CharacterInvIZ = FSolverReal(CharacterInvI);
diff --git a/Engine/Source/Runtime/Experimental/Chaos/Public/Chaos/Character/CharacterGroundConstraint.h b/Engine/Source/Runtime/Experimental/Chaos/Public/Chaos/Character/CharacterGroundConstraint.h
--- a/Engine/Source/Runtime/Experimental/Chaos/Public/Chaos/Character/CharacterGroundConstraint.h
+++ b/Engine/Source/Runtime/Experimental/Chaos/Public/Chaos/Character/CharacterGroundConstraint.h
@@ -220,6 +220,18 @@ namespace Chaos
 		FReal GetRadialForceLimit() const { return ConstraintSettings.Read().RadialForceLimit; }
 
 
+		/// The maximum force that the character can apply in the plane defined by the ground normal
+		/// to hold the character in place while standing on an unwalkable incline (defined by max
+		/// walkable slope angle)
+		void SetFrictionForceLimit(const FReal& Value)
+		{
+			ConstraintSettings.Modify(true, DirtyFlags, Proxy, [&Value](FCharacterGroundConstraintSettings& Data) {
+				Data.FrictionForceLimit = Value > FReal(0.0) ? Value : FReal(0.0);
+				});
+		}
+		FReal GetFrictionForceLimit() const { return ConstraintSettings.Read().FrictionForceLimit; }
+
+
 		/// The maximum torque that the character can apply about the vertical axis to rotate the
 		/// character to its target facing direction (or hold the character in place if no target
 		/// delta facing is set).
diff --git a/Engine/Source/Runtime/Experimental/Chaos/Public/Chaos/Character/CharacterGroundConstraintContainer.h b/Engine/Source/Runtime/Experimental/Chaos/Public/Chaos/Character/CharacterGroundConstraintContainer.h
--- a/Engine/Source/Runtime/Experimental/Chaos/Public/Chaos/Character/CharacterGroundConstraintContainer.h
+++ b/Engine/Source/Runtime/Experimental/Chaos/Public/Chaos/Character/CharacterGroundConstraintContainer.h
@@ -76,7 +76,10 @@ namespace Chaos
 					GroundParticle->RemoveConstraintHandle(this);
 				}
 				GroundParticle = InGroundParticle;
-				GroundParticle->AddConstraintHandle(this);
+				if (GroundParticle)
+				{
+					GroundParticle->AddConstraintHandle(this);
+				}
 				bGroundParticleChanged = true;
 			}
 		}
diff --git a/Engine/Source/Runtime/Experimental/Chaos/Public/Chaos/Character/CharacterGroundConstraintSettings.h b/Engine/Source/Runtime/Experimental/Chaos/Public/Chaos/Character/CharacterGroundConstraintSettings.h
--- a/Engine/Source/Runtime/Experimental/Chaos/Public/Chaos/Character/CharacterGroundConstraintSettings.h
+++ b/Engine/Source/Runtime/Experimental/Chaos/Public/Chaos/Character/CharacterGroundConstraintSettings.h
@@ -16,6 +16,7 @@ namespace Chaos
 		FVec3 VerticalAxis = FVec3(0.0, 0.0, 1.0);	/// World space up direction (default z axis)
 		FReal TargetHeight = 0.0;					/// Desired distance from the character body to the ground
 		FReal RadialForceLimit = 1500.0;			/// How much force the character can apply parallel to the ground plane to reach the target position
+		FReal FrictionForceLimit = 100.0;			/// How much force the character can apply parallel to the ground plane to reach the target position when standing on an unwalkable incline
 		FReal TwistTorqueLimit = 1000.0;			///	How much torque the character can apply about the vertical axis to reach the target facing direction
 		FReal SwingTorqueLimit = 5000.0;			///	How much torque the character can apply about the other axes to remain upright
 		FReal CosMaxWalkableSlopeAngle = 0.633;		/// Cosine of the maximum angle in degrees that the character is allowed to walk on
