Iris - Additional ObjectReferenceCache debugging

#rnx

diff --git a/Engine/Source/Runtime/Experimental/Iris/Core/Private/Iris/ReplicationSystem/ObjectReferenceCache.cpp b/Engine/Source/Runtime/Experimental/Iris/Core/Private/Iris/ReplicationSystem/ObjectReferenceCache.cpp
--- a/Engine/Source/Runtime/Experimental/Iris/Core/Private/Iris/ReplicationSystem/ObjectReferenceCache.cpp
+++ b/Engine/Source/Runtime/Experimental/Iris/Core/Private/Iris/ReplicationSystem/ObjectReferenceCache.cpp
@@ -35,8 +35,10 @@ DEFINE_LOG_CATEGORY_STATIC(LogIrisReferences, Log, All);
 
 #if UE_NET_ENABLE_REFERENCECACHE_LOG
 #	define UE_LOG_REFERENCECACHE(Verbosity, Format, ...)  UE_LOG(LogIrisReferences, Verbosity, Format, ##__VA_ARGS__)
+#	define UE_CLOG_REFERENCECACHE(Condition, Verbosity, Format, ...)  UE_CLOG(Condition, LogIrisReferences, Verbosity, Format, ##__VA_ARGS__)
 #else
 #	define UE_LOG_REFERENCECACHE(...)
+#	define UE_CLOG_REFERENCECACHE(...)
 #endif
 
 #define UE_LOG_REFERENCECACHE_WARNING(Format, ...)  UE_LOG(LogIrisReferences, Warning, Format, ##__VA_ARGS__)
@@ -521,7 +523,7 @@ void FObjectReferenceCache::RemoveReference(FNetRefHandle RefHandle, const UObje
 			}
 			else
 			{
-				check(Object == CachedItemToRemove.ObjectKey);
+				ensureAlwaysMsgf(Object == CachedItemToRemove.ObjectKey, TEXT("ObjectReferenceCache::RemoveReference: %s, Object 0x%p doesn't match cached reference 0x%p."), *RefHandle.ToString(), Object, CachedItemToRemove.ObjectKey);
 			}
 
 #if UE_NET_CLEANUP_REFERENCES_WITH_DYNAMIC_OUTER
@@ -589,6 +591,22 @@ void FObjectReferenceCache::RemoveReference(FNetRefHandle RefHandle, const UObje
 					*RefHandle.ToString(), *GetNameSafe(Object), Object);
 			}
 		}
+		else
+		{
+			UE_LOG_REFERENCECACHE(Warning, TEXT("ObjectReferenceCache::RemoveReference: %s. Trying to find RefHandle without object to remove from ObjectToNetReferenceHandle. SLOW."), *RefHandle.ToString());
+			bool bSuccesfullyRemoved = false;
+			for (auto RefHandleIt = ObjectToNetReferenceHandle.CreateIterator(); RefHandleIt; ++RefHandleIt)
+			{
+				if (RefHandleIt.Value() == RefHandle)
+				{
+					RefHandleIt.RemoveCurrent();
+					bSuccesfullyRemoved = true;
+					UE_LOG_REFERENCECACHE(Verbose, TEXT("ObjectReferenceCache::RemoveReference: %s. Successfully removed RefHandle from ObjectToNetReferenceHandle."), *RefHandle.ToString());
+					break;
+				}
+			}
+			UE_CLOG_REFERENCECACHE(!bSuccesfullyRemoved, Verbose, TEXT("ObjectReferenceCache::RemoveReference: %s. Failed removing RefHandle from ObjectToNetReferenceHandle."), *RefHandle.ToString());
+		}
 	}
 }
 
