Fix Iris's FObjectReferenceCache trying to print names of invalid pointers and causing a crash.

#rnx

diff --git a/Engine/Source/Runtime/Experimental/Iris/Core/Private/Iris/ReplicationSystem/ObjectReferenceCache.cpp b/Engine/Source/Runtime/Experimental/Iris/Core/Private/Iris/ReplicationSystem/ObjectReferenceCache.cpp
--- a/Engine/Source/Runtime/Experimental/Iris/Core/Private/Iris/ReplicationSystem/ObjectReferenceCache.cpp
+++ b/Engine/Source/Runtime/Experimental/Iris/Core/Private/Iris/ReplicationSystem/ObjectReferenceCache.cpp
@@ -251,16 +251,16 @@ bool FObjectReferenceCache::CreateObjectReferenceInternal(const UObject* Object,
 		}
 		else
 		{
-			UE_LOG_REFERENCECACHE(Verbose, TEXT("ObjectReferenceCache::CreateObjectReferenceHandle Removed %s from ObjectToNetReferenceHandle due to stale cache. Cached Object %s (0x%p).  New object %s (0x%p)"), 
-				*CachedObject.NetRefHandle.ToString(), *GetNameSafe(CachedObject.ObjectKey), CachedObject.ObjectKey, *GetNameSafe(Object), Object);
+			UE_LOG_REFERENCECACHE(Verbose, TEXT("ObjectReferenceCache::CreateObjectReferenceHandle Removed %s from ObjectToNetReferenceHandle due to stale cache. Cached Object (0x%p).  New object %s (0x%p)"), 
+				*CachedObject.NetRefHandle.ToString(), CachedObject.ObjectKey, *GetNameSafe(Object), Object);
 
 			ObjectToNetReferenceHandle.Remove(CachedObject.ObjectKey);
 			ObjectToNetReferenceHandle.Remove(Object);
 
 			if (RefHandle.IsStatic())
 			{
-				UE_LOG_REFERENCECACHE(Verbose, TEXT("ObjectReferenceCache::CreateObjectReferenceHandle clearing stale cache for static handle %s. Cached Object %s (0x%p).  New object %s (0x%p)"),
-					*CachedObject.NetRefHandle.ToString(), *GetNameSafe(CachedObject.ObjectKey), CachedObject.ObjectKey, *GetNameSafe(Object), Object);
+				UE_LOG_REFERENCECACHE(Verbose, TEXT("ObjectReferenceCache::CreateObjectReferenceHandle clearing stale cache for static handle %s. Cached Object (0x%p).  New object %s (0x%p)"),
+					*CachedObject.NetRefHandle.ToString(), CachedObject.ObjectKey, *GetNameSafe(Object), Object);
 
 				// Note we only cleanse the object reference
 				// we still keep the cachedObject data around in order to be able to serialize static destruction infos
@@ -270,8 +270,8 @@ bool FObjectReferenceCache::CreateObjectReferenceInternal(const UObject* Object,
 			}
 			else
 			{
-				UE_LOG_REFERENCECACHE(Verbose, TEXT("ObjectReferenceCache::CreateObjectReferenceHandle removing cache for handle %s. Cached Object %s (0x%p). New object %s (0x%p)"),
-					*CachedObject.NetRefHandle.ToString(), *GetNameSafe(CachedObject.ObjectKey), CachedObject.ObjectKey, *GetNameSafe(Object), Object);
+				UE_LOG_REFERENCECACHE(Verbose, TEXT("ObjectReferenceCache::CreateObjectReferenceHandle removing cache for handle %s. Cached Object (0x%p). New object %s (0x%p)"),
+					*CachedObject.NetRefHandle.ToString(), CachedObject.ObjectKey, *GetNameSafe(Object), Object);
 
 				ReferenceHandleToCachedReference.Remove(RefHandle);
 			}
