Disable connectivity hack
Sometimes the connectivity graph was not updated correctly due to some missing events. This bug couldn't be reproduced locally, so this is an attempt to remove the hack.

diff --git a/Engine/Source/Runtime/Experimental/GeometryCollectionEngine/Private/GeometryCollection/GeometryCollectionComponent.cpp b/Engine/Source/Runtime/Experimental/GeometryCollectionEngine/Private/GeometryCollection/GeometryCollectionComponent.cpp
--- a/Engine/Source/Runtime/Experimental/GeometryCollectionEngine/Private/GeometryCollection/GeometryCollectionComponent.cpp
+++ b/Engine/Source/Runtime/Experimental/GeometryCollectionEngine/Private/GeometryCollection/GeometryCollectionComponent.cpp
@@ -123,8 +123,8 @@ FAutoConsoleVariableRef CVarMaxGeometryCollectionAsyncPhysicsTickIdleTimeMs(TEXT
 float GeometryCollectionRemovalMultiplier = 1.0f;
 FAutoConsoleVariableRef CVarGeometryCollectionRemovalTimerMultiplier(TEXT("p.Chaos.GC.RemovalTimerMultiplier"), GeometryCollectionRemovalMultiplier, TEXT("Multiplier for the removal time evaluation ( > 1 : faster removal , > 1 slower"));
 
-// temporary cvar , shoul dbe removed when the root event is no longer no ecessary
-bool GeometryCollectionEmitRootBreakingEvent = true;
+// temporary cvar , should be removed when the root event is no longer no necessary
+bool GeometryCollectionEmitRootBreakingEvent = false;
 FAutoConsoleVariableRef CVarGeometryCollectionEmitRootBreakingEvent(TEXT("p.Chaos.GC.EmitRootBreakingEvent"), GeometryCollectionEmitRootBreakingEvent, TEXT("When true send a breaking event when root is breaking"));
 DEFINE_LOG_CATEGORY_STATIC(UGCC_LOG, Error, All);
 
