I/O - support for uploading on demand I/O store containers to S3 with UnrealPak

- added support for uploading containers marked as on demand directly from UnrealPak.exe
- removed C# based upload logic from automation scripts
- removed on demand I/O store writer since this change reads chunks directly from container files instead of loose files

Example usage:
UnrealPak.exe -Upload=<ContainerPathOrWildcard> -ServiceUrl=<URL> -Bucket=<BucketName> -AccessKey=<Key> -SecretKey=<Key>

Read credentials from an AWS key chain file with the following command line:
-CredentialsFile=<Path> -CredentialsFileKeyName=<EntryName>

Specify -KeepUploadedContainers to prevent UnrealPak to delete on demand containers after the upload has been completed.
Specify -BucketPrefix=<Path> to upload chunks to a specific sub directory within the bucket.

diff --git a/Engine/Source/Developer/IoStoreUtilities/Private/IoStoreUtilities.cpp b/Engine/Source/Developer/IoStoreUtilities/Private/IoStoreUtilities.cpp
--- a/Engine/Source/Developer/IoStoreUtilities/Private/IoStoreUtilities.cpp
+++ b/Engine/Source/Developer/IoStoreUtilities/Private/IoStoreUtilities.cpp
@@ -89,6 +89,7 @@ static const uint64 DefaultCompressionBlockAlignment = 64 << 10;
 static const uint64 DefaultMemoryMappingAlignment = 16 << 10;
 
 static TUniquePtr<FIoStoreReader> CreateIoStoreReader(const TCHAR* Path, const FKeyChain& KeyChain);
+bool UploadIoStoreContainerFiles(const UE::FIoStoreUploadParams& UploadParams, TConstArrayView<FString> ContainerFiles, const FKeyChain& KeyChain);
 
 class FIoStoreChunkDatabase : public IIoStoreWriterReferenceChunkDatabase
 {
@@ -450,13 +451,13 @@ struct FContainerSourceSpec
 	FName Name;
 	FString OutputPath;
 	FString OptionalOutputPath;
-	FString OnDemandOutputPath;
 	FString StageLooseFileRootPath;
 	TArray<FContainerSourceFile> SourceFiles;
 	FString PatchTargetFile;
 	TArray<FString> PatchSourceContainerFiles;
-	FGuid EncryptionKeyOverrideGuid;
+	FString EncryptionKeyOverrideGuid;
 	bool bGenerateDiffPatch = false;
+	bool bOnDemand = false;
 };
 
 struct FCookedFileStatData
@@ -942,6 +943,7 @@ struct FIoStoreArguments
 	bool bCreateDirectoryIndex = true;
 	bool bClusterByOrderFilePriority = false;
 	bool bFileRegions = false;
+	bool bUpload = false;
 	EAssetRegistryWritebackMethod WriteBackMetadataToAssetRegistry = EAssetRegistryWritebackMethod::Disabled;
 
 	FOodleDataCompression::ECompressor ShaderOodleCompressor = FOodleDataCompression::ECompressor::Mermaid;
@@ -962,7 +964,6 @@ struct FContainerTargetSpec
 	FGuid EncryptionKeyGuid;
 	FString OutputPath;
 	FString OptionalSegmentOutputPath;
-	FString OnDemandOutputPath;
 	FString StageLooseFileRootPath;
 	TSharedPtr<IIoStoreWriter> IoStoreWriter;
 	TSharedPtr<IIoStoreWriter> OptionalSegmentIoStoreWriter;
@@ -2552,17 +2553,22 @@ void InitializeContainerTargetsAndPackages(
 	{
 		FContainerTargetSpec* ContainerTarget = AddContainer(ContainerSource.Name, ContainerTargets);
 		ContainerTarget->OutputPath = ContainerSource.OutputPath;
-		ContainerTarget->OnDemandOutputPath = ContainerSource.OnDemandOutputPath;
 		ContainerTarget->StageLooseFileRootPath = ContainerSource.StageLooseFileRootPath;
 		ContainerTarget->bGenerateDiffPatch = ContainerSource.bGenerateDiffPatch;
+
+		if (ContainerSource.bOnDemand)
+		{
+			ContainerTarget->ContainerFlags |= EIoContainerFlags::OnDemand;
+		}
+
 		if (Arguments.bSign)
 		{
 			ContainerTarget->ContainerFlags |= EIoContainerFlags::Signed;
 		}
 
-		if (!ContainerTarget->EncryptionKeyGuid.IsValid())
+		if (!ContainerTarget->EncryptionKeyGuid.IsValid() && !ContainerSource.EncryptionKeyOverrideGuid.IsEmpty())
 		{
-			ContainerTarget->EncryptionKeyGuid = ContainerSource.EncryptionKeyOverrideGuid;
+			FGuid::Parse(ContainerSource.EncryptionKeyOverrideGuid, ContainerTarget->EncryptionKeyGuid);
 		}
 
 		ContainerTarget->PatchSourceReaders = CreatePatchSourceReaders(ContainerSource.PatchSourceContainerFiles, Arguments);
@@ -2609,6 +2615,12 @@ void InitializeContainerTargetsAndPackages(
 				ContainerTarget->TargetFiles.Emplace(MoveTemp(TargetFile));
 			}
 
+			if (EnumHasAnyFlags(ContainerTarget->ContainerFlags, EIoContainerFlags::OnDemand) &&
+				!EnumHasAnyFlags(ContainerTarget->ContainerFlags, EIoContainerFlags::Encrypted))
+			{
+				UE_LOG(LogIoStore, Warning, TEXT("Container '%s' set as on demand but is not encrypted, please revisit encryption strategy before publisihing to public CDN's"), *ContainerSource.Name.ToString());
+			}
+
 			if (bHasOptionalSegmentPackages)
 			{
 				if (ContainerSource.OptionalOutputPath.IsEmpty())
@@ -3753,9 +3765,9 @@ int32 CreateTarget(const FIoStoreArguments& Arguments, const FIoStoreWriterSetti
 		FIoStatus IoStatus = IoStoreWriterContext->Initialize(GeneralIoWriterSettings);
 		check(IoStatus.IsOk());
 	}
+	TArray<FString> OnDemandContainers;
 	TArray<TSharedPtr<IIoStoreWriter>> IoStoreWriters;
 	TSharedPtr<IIoStoreWriter> GlobalIoStoreWriter;
-	TUniquePtr<UE::IOnDemandIoStoreWriter> OnDemandIoStoreWriter;
 	{
 		IOSTORE_CPU_SCOPE(InitializeWriters);
 		if (!Arguments.IsDLC())
@@ -3775,7 +3787,7 @@ int32 CreateTarget(const FIoStoreArguments& Arguments, const FIoStoreWriterSetti
 			IOSTORE_CPU_SCOPE(InitializeWriter);
 			check(ContainerTarget->ContainerId.IsValid());
 
-			if (ContainerTarget->OutputPath.IsEmpty() && ContainerTarget->OnDemandOutputPath.IsEmpty())
+			if (ContainerTarget->OutputPath.IsEmpty())
 			{
 				continue;
 			}
@@ -3787,49 +3799,27 @@ int32 CreateTarget(const FIoStoreArguments& Arguments, const FIoStoreWriterSetti
 				ContainerTarget->IoStoreWriter = MakeLooseFilesIoStoreWriter(WriterSettings);
 				IoStoreWriters.Add(ContainerTarget->IoStoreWriter);
 			}
-			else if (!ContainerTarget->OnDemandOutputPath.IsEmpty())
+			else
 			{
-				if (!OnDemandIoStoreWriter.IsValid())
+				FIoContainerSettings ContainerSettings;
+				ContainerSettings.ContainerId = ContainerTarget->ContainerId;
+				if (Arguments.bCreateDirectoryIndex)
 				{
-					FIoStoreWriterSettings WriterSettings;
-					WriterSettings.CompressionMethod = NAME_Oodle;
-					OnDemandIoStoreWriter = UE::MakeOnDemandIoStoreWriter(WriterSettings, ContainerTarget->OnDemandOutputPath);
+					ContainerSettings.ContainerFlags = ContainerTarget->ContainerFlags | EIoContainerFlags::Indexed;
 				}
-
-				FIoContainerSettings ContainerSettings;
-				ContainerSettings.ContainerFlags = EIoContainerFlags::OnDemand | EIoContainerFlags::Compressed;
-
 				if (EnumHasAnyFlags(ContainerTarget->ContainerFlags, EIoContainerFlags::Encrypted))
 				{
 					if (const FNamedAESKey* Key = Arguments.KeyChain.GetEncryptionKeys().Find(ContainerTarget->EncryptionKeyGuid))
 					{
 						ContainerSettings.EncryptionKeyGuid = ContainerTarget->EncryptionKeyGuid;
 						ContainerSettings.EncryptionKey = Key->Key;
-						ContainerSettings.ContainerFlags |= EIoContainerFlags::Encrypted;
 					}
 					else
 					{
-						UE_LOG(LogIoStore, Warning, TEXT("Failed to find encryption key '%s' for container '%s'"), *ContainerTarget->EncryptionKeyGuid.ToString(), *ContainerTarget->Name.ToString());
+						UE_LOG(LogIoStore, Error, TEXT("Failed to find encryption key '%s'"), *ContainerTarget->EncryptionKeyGuid.ToString());
+						return -1;
 					}
 				}
-				ContainerTarget->IoStoreWriter = OnDemandIoStoreWriter->CreateContainer(ContainerTarget->Name.ToString(), ContainerSettings);
-				IoStoreWriters.Add(ContainerTarget->IoStoreWriter);
-			}
-			else
-			{
-				FIoContainerSettings ContainerSettings;
-				ContainerSettings.ContainerId = ContainerTarget->ContainerId;
-				if (Arguments.bCreateDirectoryIndex)
-				{
-					ContainerSettings.ContainerFlags = ContainerTarget->ContainerFlags | EIoContainerFlags::Indexed;
-				}
-				if (EnumHasAnyFlags(ContainerTarget->ContainerFlags, EIoContainerFlags::Encrypted))
-				{
-					const FNamedAESKey* Key = Arguments.KeyChain.GetEncryptionKeys().Find(ContainerTarget->EncryptionKeyGuid);
-					check(Key);
-					ContainerSettings.EncryptionKeyGuid = ContainerTarget->EncryptionKeyGuid;
-					ContainerSettings.EncryptionKey = Key->Key;
-				}
 				if (EnumHasAnyFlags(ContainerTarget->ContainerFlags, EIoContainerFlags::Signed))
 				{
 					ContainerSettings.SigningKey = Arguments.KeyChain.GetSigningKey();
@@ -3847,6 +3837,10 @@ int32 CreateTarget(const FIoStoreArguments& Arguments, const FIoStoreWriterSetti
 					ContainerTarget->OptionalSegmentIoStoreWriter->SetReferenceChunkDatabase(ChunkDatabase);
 					IoStoreWriters.Add(ContainerTarget->OptionalSegmentIoStoreWriter);
 				}
+				if (EnumHasAnyFlags(ContainerTarget->ContainerFlags, EIoContainerFlags::OnDemand))
+				{
+					OnDemandContainers.Add(ContainerTarget->OutputPath);
+				}
 			}
 		}
 	}
@@ -4003,13 +3997,6 @@ int32 CreateTarget(const FIoStoreArguments& Arguments, const FIoStoreWriterSetti
 		GlobalIoStoreWriter->Append(CreateIoChunkId(0, 0, EIoChunkType::ScriptObjects), ScriptObjectsBuffer, WriteOptions);
 	}
 
-	if (OnDemandIoStoreWriter.IsValid())
-	{
-		UE_LOG(LogIoStore, Display, TEXT("Serializing ondemand content..."));
-		IOSTORE_CPU_SCOPE(SerializingContentOnDemand);
-		OnDemandIoStoreWriter->Flush();
-	}
-
 	UE_LOG(LogIoStore, Display, TEXT("Serializing container(s)..."));
 	{
 		IOSTORE_CPU_SCOPE(Serializing);
@@ -4345,6 +4332,21 @@ int32 CreateTarget(const FIoStoreArguments& Arguments, const FIoStoreWriterSetti
 		FTaskGraphInterface::Get().WaitUntilTaskCompletes(WriteCsvFileTask);
 	}
 
+	if (Arguments.bUpload && OnDemandContainers.IsEmpty() == false)
+	{
+		TIoStatusOr<UE::FIoStoreUploadParams> UploadParams = UE::FIoStoreUploadParams::Parse(FCommandLine::Get());
+		if (UploadParams.IsOk() == false)
+		{
+			UE_LOG(LogIoStore, Warning, TEXT("Skipping upload of container file(s), reason '%s'"), *UploadParams.Status().ToString());
+			return 0;
+		}
+
+		if (UploadIoStoreContainerFiles(UploadParams.ConsumeValueOrDie(), OnDemandContainers, Arguments.KeyChain) == false)
+		{
+			return -1;
+		}
+	}
+
 	return 0;
 }
 
@@ -7432,12 +7434,7 @@ bool ParseContainerGenerationArguments(FIoStoreArguments& Arguments, FIoStoreWri
 			{
 				ContainerSpec.OutputPath = FPaths::ChangeExtension(ContainerSpec.OutputPath, TEXT(""));
 			}
-			else if (!FParse::Value(*Command, TEXT("OnDemandOutput="), ContainerSpec.OnDemandOutputPath))
-			{
-				UE_LOG(LogIoStore, Error, TEXT("Output argument missing from command '%s'"), *Command);
-				return false;
-			}
-
+			ContainerSpec.bOnDemand = FParse::Param(*Command, TEXT("OnDemand"));
 			FParse::Value(*Command, TEXT("OptionalOutput="), ContainerSpec.OptionalOutputPath);
 
 			FParse::Value(*Command, TEXT("StageLooseFileRootPath="), ContainerSpec.StageLooseFileRootPath);
@@ -7472,12 +7469,7 @@ bool ParseContainerGenerationArguments(FIoStoreArguments& Arguments, FIoStoreWri
 					UE_LOG(LogIoStore, Error, TEXT("Failed to parse Pak response file '%s'"), *ResponseFilePath);
 					return false;
 				}
-
-				FString EncryptionKeyOverrideGuidString;
-				if (FParse::Value(*Command, TEXT("EncryptionKeyOverrideGuid="), EncryptionKeyOverrideGuidString))
-				{
-					FGuid::Parse(EncryptionKeyOverrideGuidString, ContainerSpec.EncryptionKeyOverrideGuid);
-				}
+				FParse::Value(*Command, TEXT("EncryptionKeyOverrideGuid="), ContainerSpec.EncryptionKeyOverrideGuid);
 			}
 		}
 	}
@@ -7725,6 +7717,7 @@ int32 CreateIoStoreContainerFiles(const TCHAR* CmdLine)
 		
 		Arguments.DLCName = FPaths::GetBaseFilename(*Arguments.DLCPluginPath);
 		Arguments.bRemapPluginContentToGame = FParse::Param(FCommandLine::Get(), TEXT("RemapPluginContentToGame"));
+		Arguments.bUpload = FParse::Param(FCommandLine::Get(), TEXT("Upload"));
 
 		UE_LOG(LogIoStore, Display, TEXT("DLC: '%s'"), *Arguments.DLCPluginPath);
 		UE_LOG(LogIoStore, Display, TEXT("Remapping plugin content to game: '%s'"), Arguments.bRemapPluginContentToGame ? TEXT("True") : TEXT("False"));
@@ -7771,6 +7764,7 @@ int32 CreateIoStoreContainerFiles(const TCHAR* CmdLine)
 	else if (FParse::Value(FCommandLine::Get(), TEXT("CreateGlobalContainer="), Arguments.GlobalContainerPath))
 	{
 		Arguments.GlobalContainerPath = FPaths::ChangeExtension(Arguments.GlobalContainerPath, TEXT(""));
+		Arguments.bUpload = FParse::Param(FCommandLine::Get(), TEXT("Upload"));
 
 		if (!ParseContainerGenerationArguments(Arguments, WriterSettings))
 		{
@@ -7888,3 +7882,139 @@ int32 CreateIoStoreContainerFiles(const TCHAR* CmdLine)
 
 	return CreateTarget(Arguments, WriterSettings);
 }
+
+bool UploadIoStoreContainerFiles(const UE::FIoStoreUploadParams& UploadParams, TConstArrayView<FString> ContainerFiles, const FKeyChain& KeyChain)
+{
+	TMap<FGuid, FAES::FAESKey> EncryptionKeys;
+	for (const TPair<FGuid, FNamedAESKey>& KeyPair: KeyChain.GetEncryptionKeys())
+	{
+		EncryptionKeys.Add(KeyPair.Key, KeyPair.Value.Key);
+	}
+
+	TIoStatusOr<UE::FIoStoreUploadResult> Result = UE::UploadContainerFiles(UploadParams, ContainerFiles, EncryptionKeys);
+	if (Result.IsOk() == false)
+	{
+		UE_LOG(LogIoStore, Error, TEXT("Failed to upload container file(s), reason '%s'"), *Result.Status().ToString());
+		return false;
+	}
+
+	UE::FIoStoreUploadResult UploadResult = Result.ConsumeValueOrDie();
+
+	FString ConfigFilePath;
+	if (FParse::Value(FCommandLine::Get(), TEXT("ConfigFilePath="), ConfigFilePath))
+	{
+		FStringBuilderBase Sb;
+		Sb << TEXT("[Endpoint]") << TEXT("\r\n");
+
+		FString DistributionUrl;
+		if (FParse::Value(FCommandLine::Get(), TEXT("DistributionUrl="), DistributionUrl))
+		{
+			Sb << TEXT("DistributionUrl=\"") << DistributionUrl << TEXT("\"\r\n");
+		}
+		else
+		{
+			Sb << TEXT("ServiceUrl=\"") << UploadParams.ServiceUrl << TEXT("\"\r\n");
+		}
+
+		// Temporary solution to get replays working with encrypted on demand content
+		{
+			FString EncryptionKeyName;
+			if (FParse::Value(FCommandLine::Get(), TEXT("OnDemandEncryptionKeyName="), EncryptionKeyName))
+			{
+				TOptional<FNamedAESKey> EncryptionKey;
+				for (const TPair<FGuid, FNamedAESKey>& KeyPair: KeyChain.GetEncryptionKeys())
+				{
+					if (KeyPair.Value.Name.Compare(EncryptionKeyName, ESearchCase::IgnoreCase) == 0)
+					{
+						EncryptionKey.Emplace(KeyPair.Value);
+					}
+				}
+
+				if (EncryptionKey)
+				{
+					FString KeyString = FBase64::Encode(EncryptionKey.GetValue().Key.Key, FAES::FAESKey::KeySize);
+					Sb << TEXT("ContentKey=\"") << EncryptionKey.GetValue().Guid.ToString() << TEXT(":") << KeyString << TEXT("\"\r\n");
+				}
+				else
+				{
+					UE_LOG(LogIoStore, Warning, TEXT("Failed to encryption key '%s' in key chain"), *EncryptionKeyName);
+				}
+			}
+		}
+
+		if (DistributionUrl.IsEmpty())
+		{
+			// Append the bucket name when using a local/custom service
+			Sb << TEXT("TocPath=\"") << UploadParams.Bucket / UploadResult.TocPath << TEXT("\"\r\n");
+		}
+		else
+		{
+			Sb << TEXT("TocPath=\"") << UploadResult.TocPath << TEXT("\"\r\n");
+		}
+
+		UE_LOG(LogIoStore, Display, TEXT("Saving on demand config file '%s'"), *ConfigFilePath);
+		if (FFileHelper::SaveStringToFile(Sb.ToString(), *ConfigFilePath) == false)
+		{
+			UE_LOG(LogIoStore, Error, TEXT("Failed to save on demand config file '%s'"), *ConfigFilePath);
+		}
+	}
+
+	return Result.IsOk();
+}
+
+bool UploadIoStoreContainerFiles(const TCHAR* ContainerPathOrWildcard)
+{
+	check(ContainerPathOrWildcard);
+
+	FKeyChain KeyChain;
+	LoadKeyChain(FCommandLine::Get(), KeyChain);
+
+	TIoStatusOr<UE::FIoStoreUploadParams> UploadParams = UE::FIoStoreUploadParams::Parse(FCommandLine::Get());
+	if (UploadParams.IsOk() == false)
+	{
+		UE_LOG(LogIoStore, Error, TEXT("Failed to upload container file(s), reason '%s'"), *UploadParams.Status().ToString());
+		return false;
+	}
+
+	TArray<FString> ContainerFiles;
+	{
+		if (IFileManager::Get().FileExists(ContainerPathOrWildcard))
+		{
+			ContainerFiles.Add(ContainerPathOrWildcard);
+		}
+		else if (IFileManager::Get().DirectoryExists(ContainerPathOrWildcard))
+		{
+			FString Directory = ContainerPathOrWildcard;
+			FPaths::NormalizeDirectoryName(Directory);
+
+			TArray<FString> FoundContainerFiles;
+			IFileManager::Get().FindFiles(FoundContainerFiles, *(Directory / TEXT("*.utoc")), true, false);
+
+			for (const FString& Filename : FoundContainerFiles)
+			{
+				ContainerFiles.Emplace(Directory / Filename);
+			}
+		}
+		else
+		{
+			FString Directory = FPaths::GetPath(ContainerPathOrWildcard);
+			FPaths::NormalizeDirectoryName(Directory);
+
+			TArray<FString> FoundContainerFiles;
+			IFileManager::Get().FindFiles(FoundContainerFiles, ContainerPathOrWildcard, true, false);
+
+			for (const FString& Filename : FoundContainerFiles)
+			{
+				ContainerFiles.Emplace(Directory / Filename);
+			}
+		}
+	}
+
+	if (ContainerFiles.IsEmpty())
+	{
+		UE_LOG(LogIoStore, Error, TEXT("Failed to find container file(s) '%s'"), ContainerPathOrWildcard);
+		return false;
+	}
+
+	return UploadIoStoreContainerFiles(UploadParams.ConsumeValueOrDie(), ContainerFiles, KeyChain);
+}
diff --git a/Engine/Source/Developer/IoStoreUtilities/Public/IoStoreUtilities.h b/Engine/Source/Developer/IoStoreUtilities/Public/IoStoreUtilities.h
--- a/Engine/Source/Developer/IoStoreUtilities/Public/IoStoreUtilities.h
+++ b/Engine/Source/Developer/IoStoreUtilities/Public/IoStoreUtilities.h
@@ -52,3 +52,5 @@ IOSTOREUTILITIES_API bool ProcessFilesFromIoStoreContainer(
 	int32 MaxConcurrentReaders);
 
 IOSTOREUTILITIES_API bool SignIoStoreContainer(const TCHAR* InContainerFilename, const FRSAKeyHandle InSigningKey);
+
+IOSTOREUTILITIES_API bool UploadIoStoreContainerFiles(const TCHAR* Params); 
diff --git a/Engine/Source/Developer/PakFileUtilities/Private/PakFileUtilities.cpp b/Engine/Source/Developer/PakFileUtilities/Private/PakFileUtilities.cpp
--- a/Engine/Source/Developer/PakFileUtilities/Private/PakFileUtilities.cpp
+++ b/Engine/Source/Developer/PakFileUtilities/Private/PakFileUtilities.cpp
@@ -5450,11 +5450,18 @@ bool MakeBinaryConfig(const TCHAR* CmdLine)
  */
 bool ExecuteUnrealPak(const TCHAR* CmdLine)
 {
-	FString IoStoreArg;
-	if (FParse::Value(CmdLine, TEXT("-CreateGlobalContainer="), IoStoreArg) ||
-		FParse::Value(CmdLine, TEXT("-CreateDLCContainer="), IoStoreArg))
 	{
-		return CreateIoStoreContainerFiles(CmdLine) == 0;
+		FString IoStoreArg;
+		if (FParse::Value(CmdLine, TEXT("-CreateGlobalContainer="), IoStoreArg) ||
+			FParse::Value(CmdLine, TEXT("-CreateDLCContainer="), IoStoreArg))
+		{
+			return CreateIoStoreContainerFiles(CmdLine) == 0;
+		}
+
+		if (FParse::Value(CmdLine, TEXT("-Upload="), IoStoreArg))
+		{
+			return UploadIoStoreContainerFiles(*IoStoreArg) == 0;
+		}
 	}
 
 	// Parse all the non-option arguments from the command line
diff --git a/Engine/Source/Developer/S3Client/Internal/S3/S3Client.h b/Engine/Source/Developer/S3Client/Internal/S3/S3Client.h
--- a/Engine/Source/Developer/S3Client/Internal/S3/S3Client.h
+++ b/Engine/Source/Developer/S3Client/Internal/S3/S3Client.h
@@ -194,6 +194,8 @@ public:
 	FS3ListObjectResponse ListObjects(const FS3ListObjectsRequest& Request);
 	/** Upload an object described by the specified request parameters. */
 	FS3PutObjectResponse PutObject(const FS3PutObjectRequest& Request);
+	/** Retries uploading an object until succeeded or max attempts has been reached. */
+	FS3PutObjectResponse TryPutObject(const FS3PutObjectRequest& Request, int32 MaxAttempts = 3, float Delay = 0.5f);
 
 private:
 	void Setup(FS3Request& Request);
diff --git a/Engine/Source/Developer/S3Client/Private/S3Client.cpp b/Engine/Source/Developer/S3Client/Private/S3Client.cpp
--- a/Engine/Source/Developer/S3Client/Private/S3Client.cpp
+++ b/Engine/Source/Developer/S3Client/Private/S3Client.cpp
@@ -2,6 +2,7 @@
 
 #include "S3/S3Client.h"
 #include "Containers/ChunkedArray.h"
+#include "HAL/PlatformProcess.h"
 #include "Misc/ConfigCacheIni.h"
 #include "Misc/ScopeLock.h"
 
@@ -308,11 +309,11 @@ FS3CredentialsProfileStore FS3CredentialsProfileStore::FromFile(const FString& F
 			FS3ClientCredentials Credentials;
 			if (const FConfigValue* SessionToken = Section.Find(TEXT("aws_session_token")))
 			{
-				Credentials = FS3ClientCredentials(AccessKey ->GetValue(), SecretKey->GetValue(), SessionToken->GetValue());
+				Credentials = FS3ClientCredentials(AccessKey->GetValue(), SecretKey->GetValue(), SessionToken->GetValue());
 			}
 			else
 			{
-				Credentials = FS3ClientCredentials(AccessKey ->GetValue(), SecretKey->GetValue());
+				Credentials = FS3ClientCredentials(AccessKey->GetValue(), SecretKey->GetValue());
 			}
 
 			ProfileStore.Credentials.Add(ProfileName, MoveTemp(Credentials));
@@ -685,6 +686,11 @@ FS3Client::FS3Client(const FS3ClientConfig& ClientConfig, const FS3ClientCredent
 	, Credentials(ClientCredentials)
 {
 	ConnectionPool = MakeUnique<FS3Client::FConnectionPool>();
+
+	if (!Config.Region.IsEmpty())
+	{
+		Config.ServiceUrl = TEXT("https://s3.amazonaws.com");
+	}
 }
 
 FS3Client::~FS3Client()
@@ -719,6 +725,22 @@ FS3PutObjectResponse FS3Client::PutObject(const FS3PutObjectRequest& PutRequest)
 	return Request.Perform(FS3Request::EMethod::Put, Url.ToString(), FSharedBuffer::MakeView(PutRequest.ObjectData));
 }
 
+FS3PutObjectResponse FS3Client::TryPutObject(const FS3PutObjectRequest& Request, int32 MaxAttempts, float Delay)
+{
+	FS3PutObjectResponse Response;
+	for (int32 Attempt = 0; Attempt < MaxAttempts; ++Attempt)
+	{
+		Response = PutObject(Request);
+		if (Response.IsOk())
+		{
+			break;
+		}
+		FPlatformProcess::Sleep(Delay);
+	}
+
+	return Response;
+}
+
 FS3ListObjectResponse FS3Client::ListObjects(const FS3ListObjectsRequest& ListRequest)
 {
 	TAnsiStringBuilder<256> Url;
diff --git a/Engine/Source/Programs/AutomationTool/Scripts/CopyBuildToStagingDirectory.Automation.cs b/Engine/Source/Programs/AutomationTool/Scripts/CopyBuildToStagingDirectory.Automation.cs
--- a/Engine/Source/Programs/AutomationTool/Scripts/CopyBuildToStagingDirectory.Automation.cs
+++ b/Engine/Source/Programs/AutomationTool/Scripts/CopyBuildToStagingDirectory.Automation.cs
@@ -281,7 +281,7 @@ namespace AutomationScripts
 			string ContainerName,
 			FileReference PakOutputLocation,
 			DirectoryReference OptionalOutputLocation,
-			DirectoryReference OnDemandOutputLocation,
+			bool bOnDemand,
 			bool bCompressed,
 			DirectoryReference OptionalStageLooseFileRootPath,
 			EncryptionAndSigning.CryptoSettings CryptoSettings,
@@ -291,16 +291,7 @@ namespace AutomationScripts
 			bool bIsDLC)
 		{
 			StringBuilder CmdLine = new StringBuilder();
-
-			if (OnDemandOutputLocation != null)
-			{
-				FileReference OnDemandPakOutPutLocation = FileReference.Combine(OnDemandOutputLocation, ContainerName);
-				CmdLine.AppendFormat(" -OnDemandOutput={0}", MakePathSafeToUseWithCommandLine(OnDemandOutputLocation.FullName));
-			}
-			else
-			{
-				CmdLine.AppendFormat("-Output={0}", MakePathSafeToUseWithCommandLine(Path.ChangeExtension(PakOutputLocation.FullName, ".utoc")));
-			}
+			CmdLine.AppendFormat("-Output={0}", MakePathSafeToUseWithCommandLine(Path.ChangeExtension(PakOutputLocation.FullName, ".utoc")));
 
 			if (OptionalOutputLocation != null)
 			{
@@ -347,6 +338,11 @@ namespace AutomationScripts
 				}
 			}
 
+			if (bOnDemand)
+			{
+				CmdLine.AppendFormat(" -OnDemand");
+			}
+
 			return CmdLine.ToString();
 		}
 
@@ -3260,13 +3256,11 @@ namespace AutomationScripts
 				OutputFilename = OutputFilename + PostFix;
 
 				StagedFileReference OutputRelativeLocation;
-				StagedDirectoryReference OnDemandOutputRelativeLocation;
 				if (Params.HasDLCName)
 				{
 					if (Params.DLCOverrideStagedSubDir != null)
 					{
 						OutputRelativeLocation = StagedFileReference.Combine(SC.RelativeProjectRootForStage, Params.DLCOverrideStagedSubDir, "Content", "Paks", Params.DLCFile.GetFileNameWithoutExtension() + OutputFilename + ".pak");
-						OnDemandOutputRelativeLocation = StagedDirectoryReference.Combine(SC.RelativeProjectRootForStage, Params.DLCOverrideStagedSubDir, "ContentOnDemand");
 					}
 					else
 					{
@@ -3281,13 +3275,11 @@ namespace AutomationScripts
 							PluginSubdirectory = Params.DLCFile.GetFileNameWithoutAnyExtensions();
 						}
 						OutputRelativeLocation = StagedFileReference.Combine(SC.RelativeProjectRootForStage, PluginSubdirectory, "Content", "Paks", SC.FinalCookPlatform, Params.DLCFile.GetFileNameWithoutExtension() + OutputFilename + ".pak");
-						OnDemandOutputRelativeLocation = StagedDirectoryReference.Combine(SC.RelativeProjectRootForStage, PluginSubdirectory, "ContentOnDemand");
 					}
 				}
 				else
 				{
 					OutputRelativeLocation = StagedFileReference.Combine(SC.RelativeProjectRootForStage, "Content", "Paks", OutputFilename + OutputFilenameExtension);
-					OnDemandOutputRelativeLocation = StagedDirectoryReference.Combine(SC.RelativeProjectRootForStage, "ContentOnDemand");
 				}
 				if (SC.StageTargetPlatform.DeployLowerCaseFilenames(StagedFileType.UFS))
 				{
@@ -3295,12 +3287,7 @@ namespace AutomationScripts
 				}
 				OutputRelativeLocation = SC.StageTargetPlatform.Remap(OutputRelativeLocation);
 				
-				//TODO: Do we need to remap the ondemand output location?
-				//OnDemandOutputRelativeLocation = SC.StageTargetPlatform.Remap(OnDemandOutputRelativeLocation);
-
 				FileReference OutputLocation = FileReference.Combine(SC.RuntimeRootDir, OutputRelativeLocation.Name);
-				DirectoryReference OnDemandOutputLocation = DirectoryReference.Combine(SC.RuntimeRootDir, OnDemandOutputRelativeLocation.Name);
-
 				bool bCopiedExistingPak = false;
 
 				if (!PakParams.bOnDemand && !PakParams.bStageLoose && SC.StageTargetPlatform != SC.CookSourcePlatform && !Params.IgnorePaksFromDifferentCookSource)
@@ -3441,7 +3428,7 @@ namespace AutomationScripts
 								OutputLocation.GetFileNameWithoutAnyExtensions(),
 								OutputLocation,
 								SC.OptionalFileStageDirectory,
-								PakParams.bOnDemand ? OnDemandOutputLocation : null,
+								PakParams.bOnDemand,
 								bCompressContainers,
 								StageLooseFileRootPath,
 								CryptoSettings,
diff --git a/Engine/Source/Runtime/Core/Private/IO/IoStore.cpp b/Engine/Source/Runtime/Core/Private/IO/IoStore.cpp
--- a/Engine/Source/Runtime/Core/Private/IO/IoStore.cpp
+++ b/Engine/Source/Runtime/Core/Private/IO/IoStore.cpp
@@ -2166,6 +2166,18 @@ public:
 			}
 		}
 
+		~FContainerFileAccess()
+		{
+			for (int32 Index = 0; Index < NumHandlesPerFile; Index++)
+			{
+				if (Handle[Index] != nullptr)
+				{
+					delete Handle[Index];
+					Handle[Index] = nullptr;
+				}
+			}
+		}
+
 		bool IsValid() const { return bValid; }
 	};
 	
@@ -2627,7 +2639,7 @@ public:
 		return UncompressedBuffer;
 	}
 
-	TIoStatusOr<FIoStoreCompressedReadResult> ReadCompressed(const FIoChunkId& ChunkId, const FIoReadOptions& Options) const
+	TIoStatusOr<FIoStoreCompressedReadResult> ReadCompressed(const FIoChunkId& ChunkId, const FIoReadOptions& Options, bool bDecrypt) const
 	{
 		TRACE_CPUPROFILER_EVENT_SCOPE(ReadChunkCompressed);
 
@@ -2707,7 +2719,7 @@ public:
 			}
 		}
 
-		if (EnumHasAnyFlags(TocResource.Header.ContainerFlags, EIoContainerFlags::Encrypted))
+		if (bDecrypt && EnumHasAnyFlags(TocResource.Header.ContainerFlags, EIoContainerFlags::Encrypted) )
 		{
 			TRACE_CPUPROFILER_EVENT_SCOPE(Decrypt);
 
@@ -2867,9 +2879,9 @@ TIoStatusOr<FIoBuffer> FIoStoreReader::Read(const FIoChunkId& Chunk, const FIoRe
 	return Impl->Read(Chunk, Options);
 }
 
-TIoStatusOr<FIoStoreCompressedReadResult> FIoStoreReader::ReadCompressed(const FIoChunkId& Chunk, const FIoReadOptions& Options) const
+TIoStatusOr<FIoStoreCompressedReadResult> FIoStoreReader::ReadCompressed(const FIoChunkId& Chunk, const FIoReadOptions& Options, bool bDecrypt) const
 {
-	return Impl->ReadCompressed(Chunk, Options);
+	return Impl->ReadCompressed(Chunk, Options, bDecrypt);
 }
 
 UE::Tasks::TTask<TIoStatusOr<FIoBuffer>> FIoStoreReader::ReadAsync(const FIoChunkId& Chunk, const FIoReadOptions& Options) const
diff --git a/Engine/Source/Runtime/Core/Public/IO/IoDispatcher.h b/Engine/Source/Runtime/Core/Public/IO/IoDispatcher.h
--- a/Engine/Source/Runtime/Core/Public/IO/IoDispatcher.h
+++ b/Engine/Source/Runtime/Core/Public/IO/IoDispatcher.h
@@ -123,6 +123,13 @@ public:
 		return BytesToHex(Hash, 20);
 	}
 
+	FIoHash ToIoHash() const
+	{
+		FIoHash IoHash;
+		FMemory::Memcpy(IoHash.GetBytes(), Hash, sizeof(FIoHash));
+		return IoHash;
+	}
+
 	static FIoChunkHash CreateFromIoHash(const FIoHash& IoHash)
 	{
 		FIoChunkHash Result;
@@ -725,7 +732,7 @@ public:
 	// in FIoStoreCompressedReadResult::FIoBuffer as a contiguous buffer, however each block is padded during encryption, so
 	// either use FIoStoreCompressedBlockInfo::AlignedSize to advance through the buffer, or use FIoStoreCompressedBlockInfo::OffsetInBuffer
 	// directly.
-	CORE_API TIoStatusOr<FIoStoreCompressedReadResult> ReadCompressed(const FIoChunkId& Chunk, const FIoReadOptions& Options) const;
+	CORE_API TIoStatusOr<FIoStoreCompressedReadResult> ReadCompressed(const FIoChunkId& Chunk, const FIoReadOptions& Options, bool bDecrypt = true) const;
 
 	CORE_API const FIoDirectoryIndexReader& GetDirectoryIndexReader() const;
 
diff --git a/Engine/Source/Runtime/Experimental/IoStoreOnDemand/IoStoreOnDemand.Build.cs b/Engine/Source/Runtime/Experimental/IoStoreOnDemand/IoStoreOnDemand.Build.cs
--- a/Engine/Source/Runtime/Experimental/IoStoreOnDemand/IoStoreOnDemand.Build.cs
+++ b/Engine/Source/Runtime/Experimental/IoStoreOnDemand/IoStoreOnDemand.Build.cs
@@ -9,5 +9,10 @@ public class IoStoreOnDemand : ModuleRules
         PublicDependencyModuleNames.Add("Core");
         PrivateDependencyModuleNames.AddRange(new string[] { "HTTP", "Json" });
 		bAllowConfidentialPlatformDefines = true;
+
+		if (Target.Type == TargetType.Editor || Target.Type == TargetType.Program)
+		{
+			PrivateDependencyModuleNames.AddRange(new string[] { "S3Client", "RSA" });
+		}
     }
 }
diff --git a/Engine/Source/Runtime/Experimental/IoStoreOnDemand/Private/IoStoreOnDemand.cpp b/Engine/Source/Runtime/Experimental/IoStoreOnDemand/Private/IoStoreOnDemand.cpp
--- a/Engine/Source/Runtime/Experimental/IoStoreOnDemand/Private/IoStoreOnDemand.cpp
+++ b/Engine/Source/Runtime/Experimental/IoStoreOnDemand/Private/IoStoreOnDemand.cpp
@@ -11,12 +11,19 @@
 #include "Misc/Base64.h"
 #include "Misc/CommandLine.h"
 #include "Misc/ConfigCacheIni.h"
+#include "Misc/Parse.h"
 #include "Modules/ModuleManager.h"
 #include "Serialization/CompactBinaryWriter.h"
 #include "Serialization/CompactBinarySerialization.h"
 #include "Serialization/LargeMemoryWriter.h"
 #include "String/LexFromString.h"
 
+#if (IS_PROGRAM || WITH_EDITOR)
+#include "Async/Async.h"
+#include "Misc/FileHelper.h"
+#include "S3/S3Client.h"
+#endif // (IS_PROGRAM || WITH_EDITOR)
+
 DEFINE_LOG_CATEGORY(LogIas);
 
 namespace UE::IO::Private
@@ -226,6 +233,8 @@ FCbWriter& operator<<(FCbWriter& Writer, const FOnDemandTocContainerEntry& Conta
 	}
 	Writer.EndArray();
 
+	Writer.AddHash(UTF8TEXTVIEW("UTocHash"), ContainerEntry.UTocHash);
+
 	Writer.EndObject();
 
 	return Writer;
@@ -262,6 +271,8 @@ bool LoadFromCompactBinary(FCbFieldView Field, FOnDemandTocContainerEntry& OutCo
 			OutContainer.BlockHashes.Add(ArrayField.AsHash());
 		}
 
+		OutContainer.UTocHash = Obj["UTocHash"].AsHash();
+
 		return true;
 	}
 
@@ -343,6 +354,498 @@ bool LoadFromCompactBinary(FCbFieldView Field, FOnDemandToc& OutToc)
 	return false;
 }
 
+////////////////////////////////////////////////////////////////////////////////
+#if (IS_PROGRAM || WITH_EDITOR)
+class FS3UploadQueue
+{
+public:
+	FS3UploadQueue(FS3Client& Client, const FString& Bucket, int32 ThreadCount);
+	bool Enqueue(const FString& Key, FIoBuffer Payload);
+	bool Flush();
+
+private:
+	void ThreadEntry();
+
+	struct FQueueEntry
+	{
+		FString Key;
+		FIoBuffer Payload;
+	};
+
+	FS3Client& Client;
+	TArray<TFuture<void>> Threads;
+	FString Bucket;
+	FCriticalSection CriticalSection;
+	TQueue<FQueueEntry> Queue;
+	FEventRef WakeUpEvent;
+	FEventRef UploadCompleteEvent;
+	std::atomic_int32_t ConcurrentUploads{0};
+	std::atomic_int32_t ActiveThreadCount{0};
+	std::atomic_int32_t ErrorCount{0};
+	std::atomic_bool bCompleteAdding{false};
+};
+
+FS3UploadQueue::FS3UploadQueue(FS3Client& InClient, const FString& InBucket, int32 ThreadCount)
+	: Client(InClient)
+	, Bucket(InBucket)
+{
+	ActiveThreadCount = ThreadCount;
+	for (int32 Idx = 0; Idx < ThreadCount; ++Idx)
+	{
+		Threads.Add(AsyncThread([this]()
+		{
+			ThreadEntry();
+		}));
+	}
+}
+
+bool FS3UploadQueue::Enqueue(const FString& Key, FIoBuffer Payload)
+{
+	if (ActiveThreadCount == 0)
+	{
+		return false;
+	}
+
+	for(;;)
+	{
+		bool bEnqueued = false;
+		{
+			FScopeLock _(&CriticalSection);
+			if (ConcurrentUploads < Threads.Num())
+			{
+				bEnqueued = Queue.Enqueue(FQueueEntry {Key, Payload});
+			}
+		}
+
+		if (bEnqueued)
+		{
+			WakeUpEvent->Trigger();
+			break;
+		}
+
+		UploadCompleteEvent->Wait();
+	}
+
+	return true;
+}
+
+void FS3UploadQueue::ThreadEntry()
+{
+	for(;;)
+	{
+		FQueueEntry Entry;
+		bool bDequeued = false;
+		{
+			FScopeLock _(&CriticalSection);
+			bDequeued = Queue.Dequeue(Entry);
+		}
+
+		if (!bDequeued)
+		{
+			if (bCompleteAdding)
+			{
+				break;
+			}
+			WakeUpEvent->Wait();
+			continue;
+		}
+
+		ConcurrentUploads++;
+		const FS3PutObjectResponse Response = Client.TryPutObject(FS3PutObjectRequest{Bucket, Entry.Key, Entry.Payload.GetView()});
+		ConcurrentUploads--;
+		UploadCompleteEvent->Trigger();
+
+		if (Response.IsOk())
+		{
+			UE_LOG(LogIas, Display, TEXT("Uploaded chunk '%s/%s/%s'"), *Client.GetConfig().ServiceUrl, *Bucket, *Entry.Key);
+		}
+		else
+		{
+			UE_LOG(LogIas, Warning, TEXT("Failed to upload chunk '%s/%s/%s', StatusCode: %u"), *Client.GetConfig().ServiceUrl, *Bucket, *Entry.Key, Response.StatusCode);
+			ErrorCount++;
+			break;
+		}
+	}
+
+	ActiveThreadCount--;
+}
+
+bool FS3UploadQueue::Flush()
+{
+	bCompleteAdding = true;
+	for (int32 Idx = 0; Idx < Threads.Num(); ++Idx)
+	{
+		WakeUpEvent->Trigger();
+	}
+
+	for (TFuture<void>& Thread : Threads)
+	{
+		Thread.Wait();
+	}
+
+	return ErrorCount == 0;
+}
+
+////////////////////////////////////////////////////////////////////////////////
+TIoStatusOr<FIoStoreUploadParams> FIoStoreUploadParams::Parse(const TCHAR* CommandLine)
+{
+	FIoStoreUploadParams Params;
+
+	if (!FParse::Value(CommandLine, TEXT("Bucket="), Params.Bucket))
+	{
+		FIoStatus(EIoErrorCode::InvalidParameter, TEXT("Invalid bucket name"));
+	}
+
+	FParse::Value(CommandLine, TEXT("BucketPrefix="), Params.BucketPrefix);
+	FParse::Value(CommandLine, TEXT("ServiceUrl="), Params.ServiceUrl);
+	FParse::Value(CommandLine, TEXT("Region="), Params.Region);
+	FParse::Value(CommandLine, TEXT("AccessKey="), Params.AccessKey);
+	FParse::Value(CommandLine, TEXT("SecretKey="), Params.SecretKey);
+	FParse::Value(CommandLine, TEXT("SessionToken="), Params.SessionToken);
+	FParse::Value(CommandLine, TEXT("CredentialsFile="), Params.CredentialsFile);
+	FParse::Value(CommandLine, TEXT("CredentialsFileKeyName="), Params.CredentialsFileKeyName);
+	Params.bDeleteContainerFiles = FParse::Param(CommandLine, TEXT("KeepUploadedContainers")) == false;
+
+	if (Params.AccessKey.IsEmpty() &&
+		Params.SecretKey.IsEmpty() &&
+		Params.CredentialsFile.IsEmpty() &&
+		Params.CredentialsFileKeyName.IsEmpty())
+	{
+		return FIoStatus(EIoErrorCode::InvalidParameter, TEXT("Invalid credentials"));
+	}
+	
+	if (!Params.AccessKey.IsEmpty() && Params.SecretKey.IsEmpty())
+	{
+		return FIoStatus(EIoErrorCode::InvalidParameter, TEXT("Invalid secret key"));
+	}
+	else if (Params.AccessKey.IsEmpty() && !Params.SecretKey.IsEmpty())
+	{
+		return FIoStatus(EIoErrorCode::InvalidParameter, TEXT("Invalid access key"));
+	}
+
+	if (!Params.CredentialsFile.IsEmpty() && Params.CredentialsFileKeyName.IsEmpty())
+	{
+		return FIoStatus(EIoErrorCode::InvalidParameter, TEXT("Invalid credential file key name"));
+	}
+
+	if (Params.ServiceUrl.IsEmpty() && Params.Region.IsEmpty())
+	{
+		return FIoStatus(EIoErrorCode::InvalidParameter, TEXT("Service URL or AWS region needs to be specified"));
+	}
+
+	return Params;
+}
+
+TIoStatusOr<FIoStoreUploadResult> UploadContainerFiles(
+	const FIoStoreUploadParams& UploadParams,
+	TConstArrayView<FString> ContainerFiles,
+	const TMap<FGuid, FAES::FAESKey>& EncryptionKeys)
+{
+	FS3ClientConfig Config;
+	Config.ServiceUrl = UploadParams.ServiceUrl;
+	Config.Region = UploadParams.Region;
+	
+	FS3ClientCredentials Credentials;
+	if (UploadParams.CredentialsFile.IsEmpty() == false)
+	{
+		UE_LOG(LogIas, Display, TEXT("Loading credentials file '%s'"), *UploadParams.CredentialsFile);
+		FS3CredentialsProfileStore CredentialsStore = FS3CredentialsProfileStore::FromFile(UploadParams.CredentialsFile);
+		if (CredentialsStore.TryGetCredentials(UploadParams.CredentialsFileKeyName, Credentials) == false)
+		{
+			return FIoStatus(EIoErrorCode::InvalidParameter, TEXT("Failed to find valid credentials in credentials file"));
+		}
+		else
+		{
+			UE_LOG(LogIas, Display, TEXT("Found credentials for '%s'"), *UploadParams.CredentialsFileKeyName);
+		}
+	}
+	else
+	{
+		Credentials = FS3ClientCredentials(UploadParams.AccessKey, UploadParams.SecretKey, UploadParams.SessionToken);
+	}
+
+	FS3Client Client(Config, Credentials);
+	FS3UploadQueue UploadQueue(Client, UploadParams.Bucket, UploadParams.MaxConcurrentUploads);
+
+	if (ContainerFiles.Num() == 0)
+	{
+		return FIoStatus(EIoErrorCode::InvalidParameter, TEXT("No container file(s) specified"));
+	}
+
+	TSet<FIoHash> ExistingChunks;
+	{
+		TStringBuilder<256> TocsKey;
+		TocsKey << UploadParams.BucketPrefix << "/";
+
+		UE_LOG(LogIas, Display, TEXT("Fetching TOC's '%s/%s/%s'"), *Client.GetConfig().ServiceUrl, *UploadParams.Bucket, TocsKey.ToString());
+		FS3ListObjectResponse Response = Client.ListObjects(FS3ListObjectsRequest
+		{
+			UploadParams.Bucket,
+			TocsKey.ToString(),
+			TEXT('/')
+		});
+
+		for (const FS3Object& TocInfo : Response.Objects)
+		{
+			if (TocInfo.Key.EndsWith(TEXT("iochunktoc")) == false)
+			{
+				continue;
+			}
+
+			UE_LOG(LogIas, Display, TEXT("Fetching TOC '%s/%s/%s'"), *Client.GetConfig().ServiceUrl, *UploadParams.Bucket, *TocInfo.Key);
+			FS3GetObjectResponse TocResponse = Client.GetObject(FS3GetObjectRequest
+			{
+				UploadParams.Bucket,
+				TocInfo.Key
+			});
+			
+			if (TocResponse.IsOk() == false)
+			{
+				UE_LOG(LogIas, Warning, TEXT("Failed to fetch TOC '%s/%s/%s'"), *Client.GetConfig().ServiceUrl, *UploadParams.Bucket, *TocInfo.Key);
+				continue;
+			}
+
+			FOnDemandToc Toc;
+			if (UE::LoadFromCompactBinary(FCbFieldView(TocResponse.Body.GetData()), Toc) == false)
+			{
+				UE_LOG(LogIas, Warning, TEXT("Failed to load TOC '%s/%s/%s'"), *Client.GetConfig().ServiceUrl, *UploadParams.Bucket, *TocInfo.Key);
+				continue;
+			}
+
+			for (const FOnDemandTocContainerEntry& ContainerEntry : Toc.Containers)
+			{
+				for (const FOnDemandTocEntry& TocEntry : ContainerEntry.Entries)
+				{
+					ExistingChunks.Add(TocEntry.Hash);
+				}
+			}
+		}
+	}
+	UE_LOG(LogIas, Display, TEXT("Found %d existing chunks"), ExistingChunks.Num());
+	
+	FString ChunksRelativePath = UploadParams.BucketPrefix.IsEmpty()
+		? FString::Printf(TEXT("IoChunksV%u"), EOnDemandChunkVersion::Latest)
+		: FString::Printf(TEXT("%s/IoChunksV%u"), *UploadParams.BucketPrefix, EOnDemandChunkVersion::Latest);
+
+	ChunksRelativePath.ToLowerInline();
+
+	uint64 TotalUploadedChunks = 0;
+	uint64 TotalUploadedBytes = 0;
+
+	FOnDemandToc OnDemandToc;
+	OnDemandToc.Header.ChunksDirectory = FString::Printf(TEXT("IoChunksV%u"), EOnDemandChunkVersion::Latest).ToLower();
+
+	TArray<FString> UploadedFiles;
+	for (const FString& Path : ContainerFiles)
+	{
+		FIoStoreReader ContainerFileReader;
+		{
+			FIoStatus Status = ContainerFileReader.Initialize(*FPaths::ChangeExtension(Path, TEXT("")), EncryptionKeys);
+			if (!Status.IsOk())
+			{
+				UE_LOG(LogIas, Error, TEXT("Failed to open container '%s' for reading"), *Path);
+				continue;
+			}
+		}
+
+		if (EnumHasAnyFlags(ContainerFileReader.GetContainerFlags(), EIoContainerFlags::OnDemand) == false)
+		{
+			UE_LOG(LogIas, Display, TEXT("Skipping non ondemand container '%s'"), *Path);
+			continue;
+		}
+		
+		UE_LOG(LogIas, Display, TEXT("Uploading container '%s/.ucas'"), *Path);
+
+		const uint32 BlockSize = ContainerFileReader.GetCompressionBlockSize();
+		if (OnDemandToc.Header.BlockSize == 0)
+		{
+			OnDemandToc.Header.BlockSize = ContainerFileReader.GetCompressionBlockSize();
+		}
+		check(OnDemandToc.Header.BlockSize == ContainerFileReader.GetCompressionBlockSize());
+
+		TArray<FIoStoreTocChunkInfo> ChunkInfos;
+		ContainerFileReader.EnumerateChunks([&ChunkInfos](FIoStoreTocChunkInfo&& Info)
+		{ 
+			ChunkInfos.Emplace(MoveTemp(Info));
+			return true;
+		});
+		
+		FOnDemandTocContainerEntry& ContainerEntry = OnDemandToc.Containers.AddDefaulted_GetRef();
+		ContainerEntry.ContainerName = FPaths::GetBaseFilename(Path);
+		ContainerEntry.EncryptionKeyGuid = LexToString(ContainerFileReader.GetEncryptionKeyGuid());
+		
+		for (const FIoStoreTocChunkInfo& ChunkInfo : ChunkInfos)
+		{
+			const bool bDecrypt = false;
+			TIoStatusOr<FIoStoreCompressedReadResult> Status = ContainerFileReader.ReadCompressed(ChunkInfo.Id, FIoReadOptions(), bDecrypt);
+			if (!Status.IsOk())
+			{
+				return Status.Status();
+			}
+
+			FIoStoreCompressedReadResult ReadResult = Status.ConsumeValueOrDie();
+
+			const uint32 BlockOffset = ContainerEntry.BlockSizes.Num();
+			const uint32 BlockCount = ReadResult.Blocks.Num();
+			const FIoHash ChunkHash = FIoHash::HashBuffer(ReadResult.IoBuffer.GetView());
+
+			FMemoryView EncodedBlocks = ReadResult.IoBuffer.GetView();
+			uint64 RawChunkSize = 0;
+			uint64 EncodedChunkSize = 0;
+			for (const FIoStoreCompressedBlockInfo& BlockInfo : ReadResult.Blocks)
+			{
+				const uint64 EncodedBlockSize = Align(BlockInfo.CompressedSize, FAES::AESBlockSize);
+				ContainerEntry.BlockSizes.Add(EncodedBlockSize);
+
+				FMemoryView EncodedBlock = EncodedBlocks.Left(EncodedBlockSize);
+				EncodedBlocks += EncodedBlock.GetSize();
+				ContainerEntry.BlockHashes.Add(FIoHash::HashBuffer(EncodedBlock));
+
+				EncodedChunkSize += EncodedBlockSize;
+				RawChunkSize += BlockInfo.UncompressedSize;
+
+				if (OnDemandToc.Header.CompressionFormat.IsEmpty() && BlockInfo.CompressionMethod != NAME_None)
+				{
+					OnDemandToc.Header.CompressionFormat = BlockInfo.CompressionMethod.ToString();
+				}
+			}
+
+			if (EncodedChunkSize != ReadResult.IoBuffer.GetSize())
+			{
+				return FIoStatus(EIoErrorCode::ReadError, TEXT("Encoded chunk size does not match buffer"));
+			}
+
+			FOnDemandTocEntry& TocEntry = ContainerEntry.Entries.AddDefaulted_GetRef();
+			TocEntry.ChunkId = ChunkInfo.Id;
+			TocEntry.Hash = ChunkHash;
+			TocEntry.RawHash = ChunkInfo.Hash.ToIoHash();
+			TocEntry.RawSize = RawChunkSize;
+			TocEntry.EncodedSize = EncodedChunkSize;
+			TocEntry.BlockOffset = BlockOffset;
+			TocEntry.BlockCount = BlockCount;
+
+			if (ExistingChunks.Contains(TocEntry.Hash) == false)
+			{
+				const FString HashString = LexToString(ChunkHash);
+
+				TStringBuilder<256> Key;
+				Key << ChunksRelativePath
+					<< TEXT("/") << HashString.Left(2)
+					<< TEXT("/") << HashString
+					<< TEXT(".iochunk");
+
+				const bool bEnqueued = UploadQueue.Enqueue(Key.ToString(), ReadResult.IoBuffer);
+
+				if (bEnqueued)
+				{
+					UE_LOG(LogIas, Display, TEXT("Uploaded chunk '%s/%s/%s'"), *Client.GetConfig().ServiceUrl, *UploadParams.Bucket, Key.ToString());
+				}
+				else
+				{
+					return FIoStatus(EIoErrorCode::WriteError, TEXT("Failed to upload chunk"));
+				}
+
+				TotalUploadedChunks++;
+				TotalUploadedBytes += ReadResult.IoBuffer.GetSize();
+			}
+		}
+
+		{
+			const FString UTocFilePath = FPaths::ChangeExtension(Path, TEXT(".utoc"));
+			TArray<uint8> Buffer;
+			if (!FFileHelper::LoadFileToArray(Buffer, *UTocFilePath))
+			{
+				return FIoStatus(EIoErrorCode::ReadError, TEXT("Failed to upload .utoc file"));
+			}
+
+			ContainerEntry.UTocHash = FIoHash::HashBuffer(Buffer.GetData(), Buffer.Num());
+			TStringBuilder<256> Key;
+			if (UploadParams.BucketPrefix.IsEmpty() == false)
+			{
+				Key << UploadParams.BucketPrefix.ToLower() << TEXT("/");
+			}
+			Key << LexToString(ContainerEntry.UTocHash) << TEXT(".utoc");
+
+			const FS3PutObjectResponse Response = Client.TryPutObject(
+				FS3PutObjectRequest{UploadParams.Bucket, Key.ToString(), MakeMemoryView(Buffer.GetData(), Buffer.Num())});
+
+			if (Response.IsOk())
+			{
+				UE_LOG(LogIas, Display, TEXT("Uploaded '%s'"), *UTocFilePath);
+			}
+			else
+			{
+				return FIoStatus(EIoErrorCode::WriteError, FString::Printf(TEXT("Failed to upload '%s', StatusCode: %u"), *UTocFilePath, Response.StatusCode));
+			}
+		}
+		
+		UploadedFiles.Add(Path);
+	}
+
+	if (OnDemandToc.Containers.IsEmpty())
+	{
+		return FIoStatus(EIoErrorCode::InvalidParameter, TEXT("No container file(s) marked as on demand"));
+	}
+
+	if (!UploadQueue.Flush())
+	{
+		return FIoStatus(EIoErrorCode::WriteError, TEXT("Failed to upload chunk(s)"));
+	}
+
+	FIoStoreUploadResult UploadResult;
+	{
+		FCbWriter Writer;
+		Writer << OnDemandToc;
+
+		FLargeMemoryWriter Ar;
+		SaveCompactBinary(Ar, Writer.Save());
+
+		UploadResult.TocHash = FIoHash::HashBuffer(Ar.GetView());
+		TStringBuilder<256> Key;
+		if (UploadParams.BucketPrefix.IsEmpty() == false)
+		{
+			Key << UploadParams.BucketPrefix.ToLower() << TEXT("/");
+		}
+		Key << LexToString(UploadResult.TocHash) << TEXT(".iochunktoc");
+
+		UploadResult.TocPath = Key.ToString();
+
+		const FS3PutObjectResponse Response = Client.TryPutObject(FS3PutObjectRequest{UploadParams.Bucket, Key.ToString(), Ar.GetView()});
+		if (Response.IsOk())
+		{
+			UE_LOG(LogIas, Display, TEXT("Uploaded on demand TOC '%s/%s/%s'"), *Client.GetConfig().ServiceUrl, *UploadParams.Bucket, Key.ToString());
+		}
+		else
+		{
+			UE_LOG(LogIas, Warning, TEXT("Failed to upload TOC '%s/%s/%s', StatusCode: %u"), *Client.GetConfig().ServiceUrl, *UploadParams.Bucket, Key.ToString(), Response.StatusCode);
+			return FIoStatus(EIoErrorCode::WriteError, TEXT("Failed to upload TOC"));
+		}
+	}
+
+	if (UploadParams.bDeleteContainerFiles)
+	{
+		for (const FString& TocPath : UploadedFiles)
+		{
+			const FString CasPath = FPaths::ChangeExtension(TocPath, TEXT(".ucas"));
+			const TCHAR* FilePaths[] {*TocPath, *CasPath};
+			for (const TCHAR* FilePath : FilePaths)
+			{
+				if (IFileManager::Get().FileExists(FilePath))
+				{
+					UE_LOG(LogIas, Display, TEXT("Deleting '%s'"), FilePath); 
+					IFileManager::Get().Delete(FilePath);
+				}
+			}
+		}
+	}
+	
+	UE_LOG(LogIas, Display, TEXT("Done uploading %d container file(s), %d chunks of total %.2lf MiB"),
+		UploadedFiles.Num(), TotalUploadedChunks, (double(TotalUploadedBytes) / 1024 / 1024));
+	
+	return UploadResult;
+}
+#endif // (IS_PROGRAM || WITH_EDITOR)
+
 } // namespace UE
 
 ////////////////////////////////////////////////////////////////////////////////
@@ -384,7 +887,7 @@ void FIoStoreOnDemandModule::StartupModule()
 
 	{
 		FString EncryptionKey;
-		if (FParse::Value(CommandLine, TEXT("Ias.EncryptionKey="), EncryptionKey))
+		if (FParse::Value(CommandLine, TEXT("IasEncryptionKey="), EncryptionKey))
 		{
 			FGuid KeyGuid;
 			FAES::FAESKey Key;
diff --git a/Engine/Source/Runtime/Experimental/IoStoreOnDemand/Private/OnDemandIoStoreWriter.cpp b/Engine/Source/Runtime/Experimental/IoStoreOnDemand/Private/OnDemandIoStoreWriter.cpp
deleted file mode 100644
diff --git a/Engine/Source/Runtime/Experimental/IoStoreOnDemand/Public/IO/IoStoreOnDemand.h b/Engine/Source/Runtime/Experimental/IoStoreOnDemand/Public/IO/IoStoreOnDemand.h
--- a/Engine/Source/Runtime/Experimental/IoStoreOnDemand/Public/IO/IoStoreOnDemand.h
+++ b/Engine/Source/Runtime/Experimental/IoStoreOnDemand/Public/IO/IoStoreOnDemand.h
@@ -10,6 +10,11 @@
 #include "Templates/SharedPointer.h"
 #include "UObject/NameTypes.h"
 
+#if (IS_PROGRAM || WITH_EDITOR)
+#include "Misc/AES.h"
+#include "Containers/Map.h"
+#endif // (IS_PROGRAM || WITH_EDITOR)
+
 #define UE_API IOSTOREONDEMAND_API
 
 class FCbWriter;
@@ -28,6 +33,7 @@ enum class EOnDemandTocVersion : uint32
 {
 	Invalid			= 0,
 	Initial			= 1,
+	UTocHash		= 2,
 
 	LatestPlusOne,
 	Latest			= (LatestPlusOne - 1)
@@ -80,6 +86,7 @@ struct FOnDemandTocContainerEntry
 	TArray<FOnDemandTocEntry> Entries;
 	TArray<uint32> BlockSizes;
 	TArray<FIoHash> BlockHashes;
+	FIoHash UTocHash;
 
 	UE_API friend FCbWriter& operator<<(FCbWriter& Writer, const FOnDemandTocContainerEntry& ContainerEntry);
 };
@@ -98,21 +105,36 @@ struct FOnDemandToc
 
 UE_API bool LoadFromCompactBinary(FCbFieldView Field, FOnDemandToc& OutToc);
 
-////////////////////////////////////////////////////////////////////////////////
 #if (IS_PROGRAM || WITH_EDITOR)
 
-class IOnDemandIoStoreWriter
+////////////////////////////////////////////////////////////////////////////////
+struct FIoStoreUploadParams
+{
+	FString ServiceUrl;
+	FString Bucket;
+	FString BucketPrefix;
+	FString Region; 
+	FString AccessKey;
+	FString SecretKey;
+	FString SessionToken;
+	FString CredentialsFile;
+	FString CredentialsFileKeyName;
+	int32 MaxConcurrentUploads = 16;
+	bool bDeleteContainerFiles = true;
+	
+	UE_API static TIoStatusOr<FIoStoreUploadParams> Parse(const TCHAR* CommandLine);
+};
+
+struct FIoStoreUploadResult
 {
-public:
-	virtual ~IOnDemandIoStoreWriter() = default;
-	virtual TSharedPtr<IIoStoreWriter> CreateContainer(const FString& ContainerName, const FIoContainerSettings& ContainerSettings) = 0;
-	virtual void Flush() = 0;
+	FIoHash TocHash;
+	FString TocPath;
 };
 
-UE_API TUniquePtr<IOnDemandIoStoreWriter> MakeOnDemandIoStoreWriter(
-	const FIoStoreWriterSettings& WriterSettings,
-	const FString& OutputDirectory,
-	uint32 MaxConcurrentWrites = 64);
+UE_API TIoStatusOr<FIoStoreUploadResult> UploadContainerFiles(
+	const FIoStoreUploadParams& UploadParams,
+	TConstArrayView<FString> ContainerFiles,
+	const TMap<FGuid, FAES::FAESKey>& EncryptionKeys);
 
 #endif // (IS_PROGRAM || WITH_EDITOR)
 
