Added cvar to control IAS client tick's poll timeout

#rnx

diff --git a/Engine/Source/Runtime/Experimental/IoStoreOnDemand/Private/OnDemandIoDispatcherBackend.cpp b/Engine/Source/Runtime/Experimental/IoStoreOnDemand/Private/OnDemandIoDispatcherBackend.cpp
--- a/Engine/Source/Runtime/Experimental/IoStoreOnDemand/Private/OnDemandIoDispatcherBackend.cpp
+++ b/Engine/Source/Runtime/Experimental/IoStoreOnDemand/Private/OnDemandIoDispatcherBackend.cpp
@@ -49,6 +49,13 @@ static FAutoConsoleVariableRef CVar_IoDispatcherMaxHttpRetryCount (
 	TEXT("Max number of HTTP request retries before failing the I/O request.")
 );
 
+int32 GIoDispatcherHttpPollTimeoutMs = 0;
+static FAutoConsoleVariableRef CVar_IoDispatcherMaxHttpPollTimeoutMs (
+	TEXT("s.IasHttpPollTimeout"),
+	GIoDispatcherHttpPollTimeoutMs,
+	TEXT("Tick() poll timeout in milliseconds")
+);
+
 namespace UE::IO::Private
 {
 
@@ -383,7 +390,7 @@ void FHttpClient::Issue(UE::HTTP::FRequest&& Request, FIoReadCallback&& Callback
 
 bool FHttpClient::Tick()
 {
-	return EventLoop.Tick() != 0;
+	return EventLoop.Tick(GIoDispatcherHttpPollTimeoutMs) != 0;
 }
 
 ///////////////////////////////////////////////////////////////////////////////
