Use poll instead of select where appropriate and group file descriptors together into one call.

#rnx

diff --git a/Engine/Source/Runtime/Experimental/IoStoreOnDemand/Private/CoreHttp/Http.cpp b/Engine/Source/Runtime/Experimental/IoStoreOnDemand/Private/CoreHttp/Http.cpp
--- a/Engine/Source/Runtime/Experimental/IoStoreOnDemand/Private/CoreHttp/Http.cpp
+++ b/Engine/Source/Runtime/Experimental/IoStoreOnDemand/Private/CoreHttp/Http.cpp
@@ -29,6 +29,12 @@
 	using SocketType	= SOCKET;
 	using MsgFlagType	= int;
 
+#	define IAS_HTTP_USE_POLL
+	template <typename... ArgTypes> auto poll(ArgTypes... Args)
+	{
+		return WSAPoll(Forward<ArgTypes>(Args)...);
+	}
+
 	// Winsock defines "PF_MAX" indicating the protocol family count. This
 	// however competes with UE definitions related to pixel formats.
 	#if defined(PF_MAX)
@@ -39,12 +45,14 @@
 #	include <arpa/inet.h>
 #	include <fcntl.h>
 #	include <netdb.h>
+#	include <poll.h>
 #	include <sys/select.h>
 #	include <sys/socket.h>
 #	include <unistd.h>
 	using SocketType	= int;
 	using MsgFlagType	= int;
 	static int32 closesocket(int32 Socket) { return close(Socket); }
+#	define IAS_HTTP_USE_POLL
 #else
 #	include "CoreHttp/Http.inl"
 #endif
@@ -1082,8 +1090,26 @@ const char* FTicketStatus::GetErrorReason() const
 // {{{1 event-loop-int .........................................................
 
 ////////////////////////////////////////////////////////////////////////////////
-static int32 Select(SocketType Socket, int32 Mode, int32 TimeoutMs=0)
+#if !defined(IAS_HTTP_USE_POLL)
+struct FSelect
 {
+	SocketType	fd;
+	int32		events;
+	int32		revents;
+};
+static const int32 POLLIN  = 1 << 0;
+static const int32 POLLOUT = 1 << 1;
+static const int32 POLLERR = 1 << 1;
+#else
+using FSelect = pollfd;
+#endif
+
+////////////////////////////////////////////////////////////////////////////////
+static bool DoSelect(FSelect* Selects, uint32 SelectNum, int32 TimeoutMs)
+{
+#if defined(IAS_HTTP_USE_POLL)
+	return poll(Selects, SelectNum, TimeoutMs) > 0;
+#else
 	timeval TimeVal = {};
 	timeval* TimeValPtr = (TimeoutMs >= 0 ) ? &TimeVal : nullptr;
 	if (TimeoutMs > 0)
@@ -1091,21 +1117,110 @@ static int32 Select(SocketType Socket, int32 Mode, int32 TimeoutMs=0)
 		TimeVal = { TimeoutMs >> 10, TimeoutMs & ((1 << 10) - 1) };
 	}
 
-	fd_set FdSet;
-	FD_ZERO(&FdSet);
-	FD_SET(Socket, &FdSet);
+	fd_set FdSetRead;	FD_ZERO(&FdSetRead);
+	fd_set FdSetWrite;	FD_ZERO(&FdSetWrite);
+	fd_set FdSetExcept; FD_ZERO(&FdSetExcept);
 
-	fd_set FdSetExcept;
-	FD_ZERO(&FdSetExcept);
-	FD_SET(Socket, &FdSetExcept);
+	SocketType MaxFd = 0;
+	for (uint32 i = 0; i < SelectNum; ++i)
+	{
+		FSelect& Select = Selects[i];
+		fd_set* RwSet = (Select.events & POLLIN) ? &FdSetRead : &FdSetWrite;
+		FD_SET(Select.fd, RwSet);
+		FD_SET(Select.fd, &FdSetExcept);
+		MaxFd = FMath::Max(Select.fd, MaxFd);
+	}
 
-	fd_set* FdSets[] = { &FdSet, nullptr };
-	if (Mode == 1)
+	int32 Result = select(MaxFd + 1, &FdSetRead, &FdSetWrite, &FdSetExcept, TimeValPtr);
+	if (Result == 0)
 	{
-		Swap(FdSets[0], FdSets[1]);
+		return false;
 	}
-	int32 Result = select(1, FdSets[0], FdSets[1], &FdSetExcept, TimeValPtr);
-	return Result;
+
+	for (uint32 i = 0; i < SelectNum; ++i)
+	{
+		FSelect& Select = Selects[i];
+		fd_set* RwSet = (Select.events & POLLIN) ? &FdSetRead : &FdSetWrite;
+		if (FD_ISSET(Select.fd, RwSet) || FD_ISSET(Select.fd, &FdSetExcept))
+		{
+			Select.revents = Select.events;
+		}
+	}
+
+	return true;
+#endif // IAS_HTTP_USE_POLL
+}
+
+////////////////////////////////////////////////////////////////////////////////
+static uint64 ReadyCheck(FActivity** Activities, uint32 Num, uint32 TimeoutMs)
+{
+	TRACE_CPUPROFILER_EVENT_SCOPE(CoreHttp::ReadyCheck);
+
+	using EWait = FActivity::EWait;
+
+	uint64 Ret = 0;
+
+	FSelect Selects[64];
+	uint32 SelectNum = 0;
+
+	for (uint32 i = 0; i < Num; ++i)
+	{
+		FActivity* Activity = Activities[i];
+		switch (Activity->SocketWait)
+		{
+		case EWait::None:
+			Ret |= (1ull << Activity->Slot);
+			break;
+
+		case EWait::Pool:
+			if (Activity->Pool->GetState() >= FSocketPool::EState::Resolved)
+			{
+				Activities[i]->SocketWait = EWait::None;
+				Ret |= (1ull << Activity->Slot);
+			}
+			break;
+
+		case EWait::Read:
+		case EWait::Write: {
+			FSelect& Select = Selects[SelectNum];
+			Select.fd = Activity->Socket;
+			Select.events = (Activity->SocketWait == EWait::Read) ? POLLIN : POLLOUT;
+			++SelectNum;
+			} break;
+		}
+	}
+
+	// Collect result
+	if (SelectNum == 0)
+	{
+		return Ret;
+	}
+
+	if (!DoSelect(Selects, SelectNum, TimeoutMs))
+	{
+		return Ret;
+	}
+
+	const FSelect* SelectCursor = Selects;
+	for (int32 i = 0; SelectNum > 0; ++i)
+	{
+		EWait Wait = Activities[i]->SocketWait;
+		if (Wait != EWait::Read && Wait != EWait::Write)
+		{
+			continue;
+		}
+
+		if (int32(SelectCursor->revents & (POLLIN|POLLOUT|POLLERR)))
+		{
+			Activities[i]->SocketWait = EWait::None;
+			Ret |= (1ull << Activities[i]->Slot);
+		}
+
+		--SelectNum;
+		++SelectCursor;
+	}
+
+	return Ret;
 }
 
 
@@ -1114,7 +1229,6 @@ static int32 Select(SocketType Socket, int32 Mode, int32 TimeoutMs=0)
 class FEventLoopInternal
 {
 public:
-	static bool		DoIsReady(FActivity* Activity, uint32 TimeoutMs=0);
 	static int32	DoResolve(FActivity* Activity);
 	static int32	DoConnect(FActivity* Activity);
 	static int32	DoSend(FActivity* Activity);
@@ -1125,31 +1239,6 @@ public:
 	static void		Cancel(FActivity* Activity);
 };
 
-
-
-////////////////////////////////////////////////////////////////////////////////
-bool FEventLoopInternal::DoIsReady(FActivity* Activity, uint32 TimeoutMs)
-{
-	TRACE_CPUPROFILER_EVENT_SCOPE(CoreHttp::DoIsReady);
-
-	int32 Result = 0;
-	switch (Activity->SocketWait)
-	{
-	case FActivity::EWait::Pool:  Result = (Activity->Pool->GetState() >= FSocketPool::EState::Resolved);
-	case FActivity::EWait::Read:  Result = Select(Activity->Socket, 0);	break;
-	case FActivity::EWait::Write: Result = Select(Activity->Socket, 1);	break;
-	case FActivity::EWait::None:  return true;
-	}
-
-	if (Result != 0)
-	{
-		Activity->SocketWait = FActivity::EWait::None;
-		return true;
-	}
-
-	return false;
-}
-
 ////////////////////////////////////////////////////////////////////////////////
 int32 FEventLoopInternal::DoResolve(FActivity* Activity)
 {
@@ -1878,25 +1967,34 @@ uint32 FEventLoop::Tick(uint32 PollTimeoutMs)
         PrevFreeSlots = FreeSlotsLoad;
     }
 
+	uint32 BusyCount = Active.Num();
+
+	// Work out which activities are ready
+	uint64 Readies = ReadyCheck(Active.GetData(), Active.Num(), PollTimeoutMs);
+	if (Readies == 0)
+	{
+		return BusyCount;
+	}
+
 	// Tick activities
 	uint64 CancelsLoad = Cancels.load(std::memory_order_relaxed);
-	uint32 BusyCount = 0;
 	for (FActivity* Activity : Active)
 	{
-		if ((1ull << Activity->Slot) & CancelsLoad)
+		uint64 SlotBit = 1ull << Activity->Slot;
+
+		if (SlotBit & CancelsLoad)
 		{
+			--BusyCount;
 			FEventLoopInternal::Cancel(Activity);
 			continue;
 		}
 
-		++BusyCount;
-
-		int32 Result = 1;
-		if (!FEventLoopInternal::DoIsReady(Activity, PollTimeoutMs))
+		if ((Readies & SlotBit) == 0)
 		{
 			continue;
 		}
 
+		int32 Result = 1;
 		switch (Activity->State)
 		{
 		case FActivity::EState::Resolve:
