Prefixed IAS' cache logger with Ias in-keeping with other parts of the module.

#rnx

diff --git a/Engine/Source/Runtime/Experimental/IoStoreOnDemand/Private/FileIoCache.cpp b/Engine/Source/Runtime/Experimental/IoStoreOnDemand/Private/FileIoCache.cpp
--- a/Engine/Source/Runtime/Experimental/IoStoreOnDemand/Private/FileIoCache.cpp
+++ b/Engine/Source/Runtime/Experimental/IoStoreOnDemand/Private/FileIoCache.cpp
@@ -24,7 +24,7 @@
 
 #define IAS_FORCE_TOC_SAVE
 
-DEFINE_LOG_CATEGORY(LogIoCache);
+DEFINE_LOG_CATEGORY(LogIasCache);
 
 namespace UE::IO::Private
 {
@@ -612,7 +612,7 @@ TTask<TIoStatusOr<FIoBuffer>> FFileIoCache::Get(const FIoHash& Key, const FIoRea
 				IPlatformFile& Ipf = IPlatformFile::GetPlatformPhysical();
 				if (TUniquePtr<IFileHandle> FileHandle(Ipf.OpenRead(*CacheFilePath, true)); FileHandle.IsValid())
 				{
-					UE_LOG(LogIoCache, VeryVerbose, TEXT("Read chunk, Key='%s', Hash='%s', File='%s', Offset='%llu', Size='%llu'"),
+					UE_LOG(LogIasCache, VeryVerbose, TEXT("Read chunk, Key='%s', Hash='%s', File='%s', Offset='%llu', Size='%llu'"),
 						*LexToString(Entry.Key), *LexToString(Entry.Hash), *CacheFilePath, Entry.SerialOffset, Entry.SerialSize);
 					
 					FileHandle->Seek(int64(ReadOffset));
@@ -627,7 +627,7 @@ TTask<TIoStatusOr<FIoBuffer>> FFileIoCache::Get(const FIoHash& Key, const FIoRea
 					}
 
 					FOnDemandIoBackendStats::Get()->OnCacheError();
-					UE_LOG(LogIoCache, Verbose, TEXT("Read chunk failed, hash mismatch, Key='%s', Hash='%s', ExpectedHash='%s', File='%s', Offset='%llu', Size='%llu'"),
+					UE_LOG(LogIasCache, Verbose, TEXT("Read chunk failed, hash mismatch, Key='%s', Hash='%s', ExpectedHash='%s', File='%s', Offset='%llu', Size='%llu'"),
 						*LexToString(Entry.Key), *LexToString(Hash), *LexToString(ExpectedHash), *CacheFilePath, ReadOffset, ReadSize);
 
 					return TIoStatusOr<FIoBuffer>(FIoStatus(EIoErrorCode::ReadError));
@@ -635,7 +635,7 @@ TTask<TIoStatusOr<FIoBuffer>> FFileIoCache::Get(const FIoHash& Key, const FIoRea
 				else
 				{
 					FOnDemandIoBackendStats::Get()->OnCacheError();
-					UE_LOG(LogIoCache, Warning, TEXT("Read chunk failed, unable to open cache file '%s' for reading"), *CacheFilePath);
+					UE_LOG(LogIasCache, Warning, TEXT("Read chunk failed, unable to open cache file '%s' for reading"), *CacheFilePath);
 					return TIoStatusOr<FIoBuffer>(FIoStatus(EIoErrorCode::ReadError));
 				}
 			}
@@ -659,7 +659,7 @@ FIoStatus FFileIoCache::Put(const FIoHash& Key, FIoBuffer& Data)
 
 void FFileIoCache::Initialize()
 {
-	UE_LOG(LogIoCache, Log,
+	UE_LOG(LogIasCache, Log,
 		TEXT("Initializing file I/O cache, disk size %lluB, memory size %lluB"), CacheConfig.DiskQuota, CacheConfig.MemoryQuota);
 
 	WriterThread.Reset(FRunnableThread::Create(this, TEXT("Ias.FileCache"), 0, TPri_BelowNormal));
@@ -683,7 +683,7 @@ void FFileIoCache::Initialize()
 	{
 		if (FParse::Param(FCommandLine::Get(), TEXT("ClearIoCache")))
 		{
-			UE_LOG(LogIoCache, Log, TEXT("Deleting cache file '%s'"), *CacheFilePath);
+			UE_LOG(LogIasCache, Log, TEXT("Deleting cache file '%s'"), *CacheFilePath);
 			FileMgr.Delete(*CacheFilePath);
 		}
 		else
@@ -693,20 +693,20 @@ void FFileIoCache::Initialize()
 			{
 				check(WriteCursorPos != ~uint64(0));
 				
-				UE_LOG(LogIoCache, Log, TEXT("Loaded TOC '%s'"), *CacheTocPath);
+				UE_LOG(LogIasCache, Log, TEXT("Loaded TOC '%s'"), *CacheTocPath);
 				if (FileMgr.FileExists(*CacheFilePath))
 				{
 					//TODO: Integrity check?
 				}
 				else
 				{
-					UE_LOG(LogIoCache, Warning, TEXT("Failed to open cache file ''"), *CacheFilePath);
+					UE_LOG(LogIasCache, Warning, TEXT("Failed to open cache file ''"), *CacheFilePath);
 					CacheStatus = FIoStatus(EIoErrorCode::FileNotOpen);
 				}
 			}
 			else
 			{
-				UE_LOG(LogIoCache, Warning, TEXT("Failed to load TOC '%s'"), *CacheTocPath);
+				UE_LOG(LogIasCache, Warning, TEXT("Failed to load TOC '%s'"), *CacheTocPath);
 			}
 		}
 	}
@@ -736,7 +736,7 @@ void FFileIoCache::Shutdown()
 	WriterThread->Kill();
 
 	FString CacheTocPath = CacheFilePath + TEXT(".toc");
-	UE_LOG(LogIoCache, Log, TEXT("Saving TOC '%s'"), *CacheTocPath);
+	UE_LOG(LogIasCache, Log, TEXT("Saving TOC '%s'"), *CacheTocPath);
 	CacheMap.Save(CacheTocPath, WriteCursorPos);
 
 	CacheMap.Reset();
@@ -808,14 +808,14 @@ void FFileIoCache::FileWriterThreadInner()
 	}
 
 	// Open cache file and write the block to it.
-	UE_LOG(LogIoCache, VeryVerbose, TEXT("Write; cursor:%llu size:%d"), WriteCursorPos, PendingSize);
+	UE_LOG(LogIasCache, VeryVerbose, TEXT("Write; cursor:%llu size:%d"), WriteCursorPos, PendingSize);
 	TRACE_CPUPROFILER_EVENT_SCOPE(FFileIoCache::WriteCacheEntry);
 
 	IPlatformFile& Ipf = IPlatformFile::GetPlatformPhysical();
 	TUniquePtr<IFileHandle> FileHandle(Ipf.OpenWrite(*CacheFilePath, true, true));
 	if (!FileHandle.IsValid())
 	{
-		UE_LOG(LogIoCache, Warning, TEXT("Write chunks failed, unable to open file '%s' for writing"), *CacheFilePath);
+		UE_LOG(LogIasCache, Warning, TEXT("Write chunks failed, unable to open file '%s' for writing"), *CacheFilePath);
 		return;
 	}
 	FileHandle->Seek(WriteCursorPos);
diff --git a/Engine/Source/Runtime/Experimental/IoStoreOnDemand/Private/FileIoCache.h b/Engine/Source/Runtime/Experimental/IoStoreOnDemand/Private/FileIoCache.h
--- a/Engine/Source/Runtime/Experimental/IoStoreOnDemand/Private/FileIoCache.h
+++ b/Engine/Source/Runtime/Experimental/IoStoreOnDemand/Private/FileIoCache.h
@@ -12,7 +12,7 @@ class FIoBuffer;
 class FIoReadOptions;
 struct FIoHash;
 
-DECLARE_LOG_CATEGORY_EXTERN(LogIoCache, Log, All);
+DECLARE_LOG_CATEGORY_EXTERN(LogIasCache, Log, All);
 
 /** Cache for binary blobs with a 20 byte cache key. */
 class IIoCache
