Allow the governor to have a negative run-off should the write allowance be exceeded

#rnx

diff --git a/Engine/Source/Runtime/Experimental/IoStoreOnDemand/Private/FileIoCache.cpp b/Engine/Source/Runtime/Experimental/IoStoreOnDemand/Private/FileIoCache.cpp
--- a/Engine/Source/Runtime/Experimental/IoStoreOnDemand/Private/FileIoCache.cpp
+++ b/Engine/Source/Runtime/Experimental/IoStoreOnDemand/Private/FileIoCache.cpp
@@ -401,7 +401,7 @@ public:
 				FGovernorExternal();
 	void		Set(...) {}
 	uint32		TickAllowance();
-	void		Return(uint32) {}
+	void		Return(int32) {}
 
 private:
 	int64		CycleThreshold;
@@ -440,14 +440,14 @@ class FGovernorInternal
 public:
 	void		Set(uint32 InAllowance, uint32 InOps, uint32 Seconds);
 	uint32		TickAllowance();
-	void		Return(uint32 LeftOver);
+	void		Return(int32 LeftOver);
 
 private:
 	uint32		TickAllowance(int64 Cycles);
 	int64		CycleThreshold;
 	int64		CyclePrev;
 	uint32		Allowance;
-	uint32		RunOff;
+	int32		RunOff;
 };
 
 ////////////////////////////////////////////////////////////////////////////////
@@ -487,10 +487,10 @@ uint32 FGovernorInternal::TickAllowance(int64 Cycle)
 }
 
 ////////////////////////////////////////////////////////////////////////////////
-void FGovernorInternal::Return(uint32 LeftOver)
+void FGovernorInternal::Return(int32 LeftOver)
 {
 	// Roughly keep to some arbitrary limit so as to avoid excessive growth
-	uint32 OnePointFive = Allowance + (Allowance >> 1);
+	int32 OnePointFive = Allowance + (Allowance >> 1);
 	RunOff = FMath::Min(OnePointFive, LeftOver);
 }
 
@@ -498,7 +498,7 @@ void FGovernorInternal::Return(uint32 LeftOver)
 
 ////////////////////////////////////////////////////////////////////////////////
 #if !defined(UE_USE_PLATFORM_GOVERNOR)
-#	define UE_USE_PLATFORM_GOVERNOR 1
+#	define UE_USE_PLATFORM_GOVERNOR 0
 #endif
 
 #if UE_USE_PLATFORM_GOVERNOR
@@ -759,7 +759,7 @@ void FFileIoCache::FileWriterThreadInner()
 	++_TickCount;
 
 	// Update the rate limiting and see how much we are allowed to write
-	uint32 WriteAllowance = Governor.TickAllowance();
+	int32 WriteAllowance = Governor.TickAllowance();
 	if (!WriteAllowance)
 	{
 		return;
