Always peel off at least one pending cache item regardless of allowance

#rnx

diff --git a/Engine/Source/Runtime/Experimental/IoStoreOnDemand/Private/FileIoCache.cpp b/Engine/Source/Runtime/Experimental/IoStoreOnDemand/Private/FileIoCache.cpp
--- a/Engine/Source/Runtime/Experimental/IoStoreOnDemand/Private/FileIoCache.cpp
+++ b/Engine/Source/Runtime/Experimental/IoStoreOnDemand/Private/FileIoCache.cpp
@@ -274,15 +274,14 @@ int32 FCacheMap::RemovePending(FCacheEntryList& OutPending, uint32 MaxSize)
 		FCacheEntry& Entry = *Iter;
 		++Iter;
 
-		if (MaxSize < ReturnSize + Entry.Data.GetSize())
+		Pending.Remove(&Entry);
+		OutPending.AddTail(&Entry);
+
+		ReturnSize += Entry.Data.GetSize();
+		if (ReturnSize >= MaxSize)
 		{
 			break;
 		}
-
-		ReturnSize += Entry.Data.GetSize();
-
-		Pending.Remove(&Entry);
-		OutPending.AddTail(&Entry);
 	}
 
 	TotalPendingBytes -= ReturnSize;
