Join in on the bundle-patcher's PSM directory

#rnx

diff --git a/Engine/Source/Runtime/Experimental/IoStoreOnDemand/Private/FileIoCache.cpp b/Engine/Source/Runtime/Experimental/IoStoreOnDemand/Private/FileIoCache.cpp
--- a/Engine/Source/Runtime/Experimental/IoStoreOnDemand/Private/FileIoCache.cpp
+++ b/Engine/Source/Runtime/Experimental/IoStoreOnDemand/Private/FileIoCache.cpp
@@ -663,7 +663,7 @@ void FFileIoCache::Initialize()
 
 	WriterThread.Reset(FRunnableThread::Create(this, TEXT("Ias.FileCache"), 0, TPri_BelowNormal));
 
-	const FString CacheDir = FPaths::ProjectPersistentDownloadDir() / TEXT("IoCache");
+	const FString CacheDir = FPaths::ProjectPersistentDownloadDir() / TEXT("chunkdownload");
 	CacheFilePath = CacheDir / TEXT("cache.ucas");
 	WriteCursorPos = 0;
 
