Don't allow simultaneous reads and writes from the cache blob

#rnx

diff --git a/Engine/Source/Runtime/Experimental/IoStoreOnDemand/Private/FileIoCache.cpp b/Engine/Source/Runtime/Experimental/IoStoreOnDemand/Private/FileIoCache.cpp
--- a/Engine/Source/Runtime/Experimental/IoStoreOnDemand/Private/FileIoCache.cpp
+++ b/Engine/Source/Runtime/Experimental/IoStoreOnDemand/Private/FileIoCache.cpp
@@ -544,6 +544,7 @@ private:
 	TUniquePtr<FRunnableThread> WriterThread;
 	FEventRef TickWriterEvent;
 	FString CacheFilePath;
+	FCriticalSection FileCs;
 	uint64 WriteCursorPos = 0;
 	std::atomic_bool bStopRequested{false};
 	FGovernor Governor;
@@ -609,7 +610,12 @@ TTask<TIoStatusOr<FIoBuffer>> FFileIoCache::Get(const FIoHash& Key, const FIoRea
 				FIoBuffer Buffer = Options.GetTargetVa() ? FIoBuffer(FIoBuffer::Wrap, Options.GetTargetVa(), ReadSize) : FIoBuffer(ReadSize);
 
 				IPlatformFile& Ipf = IPlatformFile::GetPlatformPhysical();
-				if (TUniquePtr<IFileHandle> FileHandle(Ipf.OpenRead(*CacheFilePath, true)); FileHandle.IsValid())
+
+				// We can't read and write to the cache at the same time - if it loops around
+				// then we may end up reading partially written data.
+				FScopeLock _(&FileCs);
+
+				if (TUniquePtr<IFileHandle> FileHandle(Ipf.OpenRead(*CacheFilePath, false)); FileHandle.IsValid())
 				{
 					UE_LOG(LogIasCache, VeryVerbose, TEXT("Read chunk, Key='%s', Hash='%s', File='%s', Offset='%llu', Size='%llu'"),
 						*LexToString(Entry.Key), *LexToString(Entry.Hash), *CacheFilePath, Entry.SerialOffset, Entry.SerialSize);
@@ -811,15 +817,21 @@ void FFileIoCache::FileWriterThreadInner()
 	TRACE_CPUPROFILER_EVENT_SCOPE(FFileIoCache::WriteCacheEntry);
 
 	IPlatformFile& Ipf = IPlatformFile::GetPlatformPhysical();
-	TUniquePtr<IFileHandle> FileHandle(Ipf.OpenWrite(*CacheFilePath, true, true));
-	if (!FileHandle.IsValid())
 	{
-		UE_LOG(LogIasCache, Warning, TEXT("Write chunks failed, unable to open file '%s' for writing"), *CacheFilePath);
-		return;
+		// We can't read and write to the cache at the same time - if it loops around
+		// then we may end up reading partially written data.
+		FScopeLock _(&FileCs);
+
+		TUniquePtr<IFileHandle> FileHandle(Ipf.OpenWrite(*CacheFilePath, true, false));
+		if (!FileHandle.IsValid())
+		{
+			UE_LOG(LogIasCache, Warning, TEXT("Write chunks failed, unable to open file '%s' for writing"), *CacheFilePath);
+			return;
+		}
+		FileHandle->Seek(WriteCursorPos);
+		FileHandle->Write(Buffer, PendingSize);
+		FileHandle->Flush();
 	}
-	FileHandle->Seek(WriteCursorPos);
-	FileHandle->Write(Buffer, PendingSize);
-	FileHandle->Flush();
 	WriteCursorPos += PendingSize;
 
 	CacheMap.InsertPersisted(MoveTemp(Entries), WriteCursorPos);
