Iris ObjectReferenceCache:  more safety around mismatched cache maps to prevent crashes

#rnx

diff --git a/Engine/Source/Runtime/Experimental/Iris/Core/Private/Iris/ReplicationSystem/ObjectReferenceCache.cpp b/Engine/Source/Runtime/Experimental/Iris/Core/Private/Iris/ReplicationSystem/ObjectReferenceCache.cpp
--- a/Engine/Source/Runtime/Experimental/Iris/Core/Private/Iris/ReplicationSystem/ObjectReferenceCache.cpp
+++ b/Engine/Source/Runtime/Experimental/Iris/Core/Private/Iris/ReplicationSystem/ObjectReferenceCache.cpp
@@ -240,42 +240,52 @@ bool FObjectReferenceCache::CreateObjectReferenceInternal(const UObject* Object,
 		const FNetRefHandle RefHandle = *RefHandlePtr;
 
 		// Verify that this is the same object and clean up if it is not
-		FCachedNetObjectReference& CachedObject = ReferenceHandleToCachedReference.FindChecked(RefHandle);
-
-		if (CachedObject.Object.Get() == Object)
-		{
-			UE_LOG_REFERENCECACHE(VeryVerbose, TEXT("ObjectReferenceCache::CreateObjectReferenceHandle Found existing %s for ObjectPath %s Object: %s (0x%p), OuterNetRefHandle: %s"), 
-				*CachedObject.NetRefHandle.ToString(), ToCStr(Object->GetPathName()), *GetNameSafe(Object), Object, *CachedObject.OuterNetRefHandle.ToString());
-			OutReference = MakeNetObjectReference(CachedObject);			
-			return true;
-		}
-		else
+		if (FCachedNetObjectReference* CachedObjectPtr = ReferenceHandleToCachedReference.Find(RefHandle))
 		{
-			UE_LOG_REFERENCECACHE(Verbose, TEXT("ObjectReferenceCache::CreateObjectReferenceHandle Removed %s from ObjectToNetReferenceHandle due to stale cache. Cached Object (0x%p).  New object %s (0x%p)"), 
-				*CachedObject.NetRefHandle.ToString(), CachedObject.ObjectKey, *GetNameSafe(Object), Object);
-
-			ObjectToNetReferenceHandle.Remove(CachedObject.ObjectKey);
-			ObjectToNetReferenceHandle.Remove(Object);
-
-			if (RefHandle.IsStatic())
+			FCachedNetObjectReference& CachedObject = *CachedObjectPtr;
+			if (CachedObject.Object.Get() == Object)
 			{
-				UE_LOG_REFERENCECACHE(Verbose, TEXT("ObjectReferenceCache::CreateObjectReferenceHandle clearing stale cache for static handle %s. Cached Object (0x%p).  New object %s (0x%p)"),
-					*CachedObject.NetRefHandle.ToString(), CachedObject.ObjectKey, *GetNameSafe(Object), Object);
-
-				// Note we only cleanse the object reference
-				// we still keep the cachedObject data around in order to be able to serialize static destruction infos
-				CachedObject.ObjectKey = nullptr;
-				CachedObject.Object = nullptr;
-
+				UE_LOG_REFERENCECACHE(VeryVerbose, TEXT("ObjectReferenceCache::CreateObjectReferenceHandle Found existing %s for ObjectPath %s Object: %s (0x%p), OuterNetRefHandle: %s"),
+					*CachedObject.NetRefHandle.ToString(), ToCStr(Object->GetPathName()), *GetNameSafe(Object), Object, *CachedObject.OuterNetRefHandle.ToString());
+				OutReference = MakeNetObjectReference(CachedObject);
+				return true;
 			}
 			else
 			{
-				UE_LOG_REFERENCECACHE(Verbose, TEXT("ObjectReferenceCache::CreateObjectReferenceHandle removing cache for handle %s. Cached Object (0x%p). New object %s (0x%p)"),
+				UE_LOG_REFERENCECACHE(Verbose, TEXT("ObjectReferenceCache::CreateObjectReferenceHandle Removed %s from ObjectToNetReferenceHandle due to stale cache. Cached Object (0x%p).  New object %s (0x%p)"),
 					*CachedObject.NetRefHandle.ToString(), CachedObject.ObjectKey, *GetNameSafe(Object), Object);
 
-				ReferenceHandleToCachedReference.Remove(RefHandle);
+				ObjectToNetReferenceHandle.Remove(CachedObject.ObjectKey);
+				ObjectToNetReferenceHandle.Remove(Object);
+				RefHandlePtr = nullptr;
+
+				if (RefHandle.IsStatic())
+				{
+					UE_LOG_REFERENCECACHE(Verbose, TEXT("ObjectReferenceCache::CreateObjectReferenceHandle clearing stale cache for static handle %s. Cached Object (0x%p).  New object %s (0x%p)"),
+						*CachedObject.NetRefHandle.ToString(), CachedObject.ObjectKey, *GetNameSafe(Object), Object);
+
+					// Note we only cleanse the object reference
+					// we still keep the cachedObject data around in order to be able to serialize static destruction infos
+					CachedObject.ObjectKey = nullptr;
+					CachedObject.Object = nullptr;
+
+				}
+				else
+				{
+					UE_LOG_REFERENCECACHE(Verbose, TEXT("ObjectReferenceCache::CreateObjectReferenceHandle removing cache for handle %s. Cached Object (0x%p). New object %s (0x%p)"),
+						*CachedObject.NetRefHandle.ToString(), CachedObject.ObjectKey, *GetNameSafe(Object), Object);
+
+					ReferenceHandleToCachedReference.Remove(RefHandle);
+				}
 			}
 		}
+		else
+		{
+			UE_LOG_REFERENCECACHE(Warning, TEXT("ObjectReferenceCache::CreateObjectReferenceInternal removed %s from ObjectToNetReferenceHandle due to mismatch with cache. Object %s (0x%p) "), *RefHandle.ToString(), *GetNameSafe(Object), Object);
+			ObjectToNetReferenceHandle.Remove(Object);
+			RefHandlePtr = nullptr;
+		}
+		
 	}
 
 #if WITH_EDITOR
@@ -459,11 +469,13 @@ void FObjectReferenceCache::AddRemoteReference(FNetRefHandle RefHandle, const UO
 		// Dynamic object might have been referenced before it gets replicated, then it will then already be registered (with a path relative to its outer)
 		if (FNetRefHandle* ExistingHandle = ObjectToNetReferenceHandle.Find(Object))
 		{
-			FCachedNetObjectReference& ExistingCachedObject = ReferenceHandleToCachedReference.FindChecked(*ExistingHandle);
-			if (ExistingCachedObject.NetRefHandle != RefHandle)
+			if (FCachedNetObjectReference* ExistingCachedObject = ReferenceHandleToCachedReference.Find(*ExistingHandle))
 			{
-				UE_LOG_REFERENCECACHE_WARNING(TEXT("ObjectReferenceCache::AddRemoteReference: Detected conflicting Ref %s: Received: %s, Object: %s (0x%p)"), *ExistingCachedObject.NetRefHandle.ToString(),
-					*RefHandle.ToString(), *GetNameSafe(Object), Object);
+				if (ExistingCachedObject->NetRefHandle != RefHandle)
+				{
+					UE_LOG_REFERENCECACHE_WARNING(TEXT("ObjectReferenceCache::AddRemoteReference: Detected conflicting Ref %s: Received: %s, Object: %s (0x%p)"), *ExistingCachedObject->NetRefHandle.ToString(),
+						*RefHandle.ToString(), *GetNameSafe(Object), Object);
+				}
 			}
 		}
 	#endif
@@ -945,18 +957,27 @@ UObject* FObjectReferenceCache::ResolveObjectReferenceHandleInternal(FNetRefHand
 	FNetRefHandle* ExistingHandle = ObjectToNetReferenceHandle.Find(Object);
 	if (ExistingHandle && (*ExistingHandle != CompleteRefHandle))
 	{		
-		FCachedNetObjectReference& ExistingCachedObject = ReferenceHandleToCachedReference.FindChecked(*ExistingHandle);
-		if (ExistingCachedObject.Object.IsStale())
+		if (FCachedNetObjectReference* ExistingCachedObjectPtr = ReferenceHandleToCachedReference.Find(*ExistingHandle))
 		{
-			UE_LOG_REFERENCECACHE_WARNING(TEXT("Invalidating object for Stale Ref %s with %s as the object has been reused"), *ExistingCachedObject.NetRefHandle.ToString(), *CompleteRefHandle.ToString());
+			FCachedNetObjectReference& ExistingCachedObject = *ExistingCachedObjectPtr;
+			if (ExistingCachedObject.Object.IsStale())
+			{
+				UE_LOG_REFERENCECACHE_WARNING(TEXT("Invalidating object for Stale Ref %s with %s as the object has been reused"), *ExistingCachedObject.NetRefHandle.ToString(), *CompleteRefHandle.ToString());
+			}
+			else if (ExistingCachedObject.Object.HasSameIndexAndSerialNumber(CacheObjectPtr->Object))
+			{
+				bShouldAssign = CompleteRefHandle.GetId() > (*ExistingHandle).GetId();
+				if (bShouldAssign)
+				{
+					UE_LOG_REFERENCECACHE_WARNING(TEXT("Invalidating object for conflicting Ref %s with %s"), *ExistingCachedObject.NetRefHandle.ToString(), *CompleteRefHandle.ToString());
+				}	
+			}
 		}
-		else if (ExistingCachedObject.Object.HasSameIndexAndSerialNumber(CacheObjectPtr->Object))
+		else
 		{
-			bShouldAssign = CompleteRefHandle.GetId() > (*ExistingHandle).GetId();
-			if (bShouldAssign)
-			{
-				UE_LOG_REFERENCECACHE_WARNING(TEXT("Invalidating object for conflicting Ref %s with %s"), *ExistingCachedObject.NetRefHandle.ToString(), *CompleteRefHandle.ToString());
-			}	
+			UE_LOG_REFERENCECACHE(Warning, TEXT("ObjectReferenceCache::ResolveObjectReferenceHandleInternal removed %s from ObjectToNetReferenceHandle due to mismatch with cache. Object %s (0x%p) "), *ExistingHandle->ToString(), *GetNameSafe(Object), Object);
+			ObjectToNetReferenceHandle.Remove(Object);
+			ExistingHandle = nullptr;
 		}
 	}
 
