geometry collections emit an event when their particles are fully decayed

diff --git a/Engine/Source/Runtime/Experimental/Chaos/Private/GeometryCollectionProxyData.cpp b/Engine/Source/Runtime/Experimental/Chaos/Private/GeometryCollectionProxyData.cpp
--- a/Engine/Source/Runtime/Experimental/Chaos/Private/GeometryCollectionProxyData.cpp
+++ b/Engine/Source/Runtime/Experimental/Chaos/Private/GeometryCollectionProxyData.cpp
@@ -128,6 +128,11 @@ bool FGeometryCollectionDynamicStateFacade::IsValid() const
 		;
 }
 
+bool FGeometryCollectionDynamicStateFacade::IsActive(int32 TransformIndex) const
+{
+	return ActiveAttribute.Get()[TransformIndex];
+}
+
 bool FGeometryCollectionDynamicStateFacade::IsDynamicOrSleeping(int32 TransformIndex) const
 {
 	const int32 State = DynamicStateAttribute.Get()[TransformIndex];
@@ -147,9 +152,9 @@ bool FGeometryCollectionDynamicStateFacade::HasChildren(int32 TransformIndex) co
 
 bool FGeometryCollectionDynamicStateFacade::HasBrokenOff(int32 TransformIndex) const
 {
-	const bool IsActive = ActiveAttribute.Get()[TransformIndex];
-	const bool HasParent = (ParentAttribute.Get()[TransformIndex] != INDEX_NONE);
-	return IsActive && (!HasParent) && IsDynamicOrSleeping(TransformIndex);
+	const bool bIsActive = IsActive(TransformIndex);
+	const bool bHasParent = (ParentAttribute.Get()[TransformIndex] != INDEX_NONE);
+	return bIsActive && (!bHasParent) && IsDynamicOrSleeping(TransformIndex);
 }
 
 bool FGeometryCollectionDynamicStateFacade::HasInternalClusterParent(int32 TransformIndex) const
diff --git a/Engine/Source/Runtime/Experimental/Chaos/Public/GeometryCollectionProxyData.h b/Engine/Source/Runtime/Experimental/Chaos/Public/GeometryCollectionProxyData.h
--- a/Engine/Source/Runtime/Experimental/Chaos/Public/GeometryCollectionProxyData.h
+++ b/Engine/Source/Runtime/Experimental/Chaos/Public/GeometryCollectionProxyData.h
@@ -116,6 +116,9 @@ public:
 	/** returns true if all the necessary attributes are present */
 	CHAOS_API bool IsValid() const;
 
+	/** return true if the transform is active */
+	CHAOS_API bool IsActive(int32 TransformIndex) const;
+
 	/** return true if the transform is in a dynamic or sleeping state */
 	CHAOS_API bool IsDynamicOrSleeping(int32 TransformIndex) const;
 
diff --git a/Engine/Source/Runtime/Experimental/GeometryCollectionEngine/Private/GeometryCollection/GeometryCollectionComponent.cpp b/Engine/Source/Runtime/Experimental/GeometryCollectionEngine/Private/GeometryCollection/GeometryCollectionComponent.cpp
--- a/Engine/Source/Runtime/Experimental/GeometryCollectionEngine/Private/GeometryCollection/GeometryCollectionComponent.cpp
+++ b/Engine/Source/Runtime/Experimental/GeometryCollectionEngine/Private/GeometryCollection/GeometryCollectionComponent.cpp
@@ -2770,6 +2770,52 @@ void UGeometryCollectionComponent::TickComponent(float DeltaTime, enum ELevelTic
 	}
 }
 
+void UGeometryCollectionComponent::CheckFullyDecayed()
+{
+	if (bAlreadyFullyDecayed)
+	{
+		// Already fully decayed - don't bother doing extra work.
+		return;
+	}
+
+	if (DynamicCollection && PhysicsProxy)
+	{
+		bool bFullyDecayed = true;
+		FGeometryCollectionDecayDynamicFacade DecayFacade(*DynamicCollection);
+		FGeometryCollectionDynamicStateFacade DynamicStateFacade(*DynamicCollection);
+
+		const int32 NumTransforms = DecayFacade.GetDecayAttributeSize();
+		for (int32 TransformIdx = 0; TransformIdx < NumTransforms; ++TransformIdx)
+		{
+			// If we didn't create a particle for this transform, we shouldn't consider this particle.
+			if (!PhysicsProxy->GetExternalParticles()[TransformIdx])
+			{
+				continue;
+			}
+
+			// In an internal cluster, decay hasn't gotten to this particle yet (e.g. could be in a cluster union).
+			if (DynamicStateFacade.HasInternalClusterParent(TransformIdx))
+			{
+				bFullyDecayed = false;
+				break;
+			}
+
+			// If the particle is active, it's definitely not decayed either.
+			if (DynamicStateFacade.IsActive(TransformIdx))
+			{
+				bFullyDecayed = false;
+				break;
+			}
+		}
+
+		if (bFullyDecayed)
+		{
+			bAlreadyFullyDecayed = true;
+			OnFullyDecayedEvent.Broadcast();
+		}
+	}
+}
+
 void UGeometryCollectionComponent::AsyncPhysicsTickComponent(float DeltaTime, float SimTime)
 {
 	QUICK_SCOPE_CYCLE_COUNTER(GeometryCollectionComponent_AsyncPhysicsTick);
@@ -3276,6 +3322,8 @@ void UGeometryCollectionComponent::OnPostPhysicsSync()
 
 		// only update removal if the root is broken
 		UpdateRemovalIfNeeded();
+
+		CheckFullyDecayed();
 	}
 
 	const bool bDynamicDataIsDirty = (DynamicCollection && DynamicCollection->IsDirty() && HasVisibleGeometry());
diff --git a/Engine/Source/Runtime/Experimental/GeometryCollectionEngine/Public/GeometryCollection/GeometryCollectionComponent.h b/Engine/Source/Runtime/Experimental/GeometryCollectionEngine/Public/GeometryCollection/GeometryCollectionComponent.h
--- a/Engine/Source/Runtime/Experimental/GeometryCollectionEngine/Public/GeometryCollection/GeometryCollectionComponent.h
+++ b/Engine/Source/Runtime/Experimental/GeometryCollectionEngine/Public/GeometryCollection/GeometryCollectionComponent.h
@@ -63,6 +63,8 @@ DECLARE_DYNAMIC_MULTICAST_DELEGATE_OneParam(FOnChaosRemovalEvent, const FChaosRe
 
 DECLARE_DYNAMIC_MULTICAST_DELEGATE_OneParam(FOnChaosCrumblingEvent, const FChaosCrumblingEvent&, CrumbleEvent);
 
+DECLARE_DYNAMIC_MULTICAST_DELEGATE(FOnGeometryCollectionFullyDecayedEvent);
+
 namespace GeometryCollection
 {
 	enum class ESelectionMode : uint8
@@ -964,6 +966,7 @@ public:
 #endif  // UE_ENABLE_DEBUG_DRAWING
 
 	/**/
+	UE_DEPRECATED(5.4, "Disabled flags are no longer used.")
 	const TArray<bool>& GetDisabledFlags() const { return DisabledFlags; }
 
 	GEOMETRYCOLLECTIONENGINE_API virtual void OnCreatePhysicsState() override;
@@ -1026,6 +1029,8 @@ public:
 	// todo(chaos) remove when no longer necessary
 	FOnChaosBreakEvent OnRootBreakEvent;
 
+	FOnGeometryCollectionFullyDecayedEvent OnFullyDecayedEvent;
+
 	GEOMETRYCOLLECTIONENGINE_API void DispatchBreakEvent(const FChaosBreakEvent& Event);
 
 	GEOMETRYCOLLECTIONENGINE_API void DispatchRemovalEvent(const FChaosRemovalEvent& Event);
@@ -1296,6 +1301,9 @@ protected:
 
 	GEOMETRYCOLLECTIONENGINE_API virtual bool ProcessRepData(float DeltaTime, float SimTime);
 
+	void CheckFullyDecayed();
+	bool bAlreadyFullyDecayed = false;
+
 	int32 VersionProcessed = INDEX_NONE;
 
 	// The last time (in milliseconds) the async physics component tick fired.
