GeometryCollection ISMPool add support for ISM GPU LOD.
Also add support for Velocity from WPO option and tidy up boolean options into flags.

diff --git a/Engine/Source/Runtime/Experimental/GeometryCollectionEngine/Private/GeometryCollection/GeometryCollectionISMPoolComponent.cpp b/Engine/Source/Runtime/Experimental/GeometryCollectionEngine/Private/GeometryCollection/GeometryCollectionISMPoolComponent.cpp
--- a/Engine/Source/Runtime/Experimental/GeometryCollectionEngine/Private/GeometryCollection/GeometryCollectionISMPoolComponent.cpp
+++ b/Engine/Source/Runtime/Experimental/GeometryCollectionEngine/Private/GeometryCollection/GeometryCollectionISMPoolComponent.cpp
@@ -67,7 +67,7 @@ FGeometryCollectionISM::FGeometryCollectionISM(AActor* InOwningActor, const FGeo
 	UHierarchicalInstancedStaticMeshComponent* HISMC = nullptr;
 	UInstancedStaticMeshComponent* ISMC = nullptr;
 	
-	if (MeshInstance.Desc.bUseHISM)
+	if ((MeshInstance.Desc.Flags & FISMComponentDescription::UseHISM) != 0)
 	{
 		const FName ISMName = MakeUniqueObjectName(InOwningActor, UHierarchicalInstancedStaticMeshComponent::StaticClass(), MeshInstance.StaticMesh->GetFName());
 		ISMC = HISMC = NewObject<UHierarchicalInstancedStaticMeshComponent>(InOwningActor, ISMName, RF_Transient | RF_DuplicateTransient);
@@ -95,12 +95,14 @@ FGeometryCollectionISM::FGeometryCollectionISM(AActor* InOwningActor, const FGeo
 
 	ISMC->SetRemoveSwap();
 	ISMC->NumCustomDataFloats = MeshInstance.Desc.NumCustomDataFloats;
-	ISMC->SetReverseCulling(MeshInstance.Desc.bReverseCulling);
-	ISMC->SetMobility(MeshInstance.Desc.bIsStaticMobility ? EComponentMobility::Static : EComponentMobility::Stationary);
+	ISMC->SetReverseCulling((MeshInstance.Desc.Flags & FISMComponentDescription::ReverseCulling) != 0);
+	ISMC->SetMobility((MeshInstance.Desc.Flags & FISMComponentDescription::StaticMobility) != 0 ? EComponentMobility::Static : EComponentMobility::Stationary);
 	ISMC->SetCullDistances(MeshInstance.Desc.StartCullDistance, MeshInstance.Desc.EndCullDistance);
-	ISMC->SetCastShadow(MeshInstance.Desc.bAffectShadow);
-	ISMC->bAffectDynamicIndirectLighting = MeshInstance.Desc.bAffectDynamicIndirectLighting;
-	ISMC->bAffectDistanceFieldLighting = MeshInstance.Desc.bAffectDistanceFieldLighting;
+	ISMC->SetCastShadow((MeshInstance.Desc.Flags & FISMComponentDescription::AffectShadow) != 0);
+	ISMC->bAffectDynamicIndirectLighting = (MeshInstance.Desc.Flags & FISMComponentDescription::AffectDynamicIndirectLighting) != 0;
+	ISMC->bAffectDistanceFieldLighting = (MeshInstance.Desc.Flags & FISMComponentDescription::AffectDistanceFieldLighting) != 0;
+	ISMC->bWorldPositionOffsetWritesVelocity = (MeshInstance.Desc.Flags & FISMComponentDescription::WorldPositionOffsetWritesVelocity) != 0;
+	ISMC->bUseGpuLodSelection = (MeshInstance.Desc.Flags & FISMComponentDescription::GpuLodSelection) != 0;
 	ISMC->SetCanEverAffectNavigation(false);
 	ISMC->SetCollisionEnabled(ECollisionEnabled::NoCollision);
 	ISMC->bOverrideMinLOD = MeshInstance.Desc.MinLod > 0;	
diff --git a/Engine/Source/Runtime/Experimental/GeometryCollectionEngine/Private/GeometryCollection/GeometryCollectionISMPoolRenderer.cpp b/Engine/Source/Runtime/Experimental/GeometryCollectionEngine/Private/GeometryCollection/GeometryCollectionISMPoolRenderer.cpp
--- a/Engine/Source/Runtime/Experimental/GeometryCollectionEngine/Private/GeometryCollection/GeometryCollectionISMPoolRenderer.cpp
+++ b/Engine/Source/Runtime/Experimental/GeometryCollectionEngine/Private/GeometryCollection/GeometryCollectionISMPoolRenderer.cpp
@@ -102,7 +102,7 @@ void UGeometryCollectionISMPoolRenderer::InitMergedMeshFromGeometryCollection(UG
 
 		FGeometryCollectionStaticMeshInstance StaticMeshInstance;
 		StaticMeshInstance.StaticMesh = StaticMesh;
-		StaticMeshInstance.Desc.bUseHISM = true;
+		StaticMeshInstance.Desc.Flags |= FISMComponentDescription::UseHISM;
 
 		TArray<float> DummyCustomData;
 		MergedMeshGroup.MeshIds.Add(ISMPoolComponent->AddMeshToGroup(MergedMeshGroup.GroupIndex, StaticMeshInstance, 1, DummyCustomData));
@@ -147,7 +147,7 @@ void UGeometryCollectionISMPoolRenderer::InitInstancesFromGeometryCollection(UGe
 			{
 				StaticMeshInstance.MaterialsOverrides = AutoInstanceMesh.Materials;
 			}
-			StaticMeshInstance.Desc.bUseHISM = true;
+			StaticMeshInstance.Desc.Flags |= FISMComponentDescription::UseHISM;
 
 			TArray<float> DummyCustomData;
 			InstancesGroup.MeshIds.Add(ISMPoolComponent->AddMeshToGroup(InstancesGroup.GroupIndex, StaticMeshInstance, AutoInstanceMesh.NumInstances, DummyCustomData));
diff --git a/Engine/Source/Runtime/Experimental/GeometryCollectionEngine/Public/GeometryCollection/GeometryCollectionISMPoolComponent.h b/Engine/Source/Runtime/Experimental/GeometryCollectionEngine/Public/GeometryCollection/GeometryCollectionISMPoolComponent.h
--- a/Engine/Source/Runtime/Experimental/GeometryCollectionEngine/Public/GeometryCollection/GeometryCollectionISMPoolComponent.h
+++ b/Engine/Source/Runtime/Experimental/GeometryCollectionEngine/Public/GeometryCollection/GeometryCollectionISMPoolComponent.h
@@ -99,12 +99,19 @@ struct FInstanceGroups
  */
 struct FISMComponentDescription
 {
-	bool bUseHISM = false;
-	bool bReverseCulling = false;
-	bool bIsStaticMobility = false;
-	bool bAffectShadow = true;
-	bool bAffectDistanceFieldLighting = false;
-	bool bAffectDynamicIndirectLighting = false;
+	enum EFlags
+	{
+		UseHISM = 1 << 1,
+		GpuLodSelection = 1 << 2,
+		ReverseCulling = 1 << 3,
+		StaticMobility = 1 << 4,
+		WorldPositionOffsetWritesVelocity = 1 << 5,
+		AffectShadow = 1 << 6,
+		AffectDistanceFieldLighting = 1 << 7,
+		AffectDynamicIndirectLighting = 1 << 8,
+	};
+
+	uint32 Flags = WorldPositionOffsetWritesVelocity|AffectShadow;
 	int32 NumCustomDataFloats = 0;
 	int32 StartCullDistance = 0;
 	int32 EndCullDistance = 0;
@@ -114,12 +121,7 @@ struct FISMComponentDescription
 
 	bool operator==(const FISMComponentDescription& Other) const
 	{
-		return bUseHISM == Other.bUseHISM &&
-			bReverseCulling == Other.bReverseCulling &&
-			bIsStaticMobility == Other.bIsStaticMobility &&
-			bAffectShadow == Other.bAffectShadow &&
-			bAffectDistanceFieldLighting == Other.bAffectDistanceFieldLighting &&
-			bAffectDynamicIndirectLighting == Other.bAffectDistanceFieldLighting &&
+		return Flags == Other.Flags &&
 			NumCustomDataFloats == Other.NumCustomDataFloats &&
 			StartCullDistance == Other.StartCullDistance && 
 			EndCullDistance == Other.EndCullDistance &&
@@ -131,14 +133,7 @@ struct FISMComponentDescription
 
 FORCEINLINE uint32 GetTypeHash(const FISMComponentDescription& Desc)
 {
-	const uint32 PackedBools = 
-		(Desc.bUseHISM ? 1 : 0) | 
-		(Desc.bReverseCulling ? 2 : 0) | 
-		(Desc.bIsStaticMobility ? 4 : 0) | 
-		(Desc.bAffectShadow ? 8 : 0) | 
-		(Desc.bAffectDistanceFieldLighting ? 16 : 0) |
-		(Desc.bAffectDynamicIndirectLighting ? 32 : 0);
-	uint32 Hash = HashCombineFast(GetTypeHash(PackedBools), GetTypeHash(Desc.NumCustomDataFloats));
+	uint32 Hash = HashCombineFast(GetTypeHash(Desc.Flags), GetTypeHash(Desc.NumCustomDataFloats));
 	Hash = HashCombineFast(Hash, GetTypeHash(Desc.StartCullDistance));
 	Hash = HashCombineFast(Hash, GetTypeHash(Desc.EndCullDistance));
 	Hash = HashCombineFast(Hash, GetTypeHash(Desc.MinLod));
