fix for potentially passing a zero-length normal FindGeometryOpposingNormal in SweepQuery

diff --git a/Engine/Source/Runtime/Experimental/Chaos/Public/Chaos/GeometryQueries.h b/Engine/Source/Runtime/Experimental/Chaos/Public/Chaos/GeometryQueries.h
--- a/Engine/Source/Runtime/Experimental/Chaos/Public/Chaos/GeometryQueries.h
+++ b/Engine/Source/Runtime/Experimental/Chaos/Public/Chaos/GeometryQueries.h
@@ -302,7 +302,7 @@ namespace Chaos
 		//put back into world space
 		if (bResult && (OutTime > 0 || bComputeMTD))
 		{
-			OutNormal = ATM.TransformVectorNoScale(LocalNormal).GetSafeNormal();
+			OutNormal = ATM.TransformVectorNoScale(LocalNormal).GetSafeNormal(UE_KINDA_SMALL_NUMBER, FVec3::AxisVector(0));
 			OutPosition = ATM.TransformPositionNoScale(LocalPosition);
 			OutFaceNormal = ATM.TransformVectorNoScale(A.FindGeometryOpposingNormal(LocalDir, OutFaceIndex, OutNormal));
 		}
