Iris - Add message to checks.

diff --git a/Engine/Source/Runtime/Experimental/Iris/Core/Private/Iris/ReplicationSystem/ReplicationWriter.cpp b/Engine/Source/Runtime/Experimental/Iris/Core/Private/Iris/ReplicationSystem/ReplicationWriter.cpp
--- a/Engine/Source/Runtime/Experimental/Iris/Core/Private/Iris/ReplicationSystem/ReplicationWriter.cpp
+++ b/Engine/Source/Runtime/Experimental/Iris/Core/Private/Iris/ReplicationSystem/ReplicationWriter.cpp
@@ -100,67 +100,67 @@ void FReplicationWriter::FReplicationInfo::SetState(EReplicatedObjectState NewSt
 	{
 		case EReplicatedObjectState::PendingCreate:
 		{
-			check(CurrentState == EReplicatedObjectState::Invalid || CurrentState == EReplicatedObjectState::WaitOnCreateConfirmation);
+			checkf(CurrentState == EReplicatedObjectState::Invalid || CurrentState == EReplicatedObjectState::WaitOnCreateConfirmation, TEXT("Trying to set state %s when state is %s. IsDestructionInfo: %u IsSubObject: %u"), LexToString(NewState), LexToString(CurrentState), IsDestructionInfo, IsSubObject);
 		}
 		break;
 		case EReplicatedObjectState::WaitOnCreateConfirmation:
 		{
-			check(CurrentState == EReplicatedObjectState::PendingCreate || CurrentState == EReplicatedObjectState::CancelPendingDestroy);
+			checkf(CurrentState == EReplicatedObjectState::PendingCreate || CurrentState == EReplicatedObjectState::CancelPendingDestroy, TEXT("Trying to set state %s when state is %s. IsDestructionInfo: %u IsSubObject: %u"), LexToString(NewState), LexToString(CurrentState), IsDestructionInfo, IsSubObject);
 		}
 		break;
 		case EReplicatedObjectState::Created:
 		{
-			check(CurrentState == EReplicatedObjectState::PendingCreate || CurrentState == EReplicatedObjectState::WaitOnCreateConfirmation || CurrentState == EReplicatedObjectState::CancelPendingDestroy || CurrentState == EReplicatedObjectState::WaitOnFlush);
+			checkf(CurrentState == EReplicatedObjectState::PendingCreate || CurrentState == EReplicatedObjectState::WaitOnCreateConfirmation || CurrentState == EReplicatedObjectState::CancelPendingDestroy || CurrentState == EReplicatedObjectState::WaitOnFlush, TEXT("Trying to set state %s when state is %s. IsDestructionInfo: %u IsSubObject: %u"), LexToString(NewState), LexToString(CurrentState), IsDestructionInfo, IsSubObject);
 		}
 		break;
 		case EReplicatedObjectState::PendingTearOff:
 		{
-			check(CurrentState == EReplicatedObjectState::PendingTearOff || CurrentState == EReplicatedObjectState::WaitOnFlush || CurrentState == EReplicatedObjectState::WaitOnCreateConfirmation || CurrentState == EReplicatedObjectState::Created || CurrentState == EReplicatedObjectState::WaitOnDestroyConfirmation);
+			checkf(CurrentState == EReplicatedObjectState::PendingTearOff || CurrentState == EReplicatedObjectState::WaitOnFlush || CurrentState == EReplicatedObjectState::WaitOnCreateConfirmation || CurrentState == EReplicatedObjectState::Created || CurrentState == EReplicatedObjectState::WaitOnDestroyConfirmation, TEXT("Trying to set state %s when state is %s. IsDestructionInfo: %u IsSubObject: %u"), LexToString(NewState), LexToString(CurrentState), IsDestructionInfo, IsSubObject);
 		}
 		break;
 		case EReplicatedObjectState::SubObjectPendingDestroy:
 		{
-			check(CurrentState == EReplicatedObjectState::PendingDestroy || CurrentState == EReplicatedObjectState::SubObjectPendingDestroy || CurrentState == EReplicatedObjectState::WaitOnCreateConfirmation || CurrentState == EReplicatedObjectState::Created || CurrentState == EReplicatedObjectState::WaitOnFlush || CurrentState == EReplicatedObjectState::WaitOnDestroyConfirmation);
+			checkf(CurrentState == EReplicatedObjectState::PendingDestroy || CurrentState == EReplicatedObjectState::SubObjectPendingDestroy || CurrentState == EReplicatedObjectState::WaitOnCreateConfirmation || CurrentState == EReplicatedObjectState::Created || CurrentState == EReplicatedObjectState::WaitOnFlush || CurrentState == EReplicatedObjectState::WaitOnDestroyConfirmation, TEXT("Trying to set state %s when state is %s. IsDestructionInfo: %u IsSubObject: %u"), LexToString(NewState), LexToString(CurrentState), IsDestructionInfo, IsSubObject);
 		}
 		break;
 		case EReplicatedObjectState::WaitOnFlush:
 		{
-			check(CurrentState != EReplicatedObjectState::Invalid);
+			checkf(CurrentState != EReplicatedObjectState::Invalid, TEXT("Trying to set state %s when state is %s. IsDestructionInfo: %u IsSubObject: %u"), LexToString(NewState), LexToString(CurrentState), IsDestructionInfo, IsSubObject);
 		}
 		break;
 		case EReplicatedObjectState::PendingDestroy:
 		{
-			check(CurrentState != EReplicatedObjectState::Invalid);
+			checkf(CurrentState != EReplicatedObjectState::Invalid, TEXT("Trying to set state %s when state is %s. IsDestructionInfo: %u IsSubObject: %u"), LexToString(NewState), LexToString(CurrentState), IsDestructionInfo, IsSubObject);
 		}
 		break;
 		case EReplicatedObjectState::WaitOnDestroyConfirmation:
 		{
-			check(CurrentState >= EReplicatedObjectState::PendingTearOff);
+			checkf(CurrentState >= EReplicatedObjectState::PendingTearOff, TEXT("Trying to set state %s when state is %s. IsDestructionInfo: %u IsSubObject: %u"), LexToString(NewState), LexToString(CurrentState), IsDestructionInfo, IsSubObject);
 		}
 		break;
 		case EReplicatedObjectState::CancelPendingDestroy:
 		{
-			check(CurrentState == EReplicatedObjectState::WaitOnDestroyConfirmation || CurrentState == EReplicatedObjectState::CancelPendingDestroy);
+			checkf(CurrentState == EReplicatedObjectState::WaitOnDestroyConfirmation || CurrentState == EReplicatedObjectState::CancelPendingDestroy, TEXT("Trying to set state %s when state is %s. IsDestructionInfo: %u IsSubObject: %u"), LexToString(NewState), LexToString(CurrentState), IsDestructionInfo, IsSubObject);
 		}
 		break;
 		case EReplicatedObjectState::Destroyed:
 		{
-			check(CurrentState == EReplicatedObjectState::WaitOnDestroyConfirmation || CurrentState == EReplicatedObjectState::PendingTearOff || CurrentState == EReplicatedObjectState::CancelPendingDestroy);
+			checkf(CurrentState == EReplicatedObjectState::WaitOnDestroyConfirmation || CurrentState == EReplicatedObjectState::PendingTearOff || CurrentState == EReplicatedObjectState::CancelPendingDestroy, TEXT("Trying to set state %s when state is %s. IsDestructionInfo: %u IsSubObject: %u"), LexToString(NewState), LexToString(CurrentState), IsDestructionInfo, IsSubObject);
 		}
 		break;
 		case EReplicatedObjectState::PermanentlyDestroyed:
 		{
-			check(CurrentState == EReplicatedObjectState::Invalid || CurrentState == EReplicatedObjectState::WaitOnCreateConfirmation);
+			checkf(CurrentState == EReplicatedObjectState::Invalid || CurrentState == EReplicatedObjectState::WaitOnCreateConfirmation, TEXT("Trying to set state %s when state is %s. IsDestructionInfo: %u IsSubObject: %u"), LexToString(NewState), LexToString(CurrentState), IsDestructionInfo, IsSubObject);
 		}
 		break;
 		case EReplicatedObjectState::Invalid:
 		{
-			check(CurrentState == EReplicatedObjectState::PermanentlyDestroyed || CurrentState == EReplicatedObjectState::Destroyed || CurrentState == EReplicatedObjectState::PendingCreate);
+			checkf(CurrentState == EReplicatedObjectState::PermanentlyDestroyed || CurrentState == EReplicatedObjectState::Destroyed || CurrentState == EReplicatedObjectState::PendingCreate, TEXT("Trying to set state %s when state is %s. IsDestructionInfo: %u IsSubObject: %u"), LexToString(NewState), LexToString(CurrentState), IsDestructionInfo, IsSubObject);
 		}
 		break;
 
 		default:
-			check(false);
+			checkf(false, TEXT("Trying to set state %s when state is %s. IsDestructionInfo: %u IsSubObject: %u"), ToCStr(FString::FromInt(int(NewState))), LexToString(CurrentState), IsDestructionInfo, IsSubObject);
 			break;
 	};
 
diff --git a/Engine/Source/Runtime/Experimental/Iris/Core/Private/Iris/ReplicationSystem/ReplicationWriter.h b/Engine/Source/Runtime/Experimental/Iris/Core/Private/Iris/ReplicationSystem/ReplicationWriter.h
--- a/Engine/Source/Runtime/Experimental/Iris/Core/Private/Iris/ReplicationSystem/ReplicationWriter.h
+++ b/Engine/Source/Runtime/Experimental/Iris/Core/Private/Iris/ReplicationSystem/ReplicationWriter.h
@@ -88,7 +88,7 @@ public:
 
 	static_assert((uint8)(EReplicatedObjectState::Count) <= 32, "EReplicatedObjectState must fit in 5 bits. See FReplicationInfo::State and FReplicationRecord::FRecordInfo::ReplicatedObjectState members.");
 
-	const TCHAR* LexToString(const EReplicatedObjectState State);
+	static const TCHAR* LexToString(const EReplicatedObjectState State);
 
 	enum EFlushFlags : uint32
 	{
@@ -127,7 +127,6 @@ public:
 				uint64 LastAckedBaselineIndex : 2;						// Last acknowledged baseline index which we can use for deltacompresion
 				uint64 PendingBaselineIndex : 2;						// Baseline index pending acknowledgment from client
 				uint64 FlushFlags : 2;									// Flags indicating what we are waiting for when flushing
-				uint64 Padding : 27;
 			};
 		};
 
