fix issue where adding a particle with an existing parent into a cluster union would cause that particle to snap to a different position due to ChildToParent being set too early

diff --git a/Engine/Source/Runtime/Experimental/Chaos/Private/Chaos/ClusterUnionManager.cpp b/Engine/Source/Runtime/Experimental/Chaos/Private/Chaos/ClusterUnionManager.cpp
--- a/Engine/Source/Runtime/Experimental/Chaos/Private/Chaos/ClusterUnionManager.cpp
+++ b/Engine/Source/Runtime/Experimental/Chaos/Private/Chaos/ClusterUnionManager.cpp
@@ -447,11 +447,17 @@ namespace Chaos
 		}
 
 		Cluster->ChildParticles.Append(FinalParticlesToAdd);
+
+		// We need to call AddParticlesToCluster before we set the temporary ChildToParent. RemoveFromParent (which AddParticlesToCluster calls)
+		// will set the particles' X/R based on ChildToParent if the particle has an old parent. This ChildToParent needs to be correct relative to the old parent,
+		// not the new parent when AddParticlesToCluster is called.
+		MClustering.AddParticlesToCluster(Cluster->InternalCluster, FinalParticlesToAdd, ChildToParentMap);
+
 		for (FPBDRigidParticleHandle* Particle : FinalParticlesToAdd)
 		{
 			if (!bIsNewCluster)
 			{
-				Cluster->PendingConnectivityOperations.Add({Particle, EClusterUnionConnectivityOperation ::Add});
+				Cluster->PendingConnectivityOperations.Add({ Particle, EClusterUnionConnectivityOperation::Add });
 			}
 
 			if (!Cluster->ChildProperties.Contains(Particle))
@@ -471,8 +477,6 @@ namespace Chaos
 			}
 		}
 
-		MClustering.AddParticlesToCluster(Cluster->InternalCluster, FinalParticlesToAdd, ChildToParentMap);
-
 		// For all the particles that have been added to the cluster we need to set their parent proxy to the
 		// cluster's proxy if it exists. We need the proxy type check because for non-cluster union proxy backed unions,
 		// the cluster union's proxy will be the proxy of the most recently added particle.
