fix crash in the AABBTree when potentially accessing a particle with a nullptr debugname

diff --git a/Engine/Source/Runtime/Experimental/Chaos/Public/Chaos/AABBTree.h b/Engine/Source/Runtime/Experimental/Chaos/Public/Chaos/AABBTree.h
--- a/Engine/Source/Runtime/Experimental/Chaos/Public/Chaos/AABBTree.h
+++ b/Engine/Source/Runtime/Experimental/Chaos/Public/Chaos/AABBTree.h
@@ -1742,9 +1742,10 @@ public:
 #if CHAOS_DEBUG_NAME			
 			if constexpr (std::is_same_v<TPayloadType, FAccelerationStructureHandle>)
 			{
+				const TSharedPtr<FString, ESPMode::ThreadSafe>& DebugName = Payload.GetGeometryParticleHandle_PhysicsThread()->DebugName();
 				if (IsInPhysicsThreadContext())
 				{
-					FString DebugStr = *(Payload.GetGeometryParticleHandle_PhysicsThread()->DebugName());
+					FString DebugStr = DebugName ? *DebugName : TEXT("No Name");
 					UE_LOG(LogChaos, Warning, TEXT("AABBTree encountered invalid bounds input : %s"), *DebugStr);
 				}
 			}
