Add an asynchronous no blocking Game thread mode. The game thread is no wating for the physics thread anymore. This mode could eventually drop physics steps.

diff --git a/Engine/Source/Runtime/Experimental/Chaos/Private/Chaos/Framework/PhysicsSolverBase.cpp b/Engine/Source/Runtime/Experimental/Chaos/Private/Chaos/Framework/PhysicsSolverBase.cpp
--- a/Engine/Source/Runtime/Experimental/Chaos/Private/Chaos/Framework/PhysicsSolverBase.cpp
+++ b/Engine/Source/Runtime/Experimental/Chaos/Private/Chaos/Framework/PhysicsSolverBase.cpp
@@ -136,6 +136,7 @@ namespace Chaos
 		, PushData(InPushData)
 	{
 		CVD_GET_CURRENT_CONTEXT(CVDContext);
+		Solver.NumPendingSolverAdvanceTasks++;
 	}
 
 	void FPhysicsSolverAdvanceTask::DoTask(ENamedThreads::Type CurrentThread, const FGraphEventRef& MyCompletionGraphEvent)
@@ -153,6 +154,7 @@ namespace Chaos
 		SCOPE_CYCLE_COUNTER(STAT_ChaosTick);
 		CSV_SCOPED_TIMING_STAT_EXCLUSIVE(Physics);
 		CSV_SCOPED_TIMING_STAT(PhysicsVerbose, StepSolver);
+		PHYSICS_CSV_CUSTOM_EXPENSIVE(PhysicsCounters, NumPendingSolverAdvanceTasks, NumPendingSolverAdvanceTasks, ECsvCustomStatOp::Max);
 
 #if PHYSICS_THREAD_CONTEXT
 		FPhysicsThreadContextScope Scope(/*IsPhysicsThreadContext=*/true);
@@ -168,6 +170,8 @@ namespace Chaos
 		PushData = nullptr;
 
 		Solver.ConditionalApplyRewind_Internal();
+
+		Solver.NumPendingSolverAdvanceTasks--;
 	}
 
 	CHAOS_API int32 UseAsyncInterpolation = 1;
@@ -181,10 +185,12 @@ namespace Chaos
 
 	// 0 blocks on any physics steps generated from past GT Frames, and blocks on none of the tasks from current frame.
 	// 1 blocks on everything except the single most recent task (including tasks from current frame)
-	// 1 should gurantee we will always have a future output for interpolation from 2 frames in the past
+	// 1 should guarantee we will always have a future output for interpolation from 2 frames in the past
+	// 2 doesn't block the game thread. Physics steps could be eventually be dropped if taking too much time.
 	int32 AsyncPhysicsBlockMode = 0;
 	FAutoConsoleVariableRef CVarAsyncPhysicsBlockMode(TEXT("p.AsyncPhysicsBlockMode"), AsyncPhysicsBlockMode, TEXT("Setting to 0 blocks on any physics steps generated from past GT Frames, and blocks on none of the tasks from current frame."
-		" 1 blocks on everything except the single most recent task (including tasks from current frame). 1 should gurantee we will always have a future output for interpolation from 2 frames in the past."));
+		" 1 blocks on everything except the single most recent task (including tasks from current frame). 1 should gurantee we will always have a future output for interpolation from 2 frames in the past."
+		" 2 doesn't block the game thread, physics steps could be eventually be dropped if taking too much time."));
 
 
 	FPhysicsSolverBase::FPhysicsSolverBase(const EMultiBufferMode BufferingModeIn,const EThreadingModeTemp InThreadingMode,UObject* InOwner, Chaos::FReal InAsyncDt)
@@ -193,6 +199,7 @@ namespace Chaos
 		, PullResultsManager(MakeUnique<FChaosResultsManager>(MarshallingManager))
 		, PendingSpatialOperations_External(MakeUnique<FPendingSpatialDataQueue>())
 		, bUseCollisionResimCache(false)
+		, NumPendingSolverAdvanceTasks(0)
 		, bPaused_External(false)
 		, Owner(InOwner)
 		, ExternalDataLock_External(new FPhysSceneLock())
@@ -410,7 +417,19 @@ namespace Chaos
 			ExternalSteps++;	//we use this to average forces. It assumes external dt is about the same. 0 dt should be ignored as it typically has nothing to do with force
 		}
 
-		if(NumSteps > 0)
+		// Eventually drop physics steps in mode 2
+		if (AsyncPhysicsBlockMode == 2)
+		{
+			// Make sure not to accumulate too many physics solver tasks.
+			constexpr int32 MaxPhysicsStepToKeep = 3;
+			if (NumSteps + NumPendingSolverAdvanceTasks > MaxPhysicsStepToKeep)
+			{
+				// NumSteps + NumPendingSolverAdvanceTasks shouldn't be bigger than MaxPhysicsStepToKeep
+				NumSteps = FMath::Clamp<int32>(NumSteps, 0, MaxPhysicsStepToKeep - NumPendingSolverAdvanceTasks);
+			}
+		}
+			
+		if (NumSteps > 0)
 		{
 			//make sure any GT state is pushed into necessary buffer
 			PushPhysicsState(InternalDt, NumSteps, FMath::Max(ExternalSteps, 1));
@@ -447,13 +466,13 @@ namespace Chaos
 			else
 			{
 				// If enabled, block on all but most recent physics task, even tasks generated this frame.
-				if(AsyncPhysicsBlockMode == 1)
+				if (AsyncPhysicsBlockMode == 1)
 				{
 					BlockingTasks = PendingTasks;
 				}
 
 				FGraphEventArray Prereqs;
-				if(PendingTasks && !PendingTasks->IsComplete())
+				if (PendingTasks && !PendingTasks->IsComplete())
 				{
 					Prereqs.Add(PendingTasks);
 				}
@@ -461,15 +480,15 @@ namespace Chaos
 				PendingTasks = TGraphTask<FPhysicsSolverProcessPushDataTask>::CreateTask(&Prereqs).ConstructAndDispatchWhenReady(*this, PushData);
 				Prereqs.Add(PendingTasks);
 
-				if(bSolverHasFrozenGameThreadCallbacks)
+				if (bSolverHasFrozenGameThreadCallbacks)
 				{
 					PendingTasks = TGraphTask<FPhysicsSolverFrozenGTPreSimCallbacks>::CreateTask(&Prereqs).ConstructAndDispatchWhenReady(*this);
 					Prereqs.Add(PendingTasks);
 				}
-				
+
 				PendingTasks = TGraphTask<FPhysicsSolverAdvanceTask>::CreateTask(&Prereqs).ConstructAndDispatchWhenReady(*this, PushData);
 
-				if(IsUsingAsyncResults() == false)
+				if (IsUsingAsyncResults() == false)
 				{
 					BlockingTasks = PendingTasks;	//block right away
 				}
@@ -482,7 +501,10 @@ namespace Chaos
 				break;
 			}
 		}
-
+		if (AsyncPhysicsBlockMode == 2)
+		{
+			return {};
+		}
 		return BlockingTasks;
 	}
 
diff --git a/Engine/Source/Runtime/Experimental/Chaos/Public/Chaos/Framework/PhysicsSolverBase.h b/Engine/Source/Runtime/Experimental/Chaos/Public/Chaos/Framework/PhysicsSolverBase.h
--- a/Engine/Source/Runtime/Experimental/Chaos/Public/Chaos/Framework/PhysicsSolverBase.h
+++ b/Engine/Source/Runtime/Experimental/Chaos/Public/Chaos/Framework/PhysicsSolverBase.h
@@ -751,6 +751,9 @@ namespace Chaos
 			MarshallingManager.RegisterSimCallbackObject_External(SimCallbackObject);
 		}
 
+		// Number of ending solver task that has not been executed or that are still being executed
+		TAtomic<int32> NumPendingSolverAdvanceTasks;
+
 		/** 
 		 * Whether this solver is paused. Paused solvers will still 'tick' however they will receive a Dt of zero so they can still
 		 * build acceleration structures or accept inputs from external threads 
