Chaos Cloth: flip warp and weft directions on anisotropic constraints.

diff --git a/Engine/Source/Runtime/Experimental/Chaos/Private/Chaos/XPBDAnisotropicBendingConstraints.cpp b/Engine/Source/Runtime/Experimental/Chaos/Private/Chaos/XPBDAnisotropicBendingConstraints.cpp
--- a/Engine/Source/Runtime/Experimental/Chaos/Private/Chaos/XPBDAnisotropicBendingConstraints.cpp
+++ b/Engine/Source/Runtime/Experimental/Chaos/Private/Chaos/XPBDAnisotropicBendingConstraints.cpp
@@ -178,23 +178,24 @@ TArray<FSolverVec3> FXPBDAnisotropicBendingConstraints::GenerateWarpWeftBiasBase
 	auto SortedEdge = [](int32 P0, int32 P1) { return P0 <= P1 ? TVec2<int32>(P0, P1) : TVec2<int32>(P1, P0); };
 	auto Multiplier = [](const FVec2f& UV0, const FVec2f& UV1)
 	{
-		// Input UVs: X = Warp, Y = Weft direction
-		// Transform to X = Warp, Y = Weft, Z = Bias
+		// Input UVs: X = Weft, Y = Warp direction
+		// Transform to X = Weft, Y = Warp, Z = Bias
+		// NOTE: Weft-dominant bend is along a vertical edge (i.e., larger V(Y) diff)). Warp-dominant bend is along a horizontal edge.
 		const FSolverVec2 UVDiff = UV1 - UV0;
 		const FSolverVec2 UVDiffAbs = UVDiff.GetAbs();
 		FSolverVec3 UVDiffTransformed;
 		if (UVDiffAbs.X > UVDiffAbs.Y)
 		{
 			// Stiffness is blend between warp and bias
-			UVDiffTransformed  = FSolverVec3(UVDiffAbs.X - UVDiffAbs.Y, (FSolverReal)0.f, UVDiffAbs.Y);
+			UVDiffTransformed  = FSolverVec3((FSolverReal)0.f, UVDiffAbs.X - UVDiffAbs.Y, UVDiffAbs.Y);
 		}
 		else
 		{
 			// Stiffness is blend between weft and bias
-			UVDiffTransformed = FSolverVec3((FSolverReal)0.f, UVDiffAbs.Y - UVDiffAbs.X, UVDiffAbs.X);
+			UVDiffTransformed = FSolverVec3(UVDiffAbs.Y - UVDiffAbs.X, (FSolverReal)0.f, UVDiffAbs.X);
 		}
 		const FSolverReal Denom = UVDiffTransformed.X + UVDiffTransformed.Y + UVDiffTransformed.Z;
-		return Denom > UE_SMALL_NUMBER ? UVDiffTransformed / Denom : FSolverVec3(1.f, 0.f, 0.f); // Default to Warp if zero length
+		return Denom > UE_SMALL_NUMBER ? UVDiffTransformed / Denom : FSolverVec3(0.f, 1.f, 0.f); // Default to Warp if zero length
 	};
 
 	const TArray<TVec3<int32>>& Elements = TriangleMesh.GetElements();
@@ -499,8 +500,8 @@ void FXPBDAnisotropicBendingConstraints::Apply(FSolverParticles& Particles, cons
 			!BucklingStiffnessHasWeightMap && !BucklingStiffnessWeftHasWeightMap && !BucklingStiffnessBiasHasWeightMap &&
 			!DampingHasWeightMap)
 		{
-			const FSolverVec3 ExpStiffnessValue((FSolverReal)Stiffness, (FSolverReal)StiffnessWeft, (FSolverReal)StiffnessBias);
-			const FSolverVec3 ExpBucklingValue((FSolverReal)BucklingStiffness, (FSolverReal)BucklingStiffnessWeft, (FSolverReal)BucklingStiffnessBias);
+			const FSolverVec3 ExpStiffnessValue((FSolverReal)StiffnessWeft, (FSolverReal)Stiffness, (FSolverReal)StiffnessBias);
+			const FSolverVec3 ExpBucklingValue((FSolverReal)BucklingStiffnessWeft, (FSolverReal)BucklingStiffness, (FSolverReal)BucklingStiffnessBias);
 			const FSolverReal DampingRatioValue = (FSolverReal)DampingRatio;
 
 			if (ExpStiffnessValue.Max() < MinStiffness && ExpBucklingValue.Max() < MinStiffness)
@@ -675,8 +676,8 @@ void FXPBDAnisotropicBendingConstraints::Apply(FSolverParticles& Particles, cons
 						const FSolverReal ExpBucklingWeftValue = BucklingStiffnessWeftHasWeightMap ? BucklingStiffnessWeft[ConstraintIndex] : BucklingStiffnessWeftNoMap;
 						const FSolverReal ExpBucklingBiasValue = BucklingStiffnessBiasHasWeightMap ? BucklingStiffnessBias[ConstraintIndex] : BucklingStiffnessBiasNoMap;
 						const FSolverReal DampingRatioValue = DampingHasWeightMap ? DampingRatio[ConstraintIndex] : DampingNoMap;
-						ApplyHelper(Particles, Dt, ConstraintIndex, FSolverVec3(ExpStiffnessValue, ExpStiffnessWeftValue, ExpStiffnessBiasValue),
-							FSolverVec3(ExpBucklingValue, ExpBucklingWeftValue, ExpBucklingBiasValue), DampingRatioValue);
+						ApplyHelper(Particles, Dt, ConstraintIndex, FSolverVec3(ExpStiffnessWeftValue, ExpStiffnessValue, ExpStiffnessBiasValue),
+							FSolverVec3(ExpBucklingWeftValue, ExpBucklingValue, ExpBucklingBiasValue), DampingRatioValue);
 					});
 				}
 			}
@@ -688,8 +689,8 @@ void FXPBDAnisotropicBendingConstraints::Apply(FSolverParticles& Particles, cons
 			!BucklingStiffnessHasWeightMap && !BucklingStiffnessWeftHasWeightMap && !BucklingStiffnessBiasHasWeightMap &&
 			!DampingHasWeightMap)
 		{
-			const FSolverVec3 ExpStiffnessValue((FSolverReal)Stiffness, (FSolverReal)StiffnessWeft, (FSolverReal)StiffnessBias);
-			const FSolverVec3 ExpBucklingValue((FSolverReal)BucklingStiffness, (FSolverReal)BucklingStiffnessWeft, (FSolverReal)BucklingStiffnessBias);
+			const FSolverVec3 ExpStiffnessValue((FSolverReal)StiffnessWeft, (FSolverReal)Stiffness, (FSolverReal)StiffnessBias);
+			const FSolverVec3 ExpBucklingValue((FSolverReal)BucklingStiffnessWeft, (FSolverReal)BucklingStiffness, (FSolverReal)BucklingStiffnessBias);
 			const FSolverReal DampingRatioValue = (FSolverReal)DampingRatio;
 
 			if (ExpStiffnessValue.Max() < MinStiffness && ExpBucklingValue.Max() < MinStiffness)
@@ -720,8 +721,8 @@ void FXPBDAnisotropicBendingConstraints::Apply(FSolverParticles& Particles, cons
 				const FSolverReal ExpBucklingWeftValue = BucklingStiffnessWeftHasWeightMap ? BucklingStiffnessWeft[ConstraintIndex] : BucklingStiffnessWeftNoMap;
 				const FSolverReal ExpBucklingBiasValue = BucklingStiffnessBiasHasWeightMap ? BucklingStiffnessBias[ConstraintIndex] : BucklingStiffnessBiasNoMap;
 				const FSolverReal DampingRatioValue = DampingHasWeightMap ? DampingRatio[ConstraintIndex] : DampingNoMap;
-				ApplyHelper(Particles, Dt, ConstraintIndex, FSolverVec3(ExpStiffnessValue, ExpStiffnessWeftValue, ExpStiffnessBiasValue),
-					FSolverVec3(ExpBucklingValue, ExpBucklingWeftValue, ExpBucklingBiasValue), DampingRatioValue);
+				ApplyHelper(Particles, Dt, ConstraintIndex, FSolverVec3(ExpStiffnessWeftValue, ExpStiffnessValue, ExpStiffnessBiasValue),
+					FSolverVec3(ExpBucklingWeftValue, ExpBucklingValue, ExpBucklingBiasValue), DampingRatioValue);
 			}
 		}
 	}
diff --git a/Engine/Source/Runtime/Experimental/Chaos/Private/Chaos/XPBDAnisotropicBendingConstraints.ispc b/Engine/Source/Runtime/Experimental/Chaos/Private/Chaos/XPBDAnisotropicBendingConstraints.ispc
--- a/Engine/Source/Runtime/Experimental/Chaos/Private/Chaos/XPBDAnisotropicBendingConstraints.ispc
+++ b/Engine/Source/Runtime/Experimental/Chaos/Private/Chaos/XPBDAnisotropicBendingConstraints.ispc
@@ -272,7 +272,7 @@ export void ApplyXPBDAnisotropicBendingConstraintsWithMaps(uniform FVector4f Pan
 		#pragma ignore warning(perf)
 		const varying float StiffnessBias = StiffnessBiasTable[StiffnessBiasIndex];
 
-		const varying FVector3f StiffnessVector = {{Stiffness, StiffnessWeft, StiffnessBias}};
+		const varying FVector3f StiffnessVector = {{StiffnessWeft, Stiffness, StiffnessBias}};
 
 		#pragma ignore warning(perf)
 		const varying float BucklingStiffness = BucklingStiffnessTable[BucklingStiffnessIndex];
@@ -281,7 +281,7 @@ export void ApplyXPBDAnisotropicBendingConstraintsWithMaps(uniform FVector4f Pan
 		#pragma ignore warning(perf)
 		const varying float BucklingStiffnessBias = BucklingStiffnessBiasTable[BucklingStiffnessBiasIndex];
 		
-		const varying FVector3f BucklingStiffnessVector = {{BucklingStiffness, BucklingStiffnessWeft, BucklingStiffnessBias}};
+		const varying FVector3f BucklingStiffnessVector = {{BucklingStiffnessWeft, BucklingStiffness, BucklingStiffnessBias}};
 		
 		const varying float RestAngle = RestAngles[i];
 		const varying FVector3f WarpWeftBiasBaseMultiplier = VectorGather(&WarpWeftBiasBaseMultiplers[i]);
@@ -388,7 +388,7 @@ export void ApplyXPBDAnisotropicBendingConstraintsWithDampingAndMaps(uniform FVe
 		#pragma ignore warning(perf)
 		const varying float StiffnessBias = StiffnessBiasTable[StiffnessBiasIndex];
 
-		const varying FVector3f StiffnessVector = {{Stiffness, StiffnessWeft, StiffnessBias}};
+		const varying FVector3f StiffnessVector = {{StiffnessWeft, Stiffness, StiffnessBias}};
 
 		#pragma ignore warning(perf)
 		const varying float BucklingStiffness = BucklingStiffnessTable[BucklingStiffnessIndex];
@@ -397,7 +397,7 @@ export void ApplyXPBDAnisotropicBendingConstraintsWithDampingAndMaps(uniform FVe
 		#pragma ignore warning(perf)
 		const varying float BucklingStiffnessBias = BucklingStiffnessBiasTable[BucklingStiffnessBiasIndex];
 		
-		const varying FVector3f BucklingStiffnessVector = {{BucklingStiffness, BucklingStiffnessWeft, BucklingStiffnessBias}};
+		const varying FVector3f BucklingStiffnessVector = {{BucklingStiffnessWeft, BucklingStiffness, BucklingStiffnessBias}};
 		
 		#pragma ignore warning(perf)
 		const varying float DampingRatio = DampingTable[DampingIndex];
diff --git a/Engine/Source/Runtime/Experimental/Chaos/Private/Chaos/XPBDStretchBiasElementConstraints.cpp b/Engine/Source/Runtime/Experimental/Chaos/Private/Chaos/XPBDStretchBiasElementConstraints.cpp
--- a/Engine/Source/Runtime/Experimental/Chaos/Private/Chaos/XPBDStretchBiasElementConstraints.cpp
+++ b/Engine/Source/Runtime/Experimental/Chaos/Private/Chaos/XPBDStretchBiasElementConstraints.cpp
@@ -445,7 +445,7 @@ void FXPBDStretchBiasElementConstraints::Apply(FSolverParticles& Particles, cons
 		const int32 ConstraintColorNum = ConstraintsPerColorStartIndex.Num() - 1;
 		if (!StiffnessWarpHasWeightMap && !StiffnessWeftHasWeightMap && !StiffnessBiasHasWeightMap && !DampingHasWeightMap && !WarpScaleHasWeightMap && !WeftScaleHasWeightMap)
 		{
-			const FSolverVec3 ExpStiffnessValue(StiffnessWarpNoMap, StiffnessWeftNoMap, StiffnessBiasNoMap);
+			const FSolverVec3 ExpStiffnessValue(StiffnessWeftNoMap, StiffnessWarpNoMap, StiffnessBiasNoMap);
 			if (ExpStiffnessValue.Max() < MinStiffness)
 			{
 				return;
@@ -610,7 +610,7 @@ void FXPBDStretchBiasElementConstraints::Apply(FSolverParticles& Particles, cons
 						const FSolverReal DampingRatioValue = DampingHasWeightMap ? DampingRatio[ConstraintIndex] : DampingNoMap;
 						const FSolverReal WarpScaleValue = WarpScaleHasWeightMap ? WarpScale[ConstraintIndex] : WarpScaleNoMap;
 						const FSolverReal WeftScaleValue = WeftScaleHasWeightMap ? WeftScale[ConstraintIndex] : WeftScaleNoMap;
-						ApplyHelper(Particles, Dt, ConstraintIndex, FSolverVec3(ExpStiffnessWarpValue, ExpStiffnessWeftValue, ExpStiffnessBiasValue),DampingRatioValue, WarpScaleValue, WeftScaleValue);
+						ApplyHelper(Particles, Dt, ConstraintIndex, FSolverVec3(ExpStiffnessWeftValue, ExpStiffnessWarpValue, ExpStiffnessBiasValue),DampingRatioValue, WarpScaleValue, WeftScaleValue);
 					});
 				}
 			}
@@ -640,7 +640,7 @@ void FXPBDStretchBiasElementConstraints::Apply(FSolverParticles& Particles, cons
 				const FSolverReal DampingRatioValue = DampingHasWeightMap ? DampingRatio[ConstraintIndex] : DampingNoMap;
 				const FSolverReal WarpScaleValue = WarpScaleHasWeightMap ? WarpScale[ConstraintIndex] : WarpScaleNoMap;
 				const FSolverReal WeftScaleValue = WeftScaleHasWeightMap ? WeftScale[ConstraintIndex] : WeftScaleNoMap;
-				ApplyHelper(Particles, Dt, ConstraintIndex, FSolverVec3(ExpStiffnessWarpValue, ExpStiffnessWeftValue, ExpStiffnessBiasValue), DampingRatioValue, WarpScaleValue, WeftScaleValue);
+				ApplyHelper(Particles, Dt, ConstraintIndex, FSolverVec3(ExpStiffnessWeftValue, ExpStiffnessWarpValue, ExpStiffnessBiasValue), DampingRatioValue, WarpScaleValue, WeftScaleValue);
 			}
 		}
 	}
@@ -705,8 +705,8 @@ void FXPBDStretchBiasElementConstraints::ApplyHelper(FSolverParticles& Particles
 	const FSolverVec3 DXDvNormalized = DXDv * OneOverDXDvLen;
 
 	// constraints
-	const FSolverReal Cu = (DXDuLength - RestStretchLengths[ConstraintIndex][0]*WarpScaleValue); // stretch in warp direction
-	const FSolverReal Cv = (DXDvLength - RestStretchLengths[ConstraintIndex][1]*WeftScaleValue); // stretch in weft direction
+	const FSolverReal Cu = (DXDuLength - RestStretchLengths[ConstraintIndex][0]*WeftScaleValue); // stretch in weft direction
+	const FSolverReal Cv = (DXDvLength - RestStretchLengths[ConstraintIndex][1]*WarpScaleValue); // stretch in warp direction
 	const FSolverReal Cs = FSolverVec3::DotProduct(DXDuNormalized, DXDvNormalized);
 
 
diff --git a/Engine/Source/Runtime/Experimental/Chaos/Private/Chaos/XPBDStretchBiasElementConstraints.ispc b/Engine/Source/Runtime/Experimental/Chaos/Private/Chaos/XPBDStretchBiasElementConstraints.ispc
--- a/Engine/Source/Runtime/Experimental/Chaos/Private/Chaos/XPBDStretchBiasElementConstraints.ispc
+++ b/Engine/Source/Runtime/Experimental/Chaos/Private/Chaos/XPBDStretchBiasElementConstraints.ispc
@@ -29,8 +29,8 @@ static void CalcConstraintsAndGrads(const FVector3f& P0, const FVector3f& P1, co
 	const FVector3f DXDuNormalized = OneOverDXDuLen * DXDu;
 	const FVector3f DXDvNormalized = OneOverDXDvLen * DXDv;
 
-	Cu = DXDuLength - RestStretchLengths.V[0] * WarpScale;
-	Cv = DXDvLength - RestStretchLengths.V[1] * WeftScale;
+	Cu = DXDuLength - RestStretchLengths.V[0] * WeftScale;
+	Cv = DXDvLength - RestStretchLengths.V[1] * WarpScale;
 	Cs = VectorDot(DXDuNormalized, DXDvNormalized);
 
 	GradCu1 = DXDuNormalized * DeltaUVInv.V[0];
@@ -295,7 +295,7 @@ export void ApplyXPBDStretchBiasConstraintsWithMaps(uniform FVector4f PandInvM[]
 		const varying float StiffnessWeft = StiffnessWeftTable[StiffnessWeftIndex];
 		#pragma ignore warning(perf)
 		const varying float StiffnessBias = StiffnessBiasTable[StiffnessBiasIndex];
-		const varying FVector3f Stiffness = {{StiffnessWarp, StiffnessWeft, StiffnessBias}};
+		const varying FVector3f Stiffness = {{StiffnessWeft, StiffnessWarp, StiffnessBias}};
 
 		FVector3f Lambda = VectorLoad(&Lambdas[extract(i,0)]);
 		const varying FVector3f StiffnessScale = VectorLoad(&StiffnessScales[extract(i,0)]);
@@ -409,7 +409,7 @@ export void ApplyXPBDStretchBiasConstraintsWithDampingAndMaps(uniform FVector4f
 		const varying float StiffnessWeft = StiffnessWeftTable[StiffnessWeftIndex];
 		#pragma ignore warning(perf)
 		const varying float StiffnessBias = StiffnessBiasTable[StiffnessBiasIndex];
-		const varying FVector3f Stiffness = {{StiffnessWarp, StiffnessWeft, StiffnessBias}};
+		const varying FVector3f Stiffness = {{StiffnessWeft, StiffnessWarp, StiffnessBias}};
 
 		#pragma ignore warning(perf)
 		const varying float DampingRatio = DampingTable[DampingIndex];
