Only include S3 client headers on desktop platforms

#rnx

diff --git a/Engine/Source/Runtime/Experimental/IoStoreOnDemand/Private/IoStoreOnDemand.cpp b/Engine/Source/Runtime/Experimental/IoStoreOnDemand/Private/IoStoreOnDemand.cpp
--- a/Engine/Source/Runtime/Experimental/IoStoreOnDemand/Private/IoStoreOnDemand.cpp
+++ b/Engine/Source/Runtime/Experimental/IoStoreOnDemand/Private/IoStoreOnDemand.cpp
@@ -18,11 +18,11 @@
 #include "Serialization/LargeMemoryWriter.h"
 #include "String/LexFromString.h"
 
-#if (IS_PROGRAM || WITH_EDITOR)
+#if (PLATFORM_DESKTOP && (IS_PROGRAM || WITH_EDITOR))
 #include "Async/Async.h"
 #include "Misc/FileHelper.h"
 #include "S3/S3Client.h"
-#endif // (IS_PROGRAM || WITH_EDITOR)
+#endif // (PLATFORM_DESKTOP && (IS_PROGRAM || WITH_EDITOR))
 
 DEFINE_LOG_CATEGORY(LogIas);
 
@@ -355,7 +355,7 @@ bool LoadFromCompactBinary(FCbFieldView Field, FOnDemandToc& OutToc)
 }
 
 ////////////////////////////////////////////////////////////////////////////////
-#if (IS_PROGRAM || WITH_EDITOR)
+#if (PLATFORM_DESKTOP && (IS_PROGRAM || WITH_EDITOR))
 class FS3UploadQueue
 {
 public:
@@ -844,7 +844,7 @@ TIoStatusOr<FIoStoreUploadResult> UploadContainerFiles(
 	
 	return UploadResult;
 }
-#endif // (IS_PROGRAM || WITH_EDITOR)
+#endif // (PLATFORM_DESKTOP && (IS_PROGRAM || WITH_EDITOR))
 
 } // namespace UE
 
