fix potential array out-of-bounds in FClusterUnionManager::HandleRemoveOperation due to duplicate particles in the array of particles to remove

diff --git a/Engine/Source/Runtime/Experimental/Chaos/Private/Chaos/ClusterUnionManager.cpp b/Engine/Source/Runtime/Experimental/Chaos/Private/Chaos/ClusterUnionManager.cpp
--- a/Engine/Source/Runtime/Experimental/Chaos/Private/Chaos/ClusterUnionManager.cpp
+++ b/Engine/Source/Runtime/Experimental/Chaos/Private/Chaos/ClusterUnionManager.cpp
@@ -576,11 +576,14 @@ namespace Chaos
 		);
 		MClustering.SetGenerateClusterBreaking(false);
 
+		// Use a set to guarantee that there's no duplicates. This is necessary since further removal operations assume that the particle is only in here once.
+		const TSet<FPBDRigidParticleHandle*> ParticleSet{ Particles };
+
 		IPhysicsProxyBase* OldProxy = Cluster->InternalCluster->PhysicsProxy();
 		TArray<int32> ParticleIndicesToRemove;
-		ParticleIndicesToRemove.Reserve(Particles.Num());
+		ParticleIndicesToRemove.Reserve(ParticleSet.Num());
 
-		for (FPBDRigidParticleHandle* Handle : Particles)
+		for (FPBDRigidParticleHandle* Handle : ParticleSet)
 		{
 			const int32 ParticleIndex = Cluster->ChildParticles.Find(Handle);
 			if (ParticleIndex != INDEX_NONE)
@@ -611,7 +614,7 @@ namespace Chaos
 			Cluster->ChildParticles.RemoveAtSwap(ParticleIndex);
 		}
 
-		MClustering.RemoveParticlesFromCluster(Cluster->InternalCluster, Particles);
+		MClustering.RemoveParticlesFromCluster(Cluster->InternalCluster, ParticleSet.Array());
 
 		// Removing a particle should have no bearing on the proxy of the cluster.
 		// This gets changed because we go through an internal initialization route when we update the cluster union particle's properties.
