Iris: Cleanup verbose logs added to track ReferenceCache issue

#rnx

diff --git a/Engine/Source/Runtime/Experimental/Iris/Core/Private/Iris/ReplicationSystem/ObjectReferenceCache.cpp b/Engine/Source/Runtime/Experimental/Iris/Core/Private/Iris/ReplicationSystem/ObjectReferenceCache.cpp
--- a/Engine/Source/Runtime/Experimental/Iris/Core/Private/Iris/ReplicationSystem/ObjectReferenceCache.cpp
+++ b/Engine/Source/Runtime/Experimental/Iris/Core/Private/Iris/ReplicationSystem/ObjectReferenceCache.cpp
@@ -23,25 +23,23 @@
 #include "UObject/Package.h"
 #include "HAL/IConsoleManager.h"
 
-DEFINE_LOG_CATEGORY_STATIC(LogIrisReferences, Log, All);
-
 #ifndef UE_NET_ENABLE_REFERENCECACHE_LOG
 #	define UE_NET_ENABLE_REFERENCECACHE_LOG !(UE_BUILD_SHIPPING || UE_BUILD_TEST)
 #endif 
 
 #ifndef UE_NET_VALIDATE_REFERENCECACHE
-#	define UE_NET_VALIDATE_REFERENCECACHE !(UE_BUILD_SHIPPING || UE_BUILD_TEST)
+#	define UE_NET_VALIDATE_REFERENCECACHE 0
 #endif
 
+// Only compile warnings or errors when UE_NET_ENABLE_REFERENCECACHE_LOG is false
 #if UE_NET_ENABLE_REFERENCECACHE_LOG
-#	define UE_LOG_REFERENCECACHE(Verbosity, Format, ...)  UE_LOG(LogIrisReferences, Verbosity, Format, ##__VA_ARGS__)
-#	define UE_CLOG_REFERENCECACHE(Condition, Verbosity, Format, ...)  UE_CLOG(Condition, LogIrisReferences, Verbosity, Format, ##__VA_ARGS__)
+	DEFINE_LOG_CATEGORY_STATIC(LogIrisReferences, Log, All);
 #else
-#	define UE_LOG_REFERENCECACHE(...)
-#	define UE_CLOG_REFERENCECACHE(...)
+	DEFINE_LOG_CATEGORY_STATIC(LogIrisReferences, Warning, Warning);
 #endif
 
-#define UE_LOG_REFERENCECACHE_WARNING(Format, ...)  UE_LOG(LogIrisReferences, Warning, Format, ##__VA_ARGS__)
+#define UE_LOG_REFERENCECACHE(Verbosity, Format, ...)  UE_LOG(LogIrisReferences, Verbosity, Format, ##__VA_ARGS__)
+#define UE_CLOG_REFERENCECACHE(Condition, Verbosity, Format, ...)  UE_CLOG(Condition, LogIrisReferences, Verbosity, Format, ##__VA_ARGS__)
 
 // $TODO: If we add information about why a dynamic object is being destroyed (filtering or actual destroy) then we can 
 // try to clean up references to reduce cache pollution
@@ -243,42 +241,42 @@ bool FObjectReferenceCache::CreateObjectReferenceInternal(const UObject* Object,
 		if (FCachedNetObjectReference* CachedObjectPtr = ReferenceHandleToCachedReference.Find(RefHandle))
 		{
 			FCachedNetObjectReference& CachedObject = *CachedObjectPtr;
-			if (CachedObject.Object.Get() == Object)
-			{
+		if (CachedObject.Object.Get() == Object)
+		{
 				UE_LOG_REFERENCECACHE(VeryVerbose, TEXT("ObjectReferenceCache::CreateObjectReferenceHandle Found existing %s for ObjectPath %s Object: %s (0x%p), OuterNetRefHandle: %s"),
-					*CachedObject.NetRefHandle.ToString(), ToCStr(Object->GetPathName()), *GetNameSafe(Object), Object, *CachedObject.OuterNetRefHandle.ToString());
+				*CachedObject.NetRefHandle.ToString(), ToCStr(Object->GetPathName()), *GetNameSafe(Object), Object, *CachedObject.OuterNetRefHandle.ToString());
 				OutReference = MakeNetObjectReference(CachedObject);
-				return true;
-			}
-			else
-			{
+			return true;
+		}
+		else
+		{
 				UE_LOG_REFERENCECACHE(Verbose, TEXT("ObjectReferenceCache::CreateObjectReferenceHandle Removed %s from ObjectToNetReferenceHandle due to stale cache. Cached Object (0x%p).  New object %s (0x%p)"),
 					*CachedObject.NetRefHandle.ToString(), CachedObject.ObjectKey, *GetNameSafe(Object), Object);
 
-				ObjectToNetReferenceHandle.Remove(CachedObject.ObjectKey);
-				ObjectToNetReferenceHandle.Remove(Object);
+			ObjectToNetReferenceHandle.Remove(CachedObject.ObjectKey);
+			ObjectToNetReferenceHandle.Remove(Object);
 				RefHandlePtr = nullptr;
 
-				if (RefHandle.IsStatic())
-				{
+			if (RefHandle.IsStatic())
+			{
 					UE_LOG_REFERENCECACHE(Verbose, TEXT("ObjectReferenceCache::CreateObjectReferenceHandle clearing stale cache for static handle %s. Cached Object (0x%p).  New object %s (0x%p)"),
 						*CachedObject.NetRefHandle.ToString(), CachedObject.ObjectKey, *GetNameSafe(Object), Object);
 
-					// Note we only cleanse the object reference
-					// we still keep the cachedObject data around in order to be able to serialize static destruction infos
-					CachedObject.ObjectKey = nullptr;
-					CachedObject.Object = nullptr;
+				// Note we only cleanse the object reference
+				// we still keep the cachedObject data around in order to be able to serialize static destruction infos
+				CachedObject.ObjectKey = nullptr;
+				CachedObject.Object = nullptr;
 
-				}
-				else
-				{
+			}
+			else
+			{
 					UE_LOG_REFERENCECACHE(Verbose, TEXT("ObjectReferenceCache::CreateObjectReferenceHandle removing cache for handle %s. Cached Object (0x%p). New object %s (0x%p)"),
 						*CachedObject.NetRefHandle.ToString(), CachedObject.ObjectKey, *GetNameSafe(Object), Object);
 
-					ReferenceHandleToCachedReference.Remove(RefHandle);
-				}
+				ReferenceHandleToCachedReference.Remove(RefHandle);
 			}
 		}
+	}
 		else
 		{
 			UE_LOG_REFERENCECACHE(Warning, TEXT("ObjectReferenceCache::CreateObjectReferenceInternal removed %s from ObjectToNetReferenceHandle due to mismatch with cache. Object %s (0x%p) "), *RefHandle.ToString(), *GetNameSafe(Object), Object);
@@ -429,9 +427,8 @@ FNetRefHandle FObjectReferenceCache::GetObjectReferenceHandleFromObject(const UO
 		const FCachedNetObjectReference* CachedObjectPtr = ReferenceHandleToCachedReference.Find(*Reference);
 		if (CachedObjectPtr == nullptr)
 		{
-			ensureAlwaysMsgf(CachedObjectPtr != nullptr, TEXT("Mismatch between ReferenceCache maps. Object %s (0x%p) has handle (%s) but no cache."), *GetNameSafe(Object), Object, *Reference->ToString());
+			UE_LOG_REFERENCECACHE(Verbose, TEXT("ObjectReferenceCache::GetObjectReferenceHandleFromObject removed %s from ObjectToNetReferenceHandle due to ReferenceHandleToCachedReference not holding the handle. For Object %s (0x%p) "), *Reference->ToString(), *GetNameSafe(Object), Object);
 			const_cast<TMap<const UObject*, FNetRefHandle>&>(ObjectToNetReferenceHandle).Remove(Object);
-			UE_LOG_REFERENCECACHE(Warning, TEXT("ObjectReferenceCache::GetObjectReferenceHandleFromObject removed %s from ObjectToNetReferenceHandle due to mismatch with cache. Object %s (0x%p) "), *Reference->ToString(), *GetNameSafe(Object), Object);
 			return FNetRefHandle();
 		}
 
@@ -446,7 +443,7 @@ FNetRefHandle FObjectReferenceCache::GetObjectReferenceHandleFromObject(const UO
 		}
 		else
 		{
-			UE_LOG_REFERENCECACHE(Verbose, TEXT("ObjectReferenceCache::Found %s cached RefHandle %s for object %s (0x%p)"), 
+			UE_LOG_REFERENCECACHE(Verbose, TEXT("ObjectReferenceCache::Found %s cached RefHandle %s for object %s (0x%p) but cache was holding an invalid object."), 
 				CachedObject.Object.IsStale() ? TEXT("Stale") : TEXT("UnResolved"), *CachedObject.NetRefHandle.ToString(), *GetNameSafe(Object), Object);
 		}
 	}
@@ -473,8 +470,8 @@ void FObjectReferenceCache::AddRemoteReference(FNetRefHandle RefHandle, const UO
 			{
 				if (ExistingCachedObject->NetRefHandle != RefHandle)
 				{
-					UE_LOG_REFERENCECACHE_WARNING(TEXT("ObjectReferenceCache::AddRemoteReference: Detected conflicting Ref %s: Received: %s, Object: %s (0x%p)"), *ExistingCachedObject->NetRefHandle.ToString(),
-						*RefHandle.ToString(), *GetNameSafe(Object), Object);
+					UE_LOG_REFERENCECACHE(Warning, TEXT("ObjectReferenceCache::AddRemoteReference: Detected conflicting Ref %s: Received: %s, Object: %s (0x%p)"), *ExistingCachedObject.NetRefHandle.ToString(),
+					*RefHandle.ToString(), *GetNameSafe(Object), Object);
 				}
 			}
 		}
@@ -599,7 +596,7 @@ void FObjectReferenceCache::RemoveReference(FNetRefHandle RefHandle, const UObje
 			}
 			else
 			{
-				UE_LOG_REFERENCECACHE(Warning, TEXT("ObjectReferenceCache::RemoveReference: %s, Object: %s (0x%p). Skipped removing since ObjectToNetReferenceHandle had no entry"), 
+				UE_LOG_REFERENCECACHE(Verbose, TEXT("ObjectReferenceCache::RemoveReference: %s, Object: %s (0x%p). Skipped removing since ObjectToNetReferenceHandle had no entry"), 
 					*RefHandle.ToString(), *GetNameSafe(Object), Object);
 			}
 		}
@@ -665,7 +662,7 @@ UObject* FObjectReferenceCache::ResolveObjectReferenceHandleInternal(FNetRefHand
 			FNetRefHandle* ExistingRefHandle = ObjectToNetReferenceHandle.Find(Object);
 			if (ExistingRefHandle && *ExistingRefHandle != RefHandle)
 			{
-				UE_LOG_REFERENCECACHE_WARNING(TEXT("ObjectReferenceCache::ResolveObjectRefererenceHandle detected potential conflicting references: %s and %s"), *RefHandle.ToString(), *(*ExistingRefHandle).ToString());
+				UE_LOG_REFERENCECACHE(Warning, TEXT("ObjectReferenceCache::ResolveObjectRefererenceHandle detected potential conflicting references: %s and %s"), *RefHandle.ToString(), *(*ExistingRefHandle).ToString());
 			}
 		}
 #endif
@@ -915,7 +912,7 @@ UObject* FObjectReferenceCache::ResolveObjectReferenceHandleInternal(FNetRefHand
 
 	if (!IsValid(Object))
 	{
-		UE_LOG_REFERENCECACHE_WARNING(TEXT("GetObjectFromRefHandle: Resolved object Pointer: %p that is not valid? FullPath: %s"), Object, ToCStr(FullPath(RefHandle, ResolveContext)));
+		UE_LOG_REFERENCECACHE(Verbose, TEXT("GetObjectFromRefHandle: Resolved object Pointer: %p that is not valid? FullPath: %s"), Object, ToCStr(FullPath(RefHandle, ResolveContext)));
 	}
 
 	CacheObjectPtr->Object = MakeWeakObjectPtr(Object);
@@ -959,19 +956,19 @@ UObject* FObjectReferenceCache::ResolveObjectReferenceHandleInternal(FNetRefHand
 		if (FCachedNetObjectReference* ExistingCachedObjectPtr = ReferenceHandleToCachedReference.Find(*ExistingHandle))
 		{
 			FCachedNetObjectReference& ExistingCachedObject = *ExistingCachedObjectPtr;
-			if (ExistingCachedObject.Object.IsStale())
-			{
-				UE_LOG_REFERENCECACHE_WARNING(TEXT("Invalidating object for Stale Ref %s with %s as the object has been reused"), *ExistingCachedObject.NetRefHandle.ToString(), *CompleteRefHandle.ToString());
-			}
-			else if (ExistingCachedObject.Object.HasSameIndexAndSerialNumber(CacheObjectPtr->Object))
+		if (ExistingCachedObject.Object.IsStale())
+		{
+			UE_LOG_REFERENCECACHE(Verbose, TEXT("Invalidating object for Stale Ref %s with %s as the object has been reused"), *ExistingCachedObject.NetRefHandle.ToString(), *CompleteRefHandle.ToString());
+		}
+		else if (ExistingCachedObject.Object.HasSameIndexAndSerialNumber(CacheObjectPtr->Object))
+		{
+			bShouldAssign = CompleteRefHandle.GetId() > (*ExistingHandle).GetId();
+			if (bShouldAssign)
 			{
-				bShouldAssign = CompleteRefHandle.GetId() > (*ExistingHandle).GetId();
-				if (bShouldAssign)
-				{
-					UE_LOG_REFERENCECACHE_WARNING(TEXT("Invalidating object for conflicting Ref %s with %s"), *ExistingCachedObject.NetRefHandle.ToString(), *CompleteRefHandle.ToString());
-				}	
-			}
+				UE_LOG_REFERENCECACHE(Verbose, TEXT("Invalidating object for conflicting Ref %s with %s"), *ExistingCachedObject.NetRefHandle.ToString(), *CompleteRefHandle.ToString());
+			}	
 		}
+	}
 		else
 		{
 			UE_LOG_REFERENCECACHE(Warning, TEXT("ObjectReferenceCache::ResolveObjectReferenceHandleInternal removed %s from ObjectToNetReferenceHandle due to mismatch with cache. Object %s (0x%p) "), *ExistingHandle->ToString(), *GetNameSafe(Object), Object);
@@ -1022,7 +1019,7 @@ ENetObjectReferenceResolveResult FObjectReferenceCache::ResolveObjectReference(c
 
 				if (ResolvedObject == nullptr)
 				{
-					UE_LOG_REFERENCECACHE_WARNING(TEXT("ResolveObjectReference Failed to resolve clientassigned ref: %s due to not finding object with relative path %s."), *Reference.ToString(), *ObjectPath);	
+					UE_LOG_REFERENCECACHE(Warning, TEXT("ResolveObjectReference Failed to resolve clientassigned ref: %s due to not finding object with relative path %s."), *Reference.ToString(), *ObjectPath);	
 				}
 			}
 			else
@@ -1041,13 +1038,13 @@ ENetObjectReferenceResolveResult FObjectReferenceCache::ResolveObjectReference(c
 
 				if (ResolvedObject == nullptr)
 				{
-					UE_LOG_REFERENCECACHE_WARNING(TEXT("ResolveObjectReference Failed to resolve clientassigned ref: due to not finding object with path %s."), *ObjectPath);	
+					UE_LOG_REFERENCECACHE(Warning, TEXT("ResolveObjectReference Failed to resolve clientassigned ref: due to not finding object with path %s."), *ObjectPath);	
 				}
 			}
 		}
 		else
 		{
-			UE_LOG_REFERENCECACHE_WARNING(TEXT("ResolveObjectReference Failed to resolve clientassigned ref: %s due to missing token:"), *Reference.ToString());	
+			UE_LOG_REFERENCECACHE(Warning, TEXT("ResolveObjectReference Failed to resolve clientassigned ref: %s due to missing token:"), *Reference.ToString());	
 		}
 	}
 	else
@@ -1252,7 +1249,7 @@ void FObjectReferenceCache::ReadFullReferenceInternal(FNetSerializationContext&
 
 	if (RecursionCount > INTERNAL_READ_REF_RECURSION_LIMIT) 
 	{
-		UE_LOG_REFERENCECACHE_WARNING(TEXT("ReadFullReferenceInternal: Hit recursion limit."));
+		UE_LOG_REFERENCECACHE(Warning, TEXT("ReadFullReferenceInternal: Hit recursion limit."));
 		check(false);
 		Reader->DoOverflow();
 		OutObjectRef = FNetObjectReference();
