Avoid bounds issue with geometry collections by always using diagonal inertia in GC local frame.

Add method to PBDRigidClusteringAlgo which drops off-diagonal terms from cluster inertia tensors transformed into local space, resulting in Identity rotations of inertia (but only approximately correct inertia tensors).

Use this after geometry collection inertia calculation to avoid issues with geometry collection bounds generation which occur when there's a non-identity rotation of inertia

Update EMassOffsetType and its usages to follow the code standard

diff --git a/Engine/Source/Runtime/Experimental/Chaos/Private/Chaos/ClusterUnionManager.cpp b/Engine/Source/Runtime/Experimental/Chaos/Private/Chaos/ClusterUnionManager.cpp
--- a/Engine/Source/Runtime/Experimental/Chaos/Private/Chaos/ClusterUnionManager.cpp
+++ b/Engine/Source/Runtime/Experimental/Chaos/Private/Chaos/ClusterUnionManager.cpp
@@ -666,7 +666,7 @@ namespace Chaos
 		// Position the internal cluster at the CoM of the children
 		if (bRecomputeMassOrientation)
 		{
-			MoveClusterToMassOffset(ClusterUnion.InternalCluster, EMassOffsetType::EPosition);
+			MoveClusterToMassOffset(ClusterUnion.InternalCluster, EMassOffsetType::Position);
 			ClusterUnion.bNeedsXRInitialization = false;
 		}
 
diff --git a/Engine/Source/Runtime/Experimental/Chaos/Private/Chaos/PBDRigidClustering.cpp b/Engine/Source/Runtime/Experimental/Chaos/Private/Chaos/PBDRigidClustering.cpp
--- a/Engine/Source/Runtime/Experimental/Chaos/Private/Chaos/PBDRigidClustering.cpp
+++ b/Engine/Source/Runtime/Experimental/Chaos/Private/Chaos/PBDRigidClustering.cpp
@@ -220,7 +220,7 @@ namespace Chaos
 		
 		if (ForceMassOrientation == nullptr)
 		{
-			MoveClusterToMassOffset(NewParticle, EMassOffsetType::EPosition);
+			MoveClusterToMassOffset(NewParticle, EMassOffsetType::Position);
 		}
 		UpdateKinematicProperties(NewParticle, MChildren, MEvolution);
 		UpdateGeometry(NewParticle, ChildrenSet, MChildren, ProxyGeometry, Parameters);
@@ -442,7 +442,7 @@ namespace Chaos
 		TSet<FPBDRigidParticleHandle*> ChildrenSet(ChildrenArray);
 
 		UpdateClusterMassProperties(NewParticle, ChildrenSet);
-		MoveClusterToMassOffset(NewParticle, EMassOffsetType::EPosition);
+		MoveClusterToMassOffset(NewParticle, EMassOffsetType::Position);
 
 		UpdateKinematicProperties(NewParticle, MChildren, MEvolution);
 
diff --git a/Engine/Source/Runtime/Experimental/Chaos/Private/Chaos/PBDRigidClusteringAlgo.cpp b/Engine/Source/Runtime/Experimental/Chaos/Private/Chaos/PBDRigidClusteringAlgo.cpp
--- a/Engine/Source/Runtime/Experimental/Chaos/Private/Chaos/PBDRigidClusteringAlgo.cpp
+++ b/Engine/Source/Runtime/Experimental/Chaos/Private/Chaos/PBDRigidClusteringAlgo.cpp
@@ -89,21 +89,39 @@ namespace Chaos
 			: FReal(1) / Inertia;
 	}
 
-	void MoveClusterToMassOffset(FPBDRigidClusteredParticleHandle* Cluster, const EMassOffsetType MassOffsetTypes)
+	void AdjustClusterInertia(FPBDRigidClusteredParticleHandle* Cluster, const EInertiaOperations InertiaOperations)
 	{
-		if (MassOffsetTypes | EMassOffsetType::EPosition)
+		if ((InertiaOperations & EInertiaOperations::LocalInertiaDropOffDiagonalTerms) != EInertiaOperations::None)
 		{
+			// Discard off-diagonal terms in inertia transformed into a particle's local space, and
+			// zero out the Rotation of Inertia parameter.
+			const FMatrix33 LocalSpaceInertia = Utilities::ComputeWorldSpaceInertia(Cluster->RotationOfMass(), Cluster->I());
+			Cluster->SetI(LocalSpaceInertia.GetDiagonal());
+			Cluster->SetRotationOfMass(FQuat::Identity);
+		}
+	}
+
+	FRigidTransform3 MoveClusterToMassOffset(FPBDRigidClusteredParticleHandle* Cluster, const EMassOffsetType MassOffsetTypes)
+	{
+		FRigidTransform3 MassToLocal = FRigidTransform3::Identity;
+
+		if ((MassOffsetTypes & EMassOffsetType::Position) != EMassOffsetType::None)
+		{
+			MassToLocal.SetTranslation(Cluster->CenterOfMass());
 			Cluster->SetX(Cluster->XCom());
 			Cluster->SetP(Cluster->X());
 			Cluster->SetCenterOfMass(FVec3::ZeroVector);
 		}
 
-		if (MassOffsetTypes | EMassOffsetType::ERotation)
+		if ((MassOffsetTypes & EMassOffsetType::Rotation) != EMassOffsetType::None)
 		{
+			MassToLocal.SetRotation(Cluster->RotationOfMass());
 			Cluster->SetR(Cluster->RCom());
 			Cluster->SetQ(Cluster->R());
 			Cluster->SetRotationOfMass(FQuat::Identity);
 		}
+
+		return MassToLocal;
 	}
 
 	DECLARE_CYCLE_STAT(TEXT("TPBDRigidClustering<>::UpdateKinematicProperties()"), STAT_UpdateKinematicProperties, STATGROUP_Chaos);
diff --git a/Engine/Source/Runtime/Experimental/Chaos/Private/PhysicsProxy/GeometryCollectionPhysicsProxy.cpp b/Engine/Source/Runtime/Experimental/Chaos/Private/PhysicsProxy/GeometryCollectionPhysicsProxy.cpp
--- a/Engine/Source/Runtime/Experimental/Chaos/Private/PhysicsProxy/GeometryCollectionPhysicsProxy.cpp
+++ b/Engine/Source/Runtime/Experimental/Chaos/Private/PhysicsProxy/GeometryCollectionPhysicsProxy.cpp
@@ -115,6 +115,12 @@ FAutoConsoleVariableRef CVarGeometryCollectionAreaBasedDamageThresholdMode(
 	GeometryCollectionAreaBasedDamageThresholdMode,
 	TEXT("Area based damage threshold computation mode (0: sum of areas | 1: max of areas | 2: min of areas | 3: average of areas) [def: 0]"));
 
+int32 GeometryCollectionLocalInertiaDropOffDiagonalTerms = 0;
+FAutoConsoleVariableRef CVarGeometryCollectionLocalInertiaDropOffDiagonalTerms(
+	TEXT("p.GeometryCollection.LocalInertiaDropOffDiagonalTerms"),
+	GeometryCollectionLocalInertiaDropOffDiagonalTerms,
+	TEXT("When true, force diagonal inertia for GCs in their local space by simply dropping off-diagonal terms"));
+
 DEFINE_LOG_CATEGORY_STATIC(UGCC_LOG, Error, All);
 
 static const FSharedSimulationSizeSpecificData& GetSizeSpecificData(const TArray<FSharedSimulationSizeSpecificData>& SizeSpecificData, const FGeometryCollection& RestCollection, const int32 TransformIndex, const FBox& BoundingBox);
@@ -4435,16 +4441,23 @@ void FGeometryCollectionPhysicsProxy::InitializeSharedCollisionStructures(
 				// store them on ClusterHandle based on the combined mass properties of its children.
 				UpdateClusterMassProperties(ClusterHandle, ChildrenIndices);
 
-				// Store the relative transform from the particle's configuration space to its
-				// MassToLocal because GeometryCollections assume that the particle's XR is at its
-				// CoM and RoM - MassToLocal accounts for this offset.
-				const FRigidTransform3 MassToLocal = FRigidTransform3(ClusterHandle->CenterOfMass(), ClusterHandle->RotationOfMass());
-				CollectionMassToLocal[ClusterTransformIdx] = MassToLocal;
+				// NOTE: This method will is used to force an axis-aligned inertia tensor in local space,
+				// (ie, zero rotation of inertia) which avoids a problem with generating bounds for GCs
+				// when there's a mass rotation. The problem should go away once we remove MassToLocal
+				// and replace it with regular old CenterOfMass and RotationOfMass.
+				if (GeometryCollectionLocalInertiaDropOffDiagonalTerms)
+				{
+					AdjustClusterInertia(ClusterHandle, EInertiaOperations::LocalInertiaDropOffDiagonalTerms);
+				}
 
 				// Zero out CoM and RoM, and update XR accordingly so that when UpdateClusterMassProperties
 				// is called in the next level up with ClusterHandle as one of the the children it's XR will
 				// be correct.
-				MoveClusterToMassOffset(ClusterHandle, EMassOffsetType::EPosition | EMassOffsetType::ERotation);
+				//
+				// Store the relative transform from the particle's configuration space to its
+				// MassToLocal because GeometryCollections assume that the particle's XR is at its
+				// CoM and RoM - MassToLocal accounts for this offset.
+				CollectionMassToLocal[ClusterTransformIdx] = MoveClusterToMassOffset(ClusterHandle, EMassOffsetType::Position | EMassOffsetType::Rotation);
 
 				//update geometry
 				//merge children meshes and move them into cluster's mass space
diff --git a/Engine/Source/Runtime/Experimental/Chaos/Public/Chaos/PBDRigidClusteringAlgo.h b/Engine/Source/Runtime/Experimental/Chaos/Public/Chaos/PBDRigidClusteringAlgo.h
--- a/Engine/Source/Runtime/Experimental/Chaos/Public/Chaos/PBDRigidClusteringAlgo.h
+++ b/Engine/Source/Runtime/Experimental/Chaos/Public/Chaos/PBDRigidClusteringAlgo.h
@@ -5,20 +5,31 @@
 
 namespace Chaos
 {
-	enum CHAOS_API EMassOffsetType : uint8
+	enum class CHAOS_API EMassOffsetType : uint8
 	{
-		EPosition = 1 << 0,
-		ERotation = 1 << 1
+		None = 0,
+		Position = 1 << 0,
+		Rotation = 1 << 1
 	};
 	ENUM_CLASS_FLAGS(EMassOffsetType)
 
+	enum class CHAOS_API EInertiaOperations : uint8
+	{
+		None = 0,
+		LocalInertiaDropOffDiagonalTerms = 1 << 0
+	};
+	ENUM_CLASS_FLAGS(EInertiaOperations)
+
 	void CHAOS_API UpdateClusterMassProperties(
 		FPBDRigidClusteredParticleHandle* Parent,
 		const TSet<FPBDRigidParticleHandle*>& Children);
 
-	// If bPosition == true, set X/P of Cluster to its world CoM, and set its CoM to ZeroVector
-	// If bRotation == true, set R/Q of Cluster to its world RoM, and set its RoM to Identity
-	void CHAOS_API MoveClusterToMassOffset(FPBDRigidClusteredParticleHandle* Cluster, EMassOffsetType MassOffsetTypes);
+	// MassOffsetOperations & LocalInertiaDropOffDiagonalTerms: Discard off-diagonal inertia terms in particle's local space
+	void CHAOS_API AdjustClusterInertia(FPBDRigidClusteredParticleHandle* Cluster, EInertiaOperations InertiaOperations);
+
+	// MassOffsetTypes & Position: set X/P of Cluster to its world CoM, and set its CoM to ZeroVector
+	// MassOffsetTypes & Rotation: set R/Q of Cluster to its world RoM, and set its RoM to Identity
+	FRigidTransform3 CHAOS_API MoveClusterToMassOffset(FPBDRigidClusteredParticleHandle* Cluster, EMassOffsetType MassOffsetTypes);
 
 	void CHAOS_API UpdateKinematicProperties(
 		FPBDRigidParticleHandle* Parent,
diff --git a/Engine/Source/Runtime/Experimental/GeometryCollectionEngine/Private/GeometryCollection/DerivedDataGeometryCollectionCooker.cpp b/Engine/Source/Runtime/Experimental/GeometryCollectionEngine/Private/GeometryCollection/DerivedDataGeometryCollectionCooker.cpp
--- a/Engine/Source/Runtime/Experimental/GeometryCollectionEngine/Private/GeometryCollection/DerivedDataGeometryCollectionCooker.cpp
+++ b/Engine/Source/Runtime/Experimental/GeometryCollectionEngine/Private/GeometryCollection/DerivedDataGeometryCollectionCooker.cpp
@@ -60,7 +60,7 @@ bool FDerivedDataGeometryCollectionCooker::Build(TArray<uint8>& OutData)
 
 const TCHAR* FDerivedDataGeometryCollectionCooker::GetVersionString() const
 {
-	const TCHAR* VersionString = TEXT("E223057731B4465BBF7087001AF7ACBD");
+	const TCHAR* VersionString = TEXT("8F01FCF89D8449A09F53B80A15F1ED06");
 	return VersionString;
 }
 
