geometry collections should set debug name on GT particles

diff --git a/Engine/Source/Runtime/Experimental/Chaos/Private/PhysicsProxy/GeometryCollectionPhysicsProxy.cpp b/Engine/Source/Runtime/Experimental/Chaos/Private/PhysicsProxy/GeometryCollectionPhysicsProxy.cpp
--- a/Engine/Source/Runtime/Experimental/Chaos/Private/PhysicsProxy/GeometryCollectionPhysicsProxy.cpp
+++ b/Engine/Source/Runtime/Experimental/Chaos/Private/PhysicsProxy/GeometryCollectionPhysicsProxy.cpp
@@ -592,6 +592,10 @@ void FGeometryCollectionPhysicsProxy::Initialize(Chaos::FPBDRigidsEvolutionBase
 				{
 					P->SetSpatialIdx(Chaos::FSpatialAccelerationIdx{ 0,0 });
 				}
+
+#if !(UE_BUILD_SHIPPING || UE_BUILD_TEST)
+				P->SetDebugName(MakeShared<FString, ESPMode::ThreadSafe>(FString::Printf(TEXT("%s-%d"), *Parameters.Name, Index)));
+#endif
 			}
 			// this step is necessary for Phase 2 where we need to walk back the hierarchy from children to parent 
 			if (bGeometryCollectionAlwaysGenerateGTCollisionForClusters && GameThreadCollection.Children[Index].Num() == 0)
