IAS: Added accessor for total amount of persisted bytes
#rnx

diff --git a/Engine/Source/Runtime/Experimental/IoStoreOnDemand/Private/FileIoCache.cpp b/Engine/Source/Runtime/Experimental/IoStoreOnDemand/Private/FileIoCache.cpp
--- a/Engine/Source/Runtime/Experimental/IoStoreOnDemand/Private/FileIoCache.cpp
+++ b/Engine/Source/Runtime/Experimental/IoStoreOnDemand/Private/FileIoCache.cpp
@@ -175,6 +175,7 @@ public:
 	void InsertPersisted(FCacheEntryList&& InPersisted, const uint64 CursorPos);
 	void RemovePersisted(const uint64 RequiredSize);
 	uint64 GetPendingBytes() const { return TotalPendingBytes; }
+	uint64 GetTotalPersistedBytes() const { return TotalPersistedBytes; }
 	FIoStatus Load(const FString& FilePath, uint64& OutCursorPos);
 	FIoStatus Save(const FString& FilePath, const uint64 CursorPos);
 
