IAS: Move file cache to its won <PersistentDownloadDir>/ias folder.

#rnx

diff --git a/Engine/Source/Runtime/Experimental/IoStoreOnDemand/Private/FileIoCache.cpp b/Engine/Source/Runtime/Experimental/IoStoreOnDemand/Private/FileIoCache.cpp
--- a/Engine/Source/Runtime/Experimental/IoStoreOnDemand/Private/FileIoCache.cpp
+++ b/Engine/Source/Runtime/Experimental/IoStoreOnDemand/Private/FileIoCache.cpp
@@ -670,12 +670,21 @@ void FFileIoCache::Initialize()
 
 	WriterThread.Reset(FRunnableThread::Create(this, TEXT("Ias.FileCache"), 0, TPri_BelowNormal));
 
-	const FString CacheDir = FPaths::ProjectPersistentDownloadDir() / TEXT("chunkdownload");
+	const FString CacheDir = FPaths::ProjectPersistentDownloadDir() / TEXT("ias");
 	CacheFilePath = CacheDir / TEXT("cache.ucas");
 	WriteCursorPos = 0;
 
 	IFileManager& FileMgr = IFileManager::Get();
 
+	// TEMP: Drop file cache in old location
+	{
+		const FString OldCacheDir = FPaths::ProjectPersistentDownloadDir() / TEXT("chunkdownload");
+		FString OldCacheFilePath = OldCacheDir / TEXT("cache.ucas");
+		FString OldCacheTocPath = OldCacheFilePath + TEXT(".toc");
+		FileMgr.Delete(*OldCacheTocPath);
+		FileMgr.Delete(*OldCacheFilePath);
+	}
+
 	FString CacheTocPath = CacheFilePath + TEXT(".toc");
 	if (CacheConfig.DropCache)
 	{
