IAS: Removed duplicate cache-drop command line argument that Ridgers missed

#rnx

diff --git a/Engine/Source/Runtime/Experimental/IoStoreOnDemand/Private/FileIoCache.cpp b/Engine/Source/Runtime/Experimental/IoStoreOnDemand/Private/FileIoCache.cpp
--- a/Engine/Source/Runtime/Experimental/IoStoreOnDemand/Private/FileIoCache.cpp
+++ b/Engine/Source/Runtime/Experimental/IoStoreOnDemand/Private/FileIoCache.cpp
@@ -696,34 +696,26 @@ void FFileIoCache::Initialize()
 	FIoStatus CacheStatus(EIoErrorCode::Unknown);
 	if (FileMgr.FileExists(*CacheTocPath))
 	{
-		if (FParse::Param(FCommandLine::Get(), TEXT("ClearIoCache")))
+		CacheStatus = CacheMap.Load(CacheTocPath, WriteCursorPos);
+		if (CacheStatus.IsOk())
 		{
-			UE_LOG(LogIasCache, Log, TEXT("Deleting cache file '%s'"), *CacheFilePath);
-			FileMgr.Delete(*CacheFilePath);
-		}
-		else
-		{
-			CacheStatus = CacheMap.Load(CacheTocPath, WriteCursorPos);
-			if (CacheStatus.IsOk())
+			check(WriteCursorPos != ~uint64(0));
+
+			UE_LOG(LogIasCache, Log, TEXT("Loaded TOC '%s'"), *CacheTocPath);
+			if (FileMgr.FileExists(*CacheFilePath))
 			{
-				check(WriteCursorPos != ~uint64(0));
-				
-				UE_LOG(LogIasCache, Log, TEXT("Loaded TOC '%s'"), *CacheTocPath);
-				if (FileMgr.FileExists(*CacheFilePath))
-				{
-					//TODO: Integrity check?
-				}
-				else
-				{
-					UE_LOG(LogIasCache, Warning, TEXT("Failed to open cache file ''"), *CacheFilePath);
-					CacheStatus = FIoStatus(EIoErrorCode::FileNotOpen);
-				}
+				//TODO: Integrity check?
 			}
 			else
 			{
-				UE_LOG(LogIasCache, Warning, TEXT("Failed to load TOC '%s'"), *CacheTocPath);
+				UE_LOG(LogIasCache, Warning, TEXT("Failed to open cache file ''"), *CacheFilePath);
+				CacheStatus = FIoStatus(EIoErrorCode::FileNotOpen);
 			}
 		}
+		else
+		{
+			UE_LOG(LogIasCache, Warning, TEXT("Failed to load TOC '%s'"), *CacheTocPath);
+		}
 	}
 
 	if (!CacheStatus.IsOk())
