IAS: Improve file cache startup logging.

#rnx

diff --git a/Engine/Source/Runtime/Experimental/IoStoreOnDemand/Private/FileIoCache.cpp b/Engine/Source/Runtime/Experimental/IoStoreOnDemand/Private/FileIoCache.cpp
--- a/Engine/Source/Runtime/Experimental/IoStoreOnDemand/Private/FileIoCache.cpp
+++ b/Engine/Source/Runtime/Experimental/IoStoreOnDemand/Private/FileIoCache.cpp
@@ -665,9 +665,6 @@ FIoStatus FFileIoCache::Put(const FIoHash& Key, FIoBuffer& Data)
 
 void FFileIoCache::Initialize()
 {
-	UE_LOG(LogIasCache, Log,
-		TEXT("Initializing file I/O cache, disk size %lluB, memory size %lluB"), CacheConfig.DiskQuota, CacheConfig.MemoryQuota);
-
 	WriterThread.Reset(FRunnableThread::Create(this, TEXT("Ias.FileCache"), 0, TPri_BelowNormal));
 
 	const FString CacheDir = FPaths::ProjectPersistentDownloadDir() / TEXT("ias");
@@ -676,6 +673,10 @@ void FFileIoCache::Initialize()
 
 	IFileManager& FileMgr = IFileManager::Get();
 
+	UE_LOG(LogIasCache, Log,
+		TEXT("Initializing file cache '%s', max disk size %lluB, memory size %lluB"),
+			*CacheFilePath, CacheConfig.DiskQuota, CacheConfig.MemoryQuota);
+
 	// TEMP: Drop file cache in old location
 	{
 		const FString OldCacheDir = FPaths::ProjectPersistentDownloadDir() / TEXT("chunkdownload");
@@ -688,6 +689,7 @@ void FFileIoCache::Initialize()
 	FString CacheTocPath = CacheFilePath + TEXT(".toc");
 	if (CacheConfig.DropCache)
 	{
+		UE_LOG(LogIasCache, Log, TEXT("Dropping existing file cache '%s'"), *CacheFilePath);
 		FileMgr.Delete(*CacheTocPath);
 		FileMgr.Delete(*CacheFilePath);
 		return;
@@ -701,20 +703,22 @@ void FFileIoCache::Initialize()
 		{
 			check(WriteCursorPos != ~uint64(0));
 
-			UE_LOG(LogIasCache, Log, TEXT("Loaded TOC '%s'"), *CacheTocPath);
 			if (FileMgr.FileExists(*CacheFilePath))
 			{
 				//TODO: Integrity check?
+				UE_LOG(LogIasCache, Log, TEXT("Loaded file cache '%s', cached bytes %lluB"),
+					*CacheFilePath, CacheMap.GetTotalPersistedBytes());
 			}
 			else
 			{
-				UE_LOG(LogIasCache, Warning, TEXT("Failed to open cache file ''"), *CacheFilePath);
+				UE_LOG(LogIasCache, Warning, TEXT("Failed to open file cache '%s'"), *CacheFilePath);
 				CacheStatus = FIoStatus(EIoErrorCode::FileNotOpen);
 			}
 		}
 		else
 		{
-			UE_LOG(LogIasCache, Warning, TEXT("Failed to load TOC '%s'"), *CacheTocPath);
+			UE_LOG(LogIasCache, Warning, TEXT("Failed to load file cache toc '%s', reason '%s'"),
+				*CacheTocPath, *CacheStatus.ToString());
 		}
 	}
 
@@ -728,6 +732,7 @@ void FFileIoCache::Initialize()
 		{
 			FileMgr.MakeDirectory(*CacheDir, true);
 		}
+		UE_LOG(LogIasCache, Log, TEXT("Created file cache '%s', cached bytes 0B"), *CacheFilePath);
 	}
 }
 
@@ -743,7 +748,7 @@ void FFileIoCache::Shutdown()
 	WriterThread->Kill();
 
 	FString CacheTocPath = CacheFilePath + TEXT(".toc");
-	UE_LOG(LogIasCache, Log, TEXT("Saving TOC '%s'"), *CacheTocPath);
+	UE_LOG(LogIasCache, Log, TEXT("Saving file cache toc '%s'"), *CacheTocPath);
 	CacheMap.Save(CacheTocPath, WriteCursorPos);
 
 	CacheMap.Reset();
