IAS: Start write thread later (after cursor's set)

#rnx

diff --git a/Engine/Source/Runtime/Experimental/IoStoreOnDemand/Private/FileIoCache.cpp b/Engine/Source/Runtime/Experimental/IoStoreOnDemand/Private/FileIoCache.cpp
--- a/Engine/Source/Runtime/Experimental/IoStoreOnDemand/Private/FileIoCache.cpp
+++ b/Engine/Source/Runtime/Experimental/IoStoreOnDemand/Private/FileIoCache.cpp
@@ -13,6 +13,7 @@
 #include "HAL/RunnableThread.h"
 #include "IO/IoDispatcher.h"
 #include "IO/IoHash.h"
+#include "Misc/ScopeExit.h"
 #include "Misc/ScopeLock.h"
 #include "Misc/Paths.h"
 #include "Misc/CommandLine.h"
@@ -665,7 +666,9 @@ FIoStatus FFileIoCache::Put(const FIoHash& Key, FIoBuffer& Data)
 
 void FFileIoCache::Initialize()
 {
-	WriterThread.Reset(FRunnableThread::Create(this, TEXT("Ias.FileCache"), 0, TPri_BelowNormal));
+	ON_SCOPE_EXIT {
+		WriterThread.Reset(FRunnableThread::Create(this, TEXT("Ias.FileCache"), 0, TPri_BelowNormal));
+	};
 
 	const FString CacheDir = FPaths::ProjectPersistentDownloadDir() / TEXT("ias");
 	CacheFilePath = CacheDir / TEXT("cache.ucas");
