IAS: Improve http stats and logging.

#rnx

diff --git a/Engine/Source/Runtime/Experimental/IoStoreOnDemand/Private/OnDemandIoDispatcherBackend.cpp b/Engine/Source/Runtime/Experimental/IoStoreOnDemand/Private/OnDemandIoDispatcherBackend.cpp
--- a/Engine/Source/Runtime/Experimental/IoStoreOnDemand/Private/OnDemandIoDispatcherBackend.cpp
+++ b/Engine/Source/Runtime/Experimental/IoStoreOnDemand/Private/OnDemandIoDispatcherBackend.cpp
@@ -304,7 +304,9 @@ public:
 	bool Tick();
 
 private:
-	void Issue(UE::HTTP::FRequest&& Request, FIoReadCallback&& Callback, FAnsiStringView Url = FAnsiStringView());
+	void Issue(FAnsiStringView Url,
+		FIoReadCallback&& Callback,
+		FIoOffsetAndLength Range = FIoOffsetAndLength());
 
 	FString SvcsUrl;
 	int32 MaxConnections;
@@ -324,27 +326,23 @@ FHttpClient::FHttpClient(const FString& ServiceUrl, int32 MaxConnectionCount)
 
 void FHttpClient::Get(FAnsiStringView Url, const FIoOffsetAndLength& Range, FIoReadCallback&& Callback)
 {
-	const uint64 RangeStart = Range.GetOffset();
-	const uint64 RangeEnd = Range.GetOffset() + Range.GetLength();
-	UE::HTTP::FRequest Request = EventLoop.Get(Url, *ConnectionPool);
-	Request.Header(ANSITEXTVIEW("Range"), WriteToAnsiString<64>(ANSITEXTVIEW("bytes="), RangeStart, ANSITEXTVIEW("-"), RangeEnd));
-	
-	Issue(MoveTemp(Request), MoveTemp(Callback), Url);
+	Issue(Url, MoveTemp(Callback), Range);
 }
 
 void FHttpClient::Get(FAnsiStringView Url, FIoReadCallback&& Callback)
 {
-	Issue(EventLoop.Get(Url, *ConnectionPool), MoveTemp(Callback), Url);
+	Issue(Url, MoveTemp(Callback));
 }
 
-void FHttpClient::Issue(UE::HTTP::FRequest&& Request, FIoReadCallback&& Callback, FAnsiStringView DebugUrl)
+void FHttpClient::Issue(FAnsiStringView Url, FIoReadCallback&& Callback, FIoOffsetAndLength Range)
 {
 	using namespace UE::HTTP;
-	
+
 	auto Sink = [
 		Buffer = FIoBuffer(),
 		Callback = MoveTemp(Callback),
-		Url = FString(DebugUrl),
+		Url = FString(Url),
+		Offset = Range.GetOffset(),
 		StartTime = FPlatformTime::Cycles64(),
 		StatusCode = uint32(0)]
 		(const FTicketStatus& Status) mutable
@@ -363,14 +361,16 @@ void FHttpClient::Issue(UE::HTTP::FRequest&& Request, FIoReadCallback&& Callback
 				if (const FIoBuffer& Content = Status.GetContent(); bSuccessful && Content.GetSize() > 0)
 				{
 					UE_LOG(LogIas, VeryVerbose, TEXT("%s"),
-						*WriteToString<256>(TEXT("HTTP GET - "), Url, TEXT(" ("), StatusCode, TEXT(" "), Duration, TEXT("ms "), Content.GetSize(), TEXT(" Bytes)")));
+						*WriteToString<256>(TEXT("HTTP GET - "), Url, TEXT("("), StatusCode, TEXT(" "),
+						Duration, TEXT("ms "), Offset, TEXT(" Offset "), Content.GetSize(), TEXT(" Bytes)")));
 
 					Callback(Content);
 				}
 				else
 				{
-					UE_LOG(LogIas, VeryVerbose, TEXT("%s"),
-						*WriteToString<256>(TEXT("HTTP GET - "), Url, TEXT(" ("), StatusCode, TEXT(" "), Duration, TEXT("ms)")));
+					UE_LOG(LogIas, Warning, TEXT("%s"),
+						*WriteToString<256>(TEXT("HTTP GET - "), Url, TEXT("("), StatusCode, TEXT(" "),
+						Duration, TEXT("ms "), Offset, TEXT(" Offset 0 Bytes) : Invalid Content : "), StatusCode));
 					
 					Callback(FIoStatus(EIoErrorCode::NotFound, TEXTVIEW("Invalid Content")));
 				}
@@ -378,13 +378,22 @@ void FHttpClient::Issue(UE::HTTP::FRequest&& Request, FIoReadCallback&& Callback
 			else if (FTicketStatus::EId::Error == Status.GetId())
 			{
 				const uint64 Duration = (uint64)FPlatformTime::ToMilliseconds64(FPlatformTime::Cycles64() - StartTime);
-				UE_LOG(LogIas, VeryVerbose, TEXT("%s"),
-					*WriteToString<256>(TEXT("HTTP GET - "), Url, TEXT(" ("), StatusCode, TEXT(" "), Duration, TEXTVIEW("ms) : "), Status.GetErrorReason()));
+				UE_LOG(LogIas, Warning, TEXT("%s"),
+					*WriteToString<256>(TEXT("HTTP GET - "), Url, TEXT(" ("), StatusCode, TEXT(" "),
+					Duration, TEXTVIEW("ms "), Offset, TEXT(" Offset) : "), Status.GetErrorReason()));
 
 				Callback(FIoStatus(EIoErrorCode::ReadError, TEXTVIEW("HTTP Error")));
 			}
 		};
 
+	UE::HTTP::FRequest Request = EventLoop.Get(Url, *ConnectionPool);
+	const uint64 RangeStart = Range.GetOffset();
+	const uint64 RangeEnd = Range.GetOffset() + Range.GetLength();
+	if (RangeStart > 0 || RangeEnd > 0)
+	{
+		Request.Header(ANSITEXTVIEW("Range"), WriteToAnsiString<64>(ANSITEXTVIEW("bytes="), RangeStart, ANSITEXTVIEW("-"), RangeEnd));
+	}
+	
 	EventLoop.Send(MoveTemp(Request), MoveTemp(Sink));
 }
 
@@ -1039,7 +1048,7 @@ bool FOnDemandIoBackend::Resolve(FIoRequestImpl* Request)
 			return CompleteRequest(ChunkRequest);
 		}
 
-		Stats.OnHttpRequestEnqueue();
+		Stats.OnHttpEnqueue();
 		HttpRequests.Enqueue(ChunkRequest);
 		TickBackendEvent->Trigger();
 	};
@@ -1275,7 +1284,7 @@ uint32 FOnDemandIoBackend::Run()
 				NextChunkRequest = ChunkRequest->NextRequest;
 				ChunkRequest->NextRequest = nullptr;
 
-				Stats.OnHttpRequestDequeue();
+				Stats.OnHttpDequeue();
 
 				TAnsiStringBuilder<256> Url;
 				ChunkRequest->Params.GetUrl(Url);
@@ -1290,8 +1299,8 @@ uint32 FOnDemandIoBackend::Run()
 						{
 							if (++ChunkRequest->HttpRetryCount <= GIoDispatcherMaxHttpRetryCount)
 							{
-								Stats.OnHttpRequestFail();
-								Stats.OnHttpRequestEnqueue();
+								Stats.OnHttpRetry();
+								Stats.OnHttpEnqueue();
 								return HttpRequests.Enqueue(ChunkRequest); 
 							}
 						}
@@ -1299,11 +1308,15 @@ uint32 FOnDemandIoBackend::Run()
 						if (Status.IsOk())
 						{
 							ChunkRequest->Chunk = Status.ConsumeValueOrDie();
-							Stats.OnHttpRequestComplete(ChunkRequest->Chunk.DataSize());
+							Stats.OnHttpGet(ChunkRequest->Chunk.DataSize());
 						}
 						else
 						{
-							Stats.OnHttpRequestFail();
+							TAnsiStringBuilder<256> Url;
+							ChunkRequest->Params.GetUrl(Url);
+							UE_LOG(LogIas, Error, TEXT("%s"),
+								*WriteToString<256>(TEXT("HTTP FAILED - "), Url, TEXT(" ("), ChunkRequest->Params.ChunkRange.GetOffset(), TEXT(" Offset)")));
+							Stats.OnHttpError();
 						}
 
 						Launch(UE_SOURCE_LOCATION, [this, ChunkRequest]()
diff --git a/Engine/Source/Runtime/Experimental/IoStoreOnDemand/Private/Statistics.cpp b/Engine/Source/Runtime/Experimental/IoStoreOnDemand/Private/Statistics.cpp
--- a/Engine/Source/Runtime/Experimental/IoStoreOnDemand/Private/Statistics.cpp
+++ b/Engine/Source/Runtime/Experimental/IoStoreOnDemand/Private/Statistics.cpp
@@ -34,7 +34,7 @@ static int32 BytesToApproxKB(uint64 Bytes) { return int32(Bytes >> 10); }
 // iorequest stats
 FCounterInt			GIoRequestsMade(TEXT("Ias/IoRequestsMade"), TraceCounterDisplayHint_None);
 FCounterAtomicInt	GIoRequestsCompleted(TEXT("Ias/IoRequestsCompleted"), TraceCounterDisplayHint_None);
-FCounterAtomicInt	GIoRequestsCompletedSize(TEXT("Ias/Size/IoRequestsCompletedSize"), TraceCounterDisplayHint_Memory);
+FCounterAtomicInt	GIoRequestsCompletedSize(TEXT("Ias/IoRequestsCompletedSize"), TraceCounterDisplayHint_Memory);
 FCounterInt			GIoRequestsCancelled(TEXT("Ias/IoRequestsCancelled"), TraceCounterDisplayHint_None);
 FCounterAtomicInt	GIoRequestsFailed(TEXT("Ias/IoRequestsFailed"), TraceCounterDisplayHint_None);
 // chunkrequest stats
@@ -51,11 +51,12 @@ FCounterAtomicInt	GCachePendingBytes(TEXT("Ias/CachePendingBytes"), TraceCounter
 FCounterAtomicInt	GCacheReadBytes(TEXT("Ias/CacheReadBytes"), TraceCounterDisplayHint_Memory);
 FCounterAtomicInt	GCacheRejectBytes(TEXT("Ias/CachePutRejectBytes"), TraceCounterDisplayHint_Memory);
 // http stats
-FCounterInt			GHttpRequestsCompleted(TEXT("Ias/HttpRequestsCompleted"), TraceCounterDisplayHint_None);
-FCounterInt			GHttpRequestsFailed(TEXT("Ias/HttpRequestsFailed"), TraceCounterDisplayHint_None);
-FCounterAtomicInt	GHttpRequestsPending(TEXT("Ias/HttpRequestsPending"), TraceCounterDisplayHint_None);
-FCounterInt			GHttpRequestsInflight(TEXT("Ias/HttpRequestsInflight"), TraceCounterDisplayHint_None);
-FCounterInt			GHttpRequestsCompletedSize(TEXT("Ias/Size/HttpRequestsCompletedSize"), TraceCounterDisplayHint_Memory);
+FCounterInt			GHttpGetCount(TEXT("Ias/HttpGetCount"), TraceCounterDisplayHint_None);
+FCounterInt			GHttpErrorCount(TEXT("Ias/HttpErrorCount"), TraceCounterDisplayHint_None);
+FCounterInt			GHttpRetryCount(TEXT("Ias/HttpRetryCount"), TraceCounterDisplayHint_None);
+FCounterAtomicInt	GHttpPendingCount(TEXT("Ias/HttpPendingCount"), TraceCounterDisplayHint_None);
+FCounterInt			GHttpInflightCount(TEXT("Ias/HttpInflightCount"), TraceCounterDisplayHint_None);
+FCounterInt			GHttpDownloadedBytes(TEXT("Ias/HttpDownloadedBytes"), TraceCounterDisplayHint_Memory);
 
 ////////////////////////////////////////////////////////////////////////////////
 // CSV STATS
@@ -80,11 +81,12 @@ CSV_DEFINE_STAT(Ias, FrameCachePendingBytes);
 CSV_DEFINE_STAT(Ias, FrameCacheReadBytes);
 CSV_DEFINE_STAT(Ias, FrameCacheRejectBytes);
 // http stats
-CSV_DEFINE_STAT(Ias, FrameHttpRequestsCompleted);
-CSV_DEFINE_STAT(Ias, FrameHttpRequestsFailed);
-CSV_DEFINE_STAT(Ias, FrameHttpRequestsPending);
-CSV_DEFINE_STAT(Ias, FrameHttpRequestsInflight);
-CSV_DEFINE_STAT(Ias, FrameHttpRequestsCompletedSize);
+CSV_DEFINE_STAT(Ias, HttpGetCount);
+CSV_DEFINE_STAT(Ias, HttpRetryCount);
+CSV_DEFINE_STAT(Ias, HttpErrorCount);
+CSV_DEFINE_STAT(Ias, HttpPendingCount);
+CSV_DEFINE_STAT(Ias, HttpInflightCount);
+CSV_DEFINE_STAT(Ias, HttpDownloadKB);
 
 static FOnDemandIoBackendStats* GStatistics = nullptr;
 
@@ -190,40 +192,49 @@ void FOnDemandIoBackendStats::OnCachePersistedBytes(uint64 TotalSize)
 	CSV_CUSTOM_STAT_DEFINED(FrameCacheCachedBytes, BytesToApproxKB(GCacheCachedBytes.Get()), ECsvCustomStatOp::Set);
 }
 
-void FOnDemandIoBackendStats::OnHttpRequestEnqueue()
+void FOnDemandIoBackendStats::OnHttpEnqueue()
 {
-	GHttpRequestsPending.Add(1);
-	CSV_CUSTOM_STAT_DEFINED(FrameHttpRequestsPending, int32(GHttpRequestsPending.Get()), ECsvCustomStatOp::Set);
+	GHttpPendingCount.Add(1);
+	CSV_CUSTOM_STAT_DEFINED(HttpPendingCount, int32(GHttpPendingCount.Get()), ECsvCustomStatOp::Set);
 }
 
-void FOnDemandIoBackendStats::OnHttpRequestDequeue()
+void FOnDemandIoBackendStats::OnHttpDequeue()
 {
-	GHttpRequestsPending.Add(-1);
-	CSV_CUSTOM_STAT_DEFINED(FrameHttpRequestsPending, int32(GHttpRequestsPending.Get()), ECsvCustomStatOp::Set);
+	GHttpPendingCount.Add(-1);
+	CSV_CUSTOM_STAT_DEFINED(HttpPendingCount, int32(GHttpPendingCount.Get()), ECsvCustomStatOp::Set);
 
-	GHttpRequestsInflight.Add(1);
-	CSV_CUSTOM_STAT_DEFINED(FrameHttpRequestsPending, int32(GHttpRequestsInflight.Get()), ECsvCustomStatOp::Set);
+	GHttpInflightCount.Add(1);
+	CSV_CUSTOM_STAT_DEFINED(HttpPendingCount, int32(GHttpInflightCount.Get()), ECsvCustomStatOp::Set);
 }
 
-void FOnDemandIoBackendStats::OnHttpRequestComplete(uint64 InSize)
+void FOnDemandIoBackendStats::OnHttpGet(uint64 InSize)
 {
-	GHttpRequestsInflight.Add(-1);
-	CSV_CUSTOM_STAT_DEFINED(FrameHttpRequestsPending, int32(GHttpRequestsInflight.Get()), ECsvCustomStatOp::Set);
+	GHttpInflightCount.Add(-1);
+	CSV_CUSTOM_STAT_DEFINED(HttpPendingCount, int32(GHttpInflightCount.Get()), ECsvCustomStatOp::Set);
 
-	GHttpRequestsCompleted.Add(1);
-	CSV_CUSTOM_STAT_DEFINED(FrameHttpRequestsCompleted, int32(GHttpRequestsCompleted.Get()), ECsvCustomStatOp::Set);
+	GHttpGetCount.Add(1);
+	CSV_CUSTOM_STAT_DEFINED(HttpGetCount, int32(GHttpGetCount.Get()), ECsvCustomStatOp::Set);
 
-	GHttpRequestsCompletedSize.Add(InSize);
-	CSV_CUSTOM_STAT_DEFINED(FrameHttpRequestsCompletedSize, BytesToApproxKB(GHttpRequestsCompletedSize.Get()), ECsvCustomStatOp::Set);
+	GHttpDownloadedBytes.Add(InSize);
+	CSV_CUSTOM_STAT_DEFINED(HttpDownloadKB, BytesToApproxKB(GHttpDownloadedBytes.Get()), ECsvCustomStatOp::Set);
 }
 
-void FOnDemandIoBackendStats::OnHttpRequestFail()
+void FOnDemandIoBackendStats::OnHttpRetry()
 {
-	GHttpRequestsInflight.Add(-1);
-	CSV_CUSTOM_STAT_DEFINED(FrameHttpRequestsPending, int32(GHttpRequestsInflight.Get()), ECsvCustomStatOp::Set);
+	GHttpInflightCount.Add(-1);
+	CSV_CUSTOM_STAT_DEFINED(HttpPendingCount, int32(GHttpInflightCount.Get()), ECsvCustomStatOp::Set);
 
-	GHttpRequestsFailed.Add(1);
-	CSV_CUSTOM_STAT_DEFINED(FrameHttpRequestsFailed, int32(GHttpRequestsFailed.Get()), ECsvCustomStatOp::Set);
+	GHttpRetryCount.Add(1);
+	CSV_CUSTOM_STAT_DEFINED(HttpRetryCount, int32(GHttpRetryCount.Get()), ECsvCustomStatOp::Set);
+}
+
+void FOnDemandIoBackendStats::OnHttpError()
+{
+	GHttpInflightCount.Add(-1);
+	CSV_CUSTOM_STAT_DEFINED(HttpPendingCount, int32(GHttpInflightCount.Get()), ECsvCustomStatOp::Set);
+
+	GHttpErrorCount.Add(1);
+	CSV_CUSTOM_STAT_DEFINED(HttpErrorCount, int32(GHttpErrorCount.Get()), ECsvCustomStatOp::Set);
 }
 
 } // namespace UE::IO::Private
diff --git a/Engine/Source/Runtime/Experimental/IoStoreOnDemand/Private/Statistics.h b/Engine/Source/Runtime/Experimental/IoStoreOnDemand/Private/Statistics.h
--- a/Engine/Source/Runtime/Experimental/IoStoreOnDemand/Private/Statistics.h
+++ b/Engine/Source/Runtime/Experimental/IoStoreOnDemand/Private/Statistics.h
@@ -43,10 +43,11 @@ public:
 	void OnCachePendingBytes(uint64 TotalSize) IAS_STATISTICS_IMPL()
 	void OnCachePersistedBytes(uint64 TotalSize) IAS_STATISTICS_IMPL()
 
-	void OnHttpRequestEnqueue() IAS_STATISTICS_IMPL()
-	void OnHttpRequestDequeue() IAS_STATISTICS_IMPL()
-	void OnHttpRequestComplete(uint64 InSize) IAS_STATISTICS_IMPL()
-	void OnHttpRequestFail() IAS_STATISTICS_IMPL()
+	void OnHttpEnqueue() IAS_STATISTICS_IMPL()
+	void OnHttpDequeue() IAS_STATISTICS_IMPL()
+	void OnHttpGet(uint64 InSize) IAS_STATISTICS_IMPL()
+	void OnHttpRetry() IAS_STATISTICS_IMPL()
+	void OnHttpError() IAS_STATISTICS_IMPL()
 };
 
 #undef IAS_STATISTICS_IMPL
