Fix stack over flow for field nodes instances

diff --git a/Engine/Source/Runtime/Engine/Classes/PhysicsField/PhysicsFieldComponent.h b/Engine/Source/Runtime/Engine/Classes/PhysicsField/PhysicsFieldComponent.h
--- a/Engine/Source/Runtime/Engine/Classes/PhysicsField/PhysicsFieldComponent.h
+++ b/Engine/Source/Runtime/Engine/Classes/PhysicsField/PhysicsFieldComponent.h
@@ -192,6 +192,9 @@ public:
 
 	/** Targets offsets in the nodes array*/
 	TStaticArray<int32, EFieldPhysicsType::Field_PhysicsType_Max + 1> TargetsOffsets;
+	
+	/** Bounds offsets in the bounds array*/
+	TStaticArray<int32, EFieldPhysicsType::Field_PhysicsType_Max + 1> BoundsOffsets;
 
 	/** Nodes offsets in the paramter array */
 	TArray<int32> NodesOffsets;
@@ -199,6 +202,12 @@ public:
 	/** Nodes input parameters and connection */
 	TArray<float> NodesParams;
 
+	/** Commands bounds min sorted per target type */
+	TArray<FVector4> BoundsMin;
+
+	/** Commands bounds max sorted per target type */
+	TArray<FVector4> BoundsMax;
+
 	/** List of all the field commands in the world */
 	TArray<FFieldSystemCommand> FieldCommands;
 
diff --git a/Engine/Source/Runtime/Engine/Private/PhysicsField/PhysicsFieldComponent.cpp b/Engine/Source/Runtime/Engine/Private/PhysicsField/PhysicsFieldComponent.cpp
--- a/Engine/Source/Runtime/Engine/Private/PhysicsField/PhysicsFieldComponent.cpp
+++ b/Engine/Source/Runtime/Engine/Private/PhysicsField/PhysicsFieldComponent.cpp
@@ -672,26 +672,23 @@ void FPhysicsFieldInstance::ReleaseInstance()
 	NodesParams.Empty();
 	TargetsMax.Empty();
 	TargetsMin.Empty();
+	BoundsMin.Empty();
+	BoundsMax.Empty();
 }
 
 void FPhysicsFieldInstance::UpdateInstance(const float TimeSeconds)
 {
-	NodesOffsets.Empty();
-	NodesParams.Empty();
+	NodesOffsets.Reset();
+	NodesParams.Reset();
 
-	TStaticArray<int32, EFieldPhysicsType::Field_PhysicsType_Max + 1> BoundsOffsets;
-	TArray<FVector4> BoundsMin, BoundsMax;
-
-	int32 NumOffsets = 0;
-	int32 NumParams = 0;
-	TArray<int32> TestOffsets;
-	TArray<float> TestParams;
+	BoundsMin.Reset();
+	BoundsMax.Reset();
 
 	if (FieldResource)
 	{
 		TargetsMin.Init(FVector(0), EFieldPhysicsType::Field_PhysicsType_Max);
 		TargetsMax.Init(FVector(0), EFieldPhysicsType::Field_PhysicsType_Max);
-
+		
 		for (auto& TargetOffset : TargetsOffsets)
 		{
 			TargetOffset = 0;
@@ -702,91 +699,65 @@ void FPhysicsFieldInstance::UpdateInstance(const float TimeSeconds)
 		}
 		for (auto& TargetType : FieldResource->FieldInfos.TargetTypes)
 		{
-			TArray<FFieldNodeBase*> TargetRoots;
-			TArray<float> CreationTimes;
+			const int32 NodeOffset = NodesOffsets.Num();
+			FVector TargetMin(FLT_MAX), TargetMax(-FLT_MAX);
+			int32 NumCommands = 0;
+			
+			const EFieldOutputType OutputType = GetFieldTargetOutput(TargetType);
 			for (auto& FieldCommand : FieldCommands)
 			{
 				const EFieldPhysicsType CommandType = GetFieldPhysicsType(FieldCommand.TargetAttribute);
 				if (CommandType == TargetType)
 				{
-					const TUniquePtr<FFieldNodeBase>& RootNode = FieldCommand.RootNode;
-					TargetRoots.Add(RootNode.Get());
-					CreationTimes.Add(FieldCommand.TimeCreation);
-					BoundsOffsets[TargetType+1]++;
+					UPhysicsFieldComponent::BuildCommandBounds(FieldCommand);
+
 					BoundsMin.Add(FVector4(FieldCommand.BoundingBox.Min, 0.0));
 					BoundsMax.Add(FVector4(FieldCommand.BoundingBox.Max, FieldCommand.MaxMagnitude));
-				}
-			}
-			FFieldNodeBase* TargetNode = nullptr;
-			TMap<FFieldNodeBase*, float> CommandTimes;
-			if (TargetRoots.Num() == 1)
-			{
-				TargetNode = TargetRoots[0];
-				CommandTimes.Add(TargetNode, CreationTimes[0]);
-			}
-			else if (TargetRoots.Num() > 1)
-			{
-				const EFieldOutputType OutputType = GetFieldTargetOutput(TargetType);
-				if (OutputType == EFieldOutputType::Field_Output_Vector)
-				{
-					FFieldNode<FVector>* PreviousNode = StaticCast<FFieldNode<FVector>*>(TargetRoots[0]->NewCopy());
-					CommandTimes.Add(PreviousNode, CreationTimes[0]);
-					FFieldNode<FVector>* NextNode = nullptr;
-					for (int32 TargetIndex = 1; TargetIndex < TargetRoots.Num(); ++TargetIndex)
-					{
-						NextNode = StaticCast<FFieldNode<FVector>*>(TargetRoots[TargetIndex]->NewCopy());
-						CommandTimes.Add(NextNode, CreationTimes[TargetIndex]);
 
-						PreviousNode = new FSumVector(1.0, nullptr, PreviousNode,
-							NextNode, EFieldOperationType::Field_Add);
-					}
-					TargetNode = PreviousNode;
-				}
-				else if (OutputType == EFieldOutputType::Field_Output_Scalar)
-				{
-					FFieldNode<float>* PreviousNode = StaticCast<FFieldNode<float>*>(TargetRoots[0]->NewCopy());
-					CommandTimes.Add(PreviousNode, CreationTimes[0]);
-					FFieldNode<float>* NextNode = nullptr;
-					for (int32 TargetIndex = 1; TargetIndex < TargetRoots.Num(); ++TargetIndex)
+					TargetMin = MinVector(TargetMin, FieldCommand.BoundingBox.Min);
+					TargetMax = MaxVector(TargetMax, FieldCommand.BoundingBox.Max);
+					
+					// The order should N0,C0,N1,C1,S1,N2,C2,S2...(N: node data, C: conversion data, S: Sum data)
+					// Loop over the instances directly instead of building an instance graph to avoid stack overflow
+					if(NumCommands == 1)
 					{
-						NextNode = StaticCast<FFieldNode<float>*>(TargetRoots[TargetIndex]->NewCopy());
-						CommandTimes.Add(NextNode, CreationTimes[TargetIndex]);
-						PreviousNode = new FSumScalar(1.0, PreviousNode,
-							NextNode, EFieldOperationType::Field_Add);
+						// Pre fill datas : Only append a conversion field (int->float) for the first command 
+						// to be used by the SumScalar in the post fill
+						if (OutputType == EFieldOutputType::Field_Output_Integer)
+						{
+							FConversionField<int32, float>::FillDefaultSetup(NodesOffsets, NodesParams);
+						}
 					}
-					TargetNode = PreviousNode;
-				}
-				else if (OutputType == EFieldOutputType::Field_Output_Integer)
-				{
-					FFieldNode<float>* PreviousNode = new FConversionField<int32, float>(StaticCast<FFieldNode<int32>*>(TargetRoots[0]->NewCopy()));
-					CommandTimes.Add(PreviousNode, CreationTimes[0]);
-					FFieldNode<float>* NextNode = nullptr;
-					for (int32 TargetIndex = 1; TargetIndex < TargetRoots.Num(); ++TargetIndex)
+
+					FieldCommand.RootNode->FillSetupDatas(NodesOffsets, NodesParams, FieldCommand.TimeCreation);
+
+					if(NumCommands > 0)
 					{
-						NextNode = new FConversionField<int32, float>(StaticCast<FFieldNode<int32>*>(TargetRoots[TargetIndex]->NewCopy()));
-						CommandTimes.Add(NextNode, CreationTimes[TargetIndex]);
-						PreviousNode = new FSumScalar(1.0, PreviousNode,
-							NextNode, EFieldOperationType::Field_Add);
+						// Post fill datas : Append a sum vector field in between the previous node and the current one
+						if (OutputType == EFieldOutputType::Field_Output_Vector)
+						{
+							FSumVector::FillDefaultSetup(NodesOffsets, NodesParams);
+						}
+						// Post fill datas : Append a sum scalar field in between the previous node and the current one
+						else if (OutputType == EFieldOutputType::Field_Output_Scalar)
+						{
+							FSumScalar::FillDefaultSetup(NodesOffsets, NodesParams);
+						}
+						// Post fill datas : Append a conversion field (int->float) + sum scalar field in between the previous node and the current one
+						else if (OutputType == EFieldOutputType::Field_Output_Integer)
+						{
+							FConversionField<int32, float>::FillDefaultSetup(NodesOffsets, NodesParams);
+							FSumScalar::FillDefaultSetup(NodesOffsets, NodesParams);
+						}
 					}
-					TargetNode = PreviousNode;
+					++NumCommands;
 				}
 			}
-			const int32 PreviousNodes = NodesOffsets.Num();
-			FVector MinBound(-FLT_MAX), MaxBound(FLT_MAX);
-			if (TargetNode)
-			{
-				TargetNode->FillSetupCount(NumOffsets, NumParams);
-				NodesOffsets.Reserve(NumOffsets);
-				NodesParams.Reserve(NumParams);
-
-				TargetNode->FillSetupDatas(NodesOffsets, NodesParams, CommandTimes, 0.0);
-				TargetNode->ComputeFieldBounds(MinBound, MaxBound);
-				if (TargetRoots.Num() > 1) delete TargetNode;
-			}
-			TargetsOffsets[TargetType + 1] = NodesOffsets.Num() - PreviousNodes;
-
-			TargetsMin[TargetType] = MinBound;
-			TargetsMax[TargetType] = MaxBound;
+			BoundsOffsets[TargetType+1] = NumCommands;
+			TargetsOffsets[TargetType+1] = NodesOffsets.Num() - NodeOffset;
+			
+			TargetsMin[TargetType] = (NumCommands == 0) ? FVector(-FLT_MAX) : TargetMin;
+			TargetsMax[TargetType] = (NumCommands == 0) ? FVector(FLT_MAX) : TargetMax;
 		}
 		
 		for (uint32 FieldIndex = 1; FieldIndex < EFieldPhysicsType::Field_PhysicsType_Max + 1; ++FieldIndex)
diff --git a/Engine/Source/Runtime/Experimental/Chaos/Private/Field/FieldSystemNodes.cpp b/Engine/Source/Runtime/Experimental/Chaos/Private/Field/FieldSystemNodes.cpp
--- a/Engine/Source/Runtime/Experimental/Chaos/Private/Field/FieldSystemNodes.cpp
+++ b/Engine/Source/Runtime/Experimental/Chaos/Private/Field/FieldSystemNodes.cpp
@@ -141,10 +141,9 @@ void FUniformInteger::FillSetupCount(int32& NumOffsets, int32& NumParams) const
 	NumParams += 1;
 }
 
-void FUniformInteger::FillSetupDatas(TArray<int32>& NodesOffsets, TArray<float>& NodesParams,
-	const TMap<FFieldNodeBase*, float>& CommandTimes, const float PreviousTime) const
+void FUniformInteger::FillSetupDatas(TArray<int32>& NodesOffsets, TArray<float>& NodesParams, const float CommandTime) const
 {
-	Super::FillSetupDatas(NodesOffsets, NodesParams, CommandTimes, PreviousTime);
+	Super::FillSetupDatas(NodesOffsets, NodesParams, CommandTime);
 	NodesParams.Add(static_cast<float>(Magnitude));
 }
 
@@ -233,10 +232,9 @@ void FRadialIntMask::FillSetupCount(int32& NumOffsets, int32& NumParams) const
 	NumParams += 7;
 }
 
-void FRadialIntMask::FillSetupDatas(TArray<int32>& NodesOffsets, TArray<float>& NodesParams,
-	const TMap<FFieldNodeBase*, float>& CommandTimes, const float PreviousTime) const
+void FRadialIntMask::FillSetupDatas(TArray<int32>& NodesOffsets, TArray<float>& NodesParams, const float CommandTime) const
 {
-	Super::FillSetupDatas(NodesOffsets, NodesParams, CommandTimes, PreviousTime);
+	Super::FillSetupDatas(NodesOffsets, NodesParams, CommandTime);
 	NodesParams.Add(Radius);
 	NodesParams.Add(static_cast<float>(Position.X));
 	NodesParams.Add(static_cast<float>(Position.Y));
@@ -294,10 +292,9 @@ void FUniformScalar::FillSetupCount(int32& NumOffsets, int32& NumParams) const
 	NumParams += 1;
 }
 
-void FUniformScalar::FillSetupDatas(TArray<int32>& NodesOffsets, TArray<float>& NodesParams,
-	const TMap<FFieldNodeBase*, float>& CommandTimes, const float PreviousTime) const
+void FUniformScalar::FillSetupDatas(TArray<int32>& NodesOffsets, TArray<float>& NodesParams, const float CommandTime) const
 {
-	Super::FillSetupDatas(NodesOffsets, NodesParams, CommandTimes, PreviousTime);
+	Super::FillSetupDatas(NodesOffsets, NodesParams, CommandTime);
 	NodesParams.Add(Magnitude);
 }
 
@@ -407,19 +404,16 @@ void FWaveScalar::FillSetupCount(int32& NumOffsets, int32& NumParams) const
 	NumParams += 9;
 }
 
-void FWaveScalar::FillSetupDatas(TArray<int32>& NodesOffsets, TArray<float>& NodesParams,
-	const TMap<FFieldNodeBase*, float>& CommandTimes, const float PreviousTime) const
+void FWaveScalar::FillSetupDatas(TArray<int32>& NodesOffsets, TArray<float>& NodesParams,const float CommandTime) const
 {
-	const float NextTime = CommandTimes.Find(this) ? CommandTimes[this] : PreviousTime;
-
-	Super::FillSetupDatas(NodesOffsets, NodesParams, CommandTimes, PreviousTime);
+	Super::FillSetupDatas(NodesOffsets, NodesParams, CommandTime);
 	NodesParams.Add(Magnitude);
 	NodesParams.Add(static_cast<float>(Position.X));
 	NodesParams.Add(static_cast<float>(Position.Y));
 	NodesParams.Add(static_cast<float>(Position.Z));
 	NodesParams.Add(Wavelength);
 	NodesParams.Add(Period);
-	NodesParams.Add(NextTime);
+	NodesParams.Add(CommandTime);
 	NodesParams.Add(static_cast<float>(Function));
 	NodesParams.Add(static_cast<float>(Falloff));
 }
@@ -561,10 +555,9 @@ void FRadialFalloff::FillSetupCount(int32& NumOffsets, int32& NumParams) const
 	NumParams += 9;
 }
 
-void FRadialFalloff::FillSetupDatas(TArray<int32>& NodesOffsets, TArray<float>& NodesParams,
-	const TMap<FFieldNodeBase*, float>& CommandTimes, const float PreviousTime) const
+void FRadialFalloff::FillSetupDatas(TArray<int32>& NodesOffsets, TArray<float>& NodesParams, const float CommandTime) const
 {
-	Super::FillSetupDatas(NodesOffsets, NodesParams, CommandTimes, PreviousTime);
+	Super::FillSetupDatas(NodesOffsets, NodesParams, CommandTime);
 	NodesParams.Add(Magnitude);
 	NodesParams.Add(MinRange);
 	NodesParams.Add(MaxRange);
@@ -675,10 +668,9 @@ void FPlaneFalloff::FillSetupCount(int32& NumOffsets, int32& NumParams) const
 	NumParams += 12;
 }
 
-void FPlaneFalloff::FillSetupDatas(TArray<int32>& NodesOffsets, TArray<float>& NodesParams,
-	const TMap<FFieldNodeBase*, float>& CommandTimes, const float PreviousTime) const
+void FPlaneFalloff::FillSetupDatas(TArray<int32>& NodesOffsets, TArray<float>& NodesParams, const float CommandTime) const
 {
-	Super::FillSetupDatas(NodesOffsets, NodesParams, CommandTimes, PreviousTime);
+	Super::FillSetupDatas(NodesOffsets, NodesParams, CommandTime);
 	NodesParams.Add(Magnitude);
 	NodesParams.Add(MinRange);
 	NodesParams.Add(MaxRange);
@@ -791,10 +783,9 @@ void FBoxFalloff::FillSetupCount(int32& NumOffsets, int32& NumParams) const
 	NumParams += 15;
 }
 
-void FBoxFalloff::FillSetupDatas(TArray<int32>& NodesOffsets, TArray<float>& NodesParams,
-	const TMap<FFieldNodeBase*, float>& CommandTimes, const float PreviousTime) const
+void FBoxFalloff::FillSetupDatas(TArray<int32>& NodesOffsets, TArray<float>& NodesParams, const float CommandTime) const
 {
-	Super::FillSetupDatas(NodesOffsets, NodesParams, CommandTimes, PreviousTime);
+	Super::FillSetupDatas(NodesOffsets, NodesParams, CommandTime);
 	NodesParams.Add(Magnitude);
 	NodesParams.Add(MinRange);
 	NodesParams.Add(MaxRange);
@@ -886,10 +877,9 @@ void FNoiseField::FillSetupCount(int32& NumOffsets, int32& NumParams) const
 	NumParams += 12;
 }
 
-void FNoiseField::FillSetupDatas(TArray<int32>& NodesOffsets, TArray<float>& NodesParams,
-	const TMap<FFieldNodeBase*, float>& CommandTimes, const float PreviousTime) const
+void FNoiseField::FillSetupDatas(TArray<int32>& NodesOffsets, TArray<float>& NodesParams, const float CommandTime) const
 {
-	Super::FillSetupDatas(NodesOffsets, NodesParams, CommandTimes, PreviousTime);
+	Super::FillSetupDatas(NodesOffsets, NodesParams, CommandTime);
 	NodesParams.Add(MinRange);
 	NodesParams.Add(MaxRange);
 	NodesParams.Add(static_cast<float>(Transform.GetRotation().X));
@@ -950,10 +940,9 @@ void FUniformVector::FillSetupCount(int32& NumOffsets, int32& NumParams) const
 	NumParams += 4;
 }
 
-void FUniformVector::FillSetupDatas(TArray<int32>& NodesOffsets, TArray<float>& NodesParams,
-	const TMap<FFieldNodeBase*, float>& CommandTimes, const float PreviousTime) const 
+void FUniformVector::FillSetupDatas(TArray<int32>& NodesOffsets, TArray<float>& NodesParams, const float CommandTime) const 
 {
-	Super::FillSetupDatas(NodesOffsets, NodesParams, CommandTimes, PreviousTime);
+	Super::FillSetupDatas(NodesOffsets, NodesParams, CommandTime);
 	NodesParams.Add(Magnitude);
 	NodesParams.Add(static_cast<float>(Direction.X));
 	NodesParams.Add(static_cast<float>(Direction.Y));
@@ -1007,10 +996,9 @@ void FRadialVector::FillSetupCount(int32& NumOffsets, int32& NumParams) const
 	NumParams += 4;
 }
 
-void FRadialVector::FillSetupDatas(TArray<int32>& NodesOffsets, TArray<float>& NodesParams,
-	const TMap<FFieldNodeBase*, float>& CommandTimes, const float PreviousTime) const
+void FRadialVector::FillSetupDatas(TArray<int32>& NodesOffsets, TArray<float>& NodesParams, const float CommandTime) const
 {
-	Super::FillSetupDatas(NodesOffsets, NodesParams, CommandTimes, PreviousTime);
+	Super::FillSetupDatas(NodesOffsets, NodesParams, CommandTime);
 	NodesParams.Add(Magnitude);
 	NodesParams.Add(static_cast<float>(Position.X));
 	NodesParams.Add(static_cast<float>(Position.Y));
@@ -1065,10 +1053,9 @@ void FRandomVector::FillSetupCount(int32& NumOffsets, int32& NumParams) const
 	NumParams += 1;
 }
 
-void FRandomVector::FillSetupDatas(TArray<int32>& NodesOffsets, TArray<float>& NodesParams,
-	const TMap<FFieldNodeBase*, float>& CommandTimes, const float PreviousTime) const 
+void FRandomVector::FillSetupDatas(TArray<int32>& NodesOffsets, TArray<float>& NodesParams, const float CommandTime) const 
 {
-	Super::FillSetupDatas(NodesOffsets, NodesParams, CommandTimes, PreviousTime);
+	Super::FillSetupDatas(NodesOffsets, NodesParams, CommandTime);
 	NodesParams.Add(Magnitude);
 }
 
@@ -1219,20 +1206,17 @@ void FSumScalar::FillSetupCount(int32& NumOffsets, int32& NumParams) const
 	NumParams += 4;
 }
 
-void FSumScalar::FillSetupDatas(TArray<int32>& NodesOffsets, TArray<float>& NodesParams,
-	const TMap<FFieldNodeBase*, float>& CommandTimes, const float PreviousTime) const
+void FSumScalar::FillSetupDatas(TArray<int32>& NodesOffsets, TArray<float>& NodesParams, const float CommandTime) const
 {
-	const float NextTime = CommandTimes.Find(this) ? CommandTimes[this] : PreviousTime;
-
 	if (ScalarRight.IsValid())
 	{
-		ScalarRight->FillSetupDatas(NodesOffsets, NodesParams, CommandTimes, NextTime);
+		ScalarRight->FillSetupDatas(NodesOffsets, NodesParams, CommandTime);
 	}
 	if (ScalarLeft.IsValid())
 	{
-		ScalarLeft->FillSetupDatas(NodesOffsets, NodesParams, CommandTimes, NextTime);
+		ScalarLeft->FillSetupDatas(NodesOffsets, NodesParams, CommandTime);
 	}
-	Super::FillSetupDatas(NodesOffsets, NodesParams, CommandTimes, PreviousTime);
+	Super::FillSetupDatas(NodesOffsets, NodesParams, CommandTime);
 	NodesParams.Add(Magnitude);
 	NodesParams.Add(static_cast<float>(ScalarRight != nullptr));
 	NodesParams.Add(static_cast<float>(ScalarLeft != nullptr));
@@ -1443,24 +1427,21 @@ void FSumVector::FillSetupCount(int32& NumOffsets, int32& NumParams) const
 	NumParams += 5;
 }
 
-void FSumVector::FillSetupDatas(TArray<int32>& NodesOffsets, TArray<float>& NodesParams,
-	const TMap<FFieldNodeBase*, float>& CommandTimes, const float PreviousTime) const
+void FSumVector::FillSetupDatas(TArray<int32>& NodesOffsets, TArray<float>& NodesParams, const float CommandTime) const
 {
-	const float NextTime = CommandTimes.Find(this) ? CommandTimes[this] : PreviousTime;
-
 	if (Scalar.IsValid())
 	{
-		Scalar->FillSetupDatas(NodesOffsets, NodesParams, CommandTimes, NextTime);
+		Scalar->FillSetupDatas(NodesOffsets, NodesParams, CommandTime);
 	}
 	if (VectorRight.IsValid())
 	{
-		VectorRight->FillSetupDatas(NodesOffsets, NodesParams, CommandTimes, NextTime);
+		VectorRight->FillSetupDatas(NodesOffsets, NodesParams, CommandTime);
 	}
 	if (VectorLeft.IsValid())
 	{
-		VectorLeft->FillSetupDatas(NodesOffsets, NodesParams, CommandTimes, NextTime);
+		VectorLeft->FillSetupDatas(NodesOffsets, NodesParams, CommandTime);
 	}
-	Super::FillSetupDatas(NodesOffsets, NodesParams, CommandTimes, PreviousTime);
+	Super::FillSetupDatas(NodesOffsets, NodesParams, CommandTime);
 	NodesParams.Add(Magnitude);
 	NodesParams.Add(static_cast<float>(Scalar != nullptr));
 	NodesParams.Add(static_cast<float>(VectorRight != nullptr));
@@ -1592,17 +1573,14 @@ void FConversionField<InT, OutT>::FillSetupCount(int32& NumOffsets, int32& NumPa
 }
 
 template<class InT, class OutT>
-void FConversionField<InT, OutT>::FillSetupDatas(TArray<int32>& NodesOffsets, TArray<float>& NodesParams,
-	const TMap<FFieldNodeBase*, float>& CommandTimes, const float PreviousTime) const
+void FConversionField<InT, OutT>::FillSetupDatas(TArray<int32>& NodesOffsets, TArray<float>& NodesParams, const float CommandTime) const
 {
-	const float NextTime = CommandTimes.Find(this) ? CommandTimes[this] : PreviousTime;
-
 	if (InputField.IsValid())
 	{
-		InputField->FillSetupDatas(NodesOffsets, NodesParams, CommandTimes, NextTime);
+		InputField->FillSetupDatas(NodesOffsets, NodesParams, CommandTime);
 	}
 
-	Super::FillSetupDatas(NodesOffsets, NodesParams, CommandTimes, PreviousTime);
+	Super::FillSetupDatas(NodesOffsets, NodesParams, CommandTime);
 	NodesParams.Add(static_cast<float>(InputField != nullptr));
 }
 
@@ -1753,21 +1731,18 @@ void FCullingField<T>::FillSetupCount(int32& NumOffsets, int32& NumParams) const
 }
 
 template<class T>
-void FCullingField<T>::FillSetupDatas(TArray<int32>& NodesOffsets, TArray<float>& NodesParams,
-	const TMap<FFieldNodeBase*, float>& CommandTimes, const float PreviousTime) const
+void FCullingField<T>::FillSetupDatas(TArray<int32>& NodesOffsets, TArray<float>& NodesParams, const float CommandTime) const
 {
-	const float NextTime = CommandTimes.Find(this) ? CommandTimes[this] : PreviousTime;
-
 	if (Culling.IsValid())
 	{
-		Culling->FillSetupDatas(NodesOffsets, NodesParams, CommandTimes, NextTime);
+		Culling->FillSetupDatas(NodesOffsets, NodesParams, CommandTime);
 	}
 	if (Input.IsValid())
 	{
-		Input->FillSetupDatas(NodesOffsets, NodesParams, CommandTimes, NextTime);
+		Input->FillSetupDatas(NodesOffsets, NodesParams, CommandTime);
 	}
 
-	Super::FillSetupDatas(NodesOffsets, NodesParams, CommandTimes, PreviousTime);
+	Super::FillSetupDatas(NodesOffsets, NodesParams, CommandTime);
 	NodesParams.Add(static_cast<float>(Culling != nullptr));
 	NodesParams.Add(static_cast<float>(Input != nullptr));
 	NodesParams.Add(static_cast<float>(Operation));
diff --git a/Engine/Source/Runtime/Experimental/Chaos/Public/Field/FieldSystem.h b/Engine/Source/Runtime/Experimental/Chaos/Public/Field/FieldSystem.h
--- a/Engine/Source/Runtime/Experimental/Chaos/Public/Field/FieldSystem.h
+++ b/Engine/Source/Runtime/Experimental/Chaos/Public/Field/FieldSystem.h
@@ -450,8 +450,7 @@ public:
 	virtual void FillSetupCount(int32& NumOffsets, int32& NumParams) const {}
 
 	/** Fill the offsets/params arrays that will be used by the world physics field */
-	virtual void FillSetupDatas(TArray<int32>& NodesOffsets, TArray<float>& NodesParams,
-		const TMap<FFieldNodeBase*, float>& CommandTimes, const float PreviousTime) const {}
+	virtual void FillSetupDatas(TArray<int32>& NodesOffsets, TArray<float>& NodesParams, const float CommandTime) const {}
 
 	/** Evaluate the maximum magnitude of the field graph  */
 	virtual float EvalMaxMagnitude() const { return 1.0; }
@@ -499,8 +498,7 @@ public:
 	}
 
 	/** Fill the offsets/params arrays that will be used by the world physics field */
-	virtual void FillSetupDatas(TArray<int32>& NodesOffsets, TArray<float>& NodesParams,
-		const TMap<FFieldNodeBase*, float>& CommandTimes, const float PreviousTime) const override
+	virtual void FillSetupDatas(TArray<int32>& NodesOffsets, TArray<float>& NodesParams, const float CommandTime) const override
 	{
 		NodesOffsets.Add(NodesParams.Num());
 		NodesParams.Add(static_cast<float>(Type()));
diff --git a/Engine/Source/Runtime/Experimental/Chaos/Public/Field/FieldSystemNodes.h b/Engine/Source/Runtime/Experimental/Chaos/Public/Field/FieldSystemNodes.h
--- a/Engine/Source/Runtime/Experimental/Chaos/Public/Field/FieldSystemNodes.h
+++ b/Engine/Source/Runtime/Experimental/Chaos/Public/Field/FieldSystemNodes.h
@@ -34,8 +34,7 @@ public:
 	CHAOS_API virtual void FillSetupCount(int32& NumOffsets, int32& NumParams) const override;
 
 	/** Fill the offsets/params arrays that will be used by the world physics field */
-	CHAOS_API virtual void FillSetupDatas(TArray<int32>& NodesOffsets, TArray<float>& NodesParams,
-		const TMap<FFieldNodeBase*, float>& CommandTimes, const float PreviousTime) const override;
+	CHAOS_API virtual void FillSetupDatas(TArray<int32>& NodesOffsets, TArray<float>& NodesParams, const float CommandTime) const override;
 
 	/** Evaluate the maximum magnitude of the field graph  */
 	CHAOS_API virtual float EvalMaxMagnitude() const override;
@@ -80,8 +79,7 @@ public:
 	CHAOS_API virtual void FillSetupCount(int32& NumOffsets, int32& NumParams) const override;
 
 	/** Fill the offsets/params arrays that will be used by the world physics field */
-	CHAOS_API virtual void FillSetupDatas(TArray<int32>& NodesOffsets, TArray<float>& NodesParams,
-		const TMap<FFieldNodeBase*, float>& CommandTimes, const float PreviousTime) const override;
+	CHAOS_API virtual void FillSetupDatas(TArray<int32>& NodesOffsets, TArray<float>& NodesParams, const float CommandTime) const override;
 
 	/** Evaluate the maximum magnitude of the field graph  */
 	CHAOS_API virtual float EvalMaxMagnitude() const override;
@@ -126,8 +124,7 @@ public:
 	CHAOS_API virtual void FillSetupCount(int32& NumOffsets, int32& NumParams) const override;
 
 	/** Fill the offsets/params arrays that will be used by the world physics field */
-	CHAOS_API virtual void FillSetupDatas(TArray<int32>& NodesOffsets, TArray<float>& NodesParams,
-		const TMap<FFieldNodeBase*, float>& CommandTimes, const float PreviousTime) const override;
+	CHAOS_API virtual void FillSetupDatas(TArray<int32>& NodesOffsets, TArray<float>& NodesParams, const float CommandTime) const override;
 
 	/** Evaluate the maximum magnitude of the field graph  */
 	CHAOS_API virtual float EvalMaxMagnitude() const override;
@@ -170,8 +167,7 @@ public:
 	CHAOS_API virtual void FillSetupCount(int32& NumOffsets, int32& NumParams) const override;
 
 	/** Fill the offsets/params arrays that will be used by the world physics field */
-	CHAOS_API virtual void FillSetupDatas(TArray<int32>& NodesOffsets, TArray<float>& NodesParams,
-		const TMap<FFieldNodeBase*, float>& CommandTimes, const float PreviousTime) const override;
+	CHAOS_API virtual void FillSetupDatas(TArray<int32>& NodesOffsets, TArray<float>& NodesParams, const float CommandTime) const override;
 
 	/** Evaluate the maximum magnitude of the field graph  */
 	CHAOS_API virtual float EvalMaxMagnitude() const override;
@@ -227,8 +223,7 @@ public:
 	CHAOS_API virtual void FillSetupCount(int32& NumOffsets, int32& NumParams) const override;
 
 	/** Fill the offsets/params arrays that will be used by the world physics field */
-	CHAOS_API virtual void FillSetupDatas(TArray<int32>& NodesOffsets, TArray<float>& NodesParams,
-		const TMap<FFieldNodeBase*, float>& CommandTimes, const float PreviousTime) const override;
+	CHAOS_API virtual void FillSetupDatas(TArray<int32>& NodesOffsets, TArray<float>& NodesParams, const float CommandTime) const override;
 
 	/** Evaluate the maximum magnitude of the field graph  */
 	CHAOS_API virtual float EvalMaxMagnitude() const override;
@@ -292,8 +287,7 @@ public:
 	CHAOS_API virtual void FillSetupCount(int32& NumOffsets, int32& NumParams) const override;
 
 	/** Fill the offsets/params arrays that will be used by the world physics field */
-	CHAOS_API virtual void FillSetupDatas(TArray<int32>& NodesOffsets, TArray<float>& NodesParams,
-		const TMap<FFieldNodeBase*, float>& CommandTimes, const float PreviousTime) const override;
+	CHAOS_API virtual void FillSetupDatas(TArray<int32>& NodesOffsets, TArray<float>& NodesParams, const float CommandTime) const override;
 
 	/** Evaluate the maximum magnitude of the field graph  */
 	CHAOS_API virtual float EvalMaxMagnitude() const override;
@@ -355,8 +349,7 @@ public:
 	CHAOS_API virtual void FillSetupCount(int32& NumOffsets, int32& NumParams) const override;
 
 	/** Fill the offsets/params arrays that will be used by the world physics field */
-	CHAOS_API virtual void FillSetupDatas(TArray<int32>& NodesOffsets, TArray<float>& NodesParams,
-		const TMap<FFieldNodeBase*, float>& CommandTimes, const float PreviousTime) const override;
+	CHAOS_API virtual void FillSetupDatas(TArray<int32>& NodesOffsets, TArray<float>& NodesParams, const float CommandTime) const override;
 
 	/** Evaluate the maximum magnitude of the field graph  */
 	CHAOS_API virtual float EvalMaxMagnitude() const override;
@@ -406,8 +399,7 @@ public:
 	CHAOS_API virtual void FillSetupCount(int32& NumOffsets, int32& NumParams) const override;
 
 	/** Fill the offsets/params arrays that will be used by the world physics field */
-	CHAOS_API virtual void FillSetupDatas(TArray<int32>& NodesOffsets, TArray<float>& NodesParams,
-		const TMap<FFieldNodeBase*, float>& CommandTimes, const float PreviousTime) const override;
+	CHAOS_API virtual void FillSetupDatas(TArray<int32>& NodesOffsets, TArray<float>& NodesParams, const float CommandTime) const override;
 
 	/** Evaluate the maximum magnitude of the field graph  */
 	CHAOS_API virtual float EvalMaxMagnitude() const override;
@@ -447,8 +439,7 @@ public:
 	CHAOS_API virtual void FillSetupCount(int32& NumOffsets, int32& NumParams) const override;
 
 	/** Fill the offsets/params arrays that will be used by the world physics field */
-	CHAOS_API virtual void FillSetupDatas(TArray<int32>& NodesOffsets, TArray<float>& NodesParams,
-		const TMap<FFieldNodeBase*, float>& CommandTimes, const float PreviousTime) const override;
+	CHAOS_API virtual void FillSetupDatas(TArray<int32>& NodesOffsets, TArray<float>& NodesParams, const float CommandTime) const override;
 
 	/** Evaluate the maximum magnitude of the field graph  */
 	CHAOS_API virtual float EvalMaxMagnitude() const override;
@@ -487,8 +478,7 @@ public:
 	CHAOS_API virtual void FillSetupCount(int32& NumOffsets, int32& NumParams) const override;
 
 	/** Fill the offsets/params arrays that will be used by the world physics field */
-	CHAOS_API virtual void FillSetupDatas(TArray<int32>& NodesOffsets, TArray<float>& NodesParams,
-		const TMap<FFieldNodeBase*, float>& CommandTimes, const float PreviousTime) const override;
+	CHAOS_API virtual void FillSetupDatas(TArray<int32>& NodesOffsets, TArray<float>& NodesParams, const float CommandTime) const override;
 
 	/** Evaluate the maximum magnitude of the field graph  */
 	CHAOS_API virtual float EvalMaxMagnitude() const override;
@@ -526,8 +516,7 @@ public:
 	CHAOS_API virtual void FillSetupCount(int32& NumOffsets, int32& NumParams) const override;
 
 	/** Fill the offsets/params arrays that will be used by the world physics field */
-	CHAOS_API virtual void FillSetupDatas(TArray<int32>& NodesOffsets, TArray<float>& NodesParams,
-		const TMap<FFieldNodeBase*, float>& CommandTimes, const float PreviousTime) const override;
+	CHAOS_API virtual void FillSetupDatas(TArray<int32>& NodesOffsets, TArray<float>& NodesParams, const float CommandTime) const override;
 
 	/** Evaluate the maximum magnitude of the field graph  */
 	CHAOS_API virtual float EvalMaxMagnitude() const override;
@@ -575,8 +564,7 @@ public:
 	CHAOS_API virtual void FillSetupCount(int32& NumOffsets, int32& NumParams) const override;
 
 	/** Fill the offsets/params arrays that will be used by the world physics field */
-	CHAOS_API virtual void FillSetupDatas(TArray<int32>& NodesOffsets, TArray<float>& NodesParams,
-		const TMap<FFieldNodeBase*, float>& CommandTimes, const float PreviousTime) const override;
+	CHAOS_API virtual void FillSetupDatas(TArray<int32>& NodesOffsets, TArray<float>& NodesParams, const float CommandTime) const override;
 
 	/** Evaluate the maximum magnitude of the field graph  */
 	CHAOS_API virtual float EvalMaxMagnitude() const override;
@@ -584,6 +572,18 @@ public:
 	/** Compute the min/max spatial bounds of the field */
 	CHAOS_API virtual void ComputeFieldBounds(FVector& MinBounds, FVector& MaxBounds) const override;
 
+	/** Default setup for the sum scalar field */
+	static void FillDefaultSetup(TArray<int32>& NodesOffsets, TArray<float>& NodesParams)
+	{
+		NodesOffsets.Add(NodesParams.Num());
+		NodesParams.Add(static_cast<float>(Super::StaticType()));
+		NodesParams.Add(static_cast<float>(FFieldNodeBase::ESerializationType::FieldNode_FSumScalar));
+		NodesParams.Add(1.0);
+		NodesParams.Add(1.0);
+		NodesParams.Add(1.0);
+		NodesParams.Add(static_cast<float>(EFieldOperationType::Field_Add));
+	}
+
 	float Magnitude;
 	TUniquePtr< FFieldNode<float> > ScalarRight;
 	TUniquePtr< FFieldNode<float> > ScalarLeft;
@@ -633,8 +633,7 @@ public:
 	CHAOS_API virtual void FillSetupCount(int32& NumOffsets, int32& NumParams) const override;
 
 	/** Fill the offsets/params arrays that will be used by the world physics field */
-	CHAOS_API virtual void FillSetupDatas(TArray<int32>& NodesOffsets, TArray<float>& NodesParams,
-		const TMap<FFieldNodeBase*, float>& CommandTimes, const float PreviousTime) const override;
+	CHAOS_API virtual void FillSetupDatas(TArray<int32>& NodesOffsets, TArray<float>& NodesParams, const float CommandTime) const override;
 
 	/** Evaluate the maximum magnitude of the field graph  */
 	CHAOS_API virtual float EvalMaxMagnitude() const override;
@@ -642,6 +641,19 @@ public:
 	/** Compute the min/max spatial bounds of the field */
 	CHAOS_API virtual void ComputeFieldBounds(FVector& MinBounds, FVector& MaxBounds) const override;
 
+	/** Default setup for the sum vector field */
+	static void FillDefaultSetup(TArray<int32>& NodesOffsets, TArray<float>& NodesParams)
+	{
+		NodesOffsets.Add(NodesParams.Num());
+		NodesParams.Add(static_cast<float>(Super::StaticType()));
+		NodesParams.Add(static_cast<float>(FFieldNodeBase::ESerializationType::FieldNode_FSumVector));
+		NodesParams.Add(1.0);
+		NodesParams.Add(0.0);
+		NodesParams.Add(1.0);
+		NodesParams.Add(1.0);
+		NodesParams.Add(static_cast<float>(EFieldOperationType::Field_Add));
+	}
+
 	float Magnitude;
 	TUniquePtr< FFieldNode<float> > Scalar;
 	TUniquePtr< FFieldNode<FVector> > VectorRight;
@@ -681,8 +693,7 @@ public:
 	CHAOS_API virtual void FillSetupCount(int32& NumOffsets, int32& NumParams) const override;
 
 	/** Fill the offsets/params arrays that will be used by the world physics field */
-	CHAOS_API virtual void FillSetupDatas(TArray<int32>& NodesOffsets, TArray<float>& NodesParams,
-		const TMap<FFieldNodeBase*, float>& CommandTimes, const float PreviousTime) const override;
+	CHAOS_API virtual void FillSetupDatas(TArray<int32>& NodesOffsets, TArray<float>& NodesParams, const float CommandTime) const override;
 
 	/** Evaluate the maximum magnitude of the field graph  */
 	CHAOS_API virtual float EvalMaxMagnitude() const override;
@@ -690,6 +701,15 @@ public:
 	/** Compute the min/max spatial bounds of the field */
 	CHAOS_API virtual void ComputeFieldBounds(FVector& MinBounds, FVector& MaxBounds) const override;
 
+	/** Default setup for the conversion field */
+	static void FillDefaultSetup(TArray<int32>& NodesOffsets, TArray<float>& NodesParams)
+	{
+		NodesOffsets.Add(NodesParams.Num());
+		NodesParams.Add(static_cast<float>(Super::StaticType()));
+		NodesParams.Add(static_cast<float>(FFieldNodeBase::ESerializationType::FieldNode_FConversionField));
+		NodesParams.Add(1.0);
+	}
+
 	TUniquePtr< FFieldNode<InT> > InputField;
 };
 
@@ -734,8 +754,7 @@ public:
 	CHAOS_API virtual void FillSetupCount(int32& NumOffsets, int32& NumParams) const override;
 
 	/** Fill the offsets/params arrays that will be used by the world physics field */
-	CHAOS_API virtual void FillSetupDatas(TArray<int32>& NodesOffsets, TArray<float>& NodesParams,
-		const TMap<FFieldNodeBase*, float>& CommandTimes, const float PreviousTime) const override;
+	CHAOS_API virtual void FillSetupDatas(TArray<int32>& NodesOffsets, TArray<float>& NodesParams, const float CommandTime) const override;
 
 	/** Evaluate the maximum magnitude of the field graph  */
 	CHAOS_API virtual float EvalMaxMagnitude() const override;
