For geometry collection convex hulls created by intersection of a computed hull with an external collision shape, compute the intersection before computing any convex hulls by default

This change can be toggled w/ a cvar for now: p.gc.ComputeConvexExternalIntersectionBeforeComputingHulls
The cvar may be removed after some testing to confirm it gives better results / doesn't cause unexpected issues.

diff --git a/Engine/Source/Runtime/Experimental/Chaos/Private/GeometryCollection/GeometryCollectionConvexUtility.cpp b/Engine/Source/Runtime/Experimental/Chaos/Private/GeometryCollection/GeometryCollectionConvexUtility.cpp
--- a/Engine/Source/Runtime/Experimental/Chaos/Private/GeometryCollection/GeometryCollectionConvexUtility.cpp
+++ b/Engine/Source/Runtime/Experimental/Chaos/Private/GeometryCollection/GeometryCollectionConvexUtility.cpp
@@ -13,6 +13,9 @@
 #include "Templates/Sorting.h"
 #include "Spatial/PointHashGrid3.h"
 #include "CompGeom/ConvexDecomposition3.h"
+#include "Operations/MeshBoolean.h"
+#include "Operations/MeshSelfUnion.h"
+#include "MeshQueries.h"
 #include "GeometryCollection/GeometryCollectionClusteringUtility.h"
 
 bool UseVolumeToComputeRelativeSize = false;
@@ -21,6 +24,9 @@ FAutoConsoleVariableRef CVarUseVolumeToComputeRelativeSize(TEXT("p.gc.UseVolumeT
 bool UseLargestClusterToComputeRelativeSize = false;
 FAutoConsoleVariableRef CVarUseMaxClusterToComputeRelativeSize(TEXT("p.gc.UseLargestClusterToComputeRelativeSize"), UseVolumeToComputeRelativeSize, TEXT("Use the largest Cluster as reference for the releative size instead of the largest child (def: false)"));
 
+bool bComputeIntersectionsBeforeHull = true;
+FAutoConsoleVariableRef CVarComputeIntersectionsBeforeHull(TEXT("p.gc.ComputeConvexExternalIntersectionBeforeComputingHulls"), bComputeIntersectionsBeforeHull, TEXT("When computing convex hulls by intersection w/ an external collision shape, do this intersection on the geometry before computing hulls (def: true)"));
+
 static const Chaos::FVec3f IcoSphere_Subdiv0[] =
 {
 	{  0.000000f,  0.000000f, -1.000000f },
@@ -1209,6 +1215,85 @@ double ComputeGeometryVolume(
 	return VolOut;
 }
 
+// Helper to append the triangles of a convex hull to a dynamic mesh
+static void AddConvexHullToCompactDynamicMesh(const ::Chaos::FConvex* InConvexHull, UE::Geometry::FDynamicMesh3& Mesh, const FTransform* OptionalTransform = nullptr)
+{
+	check(Mesh.IsCompact());
+
+	const ::Chaos::FConvexStructureData& ConvexStructure = InConvexHull->GetStructureData();
+	const int32 NumV = InConvexHull->NumVertices();
+	const int32 NumP = InConvexHull->NumPlanes();
+	int32 StartV = Mesh.MaxVertexID();
+	for (int32 VIdx = 0; VIdx < NumV; ++VIdx)
+	{
+		FVector3d V = (FVector3d)InConvexHull->GetVertex(VIdx);
+		if (OptionalTransform)
+		{
+			V = OptionalTransform->TransformPosition(V);
+		}
+		int32 MeshVIdx = Mesh.AppendVertex(V);
+		checkSlow(MeshVIdx == VIdx + StartV); // Must be true because the mesh is compact
+	}
+	for (int32 PIdx = 0; PIdx < NumP; ++PIdx)
+	{
+		const int32 NumFaceV = ConvexStructure.NumPlaneVertices(PIdx);
+		const int32 V0 = StartV + ConvexStructure.GetPlaneVertex(PIdx, 0);
+		for (int32 SubIdx = 1; SubIdx + 1 < NumFaceV; ++SubIdx)
+		{
+			int32 V1 = StartV + ConvexStructure.GetPlaneVertex(PIdx, SubIdx);
+			int32 V2 = StartV + ConvexStructure.GetPlaneVertex(PIdx, SubIdx + 1);
+			constexpr bool bInvertFaces = true; // default-invert to match expected orientation of geometry meshes
+			if (bInvertFaces)
+			{
+				Swap(V1, V2);
+			}
+			int32 ResultTID = Mesh.AppendTriangle(UE::Geometry::FIndex3i(V0, V1, V2));
+			constexpr bool bFixNonmanifoldWithDuplicates = true;
+			if (bFixNonmanifoldWithDuplicates && ResultTID == UE::Geometry::FDynamicMesh3::NonManifoldID)
+			{
+				// failed to append due to a non-manifold triangle; try adding all the vertices independently so we at least capture the shape
+				// note: this should not happen for normal convex hulls, but the current convex hull algorithm does some aggressive face merging that sometimes creates weird geometry
+				UE::Geometry::FIndex3i DuplicateVerts(
+					Mesh.AppendVertex(Mesh.GetVertex(V0)),
+					Mesh.AppendVertex(Mesh.GetVertex(V1)),
+					Mesh.AppendVertex(Mesh.GetVertex(V2))
+				);
+				Mesh.AppendTriangle(DuplicateVerts);
+			}
+		}
+	}
+}
+
+// Helper to convert a geometry collection's geometry to a dynamic mesh
+static UE::Geometry::FDynamicMesh3 GeometryToDynamicMesh(const FGeometryCollection& Geometry, int32 GeomIdx, FTransform* OptionalTransform)
+{
+	UE::Geometry::FDynamicMesh3 Mesh;
+
+	int32 VStart = Geometry.VertexStart[GeomIdx];
+	int32 VCount = Geometry.VertexCount[GeomIdx];
+	int32 VEnd = VStart + VCount;
+
+	for (int32 Idx = VStart; Idx < VEnd; ++Idx)
+	{
+		FVector3d V = (FVector3d)Geometry.Vertex[Idx];
+		if (OptionalTransform)
+		{
+			V = OptionalTransform->TransformPosition(V);
+		}
+		Mesh.AppendVertex(V);
+	}
+
+	int32 FStart = Geometry.FaceStart[GeomIdx];
+	int32 FCount = Geometry.FaceCount[GeomIdx];
+	int32 FEnd = FStart + FCount;
+	for (int32 Idx = FStart; Idx < FEnd; ++Idx)
+	{
+		FIntVector Face = Geometry.Indices[Idx];
+		Mesh.AppendTriangle(Face.X - VStart, Face.Y - VStart, Face.Z - VStart);
+	}
+
+	return Mesh;
+}
 
 /// Helper to get convex hulls from a geometry collection in the format required by CreateNonoverlappingConvexHulls
 void HullsFromGeometry(
@@ -1223,7 +1308,9 @@ void HullsFromGeometry(
 	double SimplificationDistanceThreshold,
 	double OverlapRemovalShrinkPercent,
 	TFunction<bool(int32)> SkipBoneFn = nullptr,
-	const FGeometryCollectionConvexUtility::FConvexDecompositionSettings* OptionalDecompositionSettings = nullptr
+	const FGeometryCollectionConvexUtility::FConvexDecompositionSettings* OptionalDecompositionSettings = nullptr,
+	const TArray<FGeometryCollectionConvexUtility::FTransformedConvex>* OptionalIntersectConvexHulls = nullptr,
+	const TArray<TSet<int32>>* OptionalTransformToIntersectHulls = nullptr
 )
 {
 	TArray<FVector> GlobalVertices;
@@ -1276,16 +1363,63 @@ void HullsFromGeometry(
 		}
 		else if (SimulationType[Idx] == RigidType && GeomIdx != INDEX_NONE)
 		{
-			auto ComputeHull = [&Geometry, &GlobalVertices, SimplificationDistanceThreshold, OverlapRemovalShrinkPercent, GeomIdx](FVector& PivotOut) -> TUniquePtr<::Chaos::FConvex>
+			// If external intersection is requested, do the merge + intersection for the given bone here so it can be used by ComputeHull and/or the decomp
+			bool bHasExternal = OptionalIntersectConvexHulls && OptionalTransformToIntersectHulls;
+			TUniquePtr<UE::Geometry::FDynamicMesh3> GeometryIntersectedWithExternalHull = nullptr;
+			if (bHasExternal)
+			{
+				const TSet<int32>& HullInds = (*OptionalTransformToIntersectHulls)[Idx];
+				UE::Geometry::FDynamicMesh3 MergedHullMesh;
+				for (int32 HullIdx : HullInds)
+				{
+					// Append transformed external hull mesh
+					AddConvexHullToCompactDynamicMesh((*OptionalIntersectConvexHulls)[HullIdx].Convex.Get(), MergedHullMesh, &(*OptionalIntersectConvexHulls)[HullIdx].Transform);
+				}
+				if (HullInds.Num() > 1)
+				{
+					// Self-union the merged hull mesh -- mainly needed to help ensure the intersection volume accuracy
+					// i.e., in the (bAttemptDecomposition && OptionalDecompositionSettings->MaxGeoToHullVolumeRatioToDecompose < 1.0) case below
+					UE::Geometry::FMeshSelfUnion Union(&MergedHullMesh);
+					Union.Compute();
+				}
+				if (MergedHullMesh.TriangleCount() > 0)
+				{
+					UE::Geometry::FDynamicMesh3 GeometryMesh = GeometryToDynamicMesh(Geometry, GeomIdx, nullptr);
+					UE::Geometry::FMeshBoolean Boolean(&GeometryMesh, GlobalTransformArray[Idx], &MergedHullMesh, GlobalTransformArray[Idx],
+						&GeometryMesh, UE::Geometry::FMeshBoolean::EBooleanOp::Intersect);
+
+					Boolean.Compute();
+					if (GeometryMesh.VertexCount() > 0)
+					{
+						GeometryIntersectedWithExternalHull = MakeUnique<UE::Geometry::FDynamicMesh3>();
+						*GeometryIntersectedWithExternalHull = MoveTemp(GeometryMesh);
+					}
+				}
+			}
+
+			auto ComputeHull = [&Geometry, &GlobalVertices, SimplificationDistanceThreshold, OverlapRemovalShrinkPercent, GeomIdx,
+				&GeometryIntersectedWithExternalHull](FVector& PivotOut) -> TUniquePtr<::Chaos::FConvex>
 			{
-				int32 VStart = Geometry.VertexStart[GeomIdx];
-				int32 VCount = Geometry.VertexCount[GeomIdx];
-				int32 VEnd = VStart + VCount;
 				TArray<Chaos::FConvex::FVec3Type> HullPts;
-				HullPts.Reserve(VCount);
-				for (int32 VIdx = VStart; VIdx < VEnd; VIdx++)
+				if (GeometryIntersectedWithExternalHull)
 				{
-					HullPts.Add(GlobalVertices[VIdx]);
+					HullPts.Reserve(GeometryIntersectedWithExternalHull->VertexCount());
+					for (FVector3d Vertex : GeometryIntersectedWithExternalHull->VerticesItr())
+					{
+						HullPts.Add((Chaos::FConvex::FVec3Type)Vertex);
+					}
+				}
+				else
+				{
+					int32 VStart = Geometry.VertexStart[GeomIdx];
+					int32 VCount = Geometry.VertexCount[GeomIdx];
+					int32 VEnd = VStart + VCount;
+
+					HullPts.Reserve(VCount);
+					for (int32 VIdx = VStart; VIdx < VEnd; VIdx++)
+					{
+						HullPts.Add(GlobalVertices[VIdx]);
+					}
 				}
 				ensure(HullPts.Num() > 0);
 				FilterHullPoints(HullPts, SimplificationDistanceThreshold);
@@ -1301,7 +1435,14 @@ void HullsFromGeometry(
 				double InitialGeoVolume = 0.0, InitialHullVolume = 0.0;
 				if (bAttemptDecomposition && (OptionalDecompositionSettings->MinGeoVolumeToDecompose > 0.0 || OptionalDecompositionSettings->MaxGeoToHullVolumeRatioToDecompose < 1.0))
 				{
-					InitialGeoVolume = ComputeGeometryVolume(&Geometry, GeomIdx, GlobalTransformArray[Idx], 1.0);
+					if (GeometryIntersectedWithExternalHull)
+					{
+						InitialGeoVolume = UE::Geometry::TMeshQueries<UE::Geometry::FDynamicMesh3>::GetVolumeNonWatertight(*GeometryIntersectedWithExternalHull);
+					}
+					else
+					{
+						InitialGeoVolume = ComputeGeometryVolume(&Geometry, GeomIdx, GlobalTransformArray[Idx], 1.0);
+					}
 				}
 				if (bAttemptDecomposition && InitialGeoVolume < OptionalDecompositionSettings->MinGeoVolumeToDecompose)
 				{
@@ -1322,11 +1463,18 @@ void HullsFromGeometry(
 				if (bAttemptDecomposition)
 				{
 					UE::Geometry::FConvexDecomposition3 Decomposition;
-					int32 VertexStart = Geometry.VertexStart[GeomIdx];
-					TArrayView<const FVector3f> VerticesView(Geometry.Vertex.GetData() + VertexStart, Geometry.VertexCount[GeomIdx]);
-					TArrayView<const FIntVector3> FacesView(Geometry.Indices.GetData() + Geometry.FaceStart[GeomIdx], Geometry.FaceCount[GeomIdx]);
-					Decomposition.InitializeFromIndexMesh(VerticesView, FacesView, true, -VertexStart);
-					Decomposition.Compute(OptionalDecompositionSettings->MaxHullsPerGeometry, OptionalDecompositionSettings->NumAdditionalSplits, 
+					if (GeometryIntersectedWithExternalHull)
+					{
+						Decomposition.InitializeFromMesh(*GeometryIntersectedWithExternalHull, true);
+					}
+					else
+					{
+						int32 VertexStart = Geometry.VertexStart[GeomIdx];
+						TArrayView<const FVector3f> VerticesView(Geometry.Vertex.GetData() + VertexStart, Geometry.VertexCount[GeomIdx]);
+						TArrayView<const FIntVector3> FacesView(Geometry.Indices.GetData() + Geometry.FaceStart[GeomIdx], Geometry.FaceCount[GeomIdx]);
+						Decomposition.InitializeFromIndexMesh(VerticesView, FacesView, true, -VertexStart);
+					}
+					Decomposition.Compute(OptionalDecompositionSettings->MaxHullsPerGeometry, OptionalDecompositionSettings->NumAdditionalSplits,
 						OptionalDecompositionSettings->ErrorTolerance, OptionalDecompositionSettings->MinThicknessTolerance, OptionalDecompositionSettings->MaxHullsPerGeometry);
 					int32 NumHulls = Decomposition.NumHulls();
 					if ((NumHulls > 0 && !Hull) || NumHulls > 1)
@@ -1469,7 +1617,9 @@ void FGeometryCollectionConvexUtility::CreateConvexHullAttributesIfNeeded(FManag
 
 UE::GeometryCollectionConvexUtility::FConvexHulls
 FGeometryCollectionConvexUtility::ComputeLeafHulls(FGeometryCollection* GeometryCollection, const TArray<FTransform>& GlobalTransformArray, double SimplificationDistanceThreshold, double OverlapRemovalShrinkPercent,
-	TFunction<bool(int32)> SkipBoneFn, const FConvexDecompositionSettings* OptionalDecompositionSettings)
+	TFunction<bool(int32)> SkipBoneFn, const FConvexDecompositionSettings* OptionalDecompositionSettings,
+	const TArray<FTransformedConvex>* OptionalIntersectConvexHulls,
+	const TArray<TSet<int32>>* OptionalTransformToIntersectHulls)
 {
 	check(GeometryCollection);
 
@@ -1486,7 +1636,9 @@ FGeometryCollectionConvexUtility::ComputeLeafHulls(FGeometryCollection* Geometry
 	TFunctionRef<bool(int32)> HasCustomConvexFn = (CustomConvexFlags != nullptr) ? (TFunctionRef<bool(int32)>)ConvexFlagsFromArray : (TFunctionRef<bool(int32)>)ConvexFlagsAlwaysFalse;
 
 	HullsFromGeometry(*GeometryCollection, GlobalTransformArray, HasCustomConvexFn, Hulls.Hulls, Hulls.Pivots, Hulls.TransformToHullsIndices, GeometryCollection->SimulationType,
-		FGeometryCollection::ESimulationTypes::FST_Rigid, SimplificationDistanceThreshold, Hulls.OverlapRemovalShrinkPercent, SkipBoneFn, OptionalDecompositionSettings);
+		FGeometryCollection::ESimulationTypes::FST_Rigid, SimplificationDistanceThreshold, Hulls.OverlapRemovalShrinkPercent, SkipBoneFn, OptionalDecompositionSettings,
+		OptionalIntersectConvexHulls, OptionalTransformToIntersectHulls
+	);
 
 	return Hulls;
 }
@@ -1689,6 +1841,7 @@ void FGeometryCollectionConvexUtility::GenerateLeafConvexHulls(FGeometryCollecti
 			AddComputedHullsToCollection(FinalHulls, TransformToExternalHullsIndices);
 		}
 	}
+	
 	if (bUseGenerated)
 	{
 		// Identity Transform Array allows us to leave leaf hulls in their original local coordinate spaces
@@ -1704,20 +1857,18 @@ void FGeometryCollectionConvexUtility::GenerateLeafConvexHulls(FGeometryCollecti
 				return SelectionSet.Contains(BoneIdx);
 			};
 		}
-		ComputedHulls = ComputeLeafHulls(&Collection, IdentityTransformArray, Settings.SimplificationDistanceThreshold, 0, SkipBoneFn, &Settings.DecompositionSettings);
+		ComputedHulls = ComputeLeafHulls(&Collection, IdentityTransformArray, Settings.SimplificationDistanceThreshold, 0, SkipBoneFn, &Settings.DecompositionSettings,
+			bComputeIntersectionsBeforeHull ? &ExternalHulls : nullptr, 
+			bComputeIntersectionsBeforeHull ? &TransformToExternalHullsIndices : nullptr);
 
 		if (!bUseIntersect)
 		{
 			AddComputedHullsToCollection(ComputedHulls.Hulls, ComputedHulls.TransformToHullsIndices);
 		}
 	}
+	
 	if (bUseIntersect)
 	{
-		// TODO: this logic should be rewritten directly in the ComputeLeafHulls function, which should take the external hulls as an optional parameter.
-		// Specifically, ComputeLeafHulls should intersect the union of the external hulls with the geometry of the leaves, and then run convex decomposition on the result.
-		// This will avoid the possible multiplication of hulls if there are multiple external hulls, and also will allow the convex decomposition to focus
-		// only on the geometry that is not removed by the intersection.
-
 		TArray<TUniquePtr<Chaos::FConvex>> FinalHulls;
 		int32 HullIndexOffset = ConvexHullAttrib.Num();
 		for (int32 SourceTransformIdx : UseTransforms)
@@ -1756,6 +1907,19 @@ void FGeometryCollectionConvexUtility::GenerateLeafConvexHulls(FGeometryCollecti
 				}
 			}
 
+			// We've already computed the hulls w/ external geo intersection; just move it to the output
+			if (bComputeIntersectionsBeforeHull)
+			{
+				for (int32 GeoHullIdx : GeoHulls)
+				{
+					int32 HullIdx = FinalHulls.Add(MoveTemp(ComputedHulls.Hulls[GeoHullIdx]));
+					TransformToConvexIndicesAttrib[SourceTransformIdx].Add(HullIdx + HullIndexOffset);
+				}
+				continue;
+			}
+
+			// legacy !bComputeIntersectionsBeforeHull path; TODO: if we remove this cvar, remove this path
+			checkSlow(!bComputeIntersectionsBeforeHull);
 			for (int32 GeoHullIdx : GeoHulls)
 			{
 				if (TransformToExternalHullsIndices.IsEmpty())
diff --git a/Engine/Source/Runtime/Experimental/Chaos/Public/GeometryCollection/GeometryCollectionConvexUtility.h b/Engine/Source/Runtime/Experimental/Chaos/Public/GeometryCollection/GeometryCollectionConvexUtility.h
--- a/Engine/Source/Runtime/Experimental/Chaos/Public/GeometryCollection/GeometryCollectionConvexUtility.h
+++ b/Engine/Source/Runtime/Experimental/Chaos/Public/GeometryCollection/GeometryCollectionConvexUtility.h
@@ -248,19 +248,25 @@ public:
 	 */
 	static CHAOS_API void CopyChildConvexes(const FGeometryCollection* FromCollection, const TArrayView<const int32>& FromTransformIdx, FGeometryCollection* ToCollection, const TArrayView<const int32>& ToTransformIdx, bool bLeafOnly);
 
-	// Compute just the hulls of the leaf / rigid nodes that hold geometry directly, with no cluster hulls and no overlap removal by cutting
-	// This is an initial step of several algorithms: The CreateNonOverlappingConvexHullData function as well as convex-based proximity detection (TODO: and the auto-embed algorithm?)
-	// (TODO: Make auto-embed use this instead of the full hulls?)
-	// @param GlobalTransformArray		GeometryCollection's transforms to global space, as computed by GeometryCollectionAlgo::GlobalMatrices
-	static CHAOS_API UE::GeometryCollectionConvexUtility::FConvexHulls ComputeLeafHulls(FGeometryCollection* GeometryCollection, const TArray<FTransform>& GlobalTransformArray, double SimplificationDistanceThreshold = 0.0, double OverlapRemovalShrinkPercent = 0.0,
-		TFunction<bool(int32)> SkipBoneFn = nullptr, const FConvexDecompositionSettings* OptionalDecompositionSettings = nullptr);
-
 	struct FTransformedConvex
 	{
 		TSharedPtr<Chaos::FConvex> Convex;
 		FTransform Transform;
 	};
 
+	// Compute just the hulls of the leaf / rigid nodes that hold geometry directly, with no cluster hulls and no overlap removal by cutting
+	// This is an initial step of several algorithms: The CreateNonOverlappingConvexHullData function as well as convex-based proximity detection (TODO: and the auto-embed algorithm?)
+	// (TODO: Make auto-embed use this instead of the full hulls?)
+	// @param GlobalTransformArray				GeometryCollection's transforms to global space, as computed by GeometryCollectionAlgo::GlobalMatrices
+	// @param SkipBoneFn						Indicator function returns true for transform indices that do not need a convex hull to be computed, if non-null
+	// @param OptionalDecompositionSettings		Optionally generate multiple convex hulls per transform, if these settings are provided
+	// @param OptionalIntersectConvexHulls		Convex hulls to optionally intersect with the computed hulls, so that the resulting hulls will not extend outside of these provided hulls.
+	// @param OptionalTransformToIntersectHulls	Mapping from transforms to the OptionalIntersectConvexHulls
+	static CHAOS_API UE::GeometryCollectionConvexUtility::FConvexHulls ComputeLeafHulls(FGeometryCollection* GeometryCollection, const TArray<FTransform>& GlobalTransformArray, double SimplificationDistanceThreshold = 0.0, double OverlapRemovalShrinkPercent = 0.0,
+		TFunction<bool(int32)> SkipBoneFn = nullptr, const FConvexDecompositionSettings* OptionalDecompositionSettings = nullptr,
+		const TArray<FTransformedConvex>* OptionalIntersectConvexHulls = nullptr,
+		const TArray<TSet<int32>>* OptionalTransformToIntersectHulls = nullptr);
+
 	// generate a list of convex out of a hierarchy of implciit shapes
 	// suported shapes are scaled / transformed implicits as well as Boxes, convexes, spheres and capsules
 	// levelset, tapered capsule, cylinder, heightfields and trimesh are not supported
