non-GT frozen sim callbacks also need their SimTime and DeltaTime set to prevent NaNs from being set as the replicated velocity in the geometry collection

diff --git a/Engine/Source/Runtime/Experimental/Chaos/Private/PBDRigidsSolver.cpp b/Engine/Source/Runtime/Experimental/Chaos/Private/PBDRigidsSolver.cpp
--- a/Engine/Source/Runtime/Experimental/Chaos/Private/PBDRigidsSolver.cpp
+++ b/Engine/Source/Runtime/Experimental/Chaos/Private/PBDRigidsSolver.cpp
@@ -1679,6 +1679,7 @@ namespace Chaos
 		//for example a particle may not be created until marshalling occurs, and then a command could explicitly modify something like a collision setting
 		for (FSimCallbackCommandObject* SimCallbackObject : PushData.SimCommands)
 		{
+			SimCallbackObject->SetSimAndDeltaTime_Internal(GetSolverTime(), MLastDt);
 			SimCallbackObject->PreSimulate_Internal();
 			delete SimCallbackObject;
 		}
