Fix crash when moving forward in replay
In the Geometry Collection Component replication data cannot be reset on the physics thread because the reset is only called when the component is being destroyed.

diff --git a/Engine/Source/Runtime/Experimental/GeometryCollectionEngine/Private/GeometryCollection/GeometryCollectionComponent.cpp b/Engine/Source/Runtime/Experimental/GeometryCollectionEngine/Private/GeometryCollection/GeometryCollectionComponent.cpp
--- a/Engine/Source/Runtime/Experimental/GeometryCollectionEngine/Private/GeometryCollection/GeometryCollectionComponent.cpp
+++ b/Engine/Source/Runtime/Experimental/GeometryCollectionEngine/Private/GeometryCollection/GeometryCollectionComponent.cpp
@@ -1707,33 +1707,18 @@ void ActivateClusters(Chaos::FRigidClustering& Clustering, Chaos::TPBDRigidClust
 	Clustering.DeactivateClusterParticle(Cluster);
 }
 
-void UGeometryCollectionComponent::ResetRepDataCommon()
-{
-	OneOffActivatedProcessed = 0;
-	VersionProcessed = INDEX_NONE;
-	LastHardsnapTimeInMs = 0;
-}
-
 void UGeometryCollectionComponent::ResetRepData()
 {
 	ClustersToRep.Reset();
 	RepData.Reset();
 
-	if (GetNetMode() == ENetMode::NM_Client)
-	{
-		// Reset on the Physics Thread no to read/write on the client from different thread
-		if (Chaos::FPhysicsSolver* CurrSolver = GetSolver(*this))
-		{
-			CurrSolver->EnqueueCommandImmediate([this]()
-				{
-					ResetRepDataCommon();
-				});
-		}
-	}
-	else
-	{
-		ResetRepDataCommon();
-	}
+	// Those following data are initialized here from the Game Thread, 
+	// but they are read and written from the Physics Thread
+	// This function ResetRepData is only called when the component is destroyed and so the PhysicsProxy become null
+	// so at this point those data won't be changed on the Physics Thread. 
+	OneOffActivatedProcessed = 0;
+	VersionProcessed = INDEX_NONE;
+	LastHardsnapTimeInMs = 0;
 }
 
 void UGeometryCollectionComponent::UpdateRepData()
@@ -2832,13 +2817,10 @@ void UGeometryCollectionComponent::AsyncPhysicsTickComponent(float DeltaTime, fl
 		return;
 	}
 #endif
+	check(GetNetMode() != ENetMode::NM_Client);
 
 	Super::AsyncPhysicsTickComponent(DeltaTime, SimTime);
-
-	if (GetNetMode() != ENetMode::NM_Client)
-	{
-		UpdateRepData();
-	}
+	UpdateRepData();
 }
 
 void UGeometryCollectionComponent::OnHiddenInGameChanged()
diff --git a/Engine/Source/Runtime/Experimental/GeometryCollectionEngine/Public/GeometryCollection/GeometryCollectionComponent.h b/Engine/Source/Runtime/Experimental/GeometryCollectionEngine/Public/GeometryCollection/GeometryCollectionComponent.h
--- a/Engine/Source/Runtime/Experimental/GeometryCollectionEngine/Public/GeometryCollection/GeometryCollectionComponent.h
+++ b/Engine/Source/Runtime/Experimental/GeometryCollectionEngine/Public/GeometryCollection/GeometryCollectionComponent.h
@@ -1313,7 +1313,6 @@ protected:
 
 private:
 	GEOMETRYCOLLECTIONENGINE_API void ProcessRepDataOnPT();
-	GEOMETRYCOLLECTIONENGINE_API void ResetRepDataCommon();
 
 	bool bRenderStateDirty;
 	bool bEnableBoneSelection;
