Chaos - add duplicate particle checks for the ParticleViews (cvar enabled, off by default)

diff --git a/Engine/Source/Runtime/Experimental/Chaos/Private/Chaos/PBDRigidsEvolutionGBF.cpp b/Engine/Source/Runtime/Experimental/Chaos/Private/Chaos/PBDRigidsEvolutionGBF.cpp
--- a/Engine/Source/Runtime/Experimental/Chaos/Private/Chaos/PBDRigidsEvolutionGBF.cpp
+++ b/Engine/Source/Runtime/Experimental/Chaos/Private/Chaos/PBDRigidsEvolutionGBF.cpp
@@ -162,6 +162,10 @@ namespace Chaos
 		bool bChaos_Solver_TestMode  = false;
 		FAutoConsoleVariableRef CVarChaosSolverTestMode(TEXT("p.Chaos.Solver.TestMode"), bChaos_Solver_TestMode, TEXT(""));
 
+		// Set to true to enable some debug validation of the Particle Views every frame
+		bool bChaosSolverCheckParticleViews = false;
+		FAutoConsoleVariableRef CVarChaosSolverCheckParticleViews(TEXT("p.Chaos.Solver.CheckParticleViews"), bChaosSolverCheckParticleViews, TEXT(""));
+
 	}
 	using namespace CVars;
 
@@ -217,6 +221,50 @@ void SerializeToDisk(TEvolution& Evolution)
 }
 #endif
 
+// Veryify that no particle is one of the aggregate views more than once. This can happen if
+// a particle has been included in more than one orthogonal list (e.g., Active and MovingKinematic)
+// - 
+template<typename TParticleView>
+void CheckParticleViewForDupes(const FString& Name, const TParticleView& ParticleView)
+{
+	TSet<const FGeometryParticleHandle*> FoundParticles;
+	FoundParticles.Reserve(ParticleView.Num());
+
+	for (const auto& Particle : ParticleView)
+	{
+		const bool bAlreadyInView = FoundParticles.Contains(Particle.Handle());
+		if (ensureAlwaysMsgf(!bAlreadyInView, TEXT("%s duplicate Particle <%d>"), *Name, Particle.Handle()->GetHandleIdx()))
+		{
+			FoundParticles.Add(Particle.Handle());
+		}
+	}
+}
+
+void CheckParticleViewsForDupes(FPBDRigidsSOAs& Particles)
+{
+	if (CVars::bChaosSolverCheckParticleViews)
+	{
+		QUICK_SCOPE_CYCLE_COUNTER(STAT_Evolution_CheckParticleViewsForDupes);
+
+		// A particle appearing twice in either of these results in a race condition because the 
+		// collision detection loop will visit the same particle pair twice on different threads.
+		CheckParticleViewForDupes(TEXT("NonDisabledDynamicView"), Particles.GetNonDisabledDynamicView());
+		CheckParticleViewForDupes(TEXT("ActiveDynamicMovingKinematicParticlesView"), Particles.GetActiveDynamicMovingKinematicParticlesView());
+
+		// No known problems with dupes in these lists, but there should never be dupes
+		CheckParticleViewForDupes(TEXT("NonDisabledView"), Particles.GetNonDisabledView());
+		CheckParticleViewForDupes(TEXT("NonDisabledClusteredView"), Particles.GetNonDisabledClusteredView());
+		CheckParticleViewForDupes(TEXT("ActiveParticlesView"), Particles.GetActiveParticlesView());
+		//CheckParticleViewForDupes(TEXT("DirtyParticlesView"), Particles.GetDirtyParticlesView());
+		CheckParticleViewForDupes(TEXT("AllParticlesView"), Particles.GetAllParticlesView());
+		CheckParticleViewForDupes(TEXT("ActiveKinematicParticlesView"), Particles.GetActiveKinematicParticlesView());
+		CheckParticleViewForDupes(TEXT("ActiveMovingKinematicParticlesView"), Particles.GetActiveMovingKinematicParticlesView());
+		CheckParticleViewForDupes(TEXT("ActiveStaticParticlesView"), Particles.GetActiveStaticParticlesView());
+		CheckParticleViewForDupes(TEXT("ActiveDynamicMovingKinematicParticlesView"), Particles.GetActiveDynamicMovingKinematicParticlesView());
+	}
+}
+
+
 void FPBDRigidsEvolutionGBF::Advance(const FReal Dt,const FReal MaxStepDt,const int32 MaxSteps)
 {
 	// Determine how many steps we would like to take
@@ -382,6 +430,10 @@ void FPBDRigidsEvolutionGBF::AdvanceOneTimeStepImpl(const FReal Dt, const FSubSt
 		Base::ComputeIntermediateSpatialAcceleration();
 	}
 
+	// Collision detection is sensitive to duplication bugs in the particle views
+	// so this is here to help us track them down when they happen
+	CheckParticleViewsForDupes(Particles);
+
 	{
 		SCOPE_CYCLE_COUNTER(STAT_Evolution_DetectCollisions);
 		CSV_SCOPED_TIMING_STAT(PhysicsVerbose, StepSolver_DetectCollisions);
