Chaos - oprtimize BVH-BVH anf BVH-not-heightfield collision detection
- use stack rather than iteration in BVH visits
- minimize AABB transforms for BVH versus simple implicit shapes
- increase max depth of BVH

diff --git a/Engine/Source/Runtime/Experimental/Chaos/Private/Chaos/ChaosDebugDraw.cpp b/Engine/Source/Runtime/Experimental/Chaos/Private/Chaos/ChaosDebugDraw.cpp
--- a/Engine/Source/Runtime/Experimental/Chaos/Private/Chaos/ChaosDebugDraw.cpp
+++ b/Engine/Source/Runtime/Experimental/Chaos/Private/Chaos/ChaosDebugDraw.cpp
@@ -793,10 +793,10 @@ namespace Chaos
 			int32 LeafIndex = 0;
 
 			const auto& NodeVisitor = 
-				[Particle, &ShapeTransform, BVH, &LeafIndex, BoundsDrawShrink, &Duration, &Settings](const FAABB3f& NodeBounds, const int32 NodeDepth, const Private::FImplicitBVHNode& Node) -> bool
+				[Particle, &ShapeTransform, BVH, &LeafIndex, BoundsDrawShrink, &Duration, &Settings](const FAABB3f& NodeBounds, const int32 NodeDepth, const int32 NodeIndex) -> bool
 				{
 					// Draw the node if it is at the level we are asked to draw, or it is a leaf and we want to draw leaves.
-					if ((NodeDepth == ChaosDebugDebugDrawBVHLevel) || (Node.IsLeaf() && ((ChaosDebugDebugDrawBVHLevel == INDEX_NONE) || (ChaosDebugDebugDrawBVHLevel > NodeDepth))))
+					if ((NodeDepth == ChaosDebugDebugDrawBVHLevel) || (BVH->NodeIsLeaf(NodeIndex) && ((ChaosDebugDebugDrawBVHLevel == INDEX_NONE) || (ChaosDebugDebugDrawBVHLevel > NodeDepth))))
 					{
 						const FColor Color = GetIndexColor(LeafIndex++);
 
@@ -809,7 +809,7 @@ namespace Chaos
 
 						if (bChaosDebugDebugDrawBVHShapes)
 						{
-							BVH->VisitNodeObjects(Node, 
+							BVH->VisitNodeObjects(NodeIndex, 
 								[Particle, &ShapeTransform, &Color, Duration, &Settings](const FImplicitObject * Implicit, const FRigidTransform3f& RelativeTransformf, const FAABB3f& RelativeBoundsf, const int32 RootObjectIndex, const int32 LeafObjectIndex) -> void
 								{
 									const int32 ShapeIndex = (RootObjectIndex != INDEX_NONE) ? RootObjectIndex : 0;
@@ -821,7 +821,7 @@ namespace Chaos
 					return true;
 				};
 
-			BVH->VisitHierarchy(NodeVisitor);
+			BVH->VisitNodes(NodeVisitor);
 		}
 
 		void DrawParticleBVHImpl(const FRigidTransform3& SpaceTransform, const FGeometryParticleHandle* Particle, const FColor& InColor, const FChaosDebugDrawSettings& Settings)
diff --git a/Engine/Source/Runtime/Experimental/Chaos/Private/Chaos/Collision/ParticlePairMidPhase.cpp b/Engine/Source/Runtime/Experimental/Chaos/Private/Chaos/Collision/ParticlePairMidPhase.cpp
--- a/Engine/Source/Runtime/Experimental/Chaos/Private/Chaos/Collision/ParticlePairMidPhase.cpp
+++ b/Engine/Source/Runtime/Experimental/Chaos/Private/Chaos/Collision/ParticlePairMidPhase.cpp
@@ -7,6 +7,7 @@
 #include "Chaos/Collision/ContactTriangles.h"
 #include "Chaos/Collision/PBDCollisionConstraint.h"
 #include "Chaos/CollisionResolution.h"
+#include "Chaos/HeightField.h"
 #include "Chaos/ImplicitObject.h"
 #include "Chaos/ImplicitObjectBVH.h"
 #include "Chaos/ParticleHandle.h"
@@ -14,7 +15,10 @@
 #include "Chaos/PBDCollisionConstraints.h"
 #include "ChaosStats.h"
 #include "Misc/MemStack.h"
+#include "ProfilingDebugging/CountersTrace.h"
 
+TRACE_DECLARE_INT_COUNTER_EXTERN(ChaosTraceCounter_MidPhase_NumShapePair);
+TRACE_DECLARE_INT_COUNTER_EXTERN(ChaosTraceCounter_MidPhase_NumGeneric);
 
 extern bool Chaos_Collision_NarrowPhase_AABBBoundsCheck;
 
@@ -804,9 +808,6 @@ namespace Chaos
 		const FReal Dt,
 		const FCollisionContext& Context)
 	{
-		//SCOPE_CYCLE_COUNTER(STAT_Collisions_GenerateCollisions);
-		PHYSICS_CSV_SCOPED_EXPENSIVE(PhysicsVerbose, DetectCollisions_NarrowPhase);
-
 		if (!IsValid())
 		{
 			return;
@@ -990,6 +991,8 @@ namespace Chaos
 		const FReal Dt,
 		const FCollisionContext& Context)
 	{
+		TRACE_COUNTER_INCREMENT(ChaosTraceCounter_MidPhase_NumShapePair);
+
 		int32 NumActive = 0;
 		if (Flags.bIsCCD)
 		{
@@ -1065,6 +1068,8 @@ namespace Chaos
 		const FReal Dt,
 		const FCollisionContext& Context)
 	{
+		TRACE_COUNTER_INCREMENT(ChaosTraceCounter_MidPhase_NumGeneric);
+
 		const FImplicitObject* Implicit0 = GetParticle0()->Geometry().Get();
 		const FImplicitObject* Implicit1 = GetParticle1()->Geometry().Get();
 
@@ -1086,17 +1091,21 @@ namespace Chaos
 		// Create constraints for all the shape pairs whose bounds overlap.
 		// If we have a BVH, use it. Otherwise run a recursive hierarchy sweep.
 		// @todo(chaos): if we have 2 BVHs select the deepest one as BVHA?
-		if (BVH0 != nullptr)
+		if ((BVH0 != nullptr) && (BVH1 != nullptr))
 		{
-			GenerateCollisionsBVH(GetParticle0(), BVH0, GetParticle1(), Implicit1, CullDistance, Dt, Context);
+			GenerateCollisionsBVHBVH(GetParticle0(), BVH0, GetParticle1(), BVH1, CullDistance, Dt, Context);
+		}
+		else if (BVH0 != nullptr)
+		{
+			GenerateCollisionsBVHImplicitHierarchy(GetParticle0(), BVH0, GetParticle1(), Implicit1, CullDistance, Dt, Context);
 		}
 		else if (BVH1 != nullptr)
 		{
-			GenerateCollisionsBVH(GetParticle1(), BVH1, GetParticle0(), Implicit0, CullDistance, Dt, Context);
+			GenerateCollisionsBVHImplicitHierarchy(GetParticle1(), BVH1, GetParticle0(), Implicit0, CullDistance, Dt, Context);
 		}
 		else
 		{
-			GenerateCollisionsImplicit(GetParticle0(), Implicit0, GetParticle1(), Implicit1, CullDistance, Dt, Context);
+			GenerateCollisionsImplicitHierarchyImplicitHierarchy(GetParticle0(), Implicit0, GetParticle1(), Implicit1, CullDistance, Dt, Context);
 		}
 
 		// Generate manifolds for each constraint we created/recovered and (re)activate if necessary
@@ -1108,84 +1117,82 @@ namespace Chaos
 		return NumActive;
 	}
 
-	// Detect collisions betweena BVH and some other implicit object (which may be a hierarchy)
-	void FGenericParticlePairMidPhase::GenerateCollisionsBVH(
+	// Detect collisions between two BVHs
+	void FGenericParticlePairMidPhase::GenerateCollisionsBVHBVH(
 		FGeometryParticleHandle* ParticleA, const Private::FImplicitBVH* BVHA,
-		FGeometryParticleHandle* ParticleB, const FImplicitObject* RootImplicitB,
+		FGeometryParticleHandle* ParticleB, const Private::FImplicitBVH* BVHB,
 		const FReal CullDistance,
 		const FReal Dt,
 		const FCollisionContext& Context)
 	{
-		const FConstGenericParticleHandle PA = ParticleA;
-		const FConstGenericParticleHandle PB = ParticleB;
-
-		// Particle transforms
-		const FRigidTransform3 ParticleWorldTransformA = PA->GetTransformPQ();
-		const FRigidTransform3 ParticleWorldTransformB = PB->GetTransformPQ();
-		const FRigidTransform3 ParticleTransformAToB = ParticleWorldTransformA.GetRelativeTransform(ParticleWorldTransformB);
+		const FRigidTransform3 ParticleWorldTransformA = FConstGenericParticleHandle(ParticleA)->GetTransformPQ();
+		const FRigidTransform3 ParticleWorldTransformB = FConstGenericParticleHandle(ParticleB)->GetTransformPQ();
+		const FRigidTransform3 ParticleTransformBToA = ParticleWorldTransformB.GetRelativeTransform(ParticleWorldTransformA);
 
-		FMemMark Mark(FMemStack::Get());
-		TArray<bool, TMemStackAllocator<alignof(bool)>> bIsVisitedA;
-		bIsVisitedA.SetNumZeroed(BVHA->GetNumObjects());
+		const FShapeInstanceArray& ShapeInstancesA = ParticleA->ShapeInstances();
+		const FShapeInstanceArray& ShapeInstancesB = ParticleB->ShapeInstances();
 
-		// Visitor for FImplicitBVH::VisitNodeObjects
-		// Given an ImplicitObject from ParticleA (which we know overlaps the bounds of some parts of ParticleB),
-		// run collision detection on ImplicitA against the implicit object hierarchy of ParticleB.
-		// NOTE: may be called with the same ImplicitA multiple times because implicits may be in many BVH leaves
-		const auto& ObjectVisitorA = 
-			[this, ParticleA, &ParticleWorldTransformA, ParticleB, RootImplicitB, &ParticleWorldTransformB, &ParticleTransformAToB, &bIsVisitedA, CullDistance, Dt, &Context](const FImplicitObject* ImplicitA, const FRigidTransform3f& RelativeTransformfA, const FAABB3f& RelativeBoundsfA, const int32 RootObjectIndexA, const int32 LeafObjectIndexA) -> void
+		// Visit all overlapping leaf node pairs in BVHA and BVHB
+		Private::FImplicitBVH::VisitOverlappingLeafNodes(*BVHA, *BVHB, ParticleTransformBToA, 
+			[this, ParticleA, BVHA, &ParticleWorldTransformA, &ShapeInstancesA, ParticleB,
+			BVHB, &ParticleWorldTransformB, &ShapeInstancesB,
+			CullDistance, Dt, &Context]
+			(const int32 NodeIndexA, const int32 NodeIndexB)
 			{
-				if (bIsVisitedA[LeafObjectIndexA])
-				{
-					return;
-				}
-
-				bIsVisitedA[LeafObjectIndexA] = true;
-
-				// @todo(chaos): remove this float to double conversion and use floats where possible
-				const FRigidTransform3& RelativeTransformA = FRigidTransform3(RelativeTransformfA);
-				const FAABB3& RelativeBoundsA = FAABB3(RelativeBoundsfA);
+				// If we get here, BVHA(NodeIndexA) and BVHB(NodeIndexB) overlap so we must collide all object pairs in the two nodes
+				BVHA->VisitNodeObjects(NodeIndexA,
+					[this, ParticleA, &ParticleWorldTransformA, &ShapeInstancesA, ParticleB,
+					BVHB, &ParticleWorldTransformB, &ShapeInstancesB, NodeIndexB,
+					CullDistance, Dt, &Context]
+					(const FImplicitObject* ImplicitA, const FRigidTransform3f& RelativeTransformfA, const FAABB3f& RelativeBoundsfA, const int32 RootObjectIndexA, const int32 LeafObjectIndexA) -> void
+					{
+						const FRigidTransform3 RelativeTransformA = FRigidTransform3(RelativeTransformfA);
+						const FShapeInstance* ShapeInstanceA = GetShapeInstance(ShapeInstancesA, RootObjectIndexA);
+
+						BVHB->VisitNodeObjects(NodeIndexB,
+							[this, ParticleA, ImplicitA, ShapeInstanceA, &ParticleWorldTransformA, &RelativeTransformA, LeafObjectIndexA,
+							ParticleB, BVHB, &ParticleWorldTransformB, &ShapeInstancesB,
+							CullDistance, Dt, &Context]
+							(const FImplicitObject* ImplicitB, const FRigidTransform3f& RelativeTransformfB, const FAABB3f& RelativeBoundsfB, const int32 RootObjectIndexB, const int32 LeafObjectIndexB) -> void
+							{
+								const FRigidTransform3 RelativeTransformB = FRigidTransform3(RelativeTransformfB);
+								const FShapeInstance* ShapeInstanceB = GetShapeInstance(ShapeInstancesB, RootObjectIndexB);
+
+								// Detect collisions between the single implicit object pair
+								GenerateCollisionsImplicitLeafImplicitLeaf(
+									ParticleA, ImplicitA, ShapeInstanceA, ParticleWorldTransformA, RelativeTransformA, LeafObjectIndexA,
+									ParticleB, ImplicitB, ShapeInstanceB, ParticleWorldTransformB, RelativeTransformB, LeafObjectIndexB,
+									CullDistance, Dt, Context);
+							});
+					});
+			});
+	}
 
-				// If this is the first time we have seen ImplicitA, run the collision detection against ParticleB
-				GenerateCollisionsShapeHierarchy(
-					ParticleA, ImplicitA, ParticleWorldTransformA, RelativeTransformA, RelativeBoundsA, RootObjectIndexA, LeafObjectIndexA,
-					ParticleB, RootImplicitB, ParticleWorldTransformB,
-					ParticleTransformAToB,
-					CullDistance, Dt, Context);
-			};
+	// Detect collisions between a BVH and some other implicit object (which may be a hierarchy)
+	void FGenericParticlePairMidPhase::GenerateCollisionsBVHImplicitHierarchy(
+		FGeometryParticleHandle* ParticleA, const Private::FImplicitBVH* BVHA,
+		FGeometryParticleHandle* ParticleB, const FImplicitObject* RootImplicitB,
+		const FReal CullDistance, const FReal Dt, const FCollisionContext& Context)
+	{
+		const FShapeInstanceArray& ShapeInstancesB = ParticleB->ShapeInstances();
 
-		// Visitor for FImplicitBVH::VisitHierarchy
-		// This will be passed the nodes of the BVH on ParticleA and its bounds and contents. If the bounds overlap something in ParticleB
-		// we will keep recursing into the BVH. When we hit a leaf, run collision detection between the leaf contents and ParticleB.
-		const auto& NodeVisitorA = 
-			[BVHA, RootImplicitB, &ParticleTransformAToB, &ObjectVisitorA, CullDistance](const FAABB3f& NodeBoundsAf, const int32 NodeDepthA, const Private::FImplicitBVHNode& NodeA) -> bool
+		// Visit all the leaf implicits in RootImplicitB and collide against the BVH
+		RootImplicitB->VisitLeafObjects(
+			[this, ParticleA, BVHA, ParticleB, &ShapeInstancesB, CullDistance, Dt, &Context]
+			(const FImplicitObject* ImplicitB, const FRigidTransform3& RelativeTransformB, const int32 RootObjectIndexB, const int32 ObjectIndexB, const int32 LeafObjectIndexB) -> void
 			{
-				const FAABB3 NodeBoundsAInB = FAABB3(NodeBoundsAf).TransformedAABB(ParticleTransformAToB).ThickenSymmetrically(FVec3(CullDistance));
-
-				// Does this Node in A overlap anything in B?
-				// NOTE: IsOverlappingBounds performs a deep bounds check, including checking for node overlaps in BVH, Heightfield and TriMesh
-				if (!RootImplicitB->IsOverlappingBounds(NodeBoundsAInB))
-				{
-					// No overlap - stop recursing down this branch
-					return false;
-				}
-
-				// If we are at a leaf of A, we need to check all objects in the node against B
-				if (NodeA.IsLeaf())
-				{
-					BVHA->VisitNodeObjects(NodeA, ObjectVisitorA);
-				}
-
-				// Keep recursing
-				return true;
-			};
+				const FShapeInstance* ShapeInstanceB = GetShapeInstance(ShapeInstancesB, RootObjectIndexB);
 
-		// Visit all the nodes in BVHA and detect collisions with ParticleB
-		BVHA->VisitHierarchy(NodeVisitorA);
+				// ImplicitB is a single object. We perform the bounds tests in A space
+				GenerateCollisionsBVHImplicitLeaf(
+					ParticleA, BVHA,
+					ParticleB, ImplicitB, ShapeInstanceB, RelativeTransformB, LeafObjectIndexB,
+					CullDistance, Dt, Context);
+			});
 	}
 
-	// Detect collisions between two implicits, where either or both may be a hierarchy
-	void FGenericParticlePairMidPhase::GenerateCollisionsImplicit(
+	// Detect collisions between two implicits, where either or both may be a hierarchy, but neither has a BVH
+	void FGenericParticlePairMidPhase::GenerateCollisionsImplicitHierarchyImplicitHierarchy(
 		FGeometryParticleHandle* ParticleA, const FImplicitObject* RootImplicitA,
 		FGeometryParticleHandle* ParticleB, const FImplicitObject* RootImplicitB,
 		const FReal CullDistance,
@@ -1194,70 +1201,132 @@ namespace Chaos
 	{
 		const FConstGenericParticleHandle PA = ParticleA;
 		const FConstGenericParticleHandle PB = ParticleB;
+		const FShapeInstanceArray& ShapeInstancesA = ParticleA->ShapeInstances();
+		const FShapeInstanceArray& ShapeInstancesB = ParticleB->ShapeInstances();
 
 		// Particle transforms
 		const FRigidTransform3 ParticleWorldTransformA = PA->GetTransformPQ();
 		const FRigidTransform3 ParticleWorldTransformB = PB->GetTransformPQ();
-
-		// Calculate the overlapping volume in each particle's space
 		const FRigidTransform3 ParticleTransformAToB = ParticleWorldTransformA.GetRelativeTransform(ParticleWorldTransformB);
 
-		// Visitor for FImplicitBVH::VisitNodeObjects
+		// Detect collisons between Implicit Hierarchy of ParticleA and Implicit Hierarchy of ParticleB
 		// Given an ImplicitObject from ParticleA (which we know overlaps the bounds of some parts of ParticleB),
 		// run collision detection on ImplicitA against the implicit object hierarchy of ParticleB.
-		// NOTE: may be called with the same ImplicitA multiple times because implicits may be in many BVH leaves
-		const auto& ObjectVisitorA =
-			[this, ParticleA, &ParticleWorldTransformA, ParticleB, RootImplicitB, &ParticleWorldTransformB, &ParticleTransformAToB, CullDistance, Dt, &Context]
+		RootImplicitA->VisitLeafObjects(
+			[this, ParticleA, &ShapeInstancesA, &ParticleWorldTransformA,
+			ParticleB, &ShapeInstancesB, RootImplicitB, &ParticleWorldTransformB, &ParticleTransformAToB,
+			CullDistance, Dt, &Context]
 			(const FImplicitObject* ImplicitA, const FRigidTransform3& RelativeTransformA, const int32 RootObjectIndexA, const int32 ObjectIndex, const int32 LeafObjectIndexA)
 			{
 				const FAABB3 RelativeBoundsA = ImplicitA->CalculateTransformedBounds(RelativeTransformA);
+				const FAABB3 ShapeBoundsAInB = RelativeBoundsA.TransformedAABB(ParticleTransformAToB).ThickenSymmetrically(FVec3(CullDistance));
 
-				// If this is the first time we have seen ImplicitA, run the collision detection against ParticleB
-				GenerateCollisionsShapeHierarchy(
-					ParticleA, ImplicitA, ParticleWorldTransformA, RelativeTransformA, RelativeBoundsA, RootObjectIndexA, LeafObjectIndexA,
-					ParticleB, RootImplicitB, ParticleWorldTransformB,
-					ParticleTransformAToB,
-					CullDistance, Dt, Context);
-			};
+				const FShapeInstance* ShapeInstanceA = GetShapeInstance(ShapeInstancesA, RootObjectIndexA);
 
-		// Detect collisons between Implicit Hierarchy of ParticleA and Implicit Hierarchy of ParticleB
-		RootImplicitA->VisitLeafObjects(ObjectVisitorA);
+				// Detect collisons between ImplicitA and Implicit Hierarchy of ParticleB
+				RootImplicitB->VisitOverlappingLeafObjects(ShapeBoundsAInB,
+					[this, ParticleA, ImplicitA, ShapeInstanceA, &ParticleWorldTransformA, &RelativeTransformA, LeafObjectIndexA,
+					ParticleB, &ParticleWorldTransformB, &ShapeInstancesB,
+					CullDistance, Dt, &Context]
+					(const FImplicitObject* ImplicitB, const FRigidTransform3& RelativeTransformB, const int32 RootObjectIndexB, const int32 ObjectIndexB, const int32 LeafObjectIndexB)
+					{
+						const FShapeInstance* ShapeInstanceB = GetShapeInstance(ShapeInstancesB, RootObjectIndexB);
+
+						// Detect collisons between ImplicitA and ImplicitB (both leaf implicits)
+						GenerateCollisionsImplicitLeafImplicitLeaf(
+							ParticleA, ImplicitA, ShapeInstanceA, ParticleWorldTransformA, RelativeTransformA, LeafObjectIndexA,
+							ParticleB, ImplicitB, ShapeInstanceB, ParticleWorldTransformB, RelativeTransformB, LeafObjectIndexB,
+							CullDistance, Dt, Context);
+					});
+			});
 	}
 
-	// Generate collisions between lesf (non-hierarchy) implicit and some other (maybe hierarchy) implicit
-	void FGenericParticlePairMidPhase::GenerateCollisionsShapeHierarchy(
-		FGeometryParticleHandle* ParticleA, const FImplicitObject* ImplicitA, const FRigidTransform3 ParticleWorldTransformA, const FRigidTransform3& RelativeTransformA, const FAABB3& RelativeBoundsA, const int32 RootObjectIndexA, const int32 LeafObjectIndexA,
-		FGeometryParticleHandle* ParticleB, const FImplicitObject* RootImplicitB, const FRigidTransform3 ParticleWorldTransformB,
-		const FRigidTransform3 ParticleTransformAToB,
-		const FReal CullDistance,
-		const FReal Dt,
-		const FCollisionContext& Context)
+	// Detect collisions between a BVH and a leaf implicit object (not a hierarchy)
+	void FGenericParticlePairMidPhase::GenerateCollisionsBVHImplicitLeaf(
+		FGeometryParticleHandle* ParticleA, const Private::FImplicitBVH* BVHA,
+		FGeometryParticleHandle* ParticleB, const FImplicitObject* ImplicitB, const FShapeInstance* ShapeInstanceB, const FRigidTransform3& RelativeTransformB, const int32 LeafObjectIndexB,
+		const FReal CullDistance, const FReal Dt, const FCollisionContext& Context)
 	{
+		const FConstGenericParticleHandle PA = ParticleA;
+		const FConstGenericParticleHandle PB = ParticleB;
+
+		// Particle transforms
+		const FRigidTransform3 ParticleWorldTransformA = PA->GetTransformPQ();
+		const FRigidTransform3 ParticleWorldTransformB = PB->GetTransformPQ();
 		const FShapeInstanceArray& ShapeInstancesA = ParticleA->ShapeInstances();
-		const FShapeInstanceArray& ShapeInstancesB = ParticleB->ShapeInstances();
 
-		const FShapeInstance* ShapeInstanceA = GetShapeInstance(ShapeInstancesA, RootObjectIndexA);
+		// ImplicitB transforms/bounds
+		const FRigidTransform3 ImplicitTransformB = RelativeTransformB * ParticleWorldTransformB;
+		const FRigidTransform3 ImplicitTransformBToA = ImplicitTransformB.GetRelativeTransform(ParticleWorldTransformA);
+		const FAABB3 ImplicitBoundsBInA = ImplicitB->CalculateTransformedBounds(ImplicitTransformBToA);
 
-		const FAABB3 ShapeBoundsAInB = FAABB3(RelativeBoundsA).TransformedAABB(ParticleTransformAToB).ThickenSymmetrically(FVec3(CullDistance));
+		// If ImplicitB has a built-in BVH (Heightfield or TriMesh) we handle the test against the BVH differently
+		const bool bHasInternalBVHB = ImplicitB->template IsA<FHeightField>();
 
-		const auto& OverlappingLeafVisitor = 
-			[this, ParticleA, ImplicitA, ShapeInstanceA, &ParticleWorldTransformA, &RelativeTransformA, LeafObjectIndexA, ParticleB, &ParticleWorldTransformB, &ShapeInstancesB, CullDistance, Dt, &Context]
-			(const FImplicitObject* ImplicitB, const FRigidTransform3& RelativeTransformB, const int32 RootObjectIndexB, const int32 ObjectIndexB, const int32 LeafObjectIndexB)
+		// Visitor for FImplicitBVH::VisitNodeObjects
+		// Given an ImplicitObject from ParticleA (which we know overlaps the bounds of ImplicitB),
+		// run collision detection on ImplicitA against ImplicitB.
+		const auto& NodeObjectVisitorA =
+			[this, ParticleA, &ShapeInstancesA, &ParticleWorldTransformA, 
+			ParticleB, ImplicitB, ShapeInstanceB, &ParticleWorldTransformB, &RelativeTransformB, LeafObjectIndexB, 
+			CullDistance, Dt, &Context]
+			(const FImplicitObject* ImplicitA, const FRigidTransform3f& RelativeTransformfA, const FAABB3f& RelativeBoundsfA, const int32 RootObjectIndexA, const int32 LeafObjectIndexA) -> void
 			{
-				const FShapeInstance* ShapeInstanceB = GetShapeInstance(ShapeInstancesB, RootObjectIndexB);
+				const FShapeInstance* ShapeInstanceA = GetShapeInstance(ShapeInstancesA, RootObjectIndexA);
+				const FRigidTransform3 RelativeTransformA = FRigidTransform3(RelativeTransformfA);
 
-				GenerateCollisionsShapeShape(
+				GenerateCollisionsImplicitLeafImplicitLeaf(
 					ParticleA, ImplicitA, ShapeInstanceA, ParticleWorldTransformA, RelativeTransformA, LeafObjectIndexA,
 					ParticleB, ImplicitB, ShapeInstanceB, ParticleWorldTransformB, RelativeTransformB, LeafObjectIndexB,
 					CullDistance, Dt, Context);
 			};
 
-		// Detect collisons between ImplicitA and Implicit Hierarchy of ParticleB
-		RootImplicitB->VisitOverlappingLeafObjects(ShapeBoundsAInB, OverlappingLeafVisitor);
+		// Visitor for FImplicitBVH::VisitNodes
+		// This will be passed the nodes of the BVH on ParticleA and its bounds and contents. If the bounds overlap something in ParticleB
+		// we will keep recursing into the BVH. When we hit a leaf, run collision detection between the leaf contents and ParticleB.
+		const auto& NodeVisitorA =
+			[BVHA, ImplicitB, &ImplicitBoundsBInA, &ImplicitTransformBToA, bHasInternalBVHB, &NodeObjectVisitorA]
+			(const FAABB3f& RelativeNodeBoundsfA, const int32 NodeDepthA, const int32 NodeIndexA) -> bool
+			{
+				const FAABB3 RelativeNodeBoundsA = FAABB3(RelativeNodeBoundsfA);
+
+				if (bHasInternalBVHB)
+				{
+					// ImplicitB has an internal BVH (e.g., HeightField). We perform the bounds tests in B space
+					// NOTE: IsOverlappingBounds performs a deep bounds test, using any internal BVH present on ImplicitB
+					const FAABB3 NodeBoundsAInB = RelativeNodeBoundsA.InverseTransformedAABB(ImplicitTransformBToA);
+					if (!ImplicitB->IsOverlappingBounds(NodeBoundsAInB))
+					{
+						return false;
+					}
+				}
+				else
+				{
+					// ImplicitB is a single object. We perform the bounds tests in A space to avoid the node bounds transform
+					if (!ImplicitBoundsBInA.Intersects(RelativeNodeBoundsA))
+					{
+						return false;
+					}
+				}
+
+				// Is this node is a leaf, or is entirely contained in B visit all objects in the node and collide against ImplicitB
+				// @todo(chaos): support full containment check (calculated in branch above)
+				const bool bIsLeafA = BVHA->NodeIsLeaf(NodeIndexA);
+				if (bIsLeafA)
+				{
+					BVHA->VisitNodeObjects(NodeIndexA, NodeObjectVisitorA);
+				}
+
+				// Keep recursing
+				return true;
+			};
+
+		// Visit all the nodes in BVHA and detect collisions with ParticleB
+		BVHA->VisitNodes(NodeVisitorA);
 	}
 
-	// Generate collisions between two leaf (not hierahcies) implicits
-	void FGenericParticlePairMidPhase::GenerateCollisionsShapeShape(
+	// Generate collisions between two leaf (not hierarchy) implicits
+	void FGenericParticlePairMidPhase::GenerateCollisionsImplicitLeafImplicitLeaf(
 		FGeometryParticleHandle* ParticleA, const FImplicitObject* ImplicitA, const FShapeInstance* ShapeInstanceA, const FRigidTransform3 ParticleWorldTransformA, const FRigidTransform3& RelativeTransformA, const int32 LeafObjectIndexA,
 		FGeometryParticleHandle* ParticleB, const FImplicitObject* ImplicitB, const FShapeInstance* ShapeInstanceB, const FRigidTransform3 ParticleWorldTransformB, const FRigidTransform3& RelativeTransformB, const int32 LeafObjectIndexB,
 		const FReal CullDistance,
@@ -1395,16 +1464,19 @@ namespace Chaos
 		// shapes may be some children in the implicit hierarchy. The particles could be in the opposite order though, and
 		// this will depend on the shape types involved. E.g., with two particles each with a sphere and a box in a union
 		// would require up to two Sphere-Box contacts, with the particles in opposite orders.
+#if !UE_BUILD_TEST && !UE_BUILD_SHIPPING
 		if (!ensure(((InParticle0 == Particle0) && (InParticle1 == Particle1)) || ((InParticle0 == Particle1) && (InParticle1 == Particle0))))
 		{
 			// We somehow received a callback for the wrong particle pair...this should not happen
 			return nullptr;
 		}
+#endif
 
 		const FCollisionParticlePairConstraintKey CollisionKey = FCollisionParticlePairConstraintKey(InShape0, InImplicit0, InImplicitId0, InBVHParticles0, InShape1, InImplicit1, InImplicitId1, InBVHParticles1);
 		FPBDCollisionConstraint* Constraint = FindConstraint(CollisionKey);
 
 		// @todo(chaos): fix key uniqueness guarantee.  We need a truly unique key gen function
+#if !UE_BUILD_TEST && !UE_BUILD_SHIPPING
 		const bool bIsKeyCollision = (Constraint != nullptr) && ((Constraint->GetImplicit0() != InImplicit0) || (Constraint->GetImplicit1() != InImplicit1) || (Constraint->GetCollisionParticles0() != InBVHParticles0) || (Constraint->GetCollisionParticles1() != InBVHParticles1));
 		if (bIsKeyCollision)
 		{
@@ -1416,6 +1488,7 @@ namespace Chaos
 			ensure(false);
 			return nullptr;
 		}
+#endif
 
 		if (Constraint == nullptr)
 		{
diff --git a/Engine/Source/Runtime/Experimental/Chaos/Private/Chaos/ImplicitObjectUnion.cpp b/Engine/Source/Runtime/Experimental/Chaos/Private/Chaos/ImplicitObjectUnion.cpp
--- a/Engine/Source/Runtime/Experimental/Chaos/Private/Chaos/ImplicitObjectUnion.cpp
+++ b/Engine/Source/Runtime/Experimental/Chaos/Private/Chaos/ImplicitObjectUnion.cpp
@@ -11,14 +11,23 @@ namespace Chaos
 {
 	namespace CVars
 	{
-		int32 ChaosUnionBVHMinShapes = 32;
-		int32 ChaosUnionBVHMaxDepth = 8;
-		FRealSingle ChaosUnionBVHSplitBias = 0.1f;
+		// Set to false to prevent BVHs from being created in collision detection (or used, if disabled in the runtime after creation)
 		bool bChaosUnionBVHEnabled = true;
+		FAutoConsoleVariableRef CVarChaosUnionBVHEnabled(TEXT("p.Chaos.Collision.UnionBVH.Enabled"), bChaosUnionBVHEnabled, TEXT("Set to false to disable use of BVH during collision detection (without affecting creations and serialization)"));
+
+		// If a geometry hierarchy has more shapes than this, create a BVH around it
+		int32 ChaosUnionBVHMinShapes = 10;
 		FAutoConsoleVariableRef CVarChaosUnionBVHMinShapes(TEXT("p.Chaos.Collision.UnionBVH.NumShapes"), ChaosUnionBVHMinShapes, TEXT("If a geometry hierarchy has this many shapes, wrap it in a BVH for collision detection (negative to disable BVH)"));
+
+		// Maximum BVH depth. In general we want the BVH to generate leafs with a single item in them
+		int32 ChaosUnionBVHMaxDepth = 14;
 		FAutoConsoleVariableRef CVarChaosUnionBVHMaxDepth(TEXT("p.Chaos.Collision.UnionBVH.MaxDepth"), ChaosUnionBVHMaxDepth, TEXT("The allowed depth of the BVH when used to wrap a shape hiererchy"));
+
+		// A common case if objects arranged in a regular grid which is bad for the splitting algorithm.
+		// This prevents objects exactly in the middle of the bounds of a cell from being assigned to 
+		// a random child. See FImplicitBVH::Partition.
+		FRealSingle ChaosUnionBVHSplitBias = 0.1f;
 		FAutoConsoleVariableRef CVarChaosUnionBVHSplitBias(TEXT("p.Chaos.Collision.UnionBVH.SplitBias"), ChaosUnionBVHSplitBias, TEXT(""));
-		FAutoConsoleVariableRef CVarChaosUnionBVHEnabled(TEXT("p.Chaos.Collision.UnionBVH.Enabled"), bChaosUnionBVHEnabled, TEXT("Set to false to disable use of BVH during collision detection (without affecting creations and serialization)"));
 	}
 
 inline FAABB3 CalculateObjectsBounds(const TArrayView<TUniquePtr<FImplicitObject>>& Objects)
@@ -150,11 +159,11 @@ void FImplicitObjectUnion::FindAllIntersectingObjects(TArray<Pair<const FImplici
 	if (BVH.IsValid() && CVars::bChaosUnionBVHEnabled)
 	{
 		BVH->VisitOverlappingNodes(LocalBounds,
-			[this, &Out](const Private::FImplicitBVHNode& Node)
+			[this, &Out](const int32 NodeIndex)
 			{
-				if (Node.IsLeaf())
+				if (BVH->NodeIsLeaf(NodeIndex))
 				{
-					BVH->VisitNodeObjects(Node,
+					BVH->VisitNodeObjects(NodeIndex,
 						[&Out](const FImplicitObject* Implicit, const FRigidTransform3f& RelativeTransformf, const FAABB3f& RelativeBoundsf, const int32 RootObjectIndex, const int32 LeafObjectIndex) -> void
 						{
 							Out.Add(MakePair(Implicit, FRigidTransform3(RelativeTransformf)));
@@ -171,6 +180,20 @@ void FImplicitObjectUnion::FindAllIntersectingObjects(TArray<Pair<const FImplici
 	}
 }
 
+int32 FImplicitObjectUnion::CountObjectsInHierarchyImpl() const
+{
+	// Self
+	int32 NumObjects = 1;
+
+	// Children
+	for (int32 BVHObjectIndex = 0; BVHObjectIndex < MObjects.Num(); ++BVHObjectIndex)
+	{
+		NumObjects += MObjects[BVHObjectIndex]->CountObjectsInHierarchy();
+	}
+
+	return NumObjects;
+}
+
 void FImplicitObjectUnion::VisitOverlappingLeafObjectsImpl(
 	const FAABB3& LocalBounds,
 	const FRigidTransform3& ObjectTransform,
@@ -269,7 +292,6 @@ bool FImplicitObjectUnion::IsOverlappingBoundsImpl(const FAABB3& LocalBounds) co
 	return false;
 }
 
-
 TUniquePtr<FImplicitObject> FImplicitObjectUnion::Copy() const
 {
 	TArray<TUniquePtr<FImplicitObject>> CopyOfObjects;
@@ -477,11 +499,11 @@ FImplicitObjectUnionClustered::FindAllIntersectingChildren(const FAABB3& LocalBo
 	if (BVH.IsValid())
 	{
 		BVH->VisitOverlappingNodes(LocalBounds,
-			[this, &IntersectingChildren](const Private::FImplicitBVHNode& Node)
+			[this, &IntersectingChildren](const int32 NodeIndex)
 			{
-				if (Node.IsLeaf())
+				if (BVH->NodeIsLeaf(NodeIndex))
 				{
-					BVH->VisitNodeObjects(Node,
+					BVH->VisitNodeObjects(NodeIndex,
 						[this, &IntersectingChildren](const FImplicitObject* Implicit, const FRigidTransform3f& RelativeTransformf, const FAABB3f& RelativeBoundsf, const int32 RootObjectIndex, const int32 LeafObjectIndex) -> void
 						{
 							if (ensure(MOriginalParticleLookupHack.IsValidIndex(LeafObjectIndex)))
diff --git a/Engine/Source/Runtime/Experimental/Chaos/Private/PBDRigidsSolver.cpp b/Engine/Source/Runtime/Experimental/Chaos/Private/PBDRigidsSolver.cpp
--- a/Engine/Source/Runtime/Experimental/Chaos/Private/PBDRigidsSolver.cpp
+++ b/Engine/Source/Runtime/Experimental/Chaos/Private/PBDRigidsSolver.cpp
@@ -83,6 +83,8 @@ TRACE_DECLARE_INT_COUNTER(ChaosTraceCounter_NumUpdatedManifoldPoints, TEXT("Chao
 TRACE_DECLARE_INT_COUNTER(ChaosTraceCounter_NumJoints, TEXT("Chaos/Solver/Joints/NumConstraints"));
 TRACE_DECLARE_INT_COUNTER(ChaosTraceCounter_NumCharacterGroundConstraints, TEXT("Chaos/Solver/Character/NumConstraints"));
 
+TRACE_DECLARE_INT_COUNTER(ChaosTraceCounter_MidPhase_NumShapePair, TEXT("Chaos/Solver/MidPhase/NumShapePair"));
+TRACE_DECLARE_INT_COUNTER(ChaosTraceCounter_MidPhase_NumGeneric, TEXT("Chaos/Solver/MidPhase/NumGeneric"));
 
 // Stat Iteration counters
 DECLARE_DWORD_COUNTER_STAT(TEXT("NumPositionIterations"), STAT_ChaosCounter_NumPositionIterations, STATGROUP_ChaosCounters);
@@ -389,6 +391,8 @@ namespace Chaos
 			LLM_SCOPE(ELLMTag::ChaosUpdate);
 			UE_LOG(LogPBDRigidsSolver, Verbose, TEXT("AdvanceOneTimeStepTask::DoWork()"));
 
+			MSolver->ResetStatCounters();
+
 			if (FRewindData* RewindData = MSolver->GetRewindData())
 			{
 				RewindData->ApplyInputs(MSolver->GetCurrentFrame(), MSolver->GetEvolution()->IsResetting());
@@ -1998,6 +2002,12 @@ namespace Chaos
 		return GetEvolution()->GetCollisionConstraints().NumConstraints();
 	}
 
+	void FPBDRigidsSolver::ResetStatCounters()
+	{
+		TRACE_COUNTER_SET(ChaosTraceCounter_MidPhase_NumShapePair, 0);
+		TRACE_COUNTER_SET(ChaosTraceCounter_MidPhase_NumGeneric, 0);
+	}
+
 #ifndef CHAOS_COUNTER_STAT
 #define CHAOS_COUNTER_STAT(Name, Value)\
 SET_DWORD_STAT(STAT_ChaosCounter_##Name, Value); \
@@ -2005,7 +2015,6 @@ CSV_CUSTOM_STAT(PhysicsCounters, Name, Value, ECsvCustomStatOp::Set); \
 TRACE_COUNTER_SET(ChaosTraceCounter_##Name, Value)
 #endif
 
-
 	void FPBDRigidsSolver::UpdateStatCounters() const
 	{
 		const int32 NumStatic = GetEvolution()->GetParticles().GetActiveStaticParticlesView().Num();
diff --git a/Engine/Source/Runtime/Experimental/Chaos/Public/Chaos/Collision/ParticlePairMidPhase.h b/Engine/Source/Runtime/Experimental/Chaos/Public/Chaos/Collision/ParticlePairMidPhase.h
--- a/Engine/Source/Runtime/Experimental/Chaos/Public/Chaos/Collision/ParticlePairMidPhase.h
+++ b/Engine/Source/Runtime/Experimental/Chaos/Public/Chaos/Collision/ParticlePairMidPhase.h
@@ -7,6 +7,7 @@
 #include "Chaos/Collision/CollisionKeys.h"
 #include "Chaos/Collision/CollisionVisitor.h"
 #include "Chaos/Collision/PBDCollisionConstraint.h"
+#include "Chaos/ImplicitFwd.h"
 #include "Chaos/ParticleHandleFwd.h"
 #include "ProfilingDebugging/CsvProfiler.h"
 
@@ -542,33 +543,38 @@ namespace Chaos
 		CHAOS_API virtual void InjectCollisionImpl(const FPBDCollisionConstraint& Constraint, const FCollisionContext& Context) override final;
 
 	private:
-		// BVH on ParticleA versus implicit hierarchy of ParticleB
-		CHAOS_API void GenerateCollisionsBVH(
+		// BVH on ParticleA versus BVH on ParticleB
+		CHAOS_API void GenerateCollisionsBVHBVH(
 			FGeometryParticleHandle* ParticleA, const Private::FImplicitBVH* BVHA,
-			FGeometryParticleHandle* ParticleB, const FImplicitObject* RootImplicitB,
+			FGeometryParticleHandle* ParticleB, const Private::FImplicitBVH* BVHB,
 			const FReal CullDistance,
 			const FReal Dt,
 			const FCollisionContext& Context);
 
-		// Implicit hierarchy of particle A versus implicit hierarchy of ParticleB (used when no BVHs present)
-		CHAOS_API void GenerateCollisionsImplicit(
-			FGeometryParticleHandle* ParticleA, const FImplicitObject* RootImplicitA,
+		// BVH on ParticleA versus implicit hierarchy of ParticleB
+		CHAOS_API void GenerateCollisionsBVHImplicitHierarchy(
+			FGeometryParticleHandle* ParticleA, const Private::FImplicitBVH* BVHA,
 			FGeometryParticleHandle* ParticleB, const FImplicitObject* RootImplicitB,
 			const FReal CullDistance,
 			const FReal Dt,
 			const FCollisionContext& Context);
 
-		// Leaf Implicit on ParticleA versus implicit hierarchy of ParticleB
-		CHAOS_API void GenerateCollisionsShapeHierarchy(
-			FGeometryParticleHandle* ParticleA, const FImplicitObject* ImplicitA, const FRigidTransform3 ParticleWorldTransformA, const FRigidTransform3& RelativeTransformA, const FAABB3& RelativeBoundsA, const int32 RootObjectIndexA, const int32 LeafObjectIndexA,
-			FGeometryParticleHandle* ParticleB, const FImplicitObject* RootImplicitB, const FRigidTransform3 ParticleWorldTransformB,
-			const FRigidTransform3 ParticleTransformAToB,
+		// Implicit hierarchy of particle A versus implicit hierarchy of ParticleB (used when no BVH present on either)
+		CHAOS_API void GenerateCollisionsImplicitHierarchyImplicitHierarchy(
+			FGeometryParticleHandle* ParticleA, const FImplicitObject* RootImplicitA,
+			FGeometryParticleHandle* ParticleB, const FImplicitObject* RootImplicitB,
 			const FReal CullDistance,
 			const FReal Dt,
 			const FCollisionContext& Context);
 
+		// BVH on particle A versus a Leaf Implicit of ParticleB
+		CHAOS_API void GenerateCollisionsBVHImplicitLeaf(
+			FGeometryParticleHandle* ParticleA, const Private::FImplicitBVH* BVHA,
+			FGeometryParticleHandle* ParticleB, const FImplicitObject* ImplicitB, const FShapeInstance* ShapeInstanceB, const FRigidTransform3& RelativeTransformB, const int32 LeafObjectIndexB,
+			const FReal CullDistance, const FReal Dt, const FCollisionContext& Context);
+
 		// Leaf Implicit on ParticleA versus Leaf Implicit of ParticleB
-		CHAOS_API void GenerateCollisionsShapeShape(
+		CHAOS_API void GenerateCollisionsImplicitLeafImplicitLeaf(
 			FGeometryParticleHandle* ParticleA, const FImplicitObject* ImplicitA, const FShapeInstance* ShapeInstanceA, const FRigidTransform3 ParticleWorldTransformA, const FRigidTransform3& RelativeTransformA, const int32 LeafObjectIndexA,
 			FGeometryParticleHandle* ParticleB, const FImplicitObject* ImplicitB, const FShapeInstance* ShapeInstanceB, const FRigidTransform3 ParticleWorldTransformB, const FRigidTransform3& RelativeTransformB, const int32 LeafObjectIndexB,
 			const FReal CullDistance,
diff --git a/Engine/Source/Runtime/Experimental/Chaos/Public/Chaos/ImplicitFwd.h b/Engine/Source/Runtime/Experimental/Chaos/Public/Chaos/ImplicitFwd.h
--- a/Engine/Source/Runtime/Experimental/Chaos/Public/Chaos/ImplicitFwd.h
+++ b/Engine/Source/Runtime/Experimental/Chaos/Public/Chaos/ImplicitFwd.h
@@ -4,20 +4,23 @@
 
 namespace Chaos
 {
-	class FImplicitObject;
-	class FConvex;
-	class FCapsule;
-	template<class T> class TTriangle;
-	using FTriangle = TTriangle<FReal>;
-
 	template<class T, int d> class TBox;
 	template<class T, int d> class TPlane;
 	template<class T, int d> class TSphere;
+	template<class T> class TTriangle;
+
+	class FConvex;
+	class FCapsule;
+	class FHeightField;
+	class FImplicitObject;
 
-	using FImplicitObject3 = FImplicitObject;
 	using FImplicitBox3 = TBox<FReal, 3>;
 	using FImplicitCapsule3 = FCapsule;
 	using FImplicitConvex3 = FConvex;
+	using FImplicitHeightField3 = FHeightField;
+	using FImplicitObject3 = FImplicitObject;
 	using FImplicitPlane3 = TPlane<FReal, 3>;
 	using FImplicitSphere3 = TSphere<FReal, 3>;
+
+	using FTriangle = TTriangle<FReal>;
 }
\ No newline at end of file
diff --git a/Engine/Source/Runtime/Experimental/Chaos/Public/Chaos/ImplicitObject.h b/Engine/Source/Runtime/Experimental/Chaos/Public/Chaos/ImplicitObject.h
--- a/Engine/Source/Runtime/Experimental/Chaos/Public/Chaos/ImplicitObject.h
+++ b/Engine/Source/Runtime/Experimental/Chaos/Public/Chaos/ImplicitObject.h
@@ -398,6 +398,11 @@ public:
 
 	virtual uint16 GetMaterialIndex(uint32 HintIndex) const { return 0; }
 
+	int32 CountObjectsInHierarchy() const
+	{
+		return CountObjectsInHierarchyImpl();
+	}
+
 	/**
 	* Visit all the leaf objects in the hierarchy that overlap the specified local-space bounds.
 	* NOTE: Templated decorators like Instanced and Scaled cound as leafs, but object decorators like Transformed do not.
@@ -447,6 +452,12 @@ public:
 	}
 
 //protected:
+	// This should not be public, but it needs to be callable by derived classes on another instance
+	virtual int32 CountObjectsInHierarchyImpl() const
+	{
+		return 1;
+	}
+
 	// This should not be public, but it needs to be callable by derived classes on another instance
 	virtual void VisitOverlappingLeafObjectsImpl(
 		const FAABB3& LocalBounds,
@@ -485,9 +496,10 @@ public:
 		int32& LeafObjectIndex,
 		const FImplicitHierarchyVisitorBool& VisitorFunc) const
 	{
-		return VisitorFunc(this, ObjectTransform, RootObjectIndex, ObjectIndex, LeafObjectIndex);
+		const bool bResult = VisitorFunc(this, ObjectTransform, RootObjectIndex, ObjectIndex, LeafObjectIndex);
 		++ObjectIndex;
 		++LeafObjectIndex;
+		return bResult;
 	}
 
 	// This should not be public, but it needs to be callable by derived classes on another instance
diff --git a/Engine/Source/Runtime/Experimental/Chaos/Public/Chaos/ImplicitObjectBVH.h b/Engine/Source/Runtime/Experimental/Chaos/Public/Chaos/ImplicitObjectBVH.h
--- a/Engine/Source/Runtime/Experimental/Chaos/Public/Chaos/ImplicitObjectBVH.h
+++ b/Engine/Source/Runtime/Experimental/Chaos/Public/Chaos/ImplicitObjectBVH.h
@@ -42,8 +42,9 @@ namespace Chaos
 			FRigidTransform3 GetTransform() const { return FRigidTransform3(FVec3(GetX()), FRotation3(GetR())); }
 
 			// A unique index for this object in the hierarchy. E.g., if the same FImplicitObject is referenced 
-			// multiple times in the hierarchy in a union of transformed objects, each will have a different ObjectIndex 
-			// (it is the index into the pre-order depth first traversal).
+			// multiple times in the hierarchy in a union of transformed objects, each will have a different ObjectIndex.
+			// This is the ObjectIndex assigned when visiting the hierarchy via FImplicitObject::VisitHierachy and other 
+			// visit methods and can be used to index arrays initialized via those visitors.
 			int32 GetObjectIndex() const { return ObjectIndex; }
 
 			// The index of our most distant ancestor. I.e., the index in the root Union. This is used to map
@@ -143,45 +144,42 @@ namespace Chaos
 			template<typename TVisitor>
 			void VisitAllIntersections(const FAABB3& LocalBounds, const TVisitor& ObjectVisitor) const
 			{
-				const auto& NodeVisitor = [this, &ObjectVisitor](const FImplicitBVHNode& Node)
+				const auto& NodeVisitor = [this, &ObjectVisitor](const int32 NodeIndex)
 				{
-					if (Node.IsLeaf())
+					if (NodeIsLeaf(NodeIndex))
 					{
-						VisitNodeObjects(Node, ObjectVisitor);
+						VisitNodeObjects(NodeIndex, ObjectVisitor);
 					}
 				};
 
-				VisitOverlappingNodes(LocalBounds, NodeVisitor);
+				VisitOverlappingNodesStack(FAABB3f(LocalBounds), NodeVisitor);
 			}
 
 			// Calls the visitor for every overlapping leaf.
-			// @tparam TVisitor void(const FImplicitBVHNode& Node)
+			// @tparam TVisitor void(const int32 NodeIndex)
 			template<typename TVisitor>
 			void VisitOverlappingNodes(const FAABB3& LocalBounds, const TVisitor& NodeVisitor) const
 			{
-				if (!Nodes.IsEmpty())
-				{
-					VisitOverlappingNodesRecursive(Nodes[0], FAABB3f(LocalBounds), NodeVisitor);
-				}
+				VisitOverlappingNodesStack(FAABB3f(LocalBounds), NodeVisitor);
 			}
 
 			// Recursively visit all nodes in the hierarchy. Will stop visiting children
 			// if the visitor returns false. Leaf will be null when visiting an internal node.
-			// @param NodeVisitor (const FAABB3f& NodeBounds, const int32 NodeDepth, const FImplicitBVHNode& Node) -> void
+			// @param NodeVisitor (const FAABB3f& NodeBounds, const int32 NodeDepth, const int32 NodeIndex) -> void
 			template<typename TVisitor>
-			void VisitHierarchy(const TVisitor& NodeVisitor) const
+			void VisitNodes(const TVisitor& NodeVisitor) const
 			{
-				if (!Nodes.IsEmpty())
-				{
-					VisitHierarchyRecursive(Nodes[0], 0, NodeVisitor);
-				}
+				VisitNodesStack(NodeVisitor);
 			}
 
+
 			// Visit all the items in the specified leaf node (which is probably obtained from VisitHierarchy)
 			// @param ObjectVisitor (const FImplicitObject* Implicit, const FRigidTransform3f& RelativeTransformf, const FAABB3f& RelativeBoundsf, const int32 RootObjectIndex, const int32 LeafObjectIndex) -> void
 			template<typename TVisitor>
-			void VisitNodeObjects(const FImplicitBVHNode& Node, const TVisitor& ObjectVisitor) const
+			void VisitNodeObjects(const int32 NodeIndex, const TVisitor& ObjectVisitor) const
 			{
+				const FImplicitBVHNode& Node = Nodes[NodeIndex];
+
 				for (int32 NodeObjectIndex = Node.ObjectBeginIndex; NodeObjectIndex < Node.ObjectEndIndex; ++NodeObjectIndex)
 				{
 					const FImplicitBVHObject& Object = Objects[NodeObjectIndices[NodeObjectIndex]];
@@ -193,11 +191,28 @@ namespace Chaos
 			// Does the bounding box overlap any leaf nodes with items in it
 			bool IsOverlappingBounds(const FAABB3& LocalBounds) const
 			{
-				if (!Nodes.IsEmpty())
-				{
-					return IsOverlappingBoundsRecursive(Nodes[0], FAABB3f(LocalBounds));
-				}
-				return false;
+				return IsOverlappingBoundsStack(FAABB3f(LocalBounds));
+			}
+
+			template<typename TVisitor>
+			static void VisitOverlappingLeafNodes(const FImplicitBVH& BVHA, const FImplicitBVH& BVHB, const FRigidTransform3& TransformBToA, const TVisitor& LeafPairVisitor)
+			{
+				VisitOverlappingLeafNodesStack(BVHA, BVHB, TransformBToA, LeafPairVisitor);
+			}
+
+			int32 GetNumNodes() const
+			{
+				return Nodes.Num();
+			}
+
+			const FAABB3f& GetNodeBounds(const int32 NodeIndex) const
+			{
+				return Nodes[NodeIndex].Bounds;
+			}
+
+			bool NodeIsLeaf(const int32 NodeIndex) const
+			{
+				return Nodes[NodeIndex].IsLeaf();
 			}
 
 		private:
@@ -210,58 +225,185 @@ namespace Chaos
 			void Init(FObjects&& InObjects, const int32 MaxDepth, const int32 MaxLeafObjects);
 
 			template<typename TVisitor>
-			void VisitOverlappingNodesRecursive(const FImplicitBVHNode& Node, const FAABB3f& LocalBounds, const TVisitor& NodeVisitor) const
+			void VisitOverlappingNodesStack(const FAABB3f& LocalBounds, const TVisitor& NodeVisitor) const
 			{
-				if (!Node.Bounds.Intersects(LocalBounds))
+				if (Nodes.IsEmpty())
 				{
 					return;
 				}
 
-				if (Node.IsLeaf())
-				{
-					NodeVisitor(Node);
-				}
-				else
+				FMemMark Mark(FMemStack::Get());
+				TArray<int32> NodeStack;
+				NodeStack.Reserve(GetNumObjects());
+
+				int32 NodeIndex = 0;
+				while (true)
 				{
-					VisitOverlappingNodesRecursive(Nodes[Node.ChildNodeIndices[0]], LocalBounds, NodeVisitor);
-					VisitOverlappingNodesRecursive(Nodes[Node.ChildNodeIndices[1]], LocalBounds, NodeVisitor);
+					const FImplicitBVHNode& Node = Nodes[NodeIndex];
+
+					if (Node.Bounds.Intersects(LocalBounds))
+					{
+						if (Node.IsLeaf())
+						{
+							NodeVisitor(NodeIndex);
+						}
+						else
+						{
+							NodeIndex = Node.ChildNodeIndices[0];
+							NodeStack.Push(Node.ChildNodeIndices[1]);
+							continue;
+						}
+					}
+
+					if (NodeStack.IsEmpty())
+					{
+						break;
+					}
+
+					NodeIndex = NodeStack.Pop(false);
 				}
 			}
 
 			template<typename TVisitor>
-			void VisitHierarchyRecursive(const FImplicitBVHNode& Node, const int32 NodeDepth, const TVisitor& NodeVisitor) const
+			void VisitNodesStack(const TVisitor& NodeVisitor) const
 			{
-				// Visit this (non-leaf) node
-				const bool bVisitChildren = NodeVisitor(Node.Bounds, NodeDepth, Node);
+				if (Nodes.IsEmpty())
+				{
+					return;
+				}
+
+				FMemMark Mark(FMemStack::Get());
+				TArray<TPair<int32, int32>> NodeStack;	// NodeIndex, NodeDepth
+				NodeStack.Reserve(GetNumObjects());
 
-				// Visit children
-				if (bVisitChildren && !Node.IsLeaf())
+				int32 NodeIndex = 0;
+				int32 NodeDepth = 0;
+				while (true)
 				{
-					VisitHierarchyRecursive(Nodes[Node.ChildNodeIndices[0]], NodeDepth + 1, NodeVisitor);
-					VisitHierarchyRecursive(Nodes[Node.ChildNodeIndices[1]], NodeDepth + 1, NodeVisitor);
+					const FImplicitBVHNode& Node = Nodes[NodeIndex];
+
+					const bool bVisitChildren = NodeVisitor(Node.Bounds, NodeDepth, NodeIndex);
+
+					if (bVisitChildren && !Node.IsLeaf())
+					{
+						const int32 ChildNodeIndexL = Node.ChildNodeIndices[0];
+						const int32 ChildNodeIndexR = Node.ChildNodeIndices[1];
+						NodeDepth = NodeDepth + 1;
+						NodeIndex = ChildNodeIndexL;
+						NodeStack.Push({ ChildNodeIndexR, NodeDepth });
+						continue;
+					}
+
+					if (NodeStack.IsEmpty())
+					{
+						break;
+					}
+
+					NodeIndex = NodeStack.Top().Key;
+					NodeDepth = NodeStack.Top().Value;
+					NodeStack.Pop(false);
 				}
 			}
 
-			bool IsOverlappingBoundsRecursive(const FImplicitBVHNode& Node, const FAABB3f& LocalBounds) const
+			template<typename TVisitor>
+			static void VisitOverlappingLeafNodesStack(const FImplicitBVH& BVHA, const FImplicitBVH& BVHB, const FRigidTransform3& TransformBToA, const TVisitor& LeafPairVisitor)
 			{
-				if (Node.Bounds.Intersects(LocalBounds))
+				if (BVHA.Nodes.IsEmpty() || BVHB.Nodes.IsEmpty())
+				{
+					return;
+				}
+
+				// The node pair stack
+				FMemMark Mark(FMemStack::Get());
+				TArray<TVec2<int32>> NodePairStack;
+				NodePairStack.Reserve(FMath::Max(BVHA.GetNumObjects(), BVHB.GetNumObjects()));	// Can we make a better estimate of required size?
+
+				int32 NodeIndexA = 0;
+				int32 NodeIndexB = 0;
+
+				while (true)
 				{
-					// If we hit a leaf, we have an overlap
-					if (Node.IsLeaf())
+					const FImplicitBVHNode& NodeA = BVHA.Nodes[NodeIndexA];
+					const FImplicitBVHNode& NodeB = BVHB.Nodes[NodeIndexB];
+
+					const FAABB3f& BoundsA = NodeA.Bounds;
+					FAABB3f BoundsBInA = NodeB.Bounds.TransformedAABB(TransformBToA);
+
+					if (BoundsA.Intersects(BoundsBInA))
+					{
+						if (NodeA.IsLeaf() && NodeB.IsLeaf())
+						{
+							LeafPairVisitor(NodeIndexA, NodeIndexB);
+						}
+						else
+						{
+							// @todo(chaos): rule to choose whether to descend into A or B first
+							// Descend into B first, until we reach its leaf nodes
+							const bool bDescendA = NodeB.IsLeaf();
+							if (bDescendA)
+							{
+								// Descend A
+								NodeIndexA = NodeA.ChildNodeIndices[0];
+								NodePairStack.Push({ NodeA.ChildNodeIndices[1], NodeIndexB });
+								continue;
+							}
+							else
+							{
+								// Descend B
+								NodeIndexB = NodeB.ChildNodeIndices[0];
+								NodePairStack.Push({ NodeIndexA, NodeB.ChildNodeIndices[1] });
+								continue;
+							}
+						}
+					}
+
+					// If we get here we just processed a leaf or did not overlap,
+					// and need to pop an item off the stack to continue
+					if (NodePairStack.IsEmpty())
 					{
-						return true;
+						break;
 					}
+					NodeIndexA = NodePairStack.Top()[0];
+					NodeIndexB = NodePairStack.Top()[1];
+					NodePairStack.Pop(false);
+				}
+			}
+
+			bool IsOverlappingBoundsStack(const FAABB3f& LocalBounds) const
+			{
+				if (Nodes.IsEmpty())
+				{
+					return false;
+				}
+
+				FMemMark Mark(FMemStack::Get());
+				TArray<int32> NodeStack;
+				NodeStack.Reserve(GetNumObjects());
 
-					// Check children, stopping if we get an overlap
-					if (IsOverlappingBoundsRecursive(Nodes[Node.ChildNodeIndices[0]], LocalBounds))
+				int32 NodeIndex = 0;
+				while (true)
+				{
+					const FImplicitBVHNode& Node = Nodes[NodeIndex];
+
+					if (Node.Bounds.Intersects(LocalBounds))
 					{
-						return true;
+						if (Node.IsLeaf())
+						{
+							return true;
+						}
+
+						NodeIndex = Node.ChildNodeIndices[0];
+						NodeStack.Push(Node.ChildNodeIndices[1]);
+						continue;
 					}
-					if (IsOverlappingBoundsRecursive(Nodes[Node.ChildNodeIndices[1]], LocalBounds))
+
+					if (NodeStack.IsEmpty())
 					{
-						return true;
+						break;
 					}
+					NodeIndex = NodeStack.Pop(false);
 				}
+
 				return false;
 			}
 
diff --git a/Engine/Source/Runtime/Experimental/Chaos/Public/Chaos/ImplicitObjectTransformed.h b/Engine/Source/Runtime/Experimental/Chaos/Public/Chaos/ImplicitObjectTransformed.h
--- a/Engine/Source/Runtime/Experimental/Chaos/Public/Chaos/ImplicitObjectTransformed.h
+++ b/Engine/Source/Runtime/Experimental/Chaos/Public/Chaos/ImplicitObjectTransformed.h
@@ -362,6 +362,11 @@ public:
 	}
 
 protected:
+	virtual int32 CountObjectsInHierarchyImpl() const override final
+	{
+		return 1 + MObject->CountObjectsInHierarchy();
+	}
+
 	virtual void VisitOverlappingLeafObjectsImpl(
 		const FAABB3& InLocalBounds,
 		const FRigidTransform3& ObjectTransform,
diff --git a/Engine/Source/Runtime/Experimental/Chaos/Public/Chaos/ImplicitObjectUnion.h b/Engine/Source/Runtime/Experimental/Chaos/Public/Chaos/ImplicitObjectUnion.h
--- a/Engine/Source/Runtime/Experimental/Chaos/Public/Chaos/ImplicitObjectUnion.h
+++ b/Engine/Source/Runtime/Experimental/Chaos/Public/Chaos/ImplicitObjectUnion.h
@@ -3,10 +3,10 @@
 
 #include "Chaos/Array.h"
 #include "Chaos/ImplicitObject.h"
+#include "Chaos/ImplicitObjectTransformed.h"
 #include "Chaos/ISpatialAcceleration.h"
 #include "Chaos/GeometryParticles.h"
 
-#include "ImplicitObjectTransformed.h"
 #include "ChaosArchive.h"
 
 namespace Chaos
@@ -253,6 +253,8 @@ protected:
 		return ClosestIntersection;
 	}
 
+	CHAOS_API virtual int32 CountObjectsInHierarchyImpl() const override final;
+
 	CHAOS_API virtual void VisitOverlappingLeafObjectsImpl(
 		const FAABB3& LocalBounds,
 		const FRigidTransform3& ObjectTransform,
diff --git a/Engine/Source/Runtime/Experimental/Chaos/Public/PBDRigidsSolver.h b/Engine/Source/Runtime/Experimental/Chaos/Public/PBDRigidsSolver.h
--- a/Engine/Source/Runtime/Experimental/Chaos/Public/PBDRigidsSolver.h
+++ b/Engine/Source/Runtime/Experimental/Chaos/Public/PBDRigidsSolver.h
@@ -315,6 +315,7 @@ namespace Chaos
 		CHAOS_API void FieldForcesUpdateCallback();
 
 		// Update the counter in Stats and the CSV profiler
+		CHAOS_API void ResetStatCounters();
 		CHAOS_API void UpdateStatCounters() const;
 		CHAOS_API void UpdateExpensiveStatCounters() const;
 
