IAS - new commandline option -IAS to prioritize individual assset streaming over file system I/O backend

#rnx

diff --git a/Engine/Source/Runtime/Experimental/IoStoreOnDemand/Private/IoStoreOnDemand.cpp b/Engine/Source/Runtime/Experimental/IoStoreOnDemand/Private/IoStoreOnDemand.cpp
--- a/Engine/Source/Runtime/Experimental/IoStoreOnDemand/Private/IoStoreOnDemand.cpp
+++ b/Engine/Source/Runtime/Experimental/IoStoreOnDemand/Private/IoStoreOnDemand.cpp
@@ -959,7 +959,15 @@ void FIoStoreOnDemandModule::StartupModule()
 
 	TSharedPtr<UE::IOnDemandIoDispatcherBackend> Backend = UE::MakeOnDemandIoDispatcherBackend(Cache);
 	Backend->Mount(Endpoint);
-	FIoDispatcher::Get().Mount(Backend.ToSharedRef(), -10);
+	int32 BackendPriority = -10;
+#if !UE_BUILD_SHIPPING
+	if (FParse::Param(CommandLine, TEXT("Ias")))
+	{
+		// Bump the priority to be higher then the file system backend
+		BackendPriority = 10;
+	}
+#endif
+	FIoDispatcher::Get().Mount(Backend.ToSharedRef(), BackendPriority);
 #endif // !WITH_EDITOR
 }
 
