Iris
- Fixed some missing exports when using native IrisFastArraySerializer outside of testcode.

diff --git a/Engine/Source/Runtime/Experimental/Iris/Core/Public/Iris/ReplicationState/IrisFastArraySerializer.h b/Engine/Source/Runtime/Experimental/Iris/Core/Public/Iris/ReplicationState/IrisFastArraySerializer.h
--- a/Engine/Source/Runtime/Experimental/Iris/Core/Public/Iris/ReplicationState/IrisFastArraySerializer.h
+++ b/Engine/Source/Runtime/Experimental/Iris/Core/Public/Iris/ReplicationState/IrisFastArraySerializer.h
@@ -21,7 +21,7 @@ struct FIrisFastArraySerializerPrivateAccessor;
 USTRUCT()
 struct FIrisFastArraySerializer : public FFastArraySerializer
 {
-	// At the moment as we have no way to specify this we reserve 31 bits that can be used, the first bit is used for the array itself
+	// At the moment as we have no way to specify this per derived type, currently we reserve a fixed range of bits used for the changemask, the first bit is used for the array itself
 	enum { IrisFastArrayChangeMaskBits = 63U };
 	enum { IrisFastArrayChangeMaskBitOffset = 1U };
 	enum { IrisFastArrayPropertyBitIndex = 0U };	
@@ -75,10 +75,10 @@ private:
 	};
 
 	// Mark item at index as changed, and set the array as dirty
-	void InitChangeMask();
-	void InternalMarkItemChanged(int32 ItemIdx);
-	void InternalMarkAllItemsChanged();
-	void InternalMarkArrayAsDirty();
+	IRISCORE_API void InitChangeMask();
+	IRISCORE_API void InternalMarkItemChanged(int32 ItemIdx);
+	IRISCORE_API void InternalMarkAllItemsChanged();
+	IRISCORE_API void InternalMarkArrayAsDirty();
 
 	// Header for dirty state tracking needs to be just before ChangeMaskStorage. See GetReplicationStateHeader for more info.
 	UE::Net::FReplicationStateHeader ReplicationStateHeader;
