expose the 'intersect before computing hulls' option in geometry collection leaf hull compute out to the dataflow nodes rather than keeping it in a cvar

diff --git a/Engine/Plugins/Experimental/GeometryCollectionPlugin/Source/GeometryCollectionNodes/Private/Dataflow/GeometryCollectionUtilityNodes.cpp b/Engine/Plugins/Experimental/GeometryCollectionPlugin/Source/GeometryCollectionNodes/Private/Dataflow/GeometryCollectionUtilityNodes.cpp
--- a/Engine/Plugins/Experimental/GeometryCollectionPlugin/Source/GeometryCollectionNodes/Private/Dataflow/GeometryCollectionUtilityNodes.cpp
+++ b/Engine/Plugins/Experimental/GeometryCollectionPlugin/Source/GeometryCollectionNodes/Private/Dataflow/GeometryCollectionUtilityNodes.cpp
@@ -107,6 +107,7 @@ void FCreateLeafConvexHullsDataflowNode::Evaluate(Dataflow::FContext& Context, c
 			LeafSettings.DecompositionSettings.MaxHullsPerGeometry = InDecompSettings.MaxHullsPerGeometry;
 			LeafSettings.DecompositionSettings.MinThicknessTolerance = InDecompSettings.MinThicknessTolerance;
 			LeafSettings.DecompositionSettings.NumAdditionalSplits = InDecompSettings.NumAdditionalSplits;
+			LeafSettings.bComputeIntersectionsBeforeHull = bComputeIntersectionsBeforeHull;
 			FGeometryCollectionConvexUtility::GenerateLeafConvexHulls(*GeomCollection, bRestrictToSelection, SelectedBones, LeafSettings);
 			SetValue(Context, (const FManagedArrayCollection&)(*GeomCollection), &Collection);
 		}
diff --git a/Engine/Plugins/Experimental/GeometryCollectionPlugin/Source/GeometryCollectionNodes/Public/Dataflow/GeometryCollectionUtilityNodes.h b/Engine/Plugins/Experimental/GeometryCollectionPlugin/Source/GeometryCollectionNodes/Public/Dataflow/GeometryCollectionUtilityNodes.h
--- a/Engine/Plugins/Experimental/GeometryCollectionPlugin/Source/GeometryCollectionNodes/Public/Dataflow/GeometryCollectionUtilityNodes.h
+++ b/Engine/Plugins/Experimental/GeometryCollectionPlugin/Source/GeometryCollectionNodes/Public/Dataflow/GeometryCollectionUtilityNodes.h
@@ -121,15 +121,19 @@ public:
 	UPROPERTY(EditAnywhere, Category = Options)
 	EGenerateConvexMethod GenerateMethod = EGenerateConvexMethod::ExternalCollision;
 
-	/** If GenerateMethod is Intersect, only actually intersect when the volume of the Computed Hull is less than this fraction of the volume of the External Hull(s) */
+	/** If GenerateMethod is Intersect, only actually intersect when the volume of the Computed Hull is less than this fraction of the volume of the External Hull(s). */
 	UPROPERTY(EditAnywhere, Category = IntersectionFilters, meta = (ClampMin = 0.0, ClampMax = 1.0, EditCondition = "GenerateMethod == EGenerateConvexMethod::IntersectExternalWithComputed"))
 	float IntersectIfComputedIsSmallerByFactor = 1.0f;
 
-	/** If GenerateMethod is Intersect, only actually intersect if the volume of the External Hull(s) exceed this threshold */
+	/** If GenerateMethod is Intersect, only actually intersect if the volume of the External Hull(s) exceed this threshold. */
 	UPROPERTY(EditAnywhere, Category = IntersectionFilters, meta = (ClampMin = 0.0, EditCondition = "GenerateMethod == EGenerateConvexMethod::IntersectExternalWithComputed"))
 	float MinExternalVolumeToIntersect = 0.0f;
 
-	/** Computed convex hulls are simplified to keep points spaced at least this far apart (except where needed to keep the hull from collapsing to zero volume) */
+	/** Whether to compute the intersection before computing convex hulls. Typically should be enabled. */
+	UPROPERTY(EditAnywhere, Category = Convex, meta = (EditCondition = "GenerateMethod == EGenerateConvexMethod::IntersectExternalWithComputed"))
+	bool bComputeIntersectionsBeforeHull = true;
+
+	/** Computed convex hulls are simplified to keep points spaced at least this far apart (except where needed to keep the hull from collapsing to zero volume). */
 	UPROPERTY(EditAnywhere, Category = Convex, meta = (DataflowInput, ClampMin = 0.f, EditCondition = "GenerateMethod != EGenerateConvexMethod::ExternalCollision"))
 	float SimplificationDistanceThreshold = 10.f;
 
diff --git a/Engine/Source/Runtime/Experimental/Chaos/Private/GeometryCollection/GeometryCollectionConvexUtility.cpp b/Engine/Source/Runtime/Experimental/Chaos/Private/GeometryCollection/GeometryCollectionConvexUtility.cpp
--- a/Engine/Source/Runtime/Experimental/Chaos/Private/GeometryCollection/GeometryCollectionConvexUtility.cpp
+++ b/Engine/Source/Runtime/Experimental/Chaos/Private/GeometryCollection/GeometryCollectionConvexUtility.cpp
@@ -24,9 +24,6 @@ FAutoConsoleVariableRef CVarUseVolumeToComputeRelativeSize(TEXT("p.gc.UseVolumeT
 bool UseLargestClusterToComputeRelativeSize = false;
 FAutoConsoleVariableRef CVarUseMaxClusterToComputeRelativeSize(TEXT("p.gc.UseLargestClusterToComputeRelativeSize"), UseVolumeToComputeRelativeSize, TEXT("Use the largest Cluster as reference for the releative size instead of the largest child (def: false)"));
 
-bool bComputeIntersectionsBeforeHull = true;
-FAutoConsoleVariableRef CVarComputeIntersectionsBeforeHull(TEXT("p.gc.ComputeConvexExternalIntersectionBeforeComputingHulls"), bComputeIntersectionsBeforeHull, TEXT("When computing convex hulls by intersection w/ an external collision shape, do this intersection on the geometry before computing hulls (def: true)"));
-
 static const Chaos::FVec3f IcoSphere_Subdiv0[] =
 {
 	{  0.000000f,  0.000000f, -1.000000f },
@@ -1858,8 +1855,8 @@ void FGeometryCollectionConvexUtility::GenerateLeafConvexHulls(FGeometryCollecti
 			};
 		}
 		ComputedHulls = ComputeLeafHulls(&Collection, IdentityTransformArray, Settings.SimplificationDistanceThreshold, 0, SkipBoneFn, &Settings.DecompositionSettings,
-			bComputeIntersectionsBeforeHull ? &ExternalHulls : nullptr, 
-			bComputeIntersectionsBeforeHull ? &TransformToExternalHullsIndices : nullptr);
+			Settings.bComputeIntersectionsBeforeHull ? &ExternalHulls : nullptr, 
+			Settings.bComputeIntersectionsBeforeHull ? &TransformToExternalHullsIndices : nullptr);
 
 		if (!bUseIntersect)
 		{
@@ -1908,7 +1905,7 @@ void FGeometryCollectionConvexUtility::GenerateLeafConvexHulls(FGeometryCollecti
 			}
 
 			// We've already computed the hulls w/ external geo intersection; just move it to the output
-			if (bComputeIntersectionsBeforeHull)
+			if (Settings.bComputeIntersectionsBeforeHull)
 			{
 				for (int32 GeoHullIdx : GeoHulls)
 				{
@@ -1918,8 +1915,8 @@ void FGeometryCollectionConvexUtility::GenerateLeafConvexHulls(FGeometryCollecti
 				continue;
 			}
 
-			// legacy !bComputeIntersectionsBeforeHull path; TODO: if we remove this cvar, remove this path
-			checkSlow(!bComputeIntersectionsBeforeHull);
+			// The !bComputeIntersectionsBeforeHull path
+			checkSlow(!Settings.bComputeIntersectionsBeforeHull);
 			for (int32 GeoHullIdx : GeoHulls)
 			{
 				if (TransformToExternalHullsIndices.IsEmpty())
diff --git a/Engine/Source/Runtime/Experimental/Chaos/Public/GeometryCollection/GeometryCollectionConvexUtility.h b/Engine/Source/Runtime/Experimental/Chaos/Public/GeometryCollection/GeometryCollectionConvexUtility.h
--- a/Engine/Source/Runtime/Experimental/Chaos/Public/GeometryCollection/GeometryCollectionConvexUtility.h
+++ b/Engine/Source/Runtime/Experimental/Chaos/Public/GeometryCollection/GeometryCollectionConvexUtility.h
@@ -210,6 +210,10 @@ public:
 
 		// Convex decomposition settings, applied to convex hulls generated from geometry
 		FConvexDecompositionSettings DecompositionSettings;
+
+		// For GenerateMethod == IntersectExternalWithComputed, whether to compute the intersection before computing convex hulls
+		// Note: It seems more logical for this setting to be true, but we expose it as an option because in some special cases the results when it was false were preferred
+		bool bComputeIntersectionsBeforeHull = true;
 	};
 	
 	static CHAOS_API void GenerateLeafConvexHulls(FGeometryCollection& Collection, bool bRestrictToSelection, const TArrayView<const int32> TransformSubset, const FLeafConvexHullSettings& Settings);
