better cleanup of geometry collections to ensure that they are removed from cluster unions and that they aren't considered for being added into cluster unions in the breaking model

diff --git a/Engine/Source/Runtime/Experimental/Chaos/Private/Chaos/PBDRigidClustering.cpp b/Engine/Source/Runtime/Experimental/Chaos/Private/Chaos/PBDRigidClustering.cpp
--- a/Engine/Source/Runtime/Experimental/Chaos/Private/Chaos/PBDRigidClustering.cpp
+++ b/Engine/Source/Runtime/Experimental/Chaos/Private/Chaos/PBDRigidClustering.cpp
@@ -2034,12 +2034,13 @@ namespace Chaos
 		{
 			ParentParticle = ClusteredParticle->Parent();
 
-			ClusteredParticle->ClusterIds() = ClusterId();
-			ClusteredParticle->ClusterGroupIndex() = 0;
-
 			// Need to also check if the particle is a cluster union and remove from that as well.
+			// This needs to be before the call to clear our ClusterIds since HandleRemoveOperationWithClusterLookup needs to use the particle's parent to find the right cluster union.
 			ClusterUnionManager.HandleRemoveOperationWithClusterLookup({ ClusteredParticle }, EClusterUnionOperationTiming::Defer);
 
+			ClusteredParticle->ClusterIds() = ClusterId();
+			ClusteredParticle->ClusterGroupIndex() = 0;
+
 			if (MChildren.Contains(ParentParticle))
 			{
 				FRigidHandleArray& Children = MChildren[ParentParticle];
@@ -2047,8 +2048,8 @@ namespace Chaos
 				// disconnect from your parents children list
 				Children.Remove(ClusteredParticle);
 
-				// disable internal parents that have lost all their children
-				if (!MChildren[ParentParticle].Num() && ParentParticle->InternalCluster())
+				// disable internal parents that have lost all their children - do not try to disable cluster unions.
+				if (!MChildren[ParentParticle].Num() && ParentParticle->InternalCluster() && !ClusterUnionManager.IsClusterUnionParticle(ParentParticle))
 				{
 					DisableCluster(ClusteredParticle);
 				}
diff --git a/Engine/Source/Runtime/Experimental/Chaos/Private/PBDRigidsSolver.cpp b/Engine/Source/Runtime/Experimental/Chaos/Private/PBDRigidsSolver.cpp
--- a/Engine/Source/Runtime/Experimental/Chaos/Private/PBDRigidsSolver.cpp
+++ b/Engine/Source/Runtime/Experimental/Chaos/Private/PBDRigidsSolver.cpp
@@ -786,12 +786,17 @@ namespace Chaos
 				{
 					if (Handle)
 					{
-						MEvolution->DisableParticle(Handle);
-
 						// It should be safe to defer here without an immediate call to HandleDeferredClusterUnionUpdateProperties.
 						// This change doesn't really have to go through until FClusterUnionManager::FlushPendingOperations which happens
 						// prior to trying to advance the frame.
 						ClusterUnionManager.HandleRemoveOperationWithClusterLookup({ Handle }, EClusterUnionOperationTiming::Defer);
+
+						// Need to use FRigidClustering::DisableCluster instead of the evolution's to also handle the fact that the GC particle could be in the top level strained sets
+						// which would make it get considered in the breaking model which is undesirable.
+						MEvolution->DisableParticle(Handle);
+
+						// This ensures that this particle won't have any intercluster edges on it. This way no connectivity operations will try to add it into a cluster union.
+						MEvolution->GetRigidClustering().RemoveNodeConnections(Handle);
 					}
 				}
 
