IAS - broadcast encryption keys to PAK

#rnx

diff --git a/Engine/Source/Runtime/Experimental/IoStoreOnDemand/Private/IoStoreOnDemand.cpp b/Engine/Source/Runtime/Experimental/IoStoreOnDemand/Private/IoStoreOnDemand.cpp
--- a/Engine/Source/Runtime/Experimental/IoStoreOnDemand/Private/IoStoreOnDemand.cpp
+++ b/Engine/Source/Runtime/Experimental/IoStoreOnDemand/Private/IoStoreOnDemand.cpp
@@ -11,6 +11,7 @@
 #include "Misc/Base64.h"
 #include "Misc/CommandLine.h"
 #include "Misc/ConfigCacheIni.h"
+#include "Misc/CoreDelegates.h"
 #include "Misc/Parse.h"
 #include "Modules/ModuleManager.h"
 #include "Serialization/CompactBinaryWriter.h"
@@ -887,13 +888,15 @@ void FIoStoreOnDemandModule::StartupModule()
 
 	{
 		FString EncryptionKey;
-		if (FParse::Value(CommandLine, TEXT("IasEncryptionKey="), EncryptionKey))
+		if (FParse::Value(CommandLine, TEXT("Ias.EncryptionKey="), EncryptionKey))
 		{
 			FGuid KeyGuid;
 			FAES::FAESKey Key;
 			if (ParseEncryptionKeyParam(EncryptionKey, KeyGuid, Key))
 			{
+				// TODO: PAK and I/O store should share key manager
 				UE::FEncryptionKeyManager::Get().AddKey(KeyGuid, Key);
+				FCoreDelegates::GetRegisterEncryptionKeyMulticastDelegate().Broadcast(KeyGuid, Key);
 			}
 		}
 	}
@@ -936,7 +939,9 @@ void FIoStoreOnDemandModule::StartupModule()
 				FAES::FAESKey Key;
 				if (ParseEncryptionKeyParam(ContentKey, KeyGuid, Key))
 				{
+					// TODO: PAK and I/O store should share key manager
 					UE::FEncryptionKeyManager::Get().AddKey(KeyGuid, Key);
+					FCoreDelegates::GetRegisterEncryptionKeyMulticastDelegate().Broadcast(KeyGuid, Key);
 				}
 			}
 		}
