[ChaosVD] Swtiching the CVD recording method from SoA view to just record all particle handles in the RBAN solver.

In the rigid solver, we use a method that takes a SoA Particles, and use its views to record only what was changed (deltas), specifically, when we are recording "deltas", we record onnly the particles in the dirty view.

RBAN solvers don't update these views because it does not need these optimization for now, so to not miss any particle in the recording, this change ensures all particles of the RBAN solver are always recorded (it is a low amount so it doesn't have a huge impact).

diff --git a/Engine/Source/Runtime/Experimental/Chaos/Private/Chaos/Evolution/PBDMinEvolution.cpp b/Engine/Source/Runtime/Experimental/Chaos/Private/Chaos/Evolution/PBDMinEvolution.cpp
--- a/Engine/Source/Runtime/Experimental/Chaos/Private/Chaos/Evolution/PBDMinEvolution.cpp
+++ b/Engine/Source/Runtime/Experimental/Chaos/Private/Chaos/Evolution/PBDMinEvolution.cpp
@@ -144,7 +144,7 @@ namespace Chaos
 			ScatterOutput(Dt);
 		}
 
-		CVD_TRACE_PARTICLES_SOA(Particles);
+		CVD_TRACE_PARTICLES(Particles.GetParticleHandles());
 	}
 
 	// A opportunity for systems to allocate buffers for the duration of the tick, if they have enough info to do so
