turn on checks for intercluster edges by default

diff --git a/Engine/Source/Runtime/Experimental/Chaos/Private/Chaos/PBDRigidClustering.cpp b/Engine/Source/Runtime/Experimental/Chaos/Private/Chaos/PBDRigidClustering.cpp
--- a/Engine/Source/Runtime/Experimental/Chaos/Private/Chaos/PBDRigidClustering.cpp
+++ b/Engine/Source/Runtime/Experimental/Chaos/Private/Chaos/PBDRigidClustering.cpp
@@ -36,7 +36,7 @@ namespace Chaos
 	int32 UseConnectivity = 1;
 	FAutoConsoleVariableRef CVarUseConnectivity(TEXT("p.UseConnectivity"), UseConnectivity, TEXT("Whether to use connectivity graph when breaking up clusters"));
 
-	bool bCheckForInterclusterEdgesOnRelease = false;
+	bool bCheckForInterclusterEdgesOnRelease = true;
 	FAutoConsoleVariableRef CVarCheckForInterclusterEdgesOnRelease(TEXT("p.Chaos.CheckForInterclusterEdgesOnRelease"), bCheckForInterclusterEdgesOnRelease, TEXT("Whether to check for intercluster edges when removing a child from its parent cluster so that we can add the particle back into a cluster union."));
 
 	int32 ComputeClusterCollisionStrains = 1;
