add a 'per bone' negative space compute option to the merge hulls dataflow node, which can more reasonably target negative space for small leaf bones

diff --git a/Engine/Plugins/Experimental/GeometryCollectionPlugin/Source/GeometryCollectionNodes/Private/Dataflow/GeometryCollectionUtilityNodes.cpp b/Engine/Plugins/Experimental/GeometryCollectionPlugin/Source/GeometryCollectionNodes/Private/Dataflow/GeometryCollectionUtilityNodes.cpp
--- a/Engine/Plugins/Experimental/GeometryCollectionPlugin/Source/GeometryCollectionNodes/Private/Dataflow/GeometryCollectionUtilityNodes.cpp
+++ b/Engine/Plugins/Experimental/GeometryCollectionPlugin/Source/GeometryCollectionNodes/Private/Dataflow/GeometryCollectionUtilityNodes.cpp
@@ -407,33 +407,38 @@ void FMergeConvexHullsDataflowNode::Evaluate(Dataflow::FContext& Context, const
 			SelectionArray = InOptionalSelectionFilter.AsArray();
 		}
 
-		bool bHasNegativeSpace = false;
+		bool bHasPrecomputedNegativeSpace = false;
 		UE::Geometry::FSphereCovering NegativeSpace;
-		if (GetValue(Context, &bProtectNegativeSpace))
+		bool bInProtectNegativeSpace = GetValue(Context, &bProtectNegativeSpace);
+		UE::Geometry::FNegativeSpaceSampleSettings NegativeSpaceSettings;
+		if (bInProtectNegativeSpace)
 		{
-			UE::Geometry::FNegativeSpaceSampleSettings NegativeSpaceSettings;
 			NegativeSpaceSettings.TargetNumSamples = GetValue(Context, &TargetNumSamples);
 			NegativeSpaceSettings.MinRadius = GetValue(Context, &MinRadius);
 			NegativeSpaceSettings.ReduceRadiusMargin = GetValue(Context, &NegativeSpaceTolerance);
 			NegativeSpaceSettings.MinSpacing = GetValue(Context, &MinSampleSpacing);
 			NegativeSpaceSettings.SampleMethod = ConvertNegativeSpaceSampleMethodDataflowEnum(SampleMethod);
 			NegativeSpaceSettings.Sanitize();
-			bHasNegativeSpace = UE::FractureEngine::Convex::ComputeConvexHullsNegativeSpace(InCollection, NegativeSpace, NegativeSpaceSettings, bHasSelectionFilter, SelectionArray, false);
+		}
+		if (bInProtectNegativeSpace && !bComputeNegativeSpacePerBone)
+		{
+			bHasPrecomputedNegativeSpace = UE::FractureEngine::Convex::ComputeConvexHullsNegativeSpace(InCollection, NegativeSpace, NegativeSpaceSettings, bHasSelectionFilter, SelectionArray, false);
 		}
 
 		const int32 InMaxConvexCount = GetValue(Context, &MaxConvexCount);
 		const double InErrorToleranceInCm = GetValue(Context, &ErrorTolerance);
 		FGeometryCollectionConvexUtility::FMergeConvexHullSettings HullMergeSettings;
-		HullMergeSettings.EmptySpace = bHasNegativeSpace ? &NegativeSpace : nullptr;
+		HullMergeSettings.EmptySpace = bHasPrecomputedNegativeSpace ? &NegativeSpace : nullptr;
 		HullMergeSettings.ErrorToleranceInCm = InErrorToleranceInCm;
 		HullMergeSettings.MaxConvexCount = InMaxConvexCount;
+		HullMergeSettings.ComputeEmptySpacePerBoneSettings = (bInProtectNegativeSpace && bComputeNegativeSpacePerBone) ? &NegativeSpaceSettings : nullptr;
 
-		FGeometryCollectionConvexUtility::MergeHullsOnTransforms(InCollection, HullMergeSettings, bHasSelectionFilter, SelectionArray);
+		UE::Geometry::FSphereCovering UsedNegativeSpace;
+		FGeometryCollectionConvexUtility::MergeHullsOnTransforms(InCollection, HullMergeSettings, bHasSelectionFilter, SelectionArray, &UsedNegativeSpace);
 
 		SetValue(Context, MoveTemp(InCollection), &Collection);
 
-		// Move the negative space to the output container at the end, after it is no longer needed
-		Spheres.Spheres = MoveTemp(NegativeSpace);
+		Spheres.Spheres = MoveTemp(UsedNegativeSpace);
 		SetValue(Context, MoveTemp(Spheres), &SphereCovering);
 	}
 }
diff --git a/Engine/Plugins/Experimental/GeometryCollectionPlugin/Source/GeometryCollectionNodes/Public/Dataflow/GeometryCollectionUtilityNodes.h b/Engine/Plugins/Experimental/GeometryCollectionPlugin/Source/GeometryCollectionNodes/Public/Dataflow/GeometryCollectionUtilityNodes.h
--- a/Engine/Plugins/Experimental/GeometryCollectionPlugin/Source/GeometryCollectionNodes/Public/Dataflow/GeometryCollectionUtilityNodes.h
+++ b/Engine/Plugins/Experimental/GeometryCollectionPlugin/Source/GeometryCollectionNodes/Public/Dataflow/GeometryCollectionUtilityNodes.h
@@ -425,6 +425,10 @@ public:
 	UPROPERTY(EditAnywhere, Category = NegativeSpace, meta = (DataflowInput))
 	bool bProtectNegativeSpace = false;
 
+	/** Whether to compute separate negative space for each bone. Otherwise, a single negative space will be computed once and re-used for all bones. */
+	UPROPERTY(EditAnywhere, Category = NegativeSpace, meta = (EditCondition = "bProtectNegativeSpace", EditConditionHides))
+	bool bComputeNegativeSpacePerBone = false;
+
 	/** Method to use to find and sample negative space */
 	UPROPERTY(EditAnywhere, Category = NegativeSpace, meta = (EditCondition = "bProtectNegativeSpace", EditConditionHides))
 	ENegativeSpaceSampleMethodDataflowEnum SampleMethod = ENegativeSpaceSampleMethodDataflowEnum::Uniform;
diff --git a/Engine/Source/Runtime/Experimental/Chaos/Private/GeometryCollection/GeometryCollectionConvexUtility.cpp b/Engine/Source/Runtime/Experimental/Chaos/Private/GeometryCollection/GeometryCollectionConvexUtility.cpp
--- a/Engine/Source/Runtime/Experimental/Chaos/Private/GeometryCollection/GeometryCollectionConvexUtility.cpp
+++ b/Engine/Source/Runtime/Experimental/Chaos/Private/GeometryCollection/GeometryCollectionConvexUtility.cpp
@@ -1213,7 +1213,7 @@ double ComputeGeometryVolume(
 }
 
 // Helper to append the triangles of a convex hull to a dynamic mesh
-static void AddConvexHullToCompactDynamicMesh(const ::Chaos::FConvex* InConvexHull, UE::Geometry::FDynamicMesh3& Mesh, const FTransform* OptionalTransform = nullptr)
+static void AddConvexHullToCompactDynamicMesh(const ::Chaos::FConvex* InConvexHull, UE::Geometry::FDynamicMesh3& Mesh, const FTransform* OptionalTransform = nullptr, bool bInvertFaces = true)
 {
 	check(Mesh.IsCompact());
 
@@ -1239,7 +1239,6 @@ static void AddConvexHullToCompactDynamicMesh(const ::Chaos::FConvex* InConvexHu
 		{
 			int32 V1 = StartV + ConvexStructure.GetPlaneVertex(PIdx, SubIdx);
 			int32 V2 = StartV + ConvexStructure.GetPlaneVertex(PIdx, SubIdx + 1);
-			constexpr bool bInvertFaces = true; // default-invert to match expected orientation of geometry meshes
 			if (bInvertFaces)
 			{
 				Swap(V1, V2);
@@ -2228,7 +2227,8 @@ void FGeometryCollectionConvexUtility::GenerateClusterConvexHullsFromLeafOrChild
 
 }
 
-void FGeometryCollectionConvexUtility::MergeHullsOnTransforms(FManagedArrayCollection& Collection, const FGeometryCollectionConvexUtility::FMergeConvexHullSettings& Settings, bool bRestrictToSelection, const TArrayView<const int32> OptionalTransformSelection)
+void FGeometryCollectionConvexUtility::MergeHullsOnTransforms(FManagedArrayCollection& Collection, const FGeometryCollectionConvexUtility::FMergeConvexHullSettings& Settings, bool bRestrictToSelection,
+	const TArrayView<const int32> OptionalTransformSelection, UE::Geometry::FSphereCovering* OptionalSphereCoveringOut)
 {
 	static FName ConvexGroupName("Convex");
 
@@ -2258,6 +2258,15 @@ void FGeometryCollectionConvexUtility::MergeHullsOnTransforms(FManagedArrayColle
 		}
 	}
 
+	if (OptionalSphereCoveringOut)
+	{
+		OptionalSphereCoveringOut->Reset();
+	}
+	if (OptionalSphereCoveringOut && Settings.EmptySpace)
+	{
+		OptionalSphereCoveringOut->Append(*Settings.EmptySpace);
+	}
+
 	for (int32 TransformIndex : TransformsToProcess)
 	{
 		const int32 InitialNumConvex = TransformToConvexIndices[TransformIndex].Num();
@@ -2289,7 +2298,32 @@ void FGeometryCollectionConvexUtility::MergeHullsOnTransforms(FManagedArrayColle
 		};
 		UE::Geometry::FConvexDecomposition3 ConvexDecomposition;
 		ConvexDecomposition.InitializeFromHulls(Hulls.Num(), GetHullVolume, GetHullNumVertices, GetHullVertex, HullProximity);
-		ConvexDecomposition.MergeBest(Settings.MaxConvexCount, Settings.ErrorToleranceInCm, 0, true /*bAllowCompact*/, false /*bRequireHullTriangles*/, Settings.MaxConvexCount /*MaxHulls*/, Settings.EmptySpace, &GlobalTransforms[TransformIndex]);
+		UE::Geometry::FSphereCovering LocalCovering;
+		UE::Geometry::FSphereCovering* UseCovering = Settings.EmptySpace;
+		if (Settings.ComputeEmptySpacePerBoneSettings)
+		{
+			UseCovering = &LocalCovering;
+			UE::Geometry::FDynamicMesh3 HullsMesh;
+			for (const ::Chaos::FConvex* Hull : Hulls)
+			{
+				constexpr bool bInvertFaces = false; // don't invert faces, to match orientation expected by AddNegativeSpace below
+				AddConvexHullToCompactDynamicMesh(Hull, HullsMesh, &GlobalTransforms[TransformIndex], bInvertFaces);
+			}
+			UE::Geometry::FDynamicMeshAABBTree3 HullsTree(&HullsMesh, true);
+			UE::Geometry::TFastWindingTree<UE::Geometry::FDynamicMesh3> HullsWinding(&HullsTree, true);
+			LocalCovering.AddNegativeSpace(HullsWinding, *Settings.ComputeEmptySpacePerBoneSettings);
+
+			if (Settings.EmptySpace)
+			{
+				LocalCovering.Append(*Settings.EmptySpace);
+			}
+
+			if (OptionalSphereCoveringOut)
+			{
+				OptionalSphereCoveringOut->Append(LocalCovering);
+			}
+		}
+		ConvexDecomposition.MergeBest(Settings.MaxConvexCount, Settings.ErrorToleranceInCm, 0, true /*bAllowCompact*/, false /*bRequireHullTriangles*/, Settings.MaxConvexCount /*MaxHulls*/, UseCovering, &GlobalTransforms[TransformIndex]);
 
 		if (!ensure(ConvexDecomposition.Decomposition.Num() > 0))
 		{
diff --git a/Engine/Source/Runtime/Experimental/Chaos/Public/GeometryCollection/GeometryCollectionConvexUtility.h b/Engine/Source/Runtime/Experimental/Chaos/Public/GeometryCollection/GeometryCollectionConvexUtility.h
--- a/Engine/Source/Runtime/Experimental/Chaos/Public/GeometryCollection/GeometryCollectionConvexUtility.h
+++ b/Engine/Source/Runtime/Experimental/Chaos/Public/GeometryCollection/GeometryCollectionConvexUtility.h
@@ -17,6 +17,7 @@ namespace Chaos
 namespace UE::Geometry
 {
 	class FSphereCovering;
+	struct FNegativeSpaceSampleSettings;
 }
 
 UENUM()
@@ -136,10 +137,16 @@ public:
 	{
 		int32 MaxConvexCount = -1;
 		double ErrorToleranceInCm = 0.0;
+		// Optional externally-provided empty space, to be used for all hull merges
 		UE::Geometry::FSphereCovering* EmptySpace = nullptr;
+
+		// Optional settings to compute targeted empty space per-bone
+		UE::Geometry::FNegativeSpaceSampleSettings* ComputeEmptySpacePerBoneSettings = nullptr;
 	};
 	// Merge convex hulls that are currently on each (selected) transform. If convex hulls are not present, does nothing.
-	static CHAOS_API void MergeHullsOnTransforms(FManagedArrayCollection& Collection, const FGeometryCollectionConvexUtility::FMergeConvexHullSettings& Settings, bool bRestrictToSelection, const TArrayView<const int32> OptionalTransformSelection);
+	// @params OptionalSphereCoveringOut		If non-null, will be filled with spheres from all used sphere covering.
+	static CHAOS_API void MergeHullsOnTransforms(FManagedArrayCollection& Collection, const FGeometryCollectionConvexUtility::FMergeConvexHullSettings& Settings, bool bRestrictToSelection, const TArrayView<const int32> OptionalTransformSelection,
+		UE::Geometry::FSphereCovering* OptionalSphereCoveringOut = nullptr);
 
 	// Additional settings for filtering when the EGenerateConvexMethod::IntersectExternalWithComputed is applied
 	struct FIntersectionFilters
diff --git a/Engine/Source/Runtime/GeometryCore/Public/CompGeom/ConvexDecomposition3.h b/Engine/Source/Runtime/GeometryCore/Public/CompGeom/ConvexDecomposition3.h
--- a/Engine/Source/Runtime/GeometryCore/Public/CompGeom/ConvexDecomposition3.h
+++ b/Engine/Source/Runtime/GeometryCore/Public/CompGeom/ConvexDecomposition3.h
@@ -103,6 +103,12 @@ public:
 		Radius.Reset();
 	}
 
+	void Append(const FSphereCovering& Other)
+	{
+		Position.Append(Other.Position);
+		Radius.Append(Other.Radius);
+	}
+
 private:
 	// Sphere centers
 	TArray<FVector3d> Position;
