Avoid unnecessary allocations when processing lots of dirty rewind data.

diff --git a/Engine/Source/Runtime/Experimental/Chaos/Private/RewindData.cpp b/Engine/Source/Runtime/Experimental/Chaos/Private/RewindData.cpp
--- a/Engine/Source/Runtime/Experimental/Chaos/Private/RewindData.cpp
+++ b/Engine/Source/Runtime/Experimental/Chaos/Private/RewindData.cpp
@@ -512,7 +512,8 @@ void FRewindData::AdvanceFrameImp(IResimCacheBase* ResimCache)
 
 	auto AdvanceHelper = [this, EarliestFrame, FrameAndPhase](auto& DirtyObjects, const auto& DesyncFunc, const auto& AdvanceDirtyFunc)
 	{
-		for (int32 DirtyIdx = DirtyObjects.Num() - 1; DirtyIdx >= 0; --DirtyIdx)
+		const int32 InitialNumDirtyObjects = DirtyObjects.Num();
+		for (int32 DirtyIdx = InitialNumDirtyObjects - 1; DirtyIdx >= 0; --DirtyIdx)
 		{
 			auto& Info = DirtyObjects.GetDenseAt(DirtyIdx);
 
@@ -521,7 +522,8 @@ void FRewindData::AdvanceFrameImp(IResimCacheBase* ResimCache)
 			//if hasn't changed in a while stop tracking
 			if (Info.LastDirtyFrame < EarliestFrame)
 			{
-				RemoveObject(Info.GetObjectPtr());
+				constexpr bool bAllowShrinking = false;
+				RemoveObject(Info.GetObjectPtr(), bAllowShrinking);
 			}
 			else
 			{
@@ -530,18 +532,22 @@ void FRewindData::AdvanceFrameImp(IResimCacheBase* ResimCache)
 
 				if (IsResim() && !Info.bResimAsFollower)
 				{
-						DesyncIfNecessary</*bSkipDynamics=*/false>(Info, FrameAndPhase);
+					DesyncIfNecessary</*bSkipDynamics=*/false>(Info, FrameAndPhase);
 				}
 
 				if (IsResim() && Handle->SyncState() != ESyncState::InSync && !SkipDesyncTest)
 				{
 					Handle->SetEnabledDuringResim(true);	//for now just mark anything out of sync as resim enabled. TODO: use bubble
-						DesyncFunc(Handle);
+					DesyncFunc(Handle);
 				}
 
 				AdvanceDirtyFunc(Info, Handle);
 			}
 		}
+		if (InitialNumDirtyObjects > 0)
+		{
+			DirtyObjects.Shrink();
+		}
 	};
 
 	TArray<FGeometryParticleHandle*> DesyncedParticles;
diff --git a/Engine/Source/Runtime/Experimental/Chaos/Public/RewindData.h b/Engine/Source/Runtime/Experimental/Chaos/Public/RewindData.h
--- a/Engine/Source/Runtime/Experimental/Chaos/Public/RewindData.h
+++ b/Engine/Source/Runtime/Experimental/Chaos/Public/RewindData.h
@@ -1211,11 +1211,12 @@ public:
 		return nullptr;
 	}
 
-	void Remove(const TKey Key)
+	void Remove(const TKey Key, const bool bAllowShrinking)
 	{
 		if (const int32* Idx = KeyToIdx.Find(Key))
 		{
-			DenseVals.RemoveAtSwap(*Idx);
+			constexpr int32 Count = 1;
+			DenseVals.RemoveAtSwap(*Idx, Count, bAllowShrinking);
 
 			if(*Idx < DenseVals.Num())
 			{
@@ -1226,7 +1227,12 @@ public:
 			KeyToIdx.Remove(Key);
 		}
 	}
-	
+
+	void Shrink()
+	{
+		DenseVals.Shrink();
+	}
+
 	void Reset()
 	{
 		DenseVals.Reset();
@@ -1286,14 +1292,14 @@ public:
 		return Managers[Frame].DeltaTime;
 	}
 
-	void RemoveObject(const FGeometryParticleHandle* Particle)
+	void RemoveObject(const FGeometryParticleHandle* Particle, const bool bAllowShrinking=true)
 	{
-		DirtyParticles.Remove(Particle);
+		DirtyParticles.Remove(Particle, bAllowShrinking);
 	}
 
-	void RemoveObject(const FPBDJointConstraintHandle* Joint)
+	void RemoveObject(const FPBDJointConstraintHandle* Joint, const bool bAllowShrinking = true)
 	{
-		DirtyJoints.Remove(Joint);
+		DirtyJoints.Remove(Joint, bAllowShrinking);
 	}
 
 	int32 GetEarliestFrame_Internal() const { return CurFrame - FramesSaved; }
