IAS - Improved upload summary

#rnx

diff --git a/Engine/Source/Runtime/Experimental/IoStoreOnDemand/Private/IoStoreOnDemand.cpp b/Engine/Source/Runtime/Experimental/IoStoreOnDemand/Private/IoStoreOnDemand.cpp
--- a/Engine/Source/Runtime/Experimental/IoStoreOnDemand/Private/IoStoreOnDemand.cpp
+++ b/Engine/Source/Runtime/Experimental/IoStoreOnDemand/Private/IoStoreOnDemand.cpp
@@ -7,6 +7,7 @@
 
 #include "FileIoCache.h"
 #include "HAL/PlatformMisc.h"
+#include "HAL/PlatformTime.h"
 #include "HAL/FileManager.h"
 #include "Misc/Base64.h"
 #include "Misc/CommandLine.h"
@@ -542,6 +543,18 @@ TIoStatusOr<FIoStoreUploadResult> UploadContainerFiles(
 	TConstArrayView<FString> ContainerFiles,
 	const TMap<FGuid, FAES::FAESKey>& EncryptionKeys)
 {
+	struct FContainerStats
+	{
+		FString ContainerName;
+		uint64 ChunkCount = 0;
+		uint64 TotalBytes = 0;
+		uint64 UploadedChunkCount = 0;
+		uint64 UploadedBytes = 0;
+	};
+	TMap<FString, FContainerStats> ContainerSummary;
+
+	const double StartTime = FPlatformTime::Seconds();
+
 	FS3ClientConfig Config;
 	Config.ServiceUrl = UploadParams.ServiceUrl;
 	Config.Region = UploadParams.Region;
@@ -574,6 +587,8 @@ TIoStatusOr<FIoStoreUploadResult> UploadContainerFiles(
 	}
 
 	TSet<FIoHash> ExistingChunks;
+	uint64 TotalExistingTocs = 0;
+	uint64 TotalExistingBytes = 0;
 	{
 		TStringBuilder<256> TocsKey;
 		TocsKey << UploadParams.BucketPrefix << "/";
@@ -618,12 +633,14 @@ TIoStatusOr<FIoStoreUploadResult> UploadContainerFiles(
 				for (const FOnDemandTocEntry& TocEntry : ContainerEntry.Entries)
 				{
 					ExistingChunks.Add(TocEntry.Hash);
+					TotalExistingBytes += TocEntry.EncodedSize;
 				}
 			}
+
+			TotalExistingTocs++;
 		}
 	}
-	UE_LOG(LogIas, Display, TEXT("Found %d existing chunks"), ExistingChunks.Num());
-	
+
 	FString ChunksRelativePath = UploadParams.BucketPrefix.IsEmpty()
 		? FString::Printf(TEXT("IoChunksV%u"), EOnDemandChunkVersion::Latest)
 		: FString::Printf(TEXT("%s/IoChunksV%u"), *UploadParams.BucketPrefix, EOnDemandChunkVersion::Latest);
@@ -674,6 +691,7 @@ TIoStatusOr<FIoStoreUploadResult> UploadContainerFiles(
 		FOnDemandTocContainerEntry& ContainerEntry = OnDemandToc.Containers.AddDefaulted_GetRef();
 		ContainerEntry.ContainerName = FPaths::GetBaseFilename(Path);
 		ContainerEntry.EncryptionKeyGuid = LexToString(ContainerFileReader.GetEncryptionKeyGuid());
+		FContainerStats& ContainerStats = ContainerSummary.FindOrAdd(ContainerEntry.ContainerName);
 		
 		for (const FIoStoreTocChunkInfo& ChunkInfo : ChunkInfos)
 		{
@@ -725,6 +743,9 @@ TIoStatusOr<FIoStoreUploadResult> UploadContainerFiles(
 			TocEntry.BlockOffset = BlockOffset;
 			TocEntry.BlockCount = BlockCount;
 
+			ContainerStats.ChunkCount++;
+			ContainerStats.TotalBytes += EncodedChunkSize;
+
 			if (ExistingChunks.Contains(TocEntry.Hash) == false)
 			{
 				const FString HashString = LexToString(ChunkHash);
@@ -747,7 +768,9 @@ TIoStatusOr<FIoStoreUploadResult> UploadContainerFiles(
 				}
 
 				TotalUploadedChunks++;
-				TotalUploadedBytes += ReadResult.IoBuffer.GetSize();
+				TotalUploadedBytes += EncodedChunkSize;
+				ContainerStats.UploadedChunkCount++;
+				ContainerStats.UploadedBytes += EncodedChunkSize;
 			}
 		}
 
@@ -839,9 +862,47 @@ TIoStatusOr<FIoStoreUploadResult> UploadContainerFiles(
 			}
 		}
 	}
-	
-	UE_LOG(LogIas, Display, TEXT("Done uploading %d container file(s), %d chunks of total %.2lf MiB"),
-		UploadedFiles.Num(), TotalUploadedChunks, (double(TotalUploadedBytes) / 1024 / 1024));
+
+	{
+		const double Duration = FPlatformTime::Seconds() - StartTime;
+		
+		UE_LOG(LogIas, Display, TEXT(""));
+		UE_LOG(LogIas, Display, TEXT("---------------------------------------- Upload Summary --------------------------------------"));
+		UE_LOG(LogIas, Display, TEXT("%-40s: %s"), TEXT("Service URL"), *UploadParams.ServiceUrl);
+		UE_LOG(LogIas, Display, TEXT("%-40s: %s"), TEXT("Bucket"), *UploadParams.Bucket);
+		UE_LOG(LogIas, Display, TEXT("%-40s: %s"), TEXT("TOC"), *UploadResult.TocPath);
+		UE_LOG(LogIas, Display, TEXT("%-40s: %.2lf second(s)"), TEXT("Duration"), Duration);
+		UE_LOG(LogIas, Display, TEXT(""));
+		
+		UE_LOG(LogIas, Display, TEXT("%-25s %15s %15s %15s %15s"),
+			TEXT("Container"), TEXT("Chunk(s)"), TEXT("Size (MiB)"), TEXT("Uploaded"), TEXT("Uploaded (MiB)"));
+		UE_LOG(LogIas, Display, TEXT("----------------------------------------------------------------------------------------------"));
+
+		FContainerStats TotalStats;
+		for (const TTuple<FString, FContainerStats>& Kv : ContainerSummary)
+		{
+			UE_LOG(LogIas, Display, TEXT("%-25s %15llu %15.2lf %15llu %15.2lf"),
+				*Kv.Key, Kv.Value.ChunkCount, double(Kv.Value.TotalBytes) / 1024.0 / 1024.0, Kv.Value.UploadedChunkCount, double(Kv.Value.UploadedBytes) / 1024.0 / 1024.0);
+			
+			TotalStats.ChunkCount += Kv.Value.ChunkCount;
+			TotalStats.TotalBytes += Kv.Value.TotalBytes;
+			TotalStats.UploadedChunkCount += Kv.Value.UploadedChunkCount;
+			TotalStats.UploadedBytes += Kv.Value.UploadedBytes;
+		}
+		UE_LOG(LogIas, Display, TEXT("----------------------------------------------------------------------------------------------"));
+		UE_LOG(LogIas, Display, TEXT("%-25s %15llu %15.2lf %15llu %15.2lf"),
+			TEXT("Total"),TotalStats.ChunkCount, double(TotalStats.TotalBytes) / 1024.0 / 1024.0, TotalStats.UploadedChunkCount, double(TotalStats.UploadedBytes) / 1024.0 / 1024.0);
+		UE_LOG(LogIas, Display, TEXT(""));
+		
+		UE_LOG(LogIas, Display, TEXT("%-25s %15s %15s %15s"), TEXT("Bucket"), TEXT("TOC(s)"), TEXT("Chunk(s)"), TEXT("MiB"));
+		UE_LOG(LogIas, Display, TEXT("----------------------------------------------------------------------------------------------"));
+		UE_LOG(LogIas, Display, TEXT("%-25s %15llu %15d %15.2lf"), TEXT("Existing"), TotalExistingTocs, ExistingChunks.Num(), double(TotalExistingBytes) / 1024.0 / 1024.0);
+		UE_LOG(LogIas, Display, TEXT("%-25s %15llu %15d %15.2lf"), TEXT("Uploaded"), 1, TotalStats.UploadedChunkCount, double(TotalStats.UploadedBytes) / 1024.0 / 1024.0);
+		UE_LOG(LogIas, Display, TEXT("----------------------------------------------------------------------------------------------"));
+		UE_LOG(LogIas, Display, TEXT("%-25s %15llu %15d %15.2lf"), TEXT("Total"),
+			TotalExistingTocs + 1, ExistingChunks.Num() + TotalStats.UploadedChunkCount, double(TotalExistingBytes + TotalStats.UploadedBytes) / 1024.0 / 1024.0);
+		UE_LOG(LogIas, Display, TEXT(""));
+	}
 	
 	return UploadResult;
 }
