Chaos - Minor Colliison BVH Optimizations

diff --git a/Engine/Source/Runtime/Experimental/Chaos/Private/Chaos/Collision/ParticlePairMidPhase.cpp b/Engine/Source/Runtime/Experimental/Chaos/Private/Chaos/Collision/ParticlePairMidPhase.cpp
--- a/Engine/Source/Runtime/Experimental/Chaos/Private/Chaos/Collision/ParticlePairMidPhase.cpp
+++ b/Engine/Source/Runtime/Experimental/Chaos/Private/Chaos/Collision/ParticlePairMidPhase.cpp
@@ -1499,10 +1499,6 @@ namespace Chaos
 			Constraint->SetCCDEnabled(IsCCD());
 		}
 
-		// Do we want to sweep on this tick? 
-		// CCD may be temporarily disabled by the user or because we are moving slowly.
-		Constraint->SetCCDSweepEnabled(IsCCD() && bEnableSweep);
-
 		NewConstraints.Add(Constraint);
 		return Constraint;
 	}
@@ -1563,25 +1559,46 @@ namespace Chaos
 		}
 	}
 
+	void PrefetchConstraint(const TArray<FPBDCollisionConstraint*>& Constraints, const int32 ConstraintIndex)
+	{
+		if (ConstraintIndex < Constraints.Num())
+		{
+			FPlatformMisc::PrefetchBlock(Constraints[ConstraintIndex], sizeof(FPBDCollisionConstraint));
+		}
+	}
+
 	int32 FGenericParticlePairMidPhase::ProcessNewConstraints(
 		const FReal CullDistance,
 		const FReal Dt,
 		const FCollisionContext& Context)
 	{
 		int32 NumActive = 0;
+		const bool bUseCCDSweep = IsCCD() && Flags.bUseSweep;
 
-		for (FPBDCollisionConstraint* Constraint : NewConstraints)
+		const int32 NumNewConstraints = NewConstraints.Num();
+		for (int32 ConstraintIndex = 0; ConstraintIndex < NumNewConstraints; ++ConstraintIndex)
 		{
-			if (!Constraint->GetCCDSweepEnabled())
+			FPBDCollisionConstraint* Constraint = NewConstraints[ConstraintIndex];
+			PrefetchConstraint(NewConstraints, ConstraintIndex + 1);
+
+			// Do we want to sweep on this tick? 
+			// CCD may be temporarily disabled by the user or because we are moving slowly.
+			if (IsCCD())
 			{
-				UpdateCollision(Constraint, CullDistance, Dt, Context);
+				Constraint->SetCCDSweepEnabled(bUseCCDSweep);
+			}
+
+			bool bIsActive = false;
+			if (!bUseCCDSweep)
+			{
+				bIsActive = UpdateCollision(Constraint, CullDistance, Dt, Context);
 			}
 			else
 			{
-				UpdateCollisionCCD(Constraint, CullDistance, Dt, Context);
+				bIsActive = UpdateCollisionCCD(Constraint, CullDistance, Dt, Context);
 			}
 
-			if (!Constraint->GetDisabled())
+			if (bIsActive)
 			{
 				++NumActive;
 			}
@@ -1592,7 +1609,7 @@ namespace Chaos
 		return NumActive;
 	}
 
-	void FGenericParticlePairMidPhase::UpdateCollision(
+	bool FGenericParticlePairMidPhase::UpdateCollision(
 		FPBDCollisionConstraint* Constraint,
 		const FReal CullDistance,
 		const FReal Dt,
@@ -1620,6 +1637,7 @@ namespace Chaos
 		}
 
 		bool bWasManifoldRestored = false;
+		// @todo(chaos): enable manifold restore for Generic midphase
 		//if (Context.GetSettings().bAllowManifoldReuse && Flags.bEnableManifoldUpdate && bWasUpdatedLastTick)
 		//{
 		//	// Update the existing manifold. We can re-use as-is if none of the points have moved much and the bodies have not moved much
@@ -1651,16 +1669,17 @@ namespace Chaos
 		{
 			if (Context.GetAllocator()->ActivateConstraint(Constraint))
 			{
-				return;
+				return true;
 			}
 		}
 
 		// If we get here, we did not activate the constraint and it should be disabled for this tick
 		Constraint->SetDisabled(true);
+		return false;
 	}
 
 
-	void FGenericParticlePairMidPhase::UpdateCollisionCCD(
+	bool FGenericParticlePairMidPhase::UpdateCollisionCCD(
 		FPBDCollisionConstraint* Constraint,
 		const FReal CullDistance,
 		const FReal Dt,
@@ -1668,10 +1687,7 @@ namespace Chaos
 	{
 		// @todo(chaos): share this code with FSingleShapePairCollisionDetector
 
-		if (!Constraint->GetCCDSweepEnabled())
-		{
-			return UpdateCollision(Constraint, CullDistance, Dt, Context);
-		}
+		check(Constraint->GetCCDSweepEnabled());
 
 		Constraint->ResetManifold();
 		Constraint->ResetActiveManifoldContacts();
@@ -1729,6 +1745,7 @@ namespace Chaos
 		}
 
 		Context.GetAllocator()->ActivateConstraint(Constraint);
+		return true;
 	}
 
 	void FGenericParticlePairMidPhase::PruneConstraints(const int32 CurrentEpoch)
diff --git a/Engine/Source/Runtime/Experimental/Chaos/Public/Chaos/Collision/ParticlePairMidPhase.h b/Engine/Source/Runtime/Experimental/Chaos/Public/Chaos/Collision/ParticlePairMidPhase.h
--- a/Engine/Source/Runtime/Experimental/Chaos/Public/Chaos/Collision/ParticlePairMidPhase.h
+++ b/Engine/Source/Runtime/Experimental/Chaos/Public/Chaos/Collision/ParticlePairMidPhase.h
@@ -617,13 +617,13 @@ namespace Chaos
 
 		CHAOS_API void PruneConstraints(const int32 CurrentEpoch);
 
-		CHAOS_API void UpdateCollision(
+		CHAOS_API bool UpdateCollision(
 			FPBDCollisionConstraint* Constraint,
 			const FReal CullDistance,
 			const FReal Dt,
 			const FCollisionContext& Context);
 
-		CHAOS_API void UpdateCollisionCCD(
+		CHAOS_API bool UpdateCollisionCCD(
 			FPBDCollisionConstraint* Constraint,
 			const FReal CullDistance,
 			const FReal Dt,
