Fix missing breaking events
The Geometry collection component is creating its proxy once its flags are already initialized. But the proxy's flags were not initialized and overwritten them in push physics.

diff --git a/Engine/Source/Runtime/Experimental/Chaos/Private/PhysicsProxy/GeometryCollectionPhysicsProxy.cpp b/Engine/Source/Runtime/Experimental/Chaos/Private/PhysicsProxy/GeometryCollectionPhysicsProxy.cpp
--- a/Engine/Source/Runtime/Experimental/Chaos/Private/PhysicsProxy/GeometryCollectionPhysicsProxy.cpp
+++ b/Engine/Source/Runtime/Experimental/Chaos/Private/PhysicsProxy/GeometryCollectionPhysicsProxy.cpp
@@ -449,6 +449,7 @@ FGeometryCollectionPhysicsProxy::FGeometryCollectionPhysicsProxy(
 	, CollisionParticlesPerObjectFraction(CollisionParticlesPerObjectFractionDefault)
 
 	, GameThreadCollection(GameThreadCollectionIn)
+	, GameThreadPerFrameData(SimulationParameters)
 	, bIsPhysicsThreadWorldTransformDirty(false)
 	, bIsCollisionFilterDataDirty(false)
 	, CollectorGuid(InCollectorGuid)
@@ -2771,7 +2772,7 @@ void FGeometryCollectionPhysicsProxy::PushStateOnGameThread(Chaos::FPBDRigidsSol
 
 	if (GameThreadPerFrameData.GetIsNotificationDataDirty())
 	{
-		Parameters.bGenerateBreakingData = GameThreadPerFrameData.GetNotifyBreakings() || GameThreadPerFrameData.GetNotifyGlobalBreakings();
+		Parameters.bGenerateBreakingData = GameThreadPerFrameData.GetNotifyBreakings();
 		Parameters.bGenerateCrumblingData = GameThreadPerFrameData.GetNotifyCrumblings();
 		Parameters.bGenerateCrumblingChildrenData = GameThreadPerFrameData.GetCrumblingEventIncludesChildren();
 		Parameters.bGenerateGlobalBreakingData = GameThreadPerFrameData.GetNotifyGlobalBreakings();
diff --git a/Engine/Source/Runtime/Experimental/Chaos/Public/GeometryCollectionProxyData.h b/Engine/Source/Runtime/Experimental/Chaos/Public/GeometryCollectionProxyData.h
--- a/Engine/Source/Runtime/Experimental/Chaos/Public/GeometryCollectionProxyData.h
+++ b/Engine/Source/Runtime/Experimental/Chaos/Public/GeometryCollectionProxyData.h
@@ -173,6 +173,22 @@ public:
 		, bEnableStrainOnCollision(false)
 	{}
 
+	FGeometryCollectioPerFrameData(const FSimulationParameters& Parameters)
+		: bIsWorldTransformDirty(false)
+		, bIsCollisionFilterDataDirty(false)
+		, bIsNotificationDataDirty(false)
+		, bIsDamageSettingsDataDirty(false)
+		, bNotifyBreakings(Parameters.bGenerateBreakingData)
+		, bNotifyRemovals(false)
+		, bNotifyCrumblings(Parameters.bGenerateCrumblingData)
+		, bCrumblingEventIncludesChildren(Parameters.bGenerateCrumblingChildrenData)
+		, bNotifyGlobalBreakings(Parameters.bGenerateGlobalBreakingData)
+		, bNotifyGlobalRemovals(false)
+		, bNotifyGlobalCrumblings(Parameters.bGenerateGlobalCrumblingData)
+		, bGlobalCrumblingEventIncludesChildren(Parameters.bGenerateGlobalCrumblingChildrenData)
+		, bEnableStrainOnCollision(false)
+	{}
+
 	const FTransform& GetWorldTransform() const { return WorldTransform; }
 
 	void SetWorldTransform(const FTransform& InWorldTransform)
