Fix Dirty list in RigidSOA :
- Kinematic GC and Moving kinematic had duplicated particles
- Dynamic GC and Active list had duplicated particles
- Transient particles and Moving kinematic had duplicated particles

diff --git a/Engine/Source/Runtime/Experimental/Chaos/Private/Chaos/PBDRigidsEvolutionGBF.cpp b/Engine/Source/Runtime/Experimental/Chaos/Private/Chaos/PBDRigidsEvolutionGBF.cpp
--- a/Engine/Source/Runtime/Experimental/Chaos/Private/Chaos/PBDRigidsEvolutionGBF.cpp
+++ b/Engine/Source/Runtime/Experimental/Chaos/Private/Chaos/PBDRigidsEvolutionGBF.cpp
@@ -255,7 +255,7 @@ void CheckParticleViewsForDupes(FPBDRigidsSOAs& Particles)
 		CheckParticleViewForDupes(TEXT("NonDisabledView"), Particles.GetNonDisabledView());
 		CheckParticleViewForDupes(TEXT("NonDisabledClusteredView"), Particles.GetNonDisabledClusteredView());
 		CheckParticleViewForDupes(TEXT("ActiveParticlesView"), Particles.GetActiveParticlesView());
-		//CheckParticleViewForDupes(TEXT("DirtyParticlesView"), Particles.GetDirtyParticlesView());
+		CheckParticleViewForDupes(TEXT("DirtyParticlesView"), Particles.GetDirtyParticlesView());
 		CheckParticleViewForDupes(TEXT("AllParticlesView"), Particles.GetAllParticlesView());
 		CheckParticleViewForDupes(TEXT("ActiveKinematicParticlesView"), Particles.GetActiveKinematicParticlesView());
 		CheckParticleViewForDupes(TEXT("ActiveMovingKinematicParticlesView"), Particles.GetActiveMovingKinematicParticlesView());
diff --git a/Engine/Source/Runtime/Experimental/Chaos/Public/Chaos/PBDRigidsSOAs.h b/Engine/Source/Runtime/Experimental/Chaos/Public/Chaos/PBDRigidsSOAs.h
--- a/Engine/Source/Runtime/Experimental/Chaos/Public/Chaos/PBDRigidsSOAs.h
+++ b/Engine/Source/Runtime/Experimental/Chaos/Public/Chaos/PBDRigidsSOAs.h
@@ -522,7 +522,7 @@ public:
 						RemoveGeometryCollectionParticle(PBDRigidGC);
 						InsertGeometryCollectionParticle(PBDRigidGC);
 					}
-
+					check(PBDRigid->ObjectState() == EObjectStateType::Dynamic);
 					AddToActiveArray(PBDRigid);
 
 					if(!DeferUpdateViews)
@@ -722,6 +722,8 @@ public:
 		if (FPBDRigidParticleHandle* Rigid = Particle->CastToRigidParticle())
 		{
 			MovingKinematicsMapArray.Insert(Rigid);
+			// Remove from TransientDirtyMapArray to be sure not have have particle in both lists
+			TransientDirtyMapArray.Remove(Rigid);
 		}
 	}
 
@@ -1068,8 +1070,6 @@ private:
 			{
 				{&ActiveParticlesMapArray.GetArray()},
 				{&MovingKinematicsMapArray.GetArray()},
-				{&KinematicGeometryCollectionArray.GetArray()},
-				{&DynamicGeometryCollectionArray.GetArray()},
 				{&SleepingGeometryCollectionArray.GetArray()},
 				{&TransientDirtyMapArray.GetArray()}
 			};
