geometry collection navmesh export using root proxy should take into account root particle position

- also remove usage of bUpdateNavigationInTick since that's no longer relevant since it now only happens in the post physics sync callback

diff --git a/Engine/Source/Runtime/Experimental/GeometryCollectionEngine/Private/GeometryCollection/GeometryCollectionComponent.cpp b/Engine/Source/Runtime/Experimental/GeometryCollectionEngine/Private/GeometryCollection/GeometryCollectionComponent.cpp
--- a/Engine/Source/Runtime/Experimental/GeometryCollectionEngine/Private/GeometryCollection/GeometryCollectionComponent.cpp
+++ b/Engine/Source/Runtime/Experimental/GeometryCollectionEngine/Private/GeometryCollection/GeometryCollectionComponent.cpp
@@ -1058,8 +1058,10 @@ bool UGeometryCollectionComponent::DoCustomNavigableGeometryExport(FNavigableGeo
 			const TObjectPtr<UStaticMesh>& ProxyMesh = RestCollection->RootProxyData.ProxyMeshes[MeshIndex];
 			if (ProxyMesh != nullptr)
 			{
-				const FTransform& CompToWorld = GetComponentToWorld();
-				const FVector Scale3D = CompToWorld.GetScale3D();
+				const int32 RootIndex = GetRootIndex();
+				const FTransform& CompToWorld = GetComponentToWorld(); 
+				const FTransform FinalTransform = (DynamicCollection ? DynamicCollection->Transform[RootIndex] : FTransform::Identity) * CompToWorld;
+				const FVector Scale3D = FinalTransform.GetScale3D();
 				if (!Scale3D.IsZero())
 				{
 					if (const UNavCollisionBase* NavCollision = ProxyMesh->GetNavCollision())
@@ -1069,7 +1071,7 @@ bool UGeometryCollectionComponent::DoCustomNavigableGeometryExport(FNavigableGeo
 							continue;
 						}
 						
-						bHasData = NavCollision->ExportGeometry(CompToWorld, GeomExport) || bHasData;
+						bHasData = NavCollision->ExportGeometry(FinalTransform, GeomExport) || bHasData;
 					}
 				}
 			}
@@ -3353,7 +3355,7 @@ void UGeometryCollectionComponent::UpdateRenderSystemsIfNeeded(bool bDynamicColl
 
 void UGeometryCollectionComponent::UpdateNavigationDataIfNeeded(bool bDynamicCollectionDirty)
 {
-	if (bUpdateNavigationInTick && bDynamicCollectionDirty)
+	if (bDynamicCollectionDirty)
 	{
 		const UWorld* MyWorld = GetWorld();
 		if (MyWorld && MyWorld->IsGameWorld())
