IAS - support for downloading on demand container file(s) with UnrealPak

Usage: UnrealPak.exe -Download=<RemoteTocPath> -Directory=<LocalPath> -ServiceUrl=<URL> -Bucket=<BucketName>

#rnx

diff --git a/Engine/Source/Developer/IoStoreUtilities/Private/IoStoreUtilities.cpp b/Engine/Source/Developer/IoStoreUtilities/Private/IoStoreUtilities.cpp
--- a/Engine/Source/Developer/IoStoreUtilities/Private/IoStoreUtilities.cpp
+++ b/Engine/Source/Developer/IoStoreUtilities/Private/IoStoreUtilities.cpp
@@ -8506,3 +8506,23 @@ bool UploadIoStoreContainerFiles(const TCHAR* ContainerPathOrWildcard)
 
 	return UploadIoStoreContainerFiles(UploadParams.ConsumeValueOrDie(), ContainerFiles, KeyChain);
 }
+
+bool DownloadIoStoreContainerFiles(const TCHAR* TocPath)
+{
+	check(TocPath);
+
+	TIoStatusOr<UE::FIoStoreDownloadParams> Params = UE::FIoStoreDownloadParams::Parse(FCommandLine::Get());
+	if (Params.IsOk() == false)
+	{
+		UE_LOG(LogIoStore, Error, TEXT("Failed to download container file(s), reason '%s'"), *Params.Status().ToString());
+		return false;
+	}
+
+	FIoStatus Status = UE::DownloadContainerFiles(Params.ConsumeValueOrDie(), TocPath);
+	if (Status.IsOk() == false)
+	{
+		UE_LOG(LogIoStore, Error, TEXT("Failed to download container file(s), reason '%s'"), *Status.ToString());
+	}
+
+	return Status.IsOk();
+}
diff --git a/Engine/Source/Developer/IoStoreUtilities/Public/IoStoreUtilities.h b/Engine/Source/Developer/IoStoreUtilities/Public/IoStoreUtilities.h
--- a/Engine/Source/Developer/IoStoreUtilities/Public/IoStoreUtilities.h
+++ b/Engine/Source/Developer/IoStoreUtilities/Public/IoStoreUtilities.h
@@ -53,4 +53,6 @@ IOSTOREUTILITIES_API bool ProcessFilesFromIoStoreContainer(
 
 IOSTOREUTILITIES_API bool SignIoStoreContainer(const TCHAR* InContainerFilename, const FRSAKeyHandle InSigningKey);
 
-IOSTOREUTILITIES_API bool UploadIoStoreContainerFiles(const TCHAR* Params); 
+IOSTOREUTILITIES_API bool UploadIoStoreContainerFiles(const TCHAR* ContainerPathOrWildcard); 
+
+IOSTOREUTILITIES_API bool DownloadIoStoreContainerFiles(const TCHAR* TocPath);
diff --git a/Engine/Source/Developer/PakFileUtilities/Private/PakFileUtilities.cpp b/Engine/Source/Developer/PakFileUtilities/Private/PakFileUtilities.cpp
--- a/Engine/Source/Developer/PakFileUtilities/Private/PakFileUtilities.cpp
+++ b/Engine/Source/Developer/PakFileUtilities/Private/PakFileUtilities.cpp
@@ -5462,6 +5462,11 @@ bool ExecuteUnrealPak(const TCHAR* CmdLine)
 		{
 			return UploadIoStoreContainerFiles(*IoStoreArg) == 0;
 		}
+
+		if (FParse::Value(CmdLine, TEXT("-Download="), IoStoreArg))
+		{
+			return DownloadIoStoreContainerFiles(*IoStoreArg) == 0;
+		}
 	}
 
 	// Parse all the non-option arguments from the command line
diff --git a/Engine/Source/Runtime/Experimental/IoStoreOnDemand/Private/IoStoreOnDemand.cpp b/Engine/Source/Runtime/Experimental/IoStoreOnDemand/Private/IoStoreOnDemand.cpp
--- a/Engine/Source/Runtime/Experimental/IoStoreOnDemand/Private/IoStoreOnDemand.cpp
+++ b/Engine/Source/Runtime/Experimental/IoStoreOnDemand/Private/IoStoreOnDemand.cpp
@@ -21,13 +21,18 @@
 #include "String/LexFromString.h"
 
 #if (PLATFORM_DESKTOP && (IS_PROGRAM || WITH_EDITOR))
+#include "Algo/Sort.h"
 #include "Async/Async.h"
+#include "HAL/PlatformFileManager.h"
+#include "IO/IoStore.h"
 #include "Misc/FileHelper.h"
 #include "S3/S3Client.h"
 #endif // (PLATFORM_DESKTOP && (IS_PROGRAM || WITH_EDITOR))
 
 DEFINE_LOG_CATEGORY(LogIas);
 
+UE_DISABLE_OPTIMIZATION
+
 namespace UE::IO::Private
 {
 
@@ -906,6 +911,316 @@ TIoStatusOr<FIoStoreUploadResult> UploadContainerFiles(
 	
 	return UploadResult;
 }
+
+////////////////////////////////////////////////////////////////////////////////
+TIoStatusOr<FIoStoreDownloadParams> FIoStoreDownloadParams::Parse(const TCHAR* CommandLine)
+{
+	FIoStoreDownloadParams Params;
+
+	if (!FParse::Value(CommandLine, TEXT("Directory="), Params.Directory))
+	{
+		FIoStatus(EIoErrorCode::InvalidParameter, TEXT("Invalid output directory"));
+	}
+
+	if (!FParse::Value(CommandLine, TEXT("Bucket="), Params.Bucket))
+	{
+		FIoStatus(EIoErrorCode::InvalidParameter, TEXT("Invalid bucket name"));
+	}
+
+	FParse::Value(CommandLine, TEXT("ServiceUrl="), Params.ServiceUrl);
+	FParse::Value(CommandLine, TEXT("Region="), Params.Region);
+	FParse::Value(CommandLine, TEXT("AccessKey="), Params.AccessKey);
+	FParse::Value(CommandLine, TEXT("SecretKey="), Params.SecretKey);
+	FParse::Value(CommandLine, TEXT("SessionToken="), Params.SessionToken);
+	FParse::Value(CommandLine, TEXT("CredentialsFile="), Params.CredentialsFile);
+	FParse::Value(CommandLine, TEXT("CredentialsFileKeyName="), Params.CredentialsFileKeyName);
+
+	if (!Params.AccessKey.IsEmpty() && Params.SecretKey.IsEmpty())
+	{
+		return FIoStatus(EIoErrorCode::InvalidParameter, TEXT("Invalid secret key"));
+	}
+	else if (Params.AccessKey.IsEmpty() && !Params.SecretKey.IsEmpty())
+	{
+		return FIoStatus(EIoErrorCode::InvalidParameter, TEXT("Invalid access key"));
+	}
+
+	if (!Params.CredentialsFile.IsEmpty() && Params.CredentialsFileKeyName.IsEmpty())
+	{
+		return FIoStatus(EIoErrorCode::InvalidParameter, TEXT("Invalid credential file key name"));
+	}
+
+	if (Params.ServiceUrl.IsEmpty() && Params.Region.IsEmpty())
+	{
+		return FIoStatus(EIoErrorCode::InvalidParameter, TEXT("Service URL or AWS region needs to be specified"));
+	}
+
+	return Params;
+}
+
+FIoStatus DownloadContainerFiles(const FIoStoreDownloadParams& DownloadParams, const FString& TocPath)
+{
+	struct FContainerStats
+	{
+		uint64 TocEntryCount = 0;
+		uint64 TocRawSize = 0;
+		uint64 RawSize = 0;
+		uint64 CompressedSize = 0;
+		EIoContainerFlags ContainerFlags = EIoContainerFlags::None;
+		FName CompressionMethod;
+	};
+	TMap<FString, FContainerStats> ContainerSummary;
+
+	const double StartTime = FPlatformTime::Seconds();
+
+	FS3ClientConfig Config;
+	Config.ServiceUrl = DownloadParams.ServiceUrl;
+	Config.Region = DownloadParams.Region;
+	
+	FS3ClientCredentials Credentials;
+	if (DownloadParams.CredentialsFile.IsEmpty() == false)
+	{
+		UE_LOG(LogIas, Display, TEXT("Loading credentials file '%s'"), *DownloadParams.CredentialsFile);
+		FS3CredentialsProfileStore CredentialsStore = FS3CredentialsProfileStore::FromFile(DownloadParams.CredentialsFile);
+		if (CredentialsStore.TryGetCredentials(DownloadParams.CredentialsFileKeyName, Credentials) == false)
+		{
+			return FIoStatus(EIoErrorCode::InvalidParameter, TEXT("Failed to find valid credentials in credentials file"));
+		}
+		else
+		{
+			UE_LOG(LogIas, Display, TEXT("Found credentials for '%s'"), *DownloadParams.CredentialsFileKeyName);
+		}
+	}
+	else
+	{
+		Credentials = FS3ClientCredentials(DownloadParams.AccessKey, DownloadParams.SecretKey, DownloadParams.SessionToken);
+	}
+
+	FS3Client Client(Config, Credentials);
+
+	UE_LOG(LogIas, Display, TEXT("Fetching TOC '%s/%s/%s'"), *Client.GetConfig().ServiceUrl, *DownloadParams.Bucket, *TocPath);
+	FS3GetObjectResponse TocResponse = Client.GetObject(FS3GetObjectRequest
+	{
+		DownloadParams.Bucket,
+		TocPath
+	});
+
+	if (TocResponse.IsOk() == false)
+	{
+		return FIoStatus(EIoErrorCode::ReadError, TEXT("Failed to fetch TOC"));
+	}
+
+	FOnDemandToc OnDemandToc;
+	if (UE::LoadFromCompactBinary(FCbFieldView(TocResponse.Body.GetData()), OnDemandToc) == false)
+	{
+		return FIoStatus(EIoErrorCode::ReadError, TEXT("Failed to load on demand TOC"));
+	}
+
+	FStringView BucketPrefix;
+	{
+		FStringView TocPathView = TocPath;
+		int32 Index = INDEX_NONE;
+		if (TocPathView.FindLastChar(TEXT('/'), Index))
+		{
+			BucketPrefix = TocPathView.Left(Index);
+		}
+	}
+
+	for (const FOnDemandTocContainerEntry& ContainerEntry : OnDemandToc.Containers)
+	{
+		TStringBuilder<256> FileTocKey;
+		if (BucketPrefix.IsEmpty() == false)
+		{
+			FileTocKey << BucketPrefix << TEXT("/");
+		}
+		FileTocKey << ContainerEntry.UTocHash << TEXT(".utoc");
+
+		UE_LOG(LogIas, Display, TEXT("Fetching '%s/%s/%s'"), *Client.GetConfig().ServiceUrl, *DownloadParams.Bucket, FileTocKey.ToString());
+		FS3GetObjectResponse Response = Client.GetObject(FS3GetObjectRequest
+		{
+			DownloadParams.Bucket,
+			FString(FileTocKey).ToLower()
+		});
+
+		if (Response.IsOk() == false)
+		{
+			return FIoStatus(EIoErrorCode::ReadError, TEXT("Failed to load container .utoc file"));
+		}
+
+		const FString UTocPath = FString::Printf(TEXT("%s/%s.utoc"), *DownloadParams.Directory, *ContainerEntry.ContainerName);
+		const FString UCasPath = FPaths::ChangeExtension(UTocPath, TEXT(".ucas")); 
+		FContainerStats& ContainerStats = ContainerSummary.FindOrAdd(ContainerEntry.ContainerName);
+		ContainerStats.TocRawSize = Response.Body.GetSize();
+
+		if (TUniquePtr<FArchive> TocFile(IFileManager::Get().CreateFileWriter(*UTocPath)); TocFile.IsValid())
+		{
+			UE_LOG(LogIas, Display, TEXT("Writing '%s'"), *UTocPath);
+			TocFile->Serialize((void*)Response.Body.GetData(), Response.Body.GetSize());
+		}
+		else
+		{
+			return FIoStatus(EIoErrorCode::WriteError, TEXT("Failed to write container .utoc file"));
+		}
+
+		FIoStoreTocResource FileToc;
+		if (FIoStatus Status = FIoStoreTocResource::Read(*UTocPath, EIoStoreTocReadOptions::ReadAll, FileToc); Status.IsOk() == false)
+		{
+			return Status;
+		}
+
+		const int32 TocEntryCount = int32(FileToc.Header.TocEntryCount);
+		const uint64 CompressionBlockSize = FileToc.Header.CompressionBlockSize;
+		ContainerStats.TocEntryCount = TocEntryCount;
+		ContainerStats.ContainerFlags = FileToc.Header.ContainerFlags; 
+
+		TMap<FIoChunkId, FIoHash> ChunkHashes;
+		for (const FOnDemandTocEntry& Entry : ContainerEntry.Entries)
+		{
+			ChunkHashes.Add(Entry.ChunkId, Entry.Hash);
+			ContainerStats.RawSize += Entry.RawSize;
+		}
+
+		TArray<int32> SortedIndices;
+		for (int32 Idx = 0; Idx < TocEntryCount; ++Idx)
+		{
+			SortedIndices.Add(Idx);
+		}
+
+		Algo::Sort(SortedIndices, [&FileToc](int32 Lhs, int32 Rhs)
+		{
+			return FileToc.ChunkOffsetLengths[Lhs].GetOffset() < FileToc.ChunkOffsetLengths[Rhs].GetOffset();
+		});
+
+		FString ChunksRelativePath = BucketPrefix.IsEmpty()
+			? FString::Printf(TEXT("IoChunksV%u"), OnDemandToc.Header.ChunkVersion)
+			: FString::Printf(TEXT("%s/IoChunksV%u"), *FString(BucketPrefix), OnDemandToc.Header.ChunkVersion);
+
+		TUniquePtr<FArchive> CasFile(IFileManager::Get().CreateFileWriter(*UCasPath));
+		TArray<uint8> PaddingBuffer;
+
+		for (int32 Idx = 0; Idx < TocEntryCount; ++Idx)
+		{
+			const int32 SortedIdx = SortedIndices[Idx];
+			const FIoChunkId& ChunkId = FileToc.ChunkIds[SortedIdx];
+			const FIoOffsetAndLength& OffsetLength = FileToc.ChunkOffsetLengths[SortedIdx];
+			const FIoHash ChunkHash = ChunkHashes.FindChecked(ChunkId);
+			const FString HashString = LexToString(ChunkHash);
+			const int32 FirstBlockIndex = int32(OffsetLength.GetOffset() / CompressionBlockSize);
+			const int32 LastBlockIndex = int32((Align(OffsetLength.GetOffset() + OffsetLength.GetLength(), CompressionBlockSize) - 1) / CompressionBlockSize);
+			const FIoStoreTocCompressedBlockEntry& FirstBlock = FileToc.CompressionBlocks[FirstBlockIndex];
+			const FName CompressionMethod = FileToc.CompressionMethods[FirstBlock.GetCompressionMethodIndex()];
+
+			if (CompressionMethod.IsNone() == false && ContainerStats.CompressionMethod.IsNone())
+			{
+				ContainerStats.CompressionMethod = CompressionMethod;
+			}
+
+			TStringBuilder<256> ChunkKey;
+			ChunkKey << ChunksRelativePath
+				<< TEXT("/") << HashString.Left(2)
+				<< TEXT("/") << HashString
+				<< TEXT(".iochunk");
+
+			UE_LOG(LogIas, Display, TEXT("Fetching '%s/%s/%s'"), *Client.GetConfig().ServiceUrl, *DownloadParams.Bucket, ChunkKey.ToString());
+			FS3GetObjectResponse ChunkResponse = Client.GetObject(FS3GetObjectRequest
+			{
+				DownloadParams.Bucket,
+				FString(ChunkKey).ToLower()
+			});
+
+			if (ChunkResponse.IsOk() == false)
+			{
+				return FIoStatus(EIoErrorCode::ReadError, TEXT("Failed to fetch chunk"));
+			}
+
+			const uint64 Padding = FirstBlock.GetOffset() - CasFile->Tell();
+			if (Padding > PaddingBuffer.Num())
+			{
+				PaddingBuffer.SetNumZeroed(int32(Padding));
+			}
+
+			if (Padding > 0)
+			{
+				CasFile->Serialize(PaddingBuffer.GetData(), Padding);
+			}
+
+			UE_LOG(LogIas, Display, TEXT("Serializing chunk %d/%d '%s' -> '%s' (%llu B)"),
+				Idx + 1, TocEntryCount, *HashString, *LexToString(ChunkId), ChunkResponse.Body.GetSize());
+
+			check(CasFile->Tell() == FirstBlock.GetOffset());
+			CasFile->Serialize((void*)ChunkResponse.Body.GetData(), ChunkResponse.Body.GetSize());
+		}
+
+		ContainerStats.CompressedSize = CasFile->Tell();
+	}
+
+	{
+		const double Duration = FPlatformTime::Seconds() - StartTime;
+
+		UE_LOG(LogIas, Display, TEXT(""));
+		UE_LOG(LogIas, Display, TEXT("---------------------------------------- Download Summary --------------------------------------"));
+		UE_LOG(LogIas, Display, TEXT("%-40s: %s"), TEXT("Service URL"), *DownloadParams.ServiceUrl);
+		UE_LOG(LogIas, Display, TEXT("%-40s: %s"), TEXT("Bucket"), *DownloadParams.Bucket);
+		UE_LOG(LogIas, Display, TEXT("%-40s: %s"), TEXT("TOC"), *TocPath);
+		UE_LOG(LogIas, Display, TEXT("%-40s: %.2lf second(s)"), TEXT("Duration"), Duration);
+		UE_LOG(LogIas, Display, TEXT(""));
+
+		UE_LOG(LogIas, Display, TEXT("%-30s %10s %15s %15s %15s %25s"),
+			TEXT("Container"), TEXT("Flags"), TEXT("TOC Size (KB)"), TEXT("TOC Entries"), TEXT("Size (MB)"), TEXT("Compressed (MB)"));
+		UE_LOG(LogIas, Display, TEXT("-------------------------------------------------------------------------------------------------------------------------"));
+		
+		FContainerStats TotalStats;
+		for (const TTuple<FString, FContainerStats>& Kv : ContainerSummary)
+		{
+			const FString& ContainerName = Kv.Key;
+			const FContainerStats& Stats = Kv.Value;
+
+			FString CompressionInfo = TEXT("-");
+
+			if (Stats.CompressionMethod != NAME_None)
+			{
+				double Procentage = (double(Stats.RawSize - Stats.CompressedSize) / double(Stats.RawSize)) * 100.0;
+				CompressionInfo = FString::Printf(TEXT("%.2lf (%.2lf%% %s)"),
+					(double)Stats.CompressedSize / 1024.0 / 1024.0,
+					Procentage,
+					*Stats.CompressionMethod.ToString());
+			}
+
+			FString ContainerSettings = FString::Printf(TEXT("%s/%s/%s/%s/%s"),
+				EnumHasAnyFlags(Stats.ContainerFlags, EIoContainerFlags::Compressed) ? TEXT("C") : TEXT("-"),
+				EnumHasAnyFlags(Stats.ContainerFlags, EIoContainerFlags::Encrypted) ? TEXT("E") : TEXT("-"),
+				EnumHasAnyFlags(Stats.ContainerFlags, EIoContainerFlags::Signed) ? TEXT("S") : TEXT("-"),
+				EnumHasAnyFlags(Stats.ContainerFlags, EIoContainerFlags::Indexed) ? TEXT("I") : TEXT("-"),
+				EnumHasAnyFlags(Stats.ContainerFlags, EIoContainerFlags::OnDemand) ? TEXT("O") : TEXT("-"));
+
+			UE_LOG(LogIas, Display, TEXT("%-30s %10s %15.2lf %15llu %15.2lf %25s"),
+				*ContainerName,
+				*ContainerSettings,
+				(double)Stats.TocRawSize / 1024.0,
+				Stats.TocEntryCount,
+				(double)Stats.RawSize / 1024.0 / 1024.0,
+				*CompressionInfo);
+
+			TotalStats.TocEntryCount += Stats.TocEntryCount;
+			TotalStats.TocRawSize += Stats.TocRawSize;
+			TotalStats.RawSize += Stats.RawSize;
+			TotalStats.CompressedSize += Stats.CompressedSize;
+		}
+		UE_LOG(LogIas, Display, TEXT("-------------------------------------------------------------------------------------------------------------------------"));
+		UE_LOG(LogIas, Display, TEXT("%-30s %10s %15.2lf %15llu %15.2lf %25.2lf "),
+			TEXT("Total"),
+			TEXT(""),
+			(double)TotalStats.TocRawSize / 1024.0,
+			TotalStats.TocEntryCount,
+			(double)TotalStats.RawSize / 1024.0 / 1024.0,
+			(double)TotalStats.CompressedSize / 1024.0 / 1024.0);
+
+		UE_LOG(LogIas, Display, TEXT(""));
+		UE_LOG(LogIas, Display, TEXT("** Flags: (C)ompressed / (E)ncrypted / (S)igned) / (I)ndexed) / (O)nDemand **"));
+		UE_LOG(LogIas, Display, TEXT(""));
+	}
+
+	return FIoStatus::Ok;
+}
 #endif // (PLATFORM_DESKTOP && (IS_PROGRAM || WITH_EDITOR))
 
 } // namespace UE
@@ -1040,3 +1355,5 @@ void FIoStoreOnDemandModule::StartupModule()
 void FIoStoreOnDemandModule::ShutdownModule()
 {
 }
+
+UE_ENABLE_OPTIMIZATION
diff --git a/Engine/Source/Runtime/Experimental/IoStoreOnDemand/Public/IO/IoStoreOnDemand.h b/Engine/Source/Runtime/Experimental/IoStoreOnDemand/Public/IO/IoStoreOnDemand.h
--- a/Engine/Source/Runtime/Experimental/IoStoreOnDemand/Public/IO/IoStoreOnDemand.h
+++ b/Engine/Source/Runtime/Experimental/IoStoreOnDemand/Public/IO/IoStoreOnDemand.h
@@ -136,6 +136,25 @@ UE_API TIoStatusOr<FIoStoreUploadResult> UploadContainerFiles(
 	TConstArrayView<FString> ContainerFiles,
 	const TMap<FGuid, FAES::FAESKey>& EncryptionKeys);
 
+////////////////////////////////////////////////////////////////////////////////
+struct FIoStoreDownloadParams
+{
+	FString Directory;
+	FString ServiceUrl;
+	FString Bucket;
+	FString Region; 
+	FString AccessKey;
+	FString SecretKey;
+	FString SessionToken;
+	FString CredentialsFile;
+	FString CredentialsFileKeyName;
+	int32 MaxConcurrentDownloads = 16;
+	
+	UE_API static TIoStatusOr<FIoStoreDownloadParams> Parse(const TCHAR* CommandLine);
+};
+
+UE_API FIoStatus DownloadContainerFiles(const FIoStoreDownloadParams& DownloadParams, const FString& TocPath);
+
 #endif // (IS_PROGRAM || WITH_EDITOR)
 
 } // namespace UE
