Build fix - removed pragma optimization

#rnx

diff --git a/Engine/Source/Runtime/Experimental/IoStoreOnDemand/Private/IoStoreOnDemand.cpp b/Engine/Source/Runtime/Experimental/IoStoreOnDemand/Private/IoStoreOnDemand.cpp
--- a/Engine/Source/Runtime/Experimental/IoStoreOnDemand/Private/IoStoreOnDemand.cpp
+++ b/Engine/Source/Runtime/Experimental/IoStoreOnDemand/Private/IoStoreOnDemand.cpp
@@ -31,8 +31,6 @@
 
 DEFINE_LOG_CATEGORY(LogIas);
 
-UE_DISABLE_OPTIMIZATION
-
 namespace UE::IO::Private
 {
 
@@ -1355,5 +1353,3 @@ void FIoStoreOnDemandModule::StartupModule()
 void FIoStoreOnDemandModule::ShutdownModule()
 {
 }
-
-UE_ENABLE_OPTIMIZATION
