Chaos - AABB Transform optimization (disabled by cvar for now)

diff --git a/Engine/Source/Programs/HeadlessChaos/Private/HeadlessChaosTestAABB.cpp b/Engine/Source/Programs/HeadlessChaos/Private/HeadlessChaosTestAABB.cpp
--- a/Engine/Source/Programs/HeadlessChaos/Private/HeadlessChaosTestAABB.cpp
+++ b/Engine/Source/Programs/HeadlessChaos/Private/HeadlessChaosTestAABB.cpp
@@ -10,6 +10,120 @@ namespace ChaosTest
 {
 	using namespace Chaos;
 
+	void TestBounds(const FAABB3& A, const FAABB3& B, const FReal Eps = 1.e-6)
+	{
+		EXPECT_NEAR(A.Min().X, B.Min().X, Eps);
+		EXPECT_NEAR(A.Min().Y, B.Min().Y, Eps);
+		EXPECT_NEAR(A.Min().Z, B.Min().Z, Eps);
+		EXPECT_NEAR(A.Max().X, B.Max().X, Eps);
+		EXPECT_NEAR(A.Max().Y, B.Max().Y, Eps);
+		EXPECT_NEAR(A.Max().Z, B.Max().Z, Eps);
+	}
+
+	// Transformed AABB by transforming all box verts
+	// Used for reference when comparing to current TAABB::TransformedAABB results
+	FAABB3 ReferenceTransformAABB(const FAABB3& AABB, const FRigidTransform3& Transform)
+	{
+		FVec3 Extents = AABB.Extents();
+		FAABB3 NewAABB = FAABB3::EmptyAABB();
+		NewAABB.GrowToInclude(Transform.TransformPosition(FVec3(AABB.Min())));
+		NewAABB.GrowToInclude(Transform.TransformPosition(FVec3(AABB.Max())));
+
+		for (int32 AxisIndex = 0; AxisIndex < 3; ++AxisIndex)
+		{
+			NewAABB.GrowToInclude(Transform.TransformPosition(FVec3(AABB.Min() + FVec3::AxisVector(AxisIndex) * Extents)));
+			NewAABB.GrowToInclude(Transform.TransformPosition(FVec3(AABB.Max() - FVec3::AxisVector(AxisIndex) * Extents)));
+		}
+
+		return NewAABB;
+	}
+
+	// Inverse Transformed AABB by transforming all box verts
+	// Used for reference when comparing to current TAABB::InverseTransformedAABB results
+	FAABB3 ReferenceInverseTransformAABB(const FAABB3& AABB, const FRigidTransform3& Transform)
+	{
+		FVec3 Extents = AABB.Extents();
+		FAABB3 NewAABB = FAABB3::EmptyAABB();
+		NewAABB.GrowToInclude(Transform.InverseTransformPosition(FVec3(AABB.Min())));
+		NewAABB.GrowToInclude(Transform.InverseTransformPosition(FVec3(AABB.Max())));
+
+		for (int32 AxisIndex = 0; AxisIndex < 3; ++AxisIndex)
+		{
+			NewAABB.GrowToInclude(Transform.InverseTransformPosition(FVec3(AABB.Min() + FVec3::AxisVector(AxisIndex) * Extents)));
+			NewAABB.GrowToInclude(Transform.InverseTransformPosition(FVec3(AABB.Max() - FVec3::AxisVector(AxisIndex) * Extents)));
+		}
+
+		return NewAABB;
+	}
+
+	// Check that identity transform does not affect an AABB
+	GTEST_TEST(AABBTests, TestTransformAABB_Unit)
+	{
+		const FAABB3 Bounds = FAABB3(FVec3(-1,-2,-3), FVec3(1,2,3));
+		const FRigidTransform3 Transform = FRigidTransform3::Identity;
+		
+		const FAABB3 TransformedBounds = Bounds.TransformedAABB(Transform);
+		TestBounds(TransformedBounds, Bounds);
+
+		const FAABB3 InverseTransformedBounds = Bounds.InverseTransformedAABB(Transform);
+		TestBounds(InverseTransformedBounds, Bounds);
+	}
+
+	GTEST_TEST(AABBTests, TestTransformAABB_Translate)
+	{
+		const FAABB3 Bounds = FAABB3(FVec3(-3, -2, -1), FVec3(1, 2, 3));
+
+		const FVec3 Offset = FVec3(2);
+		const FRigidTransform3 Transform = FRigidTransform3(Offset, FRotation3::FromIdentity(), FVec3(1));
+
+		const FAABB3 TransformedBounds = Bounds.TransformedAABB(Transform);
+		TestBounds(TransformedBounds, FAABB3(Offset + Bounds.Min(), Offset + Bounds.Max()));
+
+		const FAABB3 InverseTransformedBounds = Bounds.InverseTransformedAABB(Transform);
+		TestBounds(InverseTransformedBounds, FAABB3(Bounds.Min() - Offset, Bounds.Max() - Offset));
+	}
+
+	GTEST_TEST(AABBTests, TestTransformAABB_Scale)
+	{
+		const FAABB3 Bounds = FAABB3(FVec3(-3, -2, -1), FVec3(1, 2, 3));
+
+		const FVec3 Scale = FVec3(2);
+		const FRigidTransform3 Transform = FRigidTransform3(FVec3(0), FRotation3::FromIdentity(), Scale);
+
+		const FAABB3 TransformedBounds = Bounds.TransformedAABB(Transform);
+		TestBounds(TransformedBounds, FAABB3(Scale * Bounds.Min(), Scale * Bounds.Max()));
+
+		const FAABB3 InverseTransformedBounds = Bounds.InverseTransformedAABB(Transform);
+		TestBounds(InverseTransformedBounds, FAABB3(Bounds.Min() / Scale, Bounds.Max() / Scale));
+	}
+
+	GTEST_TEST(AABBTests, TestTransformAABB_NegativeScale)
+	{
+		const FAABB3 Bounds = FAABB3(FVec3(-3, -2, -1), FVec3(1, 2, 3));
+
+		const FVec3 Scale = FVec3(-2);
+		const FRigidTransform3 Transform = FRigidTransform3(FVec3(0), FRotation3::FromIdentity(), Scale);
+
+		const FAABB3 TransformedBounds = Bounds.TransformedAABB(Transform);
+		TestBounds(TransformedBounds, FAABB3(Scale * Bounds.Max(), Scale * Bounds.Min()));
+
+		const FAABB3 InverseTransformedBounds = Bounds.InverseTransformedAABB(Transform);
+		TestBounds(InverseTransformedBounds, FAABB3(Bounds.Max() / Scale, Bounds.Min() / Scale));
+	}
+
+	GTEST_TEST(AABBTests, TestTransformAABB_ScaleRotateTranslate)
+	{
+		const FAABB3 Bounds = FAABB3(FVec3(-3, -2, -1), FVec3(1, 2, 3));
+		const FRigidTransform3 Transform = FRigidTransform3(FVec3(4,7,-3), FRotation3::FromAxisAngle(FVec3(4,-3,5).GetUnsafeNormal(), FMath::DegreesToRadians(36.5)), FVec3(1, -4, 17));
+
+		const FAABB3 TransformedBounds = Bounds.TransformedAABB(Transform);
+		TestBounds(TransformedBounds, ReferenceTransformAABB(Bounds, Transform));
+
+		const FAABB3 InverseTransformedBounds = Bounds.InverseTransformedAABB(Transform);
+		TestBounds(InverseTransformedBounds, ReferenceInverseTransformAABB(Bounds, Transform));
+	}
+
+
 	// A test that was failing for 8-wide ISPC architectures
 	GTEST_TEST(AABBTests, TestTransformAABB_ISPC8WideFail)
 	{
diff --git a/Engine/Source/Runtime/Experimental/Chaos/Private/Chaos/AABB.cpp b/Engine/Source/Runtime/Experimental/Chaos/Private/Chaos/AABB.cpp
--- a/Engine/Source/Runtime/Experimental/Chaos/Private/Chaos/AABB.cpp
+++ b/Engine/Source/Runtime/Experimental/Chaos/Private/Chaos/AABB.cpp
@@ -22,6 +22,9 @@ static bool bChaos_AABBTransform_ISPC_Enabled = CHAOS_AABB_TRANSFORM_ISPC_ENABLE
 static FAutoConsoleVariableRef CVarChaosAABBTransformISPCEnabled(TEXT("p.Chaos.AABBTransform.ISPC"), bChaos_AABBTransform_ISPC_Enabled, TEXT("Whether to use ISPC optimizations when computing AABB transforms"));
 #endif
 
+static bool bChaos_AABBTransform_Optimized = false;
+static FAutoConsoleVariableRef CVarChaosAABBTransformOptimized(TEXT("p.Chaos.AABBTransform.Optimized"), bChaos_AABBTransform_Optimized, TEXT("Whether to use optimized AABB transform"));
+
 namespace Chaos
 {
 
@@ -270,6 +273,94 @@ bool TAABB<T, d>::Raycast(const TVector<FReal, d>& StartPoint, const TVector<FRe
 	return true;
 }
 
+template<typename T, typename U>
+inline TAABB<T, 3> TransformedAABBHelper2(const TAABB<T, 3>& AABB, const TRigidTransform<U, 3>& SpaceTransform)
+{
+	// Full and empty bounds do not transform
+	if (AABB.IsFull() || AABB.IsEmpty())
+	{
+		return AABB;
+	}
+
+	// NOTE: This is required for TAABB3<float> and TAABB3<double> with either float or double transform
+	const TVec3<T> Translation = TVec3<T>(SpaceTransform.GetTranslation());
+	const TRotation3<T> Rotation = TRotation3<T>(SpaceTransform.GetRotation());
+	const TVec3<T> Scale = TVec3<T>(SpaceTransform.GetScale3D());
+
+	// Center-relative scaled verts
+	// NOTE: Scale may be negative, but it does not impact the bounds calculation
+	const TVec3<T> Extent = T(0.5) * Scale * (AABB.Max() - AABB.Min());
+	const TVec3<T> Vert0 = TVec3<T>(Extent.X, Extent.Y, Extent.Z);
+	const TVec3<T> Vert1 = TVec3<T>(Extent.X, Extent.Y, -Extent.Z);
+	const TVec3<T> Vert2 = TVec3<T>(Extent.X, -Extent.Y, Extent.Z);
+	const TVec3<T> Vert3 = TVec3<T>(Extent.X, -Extent.Y, -Extent.Z);
+
+	// Rotated center-relative scaled verts
+	const TVec3<T> RVert0 = Rotation * Vert0;
+	const TVec3<T> RVert1 = Rotation * Vert1;
+	const TVec3<T> RVert2 = Rotation * Vert2;
+	const TVec3<T> RVert3 = Rotation * Vert3;
+
+	// Max rotated scaled center-relative extent
+	TVec3<T> RExtent = TVec3<T>(0);
+	RExtent = TVec3<T>::Max(RExtent, RVert0.GetAbs());
+	RExtent = TVec3<T>::Max(RExtent, RVert1.GetAbs());
+	RExtent = TVec3<T>::Max(RExtent, RVert2.GetAbs());
+	RExtent = TVec3<T>::Max(RExtent, RVert3.GetAbs());
+
+	// Transformed center
+	// NOTE: This is where positive/negative scales matters
+	const TVec3<T> Center = T(0.5) * (AABB.Min() + AABB.Max());
+	const TVec3<T> TCenter = Translation + Rotation * (Scale * Center);
+
+	// Transformed bounds
+	return TAABB<T, 3>(TCenter - RExtent, TCenter + RExtent);
+}
+
+template<typename T, typename U>
+inline TAABB<T, 3> InverseTransformedAABBHelper2(const TAABB<T, 3>& AABB, const TRigidTransform<U, 3>& SpaceTransform)
+{
+	// Full and empty bounds do not transform
+	if (AABB.IsFull() || AABB.IsEmpty())
+	{
+		return AABB;
+	}
+
+	// NOTE: This is required for TAABB3<float> and TAABB3<double> with either float or double transform
+	const TVec3<T> InvTranslation = -TVec3<T>(SpaceTransform.GetTranslation());
+	const TRotation3<T> InvRotation = TRotation3<T>(SpaceTransform.GetRotation()).Inverse();
+	const TVec3<T> InvScale = TVec3<T>(1) / TVec3<T>(SpaceTransform.GetScale3D());
+
+	// Center-relative scaled verts
+	const TVec3<T> Extent = T(0.5) * (AABB.Max() - AABB.Min());
+	const TVec3<T> Vert0 = TVec3<T>(Extent.X, Extent.Y, Extent.Z);
+	const TVec3<T> Vert1 = TVec3<T>(Extent.X, Extent.Y, -Extent.Z);
+	const TVec3<T> Vert2 = TVec3<T>(Extent.X, -Extent.Y, Extent.Z);
+	const TVec3<T> Vert3 = TVec3<T>(Extent.X, -Extent.Y, -Extent.Z);
+
+	// Rotated center-relative scaled verts
+	// Note: Scale may be negative but it does not affect the bounds calculation
+	const TVec3<T> RVert0 = InvScale * (InvRotation * Vert0);
+	const TVec3<T> RVert1 = InvScale * (InvRotation * Vert1);
+	const TVec3<T> RVert2 = InvScale * (InvRotation * Vert2);
+	const TVec3<T> RVert3 = InvScale * (InvRotation * Vert3);
+
+	// Max rotated scaled center-relative extent
+	TVec3<T> RExtent = TVec3<T>(0);
+	RExtent = TVec3<T>::Max(RExtent, RVert0.GetAbs());
+	RExtent = TVec3<T>::Max(RExtent, RVert1.GetAbs());
+	RExtent = TVec3<T>::Max(RExtent, RVert2.GetAbs());
+	RExtent = TVec3<T>::Max(RExtent, RVert3.GetAbs());
+
+	// Transformed center
+	// NOTE: This is where positive/negative scales matters
+	const TVec3<T> Center = T(0.5) * (AABB.Min() + AABB.Max());
+	const TVec3<T> TCenter =  InvScale * (InvRotation * (InvTranslation + Center));
+
+	// Transformed bounds
+	return TAABB<T, 3>(TCenter - RExtent, TCenter + RExtent);
+}
+
 template<typename T>
 inline TAABB<T, 3> TransformedAABBHelper(const TAABB<T, 3>& AABB, const FMatrix44& SpaceTransform)
 {
@@ -330,6 +421,56 @@ inline TAABB<FReal, 3> TransformedAABBHelperISPC(const TAABB<FReal, 3>& AABB, co
 
 inline TAABB<Chaos::FRealSingle, 3> TransformedAABBHelperISPC(const TAABB<Chaos::FRealSingle, 3>& AABB, const FTransform& SpaceTransform)
 {
+	// Full and empty bounds do not transform
+	if (AABB.IsFull() || AABB.IsEmpty())
+	{
+		return AABB;
+	}
+
+	check(bRealTypeCompatibleWithISPC);
+#if INTEL_ISPC
+	static_assert(sizeof(ispc::FVector3f) == sizeof(Chaos::TVector<Chaos::FRealSingle, 3>), "sizeof(ispc::FVector3f) != sizeof(Chaos::TVector<Chaos::FRealSingle, 3>)");
+
+	TVector<Chaos::FRealSingle, 3> NewMin, NewMax;
+	ispc::TransformedAABBMixed((const ispc::FTransform&)SpaceTransform, (const ispc::FVector3f&)AABB.Min(), (const ispc::FVector3f&)AABB.Max(), (ispc::FVector3f&)NewMin, (ispc::FVector3f&)NewMax);
+
+	TAABB<Chaos::FRealSingle, 3> NewAABB(NewMin, NewMax);
+	return NewAABB;
+#else
+	check(false);
+	return TAABB<Chaos::FRealSingle, 3>::EmptyAABB();
+#endif
+}
+
+inline TAABB<FReal, 3> TransformedAABBHelperISPC2(const TAABB<FReal, 3>& AABB, const FTransform& SpaceTransform)
+{
+	// Full and empty bounds do not transform
+	if (AABB.IsFull() || AABB.IsEmpty())
+	{
+		return AABB;
+	}
+
+	check(bRealTypeCompatibleWithISPC);
+#if INTEL_ISPC
+	TVector<Chaos::FReal, 3> NewMin, NewMax;
+	ispc::TransformedAABB2((const ispc::FTransform&)SpaceTransform, (const ispc::FVector&)AABB.Min(), (const ispc::FVector&)AABB.Max(), (ispc::FVector&)NewMin, (ispc::FVector&)NewMax);
+
+	TAABB<FReal, 3> NewAABB(NewMin, NewMax);
+	return NewAABB;
+#else
+	check(false);
+	return TAABB<FReal, 3>::EmptyAABB();
+#endif
+}
+
+inline TAABB<Chaos::FRealSingle, 3> TransformedAABBHelperISPC2(const TAABB<Chaos::FRealSingle, 3>& AABB, const FTransform& SpaceTransform)
+{
+	// Full and empty bounds do not transform
+	if (AABB.IsFull() || AABB.IsEmpty())
+	{
+		return AABB;
+	}
+
 	check(bRealTypeCompatibleWithISPC);
 #if INTEL_ISPC
 	static_assert(sizeof(ispc::FVector3f) == sizeof(Chaos::TVector<Chaos::FRealSingle, 3>), "sizeof(ispc::FVector3f) != sizeof(Chaos::TVector<Chaos::FRealSingle, 3>)");
@@ -348,13 +489,27 @@ inline TAABB<Chaos::FRealSingle, 3> TransformedAABBHelperISPC(const TAABB<Chaos:
 template<typename T, int d>
 TAABB<T, d> TAABB<T, d>::TransformedAABB(const Chaos::TRigidTransform<FReal, 3>& SpaceTransform) const
 {
-	if (bRealTypeCompatibleWithISPC && INTEL_ISPC && bChaos_AABBTransform_ISPC_Enabled)
+	if (bRealTypeCompatibleWithISPC && INTEL_ISPC && bChaos_AABBTransform_ISPC_Enabled )
 	{
-		return TransformedAABBHelperISPC(*this, SpaceTransform);
+		if (bChaos_AABBTransform_Optimized)
+		{
+			return TransformedAABBHelperISPC2(*this, SpaceTransform);
+		}
+		else
+		{
+			return TransformedAABBHelperISPC(*this, SpaceTransform);
+		}
 	}
-	else
+	else 
 	{
-		return TransformedAABBHelper<T>(*this, SpaceTransform.ToMatrixWithScale());
+		if (bChaos_AABBTransform_Optimized)
+		{
+			return TransformedAABBHelper2<T>(*this, SpaceTransform);
+		}
+		else
+		{
+			return TransformedAABBHelper<T>(*this, SpaceTransform.ToMatrixWithScale());
+		}
 	}
 }
 
@@ -376,36 +531,57 @@ TAABB<T, d> TAABB<T, d>::TransformedAABB(const FTransform& SpaceTransform) const
 {
 	if (bRealTypeCompatibleWithISPC && INTEL_ISPC && bChaos_AABBTransform_ISPC_Enabled)
 	{
-		return TransformedAABBHelperISPC(*this, SpaceTransform);
+		if (bChaos_AABBTransform_Optimized)
+		{
+			return TransformedAABBHelperISPC2(*this, SpaceTransform);
+		}
+		else
+		{
+			return TransformedAABBHelperISPC(*this, SpaceTransform);
+		}
 	}
-	else
+	else 
 	{
-		return TransformedAABBHelper<T>(*this, SpaceTransform.ToMatrixWithScale());
+		if (bChaos_AABBTransform_Optimized)
+		{
+			return TransformedAABBHelper2<T>(*this, Chaos::TRigidTransform<FReal, 3>(SpaceTransform));
+		}
+		else
+		{
+			return TransformedAABBHelper<T>(*this, SpaceTransform.ToMatrixWithScale());
+		}
 	}
 }
 
 template<typename T, int d>
 TAABB<T, d> TAABB<T, d>::InverseTransformedAABB(const FRigidTransform3& SpaceTransform) const
 {
-	// Full and empty bounds do not transform
-	if (IsFull() || IsEmpty())
+	if (bChaos_AABBTransform_Optimized)
 	{
-		return *this;
+		return InverseTransformedAABBHelper2<T>(*this, SpaceTransform);
 	}
+	else
+	{
+		// Full and empty bounds do not transform
+		if (IsFull() || IsEmpty())
+		{
+			return *this;
+		}
 
-	TVector<T, d> CurrentExtents = Extents();
-	int32 Idx = 0;
-	const TVector<T, d> MinToNewSpace = SpaceTransform.InverseTransformPosition(FVector(Min()));
-	TAABB<T, d> NewAABB(MinToNewSpace, MinToNewSpace);
-	NewAABB.GrowToInclude(SpaceTransform.InverseTransformPosition(FVector(Max())));
+		TVector<T, d> CurrentExtents = Extents();
+		int32 Idx = 0;
+		const TVector<T, d> MinToNewSpace = SpaceTransform.InverseTransformPosition(FVector(Min()));
+		TAABB<T, d> NewAABB(MinToNewSpace, MinToNewSpace);
+		NewAABB.GrowToInclude(SpaceTransform.InverseTransformPosition(FVector(Max())));
 
-	for (int32 j = 0; j < d; ++j)
-	{
-		NewAABB.GrowToInclude(SpaceTransform.InverseTransformPosition(FVector(Min() + TVector<T, d>::AxisVector(j) * CurrentExtents)));
-		NewAABB.GrowToInclude(SpaceTransform.InverseTransformPosition(FVector(Max() - TVector<T, d>::AxisVector(j) * CurrentExtents)));
-	}
+		for (int32 j = 0; j < d; ++j)
+		{
+			NewAABB.GrowToInclude(SpaceTransform.InverseTransformPosition(FVector(Min() + TVector<T, d>::AxisVector(j) * CurrentExtents)));
+			NewAABB.GrowToInclude(SpaceTransform.InverseTransformPosition(FVector(Max() - TVector<T, d>::AxisVector(j) * CurrentExtents)));
+		}
 
-	return NewAABB;
+		return NewAABB;
+	}
 }
 }
 
diff --git a/Engine/Source/Runtime/Experimental/Chaos/Private/Chaos/AABB.ispc b/Engine/Source/Runtime/Experimental/Chaos/Private/Chaos/AABB.ispc
--- a/Engine/Source/Runtime/Experimental/Chaos/Private/Chaos/AABB.ispc
+++ b/Engine/Source/Runtime/Experimental/Chaos/Private/Chaos/AABB.ispc
@@ -5,6 +5,16 @@
 #include "Math/Vector.isph"
 #include "Math/Transform.isph"
 
+static const uniform struct FVector4d DoubleNegateXVector = { {-1.0d, 1.0d, 1.0d, 1.0d} };
+static const uniform struct FVector4d DoubleNegateYVector = { {1.0d, -1.0d, 1.0d, 1.0d} };
+static const uniform struct FVector4d DoubleNegateZVector = { {1.0d, 1.0d, -1.0d, 1.0d} };
+
+#define NegateXVector DoubleNegateXVector
+#define NegateYVector DoubleNegateYVector
+#define NegateZVector DoubleNegateZVector
+#define NegateWVector DoubleNegateWVector
+
+
 inline static void GrowToInclude(uniform FVector4 &Min, uniform FVector4 &Max, const uniform FVector4& V)
 {
 	Min = VectorMin(Min, V);
@@ -59,3 +69,39 @@ export void TransformedAABBMixed(const uniform FTransform &SpaceTransform, const
 	NewMin = SetVector(MinToNewSpace4f);
 	NewMax = SetVector(MaxToNewSpace4f);
 }
+
+export void TransformedAABB2(const uniform FTransform &SpaceTransform, const uniform FVector &MMin, const uniform FVector &MMax, uniform FVector &NewMin, uniform FVector &NewMax)
+{
+	const uniform FVector4 Min = SetVector4(MMin.V[0], MMin.V[1], MMin.V[2], ZERO);
+	const uniform FVector4 Max = SetVector4(MMax.V[0], MMax.V[1], MMax.V[2], ZERO);
+
+	// Center-relative scaled verts
+	// NOTE: Scale may be negative, but it does not impact the bounds calculation
+	const uniform FVector4 Extent = 0.5 * SpaceTransform.Scale3D * (Max - Min);
+	const uniform FVector4 Vert0 = Extent;										// (Extent.V[0],  Extent.V[1],  Extent.V[2])
+	const uniform FVector4 Vert1 = Extent * NegateZVector;						// (Extent.V[0],  Extent.V[1], -Extent.V[2])
+	const uniform FVector4 Vert2 = Extent * NegateYVector;						// (Extent.V[0], -Extent.V[1],  Extent.V[2])
+	const uniform FVector4 Vert3 = Extent * (NegateYVector * NegateZVector);	// (Extent.V[0], -Extent.V[1], -Extent.V[2])
+
+	// Rotated center-relative scaled verts
+	const uniform FVector4 RVert0 = VectorQuaternionRotateVector(SpaceTransform.Rotation, Vert0);
+	const uniform FVector4 RVert1 = VectorQuaternionRotateVector(SpaceTransform.Rotation, Vert1);
+	const uniform FVector4 RVert2 = VectorQuaternionRotateVector(SpaceTransform.Rotation, Vert2);
+	const uniform FVector4 RVert3 = VectorQuaternionRotateVector(SpaceTransform.Rotation, Vert3);
+
+	// Max rotated scaled center-relative extent
+	uniform FVector4 RExtent = VectorZero;
+	RExtent = VectorMax(RExtent, VectorAbs(RVert0));
+	RExtent = VectorMax(RExtent, VectorAbs(RVert1));
+	RExtent = VectorMax(RExtent, VectorAbs(RVert2));
+	RExtent = VectorMax(RExtent, VectorAbs(RVert3));
+
+	// Transformed center
+	// NOTE: This is where positive/negative scales matters
+	const uniform FVector4 Center = 0.5 * (Min + Max);
+	const uniform FVector4 TCenter = SpaceTransform.Translation + SpaceTransform.Rotation * (SpaceTransform.Scale3D * Center);
+
+	// Transformed bounds
+	NewMin = SetVector(TCenter - RExtent);
+	NewMax = SetVector(TCenter + RExtent);
+}
