Iris - Load right type of vector.

#rnx

diff --git a/Engine/Source/Runtime/Experimental/Iris/Core/Private/Iris/ReplicationSystem/Prioritization/LocationBasedNetObjectPrioritizer.cpp b/Engine/Source/Runtime/Experimental/Iris/Core/Private/Iris/ReplicationSystem/Prioritization/LocationBasedNetObjectPrioritizer.cpp
--- a/Engine/Source/Runtime/Experimental/Iris/Core/Private/Iris/ReplicationSystem/Prioritization/LocationBasedNetObjectPrioritizer.cpp
+++ b/Engine/Source/Runtime/Experimental/Iris/Core/Private/Iris/ReplicationSystem/Prioritization/LocationBasedNetObjectPrioritizer.cpp
@@ -111,6 +111,6 @@ void ULocationBasedNetObjectPrioritizer::UpdateLocation(const uint32 ObjectIndex
 		TArrayView<const UE::Net::FReplicationInstanceProtocol::FFragmentData> FragmentDatas = MakeArrayView(InstanceProtocol->FragmentData, InstanceProtocol->FragmentCount);
 		const UE::Net::FReplicationInstanceProtocol::FFragmentData& FragmentData = FragmentDatas[Info.GetLocationStateIndex()];
 		const uint8* LocationOffset = FragmentData.ExternalSrcBuffer + Info.GetLocationStateOffset();
-		SetLocation(Info, VectorLoadFloat3_W0(reinterpret_cast<const float*>(LocationOffset)));
+		SetLocation(Info, VectorLoadFloat3_W0(reinterpret_cast<const FVector*>(LocationOffset)));
 	}
 }
