Added send/recv buffer pool parameters

#rnx

diff --git a/Engine/Source/Runtime/Experimental/IoStoreOnDemand/Private/CoreHttp/Client.h b/Engine/Source/Runtime/Experimental/IoStoreOnDemand/Private/CoreHttp/Client.h
--- a/Engine/Source/Runtime/Experimental/IoStoreOnDemand/Private/CoreHttp/Client.h
+++ b/Engine/Source/Runtime/Experimental/IoStoreOnDemand/Private/CoreHttp/Client.h
@@ -60,6 +60,8 @@ public:
 	{
 		int32				SetHostFromUrl(FAnsiStringView Url);
 		uint32				ConnectionCount;
+		int32				SendBufSize = -1;
+		int32				RecvBufSize = -1;
 		struct {
 			FAnsiStringView	Name;
 			uint32			Port = 80;
diff --git a/Engine/Source/Runtime/Experimental/IoStoreOnDemand/Private/CoreHttp/Http.cpp b/Engine/Source/Runtime/Experimental/IoStoreOnDemand/Private/CoreHttp/Http.cpp
--- a/Engine/Source/Runtime/Experimental/IoStoreOnDemand/Private/CoreHttp/Http.cpp
+++ b/Engine/Source/Runtime/Experimental/IoStoreOnDemand/Private/CoreHttp/Http.cpp
@@ -580,6 +580,7 @@ class FSocketPool
 {
 public:
 	enum class EState : uint8 { Unresolved, Busy, Resolved, Error };
+	enum class EDirection : uint8 { Send, Recv };
 
 					FSocketPool(FAnsiStringView InHostName, uint32 InPort, uint32 InMaxLeases);
 					~FSocketPool();
@@ -587,6 +588,8 @@ public:
 	bool			LeaseSocket(SocketType& Out);
 	void			ReturnLease(SocketType Socket);
 	bool			AddIpAddress(uint32 Address);
+	void			SetBufferSize(EDirection Dir, int32 Size);
+	int32			GetBufferSize(EDirection Dir) const;
 	uint32			GetIpAddress() const	{ return IpAddresses[0]; }
 	FAnsiStringView	GetHostName() const		{ return HostName; }
 	uint32			GetPort() const			{ return Port; }
@@ -596,6 +599,8 @@ public:
 private:
 	FAnsiStringView	HostName;
 	uint32			IpAddresses[4] = {};
+	int16			SendBufKb;
+	int16			RecvBufKb;
 	uint16			Port;
 	uint8			LeaseCount = 0;
 	uint8			MaxLeases : 6;
@@ -679,6 +684,18 @@ bool FSocketPool::AddIpAddress(uint32 Address)
 	return false;
 }
 
+////////////////////////////////////////////////////////////////////////////////
+void FSocketPool::SetBufferSize(EDirection Dir, int32 Size)
+{
+	((Dir == EDirection::Send) ? SendBufKb : RecvBufKb) = Size >> 10;
+}
+
+////////////////////////////////////////////////////////////////////////////////
+int32 FSocketPool::GetBufferSize(EDirection Dir) const
+{
+	return int32((Dir == EDirection::Send) ? SendBufKb : RecvBufKb) << 10;
+}
+
 
 
 ////////////////////////////////////////////////////////////////////////////////
@@ -724,6 +741,9 @@ FConnectionPool::FConnectionPool(const FParams& Params)
 		Params.Host.Port,
 		Params.ConnectionCount
 	);
+	Internal->SetBufferSize(FSocketPool::EDirection::Send, Params.SendBufSize);
+	Internal->SetBufferSize(FSocketPool::EDirection::Recv, Params.RecvBufSize);
+
 	Ptr = Internal;
 }
 
@@ -1391,6 +1411,17 @@ int32 FEventLoopInternal::DoConnect(FActivity* Activity)
 		}
 	}
 
+	// Adjust socket send and recv buffer sizes
+	if (int32 OptValue = Activity->Pool->GetBufferSize(FSocketPool::EDirection::Send); OptValue >= 0)
+	{
+		setsockopt(Candidate, SOL_SOCKET, SO_SNDBUF, &(char&)OptValue, sizeof(OptValue));
+	}
+	if (int32 OptValue = Activity->Pool->GetBufferSize(FSocketPool::EDirection::Recv); OptValue >= 0)
+	{
+		setsockopt(Candidate, SOL_SOCKET, SO_RCVBUF, &(char&)OptValue, sizeof(OptValue));
+	}
+
+	// connect
 	uint32 IpAddress = Activity->Pool->GetIpAddress();
 	if (IpAddress == 0)
 	{
@@ -1398,7 +1429,6 @@ int32 FEventLoopInternal::DoConnect(FActivity* Activity)
 		return -1;
 	}
 
-	// connect
 	sockaddr_in AddrInet = { sizeof(sockaddr_in) };
 	AddrInet.sin_family = AF_INET;
 	AddrInet.sin_port = htons(Activity->Pool->GetPort());
