Activity recvs return the number of bytes they received

#rnx

diff --git a/Engine/Source/Runtime/Experimental/IoStoreOnDemand/Private/CoreHttp/Http.cpp b/Engine/Source/Runtime/Experimental/IoStoreOnDemand/Private/CoreHttp/Http.cpp
--- a/Engine/Source/Runtime/Experimental/IoStoreOnDemand/Private/CoreHttp/Http.cpp
+++ b/Engine/Source/Runtime/Experimental/IoStoreOnDemand/Private/CoreHttp/Http.cpp
@@ -1249,12 +1249,18 @@ static uint64 ReadyCheck(FActivity** Activities, uint32 Num, uint32 TimeoutMs)
 class FEventLoopInternal
 {
 public:
+	struct FResult
+	{
+		int32	Result;
+		uint32	Size;
+	};
+
 	static int32	DoResolve(FActivity* Activity);
 	static int32	DoConnect(FActivity* Activity);
 	static int32	DoSend(FActivity* Activity);
 	static int32	DoRecvMessage(FActivity* Activity);
-	static int32	DoRecvContent(FActivity* Activity);
-	static int32	DoRecvStream(FActivity* Activity);
+	static FResult	DoRecvContent(FActivity* Activity, uint32 MaxRecvSize);
+	static FResult	DoRecvStream(FActivity* Activity, uint32 MaxRecSize);
 	static int32	DoRecvDone(FActivity* Activity);
 	static void		Cancel(FActivity* Activity);
 };
@@ -1717,13 +1723,17 @@ int32 FEventLoopInternal::DoRecvMessage(FActivity* Activity)
 }
 
 ////////////////////////////////////////////////////////////////////////////////
-int32 FEventLoopInternal::DoRecvContent(FActivity* Activity)
+FEventLoopInternal::FResult FEventLoopInternal::DoRecvContent(
+	FActivity* Activity,
+	uint32 MaxRecvSize)
 {
 	TRACE_CPUPROFILER_EVENT_SCOPE(CoreHttp::DoRecvContent);
 
 	auto& Response = *(FResponseInternal*)(Activity->Buffer.GetData());
 
 	FMutableMemoryView DestView = Response.Dest->GetMutableView();
+
+	uint32 RecvSize = 0;
 	while (true)
 	{
 		uint32 Size = (Response.ContentLength - Activity->StateParam);
@@ -1732,6 +1742,12 @@ int32 FEventLoopInternal::DoRecvContent(FActivity* Activity)
 			break;
 		}
 
+		Size = FMath::Min(Size, MaxRecvSize - RecvSize);
+		if (Size == 0)
+		{
+			return { 1, RecvSize };
+		}
+
 		char* Cursor = (char*)(DestView.GetData()) + Activity->StateParam;
 		int32 Result = recv(Activity->Socket, Cursor, Size, MsgFlagType(0));
 		if (Result < 0)
@@ -1739,26 +1755,27 @@ int32 FEventLoopInternal::DoRecvContent(FActivity* Activity)
 			if (IsSocketResult(EWOULDBLOCK))
 			{
 				Activity->SocketWait = FActivity::EWait::Read;
-				return 1;
+				return { 1, RecvSize };
 			}
 
 			Activity_SetError(Activity, "Socket error while receiving content");
-			return -1;
+			return { -1 };
 		}
 
 		if (Result == 0 && (Activity->StateParam != Response.ContentLength))
 		{
 			Activity_SetError(Activity, "ATH0.RecvContent");
-			return -1;
+			return { -1 };
 		}
 
 		Activity->StateParam += Result;
+		RecvSize += Result;
 	}
 
 	FLatencyInjector::Begin(FLatencyInjector::EType::Network, Activity->StateParam);
 
 	Activity->State = FActivity::EState::RecvDone;
-	return 0;
+	return { 0, RecvSize };
 }
 
 ////////////////////////////////////////////////////////////////////////////////
@@ -1778,12 +1795,14 @@ int32 FEventLoopInternal::DoRecvDone(FActivity* Activity)
 }
 
 ////////////////////////////////////////////////////////////////////////////////
-int32 FEventLoopInternal::DoRecvStream(FActivity* Activity)
+FEventLoopInternal::FResult FEventLoopInternal::DoRecvStream(
+	FActivity* Activity,
+	uint32 RecvSize)
 {
 	TRACE_CPUPROFILER_EVENT_SCOPE(CoreHttp::DoRecvStream);
 
 	check(false); // not implemented yet
-	return -1;
+	return { -1 };
 }
 
 ////////////////////////////////////////////////////////////////////////////////
@@ -2048,17 +2067,22 @@ uint32 FEventLoop::Tick(uint32 PollTimeoutMs)
 				break;
 
 		case FActivity::EState::RecvContent:
-		case FActivity::EState::RecvStream:
+		case FActivity::EState::RecvStream: {
+			decltype(FEventLoopInternal::DoRecvContent)* Handler;
 			if (Activity->State == FActivity::EState::RecvContent)
 			{
-				Result = FEventLoopInternal::DoRecvContent(Activity);
+				Handler = &FEventLoopInternal::DoRecvContent;
 			}
 			else
 			{
-				Result = FEventLoopInternal::DoRecvStream(Activity);
+				Handler = &FEventLoopInternal::DoRecvStream;
 			}
+
+			auto [ResultInner, RecvSize] = Handler(Activity, ~0u);
+			Result = ResultInner;
 			if (Result)
 				break;
+		}
 
 		case FActivity::EState::RecvDone:
 			Result = FEventLoopInternal::DoRecvDone(Activity);
