IAS: Add ias.NoCache debug command line parameter to disable the file cache and run with streaming only.

#rnx

diff --git a/Engine/Source/Runtime/Experimental/IoStoreOnDemand/Private/IoStoreOnDemand.cpp b/Engine/Source/Runtime/Experimental/IoStoreOnDemand/Private/IoStoreOnDemand.cpp
--- a/Engine/Source/Runtime/Experimental/IoStoreOnDemand/Private/IoStoreOnDemand.cpp
+++ b/Engine/Source/Runtime/Experimental/IoStoreOnDemand/Private/IoStoreOnDemand.cpp
@@ -126,6 +126,10 @@ static FFileIoCacheConfig GetFileIoCacheConfig(const TCHAR* CommandLine)
 	{
 		Ret.DropCache = true;
 	}
+	if (FParse::Param(CommandLine, TEXT("Ias.NoCache")))
+	{
+		Ret.DiskQuota = 0;
+	}
 #endif
 
 	return Ret;
@@ -1335,6 +1339,10 @@ void FIoStoreOnDemandModule::StartupModule()
 	{
 		Cache = MakeShareable(MakeFileIoCache(Config).Release());
 	}
+	else
+	{
+		UE_LOG(LogIas, Log, TEXT("File cache disabled. Streaming only."));
+	}
 
 	TSharedPtr<UE::IOnDemandIoDispatcherBackend> Backend = UE::MakeOnDemandIoDispatcherBackend(Cache);
 	Backend->Mount(Endpoint);
