Output more useful LogIas information before the URL

#rnx

diff --git a/Engine/Source/Runtime/Experimental/IoStoreOnDemand/Private/OnDemandIoDispatcherBackend.cpp b/Engine/Source/Runtime/Experimental/IoStoreOnDemand/Private/OnDemandIoDispatcherBackend.cpp
--- a/Engine/Source/Runtime/Experimental/IoStoreOnDemand/Private/OnDemandIoDispatcherBackend.cpp
+++ b/Engine/Source/Runtime/Experimental/IoStoreOnDemand/Private/OnDemandIoDispatcherBackend.cpp
@@ -59,6 +59,13 @@ static FAutoConsoleVariableRef CVar_IoDispatcherMaxHttpPollTimeoutMs (
 namespace UE::IO::Private
 {
 
+///////////////////////////////////////////////////////////////////////////////
+static void LogHttpResult(const TCHAR* Url, uint32 StatusCode, uint32 Duration, uint32 Size, uint32 Offset, const char* Memo="")
+{
+	Size >>= 10;
+	UE_LOG(LogIas, VeryVerbose, TEXT("http:%u %ums %uKiB %u %S %s"), StatusCode, Duration, Size, Offset, Memo, Url);
+};
+
 using namespace UE::Tasks;
 
 ///////////////////////////////////////////////////////////////////////////////
@@ -356,32 +363,24 @@ void FHttpClient::Issue(FAnsiStringView Url, FIoReadCallback&& Callback, FIoOffs
 			else if (FTicketStatus::EId::Content == Status.GetId())
 			{
 				const uint64 Duration = (uint64)FPlatformTime::ToMilliseconds64(FPlatformTime::Cycles64() - StartTime);
+				const FIoBuffer& Content = Status.GetContent(); 
+
+				LogHttpResult(*Url, StatusCode, Duration, Content.GetSize(), Offset);
 
 				const bool bSuccessful = StatusCode > 199 && StatusCode < 300;
-				if (const FIoBuffer& Content = Status.GetContent(); bSuccessful && Content.GetSize() > 0)
+				if (bSuccessful && Content.GetSize() > 0)
 				{
-					UE_LOG(LogIas, VeryVerbose, TEXT("%s"),
-						*WriteToString<256>(TEXT("HTTP GET - "), Url, TEXT("("), StatusCode, TEXT(" "),
-						Duration, TEXT("ms "), Offset, TEXT(" Offset "), Content.GetSize(), TEXT(" Bytes)")));
-
 					Callback(Content);
 				}
 				else
 				{
-					UE_LOG(LogIas, Warning, TEXT("%s"),
-						*WriteToString<256>(TEXT("HTTP GET - "), Url, TEXT("("), StatusCode, TEXT(" "),
-						Duration, TEXT("ms "), Offset, TEXT(" Offset 0 Bytes) : Invalid Content : "), StatusCode));
-					
 					Callback(FIoStatus(EIoErrorCode::NotFound, TEXTVIEW("Invalid Content")));
 				}
 			}
 			else if (FTicketStatus::EId::Error == Status.GetId())
 			{
 				const uint64 Duration = (uint64)FPlatformTime::ToMilliseconds64(FPlatformTime::Cycles64() - StartTime);
-				UE_LOG(LogIas, Warning, TEXT("%s"),
-					*WriteToString<256>(TEXT("HTTP GET - "), Url, TEXT(" ("), StatusCode, TEXT(" "),
-					Duration, TEXTVIEW("ms "), Offset, TEXT(" Offset) : "), Status.GetErrorReason()));
-
+				LogHttpResult(*Url, StatusCode, Duration, 0, Offset, Status.GetErrorReason());
 				Callback(FIoStatus(EIoErrorCode::ReadError, TEXTVIEW("HTTP Error")));
 			}
 		};
@@ -1314,8 +1313,7 @@ uint32 FOnDemandIoBackend::Run()
 						{
 							TAnsiStringBuilder<256> Url;
 							ChunkRequest->Params.GetUrl(Url);
-							UE_LOG(LogIas, Error, TEXT("%s"),
-								*WriteToString<256>(TEXT("HTTP FAILED - "), Url, TEXT(" ("), ChunkRequest->Params.ChunkRange.GetOffset(), TEXT(" Offset)")));
+							LogHttpResult(StringCast<TCHAR>(*Url).Get(), -1, 0, 0, ChunkRequest->Params.ChunkRange.GetOffset(), "HTTP FAILED");
 							Stats.OnHttpError();
 						}
 
