Dataflow:
- Added support for structs' child properties as inputs/outputs using property chains.
- Updated the property offset calculations to work with any properties from a struct property hierarchy.
- Replaced all TFieldIterator<FProperty> with FPropertyValueIterator to iterate through the property hierarchies.
- Made the FDataflowConnection and FContextCacheElementproperty pointers const to allow use with the FPropertyValueIterator API.
- Renamed FDataflowOutput::SetPassthroughOffsetAddress() to SetPassthroughOffset to prevent confusion with property addresses and offsets.

#rnx

diff --git a/Engine/Source/Runtime/Experimental/Dataflow/Core/Private/Dataflow/DataflowConnection.cpp b/Engine/Source/Runtime/Experimental/Dataflow/Core/Private/Dataflow/DataflowConnection.cpp
--- a/Engine/Source/Runtime/Experimental/Dataflow/Core/Private/Dataflow/DataflowConnection.cpp
+++ b/Engine/Source/Runtime/Experimental/Dataflow/Core/Private/Dataflow/DataflowConnection.cpp
@@ -7,7 +7,7 @@
 
 #include UE_INLINE_GENERATED_CPP_BY_NAME(DataflowConnection)
 
-FDataflowConnection::FDataflowConnection(Dataflow::FPin::EDirection InDirection, FName InType, FName InName, FDataflowNode* InOwningNode, FProperty* InProperty, FGuid InGuid)
+FDataflowConnection::FDataflowConnection(Dataflow::FPin::EDirection InDirection, FName InType, FName InName, FDataflowNode* InOwningNode, const FProperty* InProperty, FGuid InGuid)
 	: Direction(InDirection)
 	, Type(InType)
 	, Name(InName)
@@ -16,11 +16,11 @@ FDataflowConnection::FDataflowConnection(Dataflow::FPin::EDirection InDirection,
 	, Guid(InGuid)
 {}
 
-int32 FDataflowConnection::GetOffset() const
+uint32 FDataflowConnection::GetOffset() const
 {
-	if (ensure(Property != nullptr))
+	if (ensure(OwningNode))
 	{
-		return Property->GetOffset_ForInternal();
+		return OwningNode->GetPropertyOffset(Name);
 	}
 	return INDEX_NONE;
 }
diff --git a/Engine/Source/Runtime/Experimental/Dataflow/Core/Private/Dataflow/DataflowNode.cpp b/Engine/Source/Runtime/Experimental/Dataflow/Core/Private/Dataflow/DataflowNode.cpp
--- a/Engine/Source/Runtime/Experimental/Dataflow/Core/Private/Dataflow/DataflowNode.cpp
+++ b/Engine/Source/Runtime/Experimental/Dataflow/Core/Private/Dataflow/DataflowNode.cpp
@@ -32,13 +32,14 @@ void FDataflowNode::AddInput(FDataflowInput* InPtr)
 			ensureMsgf(!In->GetName().IsEqual(InPtr->GetName()), TEXT("Add Input Failed: Existing Node input already defined with name (%s)"), *InPtr->GetName().ToString());
 		}
 
-		if (!Inputs.Contains(InPtr->Property->GetOffset_ForInternal()))
+		const uint32 PropertyOffset = InPtr->GetOffset();
+		if (!Inputs.Contains(PropertyOffset))
 		{
-			Inputs.Add(InPtr->Property->GetOffset_ForInternal(), InPtr);
+			Inputs.Add(PropertyOffset, InPtr);
 		}
 		else
 		{
-			Inputs[InPtr->Property->GetOffset_ForInternal()] = InPtr;
+			Inputs[PropertyOffset] = InPtr;
 		}
 	}
 }
@@ -116,13 +117,14 @@ void FDataflowNode::AddOutput(FDataflowOutput* InPtr)
 			ensureMsgf(!Out->GetName().IsEqual(InPtr->GetName()), TEXT("Add Output Failed: Existing Node output already defined with name (%s)"), *InPtr->GetName().ToString());
 		}
 
-		if (!Outputs.Contains(InPtr->Property->GetOffset_ForInternal()))
+		const uint32 PropertyOffset = InPtr->GetOffset();
+		if (!Outputs.Contains(PropertyOffset))
 		{
-			Outputs.Add(InPtr->Property->GetOffset_ForInternal(), InPtr);
+			Outputs.Add(PropertyOffset, InPtr);
 		}
 		else
 		{
-			Outputs[InPtr->Property->GetOffset_ForInternal()] = InPtr;
+			Outputs[PropertyOffset] = InPtr;
 		}
 	}
 }
@@ -238,106 +240,191 @@ void FDataflowNode::Invalidate(const Dataflow::FTimestamp& InModifiedTimestamp)
 	}
 }
 
-void FDataflowNode::RegisterInputConnection(const void* InProperty)
+const FProperty* FDataflowNode::FindProperty(const UStruct* Struct, const void* InProperty, const FName& PropertyName, TArray<const FProperty*>* OutPropertyChain) const
+{
+	const FProperty* Property = nullptr;
+	for (FPropertyValueIterator PropertyIt(FProperty::StaticClass(), Struct, this); PropertyIt; ++PropertyIt)
+	{
+		if (InProperty == PropertyIt.Value() && (PropertyName == NAME_None || PropertyName == PropertyIt.Key()->GetName()))
+		{
+			Property = PropertyIt.Key();
+			if (OutPropertyChain)
+			{
+				PropertyIt.GetPropertyChain(*OutPropertyChain);
+			}
+			break;
+		}
+	}
+	return Property;
+}
+
+const FProperty* FDataflowNode::FindProperty(const UStruct* Struct, const FName& PropertyFullName, TArray<const FProperty*>* OutPropertyChain) const
+{
+	const FProperty* Property = nullptr;
+	for (FPropertyValueIterator PropertyIt(FProperty::StaticClass(), Struct, this); PropertyIt; ++PropertyIt)
+	{
+		TArray<const FProperty*> PropertyChain;
+		PropertyIt.GetPropertyChain(PropertyChain);
+		if (GetPropertyFullName(PropertyChain) == PropertyFullName)
+		{
+			Property = PropertyIt.Key();
+			if (OutPropertyChain)
+			{
+				*OutPropertyChain = MoveTemp(PropertyChain);
+			}
+			break;
+		}
+	}
+	return Property;
+}
+
+uint32 FDataflowNode::GetPropertyOffset(const TArray<const FProperty*>& PropertyChain)
+{
+	uint32 Offset = 0;
+	for (const FProperty* const Property : PropertyChain)
+	{
+		Offset += (uint32)Property->GetOffset_ForInternal();
+	}
+	return Offset;
+}
+
+uint32 FDataflowNode::GetPropertyOffset(const FName& PropertyFullName) const
+{
+	uint32 Offset = 0;
+	if (const TUniquePtr<const FStructOnScope> ScriptOnStruct =
+		TUniquePtr<FStructOnScope>(const_cast<FDataflowNode*>(this)->NewStructOnScope()))  // The mutable Struct Memory is not accessed here, allowing for the const_cast and keeping this method const
+	{
+		if (const UStruct* const Struct = ScriptOnStruct->GetStruct())
+		{
+			TArray<const FProperty*> PropertyChain;
+			FindProperty(Struct, PropertyFullName, &PropertyChain);
+			Offset = GetPropertyOffset(PropertyChain);
+		}
+	}
+	return Offset;
+}
+
+FName FDataflowNode::GetPropertyFullName(const TArray<const FProperty*>& PropertyChain)
+{
+	FString PropertyFullName;
+	for (const FProperty* const Property : PropertyChain)
+	{
+		const FString PropertyName = Property->GetName();
+		PropertyFullName = PropertyFullName.IsEmpty() ?
+			PropertyName :
+			FString::Format(TEXT("{0}.{1}"), { PropertyName, PropertyFullName });
+	}
+	return FName(*PropertyFullName);
+}
+
+FText FDataflowNode::GetPropertyDisplayNameText(const TArray<const FProperty*>& PropertyChain)
+{
+#if WITH_EDITORONLY_DATA  // GetDisplayNameText() is only available if WITH_EDITORONLY_DATA
+	FText PropertyText;
+	for (const FProperty* const Property : PropertyChain)
+	{
+		static const FTextFormat TextFormat(NSLOCTEXT("DataflowNode", "PropertyDisplayNameTextConcatenator", "{0}.{1}"));
+		PropertyText = PropertyText.IsEmpty() ?
+			Property->GetDisplayNameText() :
+			FText::Format(TextFormat, Property->GetDisplayNameText(), PropertyText);
+	}
+	return PropertyText;
+#else
+	return FText::FromName(GetPropertyFullName(PropertyChain));
+#endif
+}
+
+void FDataflowNode::RegisterInputConnection(const void* InProperty, const FName& PropertyName)
 {
 	if (TUniquePtr<FStructOnScope> ScriptOnStruct = TUniquePtr<FStructOnScope>(NewStructOnScope()))
 	{
 		if (const UStruct* Struct = ScriptOnStruct->GetStruct())
 		{
-			for (TFieldIterator<FProperty> PropertyIt(Struct); PropertyIt; ++PropertyIt)
+			TArray<const FProperty*> PropertyChain;
+			const FProperty* const Property =
+				FindProperty(Struct, InProperty, PropertyName, &PropertyChain);
+			if (ensure(Property && PropertyChain.Num()))
 			{
-				FProperty* Property = *PropertyIt;
-				size_t RealAddress = (size_t)this + Property->GetOffset_ForInternal();
-				if (RealAddress == (size_t)InProperty)
-				{
-					FName PropName(Property->GetName());
-					FName PropType(Property->GetCPPType());
-					AddInput(new FDataflowInput({ PropType, PropName, this, Property }));
-				}
+				const FName PropName(GetPropertyFullName(PropertyChain));
+				const FName PropType(Property->GetCPPType());
+				AddInput(new FDataflowInput({ PropType, PropName, this, Property }));
 			}
 		}
 	}
 }
 
-void FDataflowNode::UnregisterInputConnection(const void* InProperty)
+void FDataflowNode::UnregisterInputConnection(const void* InProperty, const FName& PropertyName)
 {
 	if (TUniquePtr<FStructOnScope> ScriptOnStruct = TUniquePtr<FStructOnScope>(NewStructOnScope()))
 	{
 		if (const UStruct* const Struct = ScriptOnStruct->GetStruct())
 		{
-			for (TFieldIterator<FProperty> PropertyIt(Struct); PropertyIt; ++PropertyIt)
+			TArray<const FProperty*> PropertyChain;
+			const FProperty* const Property =
+				FindProperty(Struct, InProperty, PropertyName, &PropertyChain);
+			if (ensure(Property && PropertyChain.Num()))
 			{
-				const FProperty* const Property = *PropertyIt;
-				const int32 Offset = Property->GetOffset_ForInternal();
-				const size_t RealAddress = (size_t)this + (size_t)Offset;
-				if (RealAddress == (size_t)InProperty)
+				const uint32 Offset = GetPropertyOffset(PropertyChain);
+				if (FDataflowInput* const* const Input = Inputs.Find(Offset))
 				{
-					if (FDataflowInput* const* const Input = Inputs.Find(Offset))
-					{
-						Inputs.Remove(Offset);
-						delete *Input;
+					Inputs.Remove(Offset);
+					delete *Input;
 
-						// Invalidate graph as this input might have had connections
-						Invalidate();
-					}
-					break;
+					// Invalidate graph as this input might have had connections
+					Invalidate();
 				}
 			}
 		}
 	}
 }
 
-void FDataflowNode::RegisterOutputConnection(const void* InProperty, const void* Passthrough)
+void FDataflowNode::RegisterOutputConnection(const void* InProperty, const void* Passthrough, const FName& PropertyName, const FName& PassthroughName)
 {
 	if (TUniquePtr<FStructOnScope> ScriptOnStruct = TUniquePtr<FStructOnScope>(NewStructOnScope()))
 	{
 		if (const UStruct* Struct = ScriptOnStruct->GetStruct())
 		{
 			FDataflowOutput* OutputConnection = nullptr;
-			size_t PassthroughOffset = INDEX_NONE;
-			for (TFieldIterator<FProperty> PropertyIt(Struct); PropertyIt; ++PropertyIt)
+			TArray<const FProperty*> PropertyChain;
+			const FProperty* const Property =
+				FindProperty(Struct, InProperty, PropertyName, &PropertyChain);
+			if (ensure(Property && PropertyChain.Num()))
 			{
-				FProperty* Property = *PropertyIt;
-				size_t RealAddress = (size_t)this + Property->GetOffset_ForInternal();
-				if (RealAddress == (size_t)InProperty)
-				{
-					FName PropName(Property->GetName());
-					FName PropType(Property->GetCPPType());
-					OutputConnection = new FDataflowOutput({ PropType, PropName, this, Property });
-				}
-				if (RealAddress == (size_t)Passthrough)
-				{
-					PassthroughOffset = (uint32)Property->GetOffset_ForInternal();
-				}
-			}
-			if(OutputConnection != nullptr)
-			{
-				if(PassthroughOffset != INDEX_NONE)
+				const FName PropName(GetPropertyFullName(PropertyChain));
+				const FName PropType(Property->GetCPPType());
+				OutputConnection = new FDataflowOutput({ PropType, PropName, this, Property });
+
+				TArray<const FProperty*> PassthroughPropertyChain;
+				if (FindProperty(Struct, Passthrough, PassthroughName, &PassthroughPropertyChain))
 				{
-					OutputConnection->SetPassthroughOffsetAddress(PassthroughOffset);
+					const uint32 PassthroughOffset = GetPropertyOffset(PassthroughPropertyChain);
+					OutputConnection->SetPassthroughOffset(PassthroughOffset);
 				}
 				AddOutput(OutputConnection);
 			}
-			
 		}
 	}
 }
 
-
 bool FDataflowNode::ValidateConnections()
 {
 	bHasValidConnections = true;
+#if WITH_EDITORONLY_DATA
 	if (const TUniquePtr<FStructOnScope> ScriptOnStruct = TUniquePtr<FStructOnScope>(NewStructOnScope()))
 	{
-		if (const UStruct* Struct = ScriptOnStruct->GetStruct())
+		if (const UStruct* const Struct = ScriptOnStruct->GetStruct())
 		{
-			for (TFieldIterator<FProperty> PropertyIt(Struct); PropertyIt; ++PropertyIt)
+			for (FPropertyValueIterator PropertyIt(FProperty::StaticClass(), Struct, ScriptOnStruct->GetStructMemory()); PropertyIt; ++PropertyIt)
 			{
-				const FProperty* Property = *PropertyIt;
-				FName PropName(Property->GetName());
-#if WITH_EDITORONLY_DATA
+				const FProperty* const Property = PropertyIt.Key();
+				check(Property);
+
 				if (Property->HasMetaData(FDataflowNode::DataflowInput))
 				{
+					TArray<const FProperty*> PropertyChain;
+					PropertyIt.GetPropertyChain(PropertyChain);
+					const FName PropName(GetPropertyFullName(PropertyChain));
+
 					if (!FindInput(PropName))
 					{
 						UE_LOG(LogChaos, Warning, TEXT("Missing dataflow RegisterInputConnection in constructor for (%s:%s)"), *GetName().ToString(), *PropName.ToString())
@@ -346,7 +433,11 @@ bool FDataflowNode::ValidateConnections()
 				}
 				if (Property->HasMetaData(FDataflowNode::DataflowOutput))
 				{
-					const FDataflowOutput* OutputConnection = FindOutput(PropName);
+					TArray<const FProperty*> PropertyChain;
+					PropertyIt.GetPropertyChain(PropertyChain);
+					const FName PropName(GetPropertyFullName(PropertyChain));
+
+					const FDataflowOutput* const OutputConnection = FindOutput(PropName);
 					if(!OutputConnection)
 					{
 						UE_LOG(LogChaos, Warning, TEXT("Missing dataflow RegisterOutputConnection in constructor for (%s:%s)"), *GetName().ToString(),*PropName.ToString());
@@ -389,10 +480,10 @@ bool FDataflowNode::ValidateConnections()
 						bHasValidConnections = false;
 					}
 				}
-#endif
 			}
 		}
 	}
+#endif
 	return bHasValidConnections;
 }
 
@@ -404,14 +495,15 @@ FString FDataflowNode::GetToolTip()
 		if (const UStruct* Struct = ScriptOnStruct->GetStruct())
 		{
 			FString OutStr, InputsStr, OutputsStr;
-			
+
 			FText StructText = Struct->GetToolTipText();
 
 			OutStr.Appendf(TEXT("%s\n\n%s\n"), *GetDisplayName().ToString(), *StructText.ToString());
-			
-			for (TFieldIterator<FProperty> PropertyIt(Struct); PropertyIt; ++PropertyIt)
+
+			for (FPropertyValueIterator PropertyIt(FProperty::StaticClass(), Struct, ScriptOnStruct->GetStructMemory()); PropertyIt; ++PropertyIt)
 			{
-				const FProperty* Property = *PropertyIt;
+				const FProperty* const Property = PropertyIt.Key();
+				check(Property);
 
 				if (Property->HasMetaData(TEXT("Tooltip")))
 				{
@@ -425,6 +517,12 @@ FString FDataflowNode::GetToolTip()
 						{
 							break;
 						}
+
+						TArray<const FProperty*> PropertyChain;
+						PropertyIt.GetPropertyChain(PropertyChain);
+
+						const FName PropName(GetPropertyFullName(PropertyChain));
+
 						const FString& MainTooltipText = (OutArr.Num() > 1) ? OutArr[1] : OutArr[0];
 
 						if (Property->HasMetaData(FDataflowNode::DataflowInput) &&
@@ -433,29 +531,29 @@ FString FDataflowNode::GetToolTip()
 						{
 							if (Property->HasMetaData(FDataflowNode::DataflowIntrinsic))
 							{
-								InputsStr.Appendf(TEXT("    %s [Intrinsic] - %s\n"), *Property->GetName(), *MainTooltipText);
+								InputsStr.Appendf(TEXT("    %s [Intrinsic] - %s\n"), *PropName.ToString(), *MainTooltipText);
 							}
 							else
 							{
-								InputsStr.Appendf(TEXT("    %s - %s\n"), *Property->GetName(), *MainTooltipText);
+								InputsStr.Appendf(TEXT("    %s - %s\n"), *PropName.ToString(), *MainTooltipText);
 							}
 
-							OutputsStr.Appendf(TEXT("    %s [Passthrough] - %s\n"), *Property->GetName(), *MainTooltipText);
+							OutputsStr.Appendf(TEXT("    %s [Passthrough] - %s\n"), *PropName.ToString(), *MainTooltipText);
 						}					
 						else if (Property->HasMetaData(FDataflowNode::DataflowInput))
 						{
 							if (Property->HasMetaData(FDataflowNode::DataflowIntrinsic))
 							{
-								InputsStr.Appendf(TEXT("    %s [Intrinsic] - %s\n"), *Property->GetName(), *MainTooltipText);
+								InputsStr.Appendf(TEXT("    %s [Intrinsic] - %s\n"), *PropName.ToString(), *MainTooltipText);
 							}
 							else
 							{
-								InputsStr.Appendf(TEXT("    %s - %s\n"), *Property->GetName(), *MainTooltipText);
+								InputsStr.Appendf(TEXT("    %s - %s\n"), *PropName.ToString(), *MainTooltipText);
 							}
 						}
 						else if (Property->HasMetaData(FDataflowNode::DataflowOutput))
 						{
-							OutputsStr.Appendf(TEXT("    %s - %s\n"), *Property->GetName(), *MainTooltipText);
+							OutputsStr.Appendf(TEXT("    %s - %s\n"), *PropName.ToString(), *MainTooltipText);
 						}
 					}
 				}
@@ -478,99 +576,82 @@ FString FDataflowNode::GetToolTip()
 	return "";
 }
 
-FText FDataflowNode::GetPinDisplayName(const FName& PropertyName)
+FText FDataflowNode::GetPinDisplayName(const FName& PropertyFullName)
 {
-	// GetDisplayNameText() is only available if WITH_EDITORONLY_DATA
-#if WITH_EDITORONLY_DATA
 	if (const TUniquePtr<FStructOnScope> ScriptOnStruct = TUniquePtr<FStructOnScope>(NewStructOnScope()))
 	{
-		if (const UStruct* Struct = ScriptOnStruct->GetStruct())
+		if (const UStruct* const Struct = ScriptOnStruct->GetStruct())
 		{
-			for (TFieldIterator<FProperty> PropertyIt(Struct); PropertyIt; ++PropertyIt)
+			TArray<const FProperty*> PropertyChain;
+			if (FindProperty(Struct, PropertyFullName, &PropertyChain))
 			{
-				const FProperty* Property = *PropertyIt;
-
-				if (Property->GetName() == PropertyName.ToString())
-				{
-					return Property->GetDisplayNameText();
-				}
+				return GetPropertyDisplayNameText(PropertyChain);
 			}
 		}
 	}
-#endif
 
 	return FText();
 }
 
-FString FDataflowNode::GetPinToolTip(const FName& PropertyName)
+FString FDataflowNode::GetPinToolTip(const FName& PropertyFullName)
 {
+#if WITH_EDITORONLY_DATA
 	if (const TUniquePtr<FStructOnScope> ScriptOnStruct = TUniquePtr<FStructOnScope>(NewStructOnScope()))
 	{
-		if (const UStruct* Struct = ScriptOnStruct->GetStruct())
+		if (const UStruct* const Struct = ScriptOnStruct->GetStruct())
 		{
-			for (TFieldIterator<FProperty> PropertyIt(Struct); PropertyIt; ++PropertyIt)
+			if (const FProperty* const Property = FindProperty(Struct, PropertyFullName))
 			{
-				const FProperty* Property = *PropertyIt;
-
-#if WITH_EDITORONLY_DATA
-				if (Property->GetName() == PropertyName.ToString())
+				if (Property->HasMetaData(TEXT("Tooltip")))
 				{
-					if (Property->HasMetaData(TEXT("Tooltip")))
+					const FString ToolTipStr = Property->GetToolTipText(true).ToString();
+					if (ToolTipStr.Len() > 0)
 					{
-						FString ToolTipStr = Property->GetToolTipText(true).ToString();
-						if (ToolTipStr.Len() > 0)
-						{
-							TArray<FString> OutArr;
-							int32 NumElems = ToolTipStr.ParseIntoArray(OutArr, TEXT(":\r\n"));
+						TArray<FString> OutArr;
+						const int32 NumElems = ToolTipStr.ParseIntoArray(OutArr, TEXT(":\r\n"));
 
-							if (NumElems == 2)
-							{
-								return OutArr[1];  // Return tooltip meta text
-							}
-							else if (NumElems == 1)
-							{
-								return OutArr[0];  // Return doc comment
-							}
+						if (NumElems == 2)
+						{
+							return OutArr[1];  // Return tooltip meta text
+						}
+						else if (NumElems == 1)
+						{
+							return OutArr[0];  // Return doc comment
 						}
 					}
 				}
-#endif
 			}
 		}
 	}
+#endif
 
 	return "";
 }
 
-TArray<FString> FDataflowNode::GetPinMetaData(const FName& PropertyName)
+TArray<FString> FDataflowNode::GetPinMetaData(const FName& PropertyFullName)
 {
+#if WITH_EDITORONLY_DATA
 	if (const TUniquePtr<FStructOnScope> ScriptOnStruct = TUniquePtr<FStructOnScope>(NewStructOnScope()))
 	{
-		if (const UStruct* Struct = ScriptOnStruct->GetStruct())
+		if (const UStruct* const Struct = ScriptOnStruct->GetStruct())
 		{
-			for (TFieldIterator<FProperty> PropertyIt(Struct); PropertyIt; ++PropertyIt)
+			if (const FProperty* const Property = FindProperty(Struct, PropertyFullName))
 			{
-				const FProperty* Property = *PropertyIt;
-
-#if WITH_EDITORONLY_DATA
-				if (Property->GetName() == PropertyName.ToString())
+				TArray<FString> MetaDataStrArr;
+				if (Property->HasMetaData(FDataflowNode::DataflowPassthrough))
 				{
-					TArray<FString> MetaDataStrArr;
-					if (Property->HasMetaData(FDataflowNode::DataflowPassthrough))
-					{
-						MetaDataStrArr.Add("Passthrough");
-					}
-					if (Property->HasMetaData(FDataflowNode::DataflowIntrinsic))
-					{
-						MetaDataStrArr.Add("Intrinsic");
-					}
-
-					return MetaDataStrArr;
+					MetaDataStrArr.Add("Passthrough");
 				}
-#endif
+				if (Property->HasMetaData(FDataflowNode::DataflowIntrinsic))
+				{
+					MetaDataStrArr.Add("Intrinsic");
+				}
+
+				return MetaDataStrArr;
 			}
 		}
 	}
+#endif
 
 	return TArray<FString>();
 }
diff --git a/Engine/Source/Runtime/Experimental/Dataflow/Core/Public/Dataflow/DataflowConnection.h b/Engine/Source/Runtime/Experimental/Dataflow/Core/Public/Dataflow/DataflowConnection.h
--- a/Engine/Source/Runtime/Experimental/Dataflow/Core/Public/Dataflow/DataflowConnection.h
+++ b/Engine/Source/Runtime/Experimental/Dataflow/Core/Public/Dataflow/DataflowConnection.h
@@ -48,21 +48,21 @@ protected:
 	FName Type;
 	FName Name;
 	FDataflowNode* OwningNode = nullptr;
-	FProperty* Property = nullptr;
+	const FProperty* Property = nullptr;
 	FGuid  Guid;
 
 	friend struct FDataflowNode;
 
 public:
 	FDataflowConnection() {};
-	DATAFLOWCORE_API FDataflowConnection(Dataflow::FPin::EDirection Direction, FName InType, FName InName, FDataflowNode* OwningNode = nullptr, FProperty* InProperty = nullptr, FGuid InGuid = FGuid::NewGuid());
+	DATAFLOWCORE_API FDataflowConnection(Dataflow::FPin::EDirection Direction, FName InType, FName InName, FDataflowNode* OwningNode = nullptr, const FProperty* InProperty = nullptr, FGuid InGuid = FGuid::NewGuid());
 	virtual ~FDataflowConnection() {};
 
 	FDataflowNode* GetOwningNode() { return OwningNode; }
 	const FDataflowNode* GetOwningNode() const { return OwningNode; }
 
 	Dataflow::FPin::EDirection GetDirection() const { return Direction; }
-	DATAFLOWCORE_API int32 GetOffset( ) const;
+	DATAFLOWCORE_API uint32 GetOffset() const;
 
 	FName GetType() const { return Type; }
 
diff --git a/Engine/Source/Runtime/Experimental/Dataflow/Core/Public/Dataflow/DataflowInputOutput.h b/Engine/Source/Runtime/Experimental/Dataflow/Core/Public/Dataflow/DataflowInputOutput.h
--- a/Engine/Source/Runtime/Experimental/Dataflow/Core/Public/Dataflow/DataflowInputOutput.h
+++ b/Engine/Source/Runtime/Experimental/Dataflow/Core/Public/Dataflow/DataflowInputOutput.h
@@ -21,7 +21,7 @@ struct FDataflowOutput;
 namespace Dataflow
 {
 	struct FInputParameters {
-		FInputParameters(FName InType = FName(""), FName InName = FName(""), FDataflowNode * InOwner = nullptr, FProperty * InProperty = nullptr)
+		FInputParameters(FName InType = FName(""), FName InName = FName(""), FDataflowNode* InOwner = nullptr, const FProperty* InProperty = nullptr)
 			: Type(InType)
 			, Name(InName)
 			, Owner(InOwner)
@@ -29,7 +29,7 @@ namespace Dataflow
 		FName Type;
 		FName Name;
 		FDataflowNode* Owner = nullptr;
-		FProperty* Property = nullptr;
+		const FProperty* Property = nullptr;
 	};
 }
 
@@ -71,7 +71,7 @@ namespace Dataflow
 {
 	struct FOutputParameters
 	{
-		FOutputParameters(FName InType = FName(""), FName InName = FName(""), FDataflowNode* InOwner = nullptr, FProperty* InProperty = nullptr)
+		FOutputParameters(FName InType = FName(""), FName InName = FName(""), FDataflowNode* InOwner = nullptr, const FProperty* InProperty = nullptr)
 			: Type(InType)
 			, Name(InName)
 			, Owner(InOwner)
@@ -80,7 +80,7 @@ namespace Dataflow
 		FName Type;
 		FName Name;
 		FDataflowNode* Owner = nullptr;
-		FProperty* Property = nullptr;
+		const FProperty* Property = nullptr;
 	};
 }
 USTRUCT()
@@ -92,7 +92,7 @@ struct FDataflowOutput : public FDataflowConnection
 	
 	TArray< FDataflowInput* > Connections;
 
-	size_t PassthroughOffsetAddress = INDEX_NONE;
+	uint32 PassthroughOffset = INDEX_NONE;
 
 public:
 	static DATAFLOWCORE_API FDataflowOutput NoOpOutput;
@@ -111,16 +111,16 @@ public:
 
 	DATAFLOWCORE_API virtual bool RemoveConnection(FDataflowConnection* InInput) override;
 
-	virtual FORCEINLINE void SetPassthroughOffsetAddress(const size_t InPassthroughOffsetAddress)
+	virtual FORCEINLINE void SetPassthroughOffset(const uint32 InPassthroughOffset)
 	{
-		PassthroughOffsetAddress = InPassthroughOffsetAddress;
+		PassthroughOffset = InPassthroughOffset;
 	}
 
 	virtual FORCEINLINE void* GetPassthroughRealAddress() const
 	{
-		if(PassthroughOffsetAddress != INDEX_NONE)
+		if(PassthroughOffset != INDEX_NONE)
 		{
-			return (void*)((size_t)OwningNode + PassthroughOffsetAddress);
+			return (void*)((size_t)OwningNode + (size_t)PassthroughOffset);
 		}
 		return nullptr;
 	}
diff --git a/Engine/Source/Runtime/Experimental/Dataflow/Core/Public/Dataflow/DataflowNode.h b/Engine/Source/Runtime/Experimental/Dataflow/Core/Public/Dataflow/DataflowNode.h
--- a/Engine/Source/Runtime/Experimental/Dataflow/Core/Public/Dataflow/DataflowNode.h
+++ b/Engine/Source/Runtime/Experimental/Dataflow/Core/Public/Dataflow/DataflowNode.h
@@ -117,6 +117,9 @@ struct FDataflowNode
 	DATAFLOWCORE_API const FDataflowOutput* FindOutput(FName Name) const;
 	DATAFLOWCORE_API const FDataflowOutput* FindOutput(const void* Reference) const;
 
+	/** Return a property's byte offset from the dataflow base node address using the full property name (must includes its parent struct property names). */
+	uint32 GetPropertyOffset(const FName& PropertyFullName) const;
+
 	static DATAFLOWCORE_API const FName DataflowInput;
 	static DATAFLOWCORE_API const FName DataflowOutput;
 	static DATAFLOWCORE_API const FName DataflowPassthrough;
@@ -136,11 +139,15 @@ struct FDataflowNode
 	virtual FStructOnScope* NewStructOnScope() { return nullptr; }
 	virtual const UScriptStruct* TypedScriptStruct() const { return nullptr; }
 
-	/** Register the Input and Outputs after the creation in the factory */
-	DATAFLOWCORE_API void RegisterInputConnection(const void* Property);
-	DATAFLOWCORE_API void RegisterOutputConnection(const void* Property, const void* Passthrough = nullptr);
+	/** Register the Input and Outputs after the creation in the factory. Use PropertyName to disambiguate a struct name from its first property. */
+	DATAFLOWCORE_API void RegisterInputConnection(const void* Property, const FName& PropertyName = NAME_None);
+	DATAFLOWCORE_API void RegisterOutputConnection(
+		const void* Property,
+		const void* Passthrough = nullptr,
+		const FName& PropertyName = NAME_None,
+		const FName& PassthroughName = NAME_None);
 	/** Unregister the input connection if one exists matching this property, and then invalidate the graph. */
-	DATAFLOWCORE_API void UnregisterInputConnection(const void* Property);
+	DATAFLOWCORE_API void UnregisterInputConnection(const void* Property, const FName& PropertyName = NAME_None);
 
 	//
 	// Evaluation
@@ -259,6 +266,20 @@ struct FDataflowNode
 	FOnNodeInvalidated& GetOnNodeInvalidatedDelegate() { return OnNodeInvalidatedDelegate; }
 
 private:
+	static FName GetPropertyFullName(const TArray<const FProperty*>& PropertyChain);
+	static FText GetPropertyDisplayNameText(const TArray<const FProperty*>& PropertyChain);
+	static uint32 GetPropertyOffset(const TArray<const FProperty*>& PropertyChain);
+
+	/**
+	* Find a property using the property address and name (not including its parent struct property names).
+	* If NAME_None is used as the name, and the same address is shared by a parent structure property and
+	* its first child property, then the parent will be returned.
+	*/
+	const FProperty* FindProperty(const UStruct* Struct, const void* Property, const FName& PropertyName, TArray<const FProperty*>* OutPropertyChain = nullptr) const;
+
+	/** Find a property using the property full name (must includes its parent struct property names). */
+	const FProperty* FindProperty(const UStruct* Struct, const FName& PropertyFullName, TArray<const FProperty*>* OutPropertyChain = nullptr) const;
+
 	virtual TArray<Dataflow::FRenderingParameter> GetRenderParametersImpl() const { return TArray<Dataflow::FRenderingParameter>(); }
 
 
diff --git a/Engine/Source/Runtime/Experimental/Dataflow/Core/Public/Dataflow/DataflowNodeParameters.h b/Engine/Source/Runtime/Experimental/Dataflow/Core/Public/Dataflow/DataflowNodeParameters.h
--- a/Engine/Source/Runtime/Experimental/Dataflow/Core/Public/Dataflow/DataflowNodeParameters.h
+++ b/Engine/Source/Runtime/Experimental/Dataflow/Core/Public/Dataflow/DataflowNodeParameters.h
@@ -44,7 +44,7 @@ namespace Dataflow
 
 	struct FContextCacheElementBase 
 	{
-		FContextCacheElementBase(FProperty* InProperty = nullptr, FTimestamp InTimestamp = FTimestamp::Invalid)
+		FContextCacheElementBase(const FProperty* InProperty = nullptr, FTimestamp InTimestamp = FTimestamp::Invalid)
 			: Property(InProperty)
 			, Timestamp(InTimestamp)
 		{}
@@ -53,14 +53,14 @@ namespace Dataflow
 		template<typename T>
 		const T& GetTypedData(const FProperty* PropertyIn) const;
 		
-		FProperty* Property = nullptr;
+		const FProperty* Property = nullptr;
 		FTimestamp Timestamp = FTimestamp::Invalid;
 	};
 
 	template<class T>
 	struct FContextCacheElement : public FContextCacheElementBase 
 	{
-		FContextCacheElement(FProperty* InProperty, T&& InData, FTimestamp Timestamp)
+		FContextCacheElement(const FProperty* InProperty, T&& InData, FTimestamp Timestamp)
 			: FContextCacheElementBase(InProperty, Timestamp)
 			, Data(Forward<T>(InData))
 		{}
@@ -121,7 +121,7 @@ namespace Dataflow
 		virtual void SetDataImpl(int64 Key, TUniquePtr<FContextCacheElementBase>&& DataStoreEntry) = 0;
 		
 		template<typename T>
-		void SetData(size_t Key, FProperty* Property, T&& Value)
+		void SetData(size_t Key, const FProperty* Property, T&& Value)
 		{
 			int64 IntKey = (int64)Key;
 			TUniquePtr<FContextCacheElement<T>> DataStoreEntry = MakeUnique<FContextCacheElement<T>>(Property, Forward<T>(Value), FTimestamp::Current());
@@ -133,7 +133,7 @@ namespace Dataflow
 		virtual TUniquePtr<FContextCacheElementBase>* GetDataImpl(int64 Key) = 0;
 
 		template<class T>
-		const T& GetData(size_t Key, FProperty* Property, const T& Default = T())
+		const T& GetData(size_t Key, const FProperty* Property, const T& Default = T())
 		{
 			if (TUniquePtr<FContextCacheElementBase>* Cache = GetDataImpl(Key))
 			{
