Removed superflous inner class.

#rnx

diff --git a/Engine/Source/Runtime/Experimental/IoStoreOnDemand/Private/CoreHttp/Http.cpp b/Engine/Source/Runtime/Experimental/IoStoreOnDemand/Private/CoreHttp/Http.cpp
--- a/Engine/Source/Runtime/Experimental/IoStoreOnDemand/Private/CoreHttp/Http.cpp
+++ b/Engine/Source/Runtime/Experimental/IoStoreOnDemand/Private/CoreHttp/Http.cpp
@@ -1246,27 +1246,14 @@ static uint64 ReadyCheck(FActivity** Activities, uint32 Num, uint32 TimeoutMs)
 
 
 ////////////////////////////////////////////////////////////////////////////////
-class FEventLoopInternal
+struct FHandlerResult
 {
-public:
-	struct FResult
-	{
-		int32	Result;
-		uint32	Size;
-	};
-
-	static int32	DoResolve(FActivity* Activity);
-	static int32	DoConnect(FActivity* Activity);
-	static int32	DoSend(FActivity* Activity);
-	static int32	DoRecvMessage(FActivity* Activity);
-	static FResult	DoRecvContent(FActivity* Activity, uint32 MaxRecvSize);
-	static FResult	DoRecvStream(FActivity* Activity, uint32 MaxRecSize);
-	static int32	DoRecvDone(FActivity* Activity);
-	static void		Cancel(FActivity* Activity);
+	int32	Result;
+	uint32	Size;
 };
 
 ////////////////////////////////////////////////////////////////////////////////
-int32 FEventLoopInternal::DoResolve(FActivity* Activity)
+static int32 DoResolve(FActivity* Activity)
 {
 	TRACE_CPUPROFILER_EVENT_SCOPE(CoreHttp::DoResolve);
 
@@ -1350,7 +1337,7 @@ int32 FEventLoopInternal::DoResolve(FActivity* Activity)
 }
 
 ////////////////////////////////////////////////////////////////////////////////
-int32 FEventLoopInternal::DoConnect(FActivity* Activity)
+static int32 DoConnect(FActivity* Activity)
 {
 	TRACE_CPUPROFILER_EVENT_SCOPE(CoreHttp::DoConnect);
 
@@ -1457,7 +1444,7 @@ int32 FEventLoopInternal::DoConnect(FActivity* Activity)
 }
 
 ////////////////////////////////////////////////////////////////////////////////
-int32 FEventLoopInternal::DoSend(FActivity* Activity)
+static int32 DoSend(FActivity* Activity)
 {
 	TRACE_CPUPROFILER_EVENT_SCOPE(CoreHttp::DoSend);
 
@@ -1540,7 +1527,7 @@ int32 FEventLoopInternal::DoSend(FActivity* Activity)
 }
 
 ////////////////////////////////////////////////////////////////////////////////
-int32 FEventLoopInternal::DoRecvMessage(FActivity* Activity)
+static int32 DoRecvMessage(FActivity* Activity)
 {
 	TRACE_CPUPROFILER_EVENT_SCOPE(CoreHttp::DoRecvMessage);
 
@@ -1723,9 +1710,7 @@ int32 FEventLoopInternal::DoRecvMessage(FActivity* Activity)
 }
 
 ////////////////////////////////////////////////////////////////////////////////
-FEventLoopInternal::FResult FEventLoopInternal::DoRecvContent(
-	FActivity* Activity,
-	uint32 MaxRecvSize)
+static FHandlerResult DoRecvContent(FActivity* Activity, uint32 MaxRecvSize)
 {
 	TRACE_CPUPROFILER_EVENT_SCOPE(CoreHttp::DoRecvContent);
 
@@ -1779,7 +1764,7 @@ FEventLoopInternal::FResult FEventLoopInternal::DoRecvContent(
 }
 
 ////////////////////////////////////////////////////////////////////////////////
-int32 FEventLoopInternal::DoRecvDone(FActivity* Activity)
+static int32 DoRecvDone(FActivity* Activity)
 {
 	if (!FLatencyInjector::HasExpired(Activity->StateParam))
 	{
@@ -1795,9 +1780,7 @@ int32 FEventLoopInternal::DoRecvDone(FActivity* Activity)
 }
 
 ////////////////////////////////////////////////////////////////////////////////
-FEventLoopInternal::FResult FEventLoopInternal::DoRecvStream(
-	FActivity* Activity,
-	uint32 RecvSize)
+static FHandlerResult DoRecvStream(FActivity* Activity, uint32 RecvSize)
 {
 	TRACE_CPUPROFILER_EVENT_SCOPE(CoreHttp::DoRecvStream);
 
@@ -1806,7 +1789,7 @@ FEventLoopInternal::FResult FEventLoopInternal::DoRecvStream(
 }
 
 ////////////////////////////////////////////////////////////////////////////////
-void FEventLoopInternal::Cancel(FActivity* Activity)
+static void DoCancel(FActivity* Activity)
 {
 	if (Activity->State >= FActivity::EState::Completed)
 	{
@@ -2034,7 +2017,7 @@ uint32 FEventLoop::Tick(uint32 PollTimeoutMs)
 		if (SlotBit & CancelsLoad)
 		{
 			--BusyCount;
-			FEventLoopInternal::Cancel(Activity);
+			DoCancel(Activity);
 			continue;
 		}
 
@@ -2047,35 +2030,35 @@ uint32 FEventLoop::Tick(uint32 PollTimeoutMs)
 		switch (Activity->State)
 		{
 		case FActivity::EState::Resolve:
-			Result = FEventLoopInternal::DoResolve(Activity);
+			Result = DoResolve(Activity);
 			if (Result)
 				break;
 
 		case FActivity::EState::Connect:
-			Result = FEventLoopInternal::DoConnect(Activity);
+			Result = DoConnect(Activity);
 			if (Result)
 				break;
 
 		case FActivity::EState::Send:
-			Result = FEventLoopInternal::DoSend(Activity);
+			Result = DoSend(Activity);
 			if (Result)
 				break;
 
 		case FActivity::EState::RecvMessage:
-			Result = FEventLoopInternal::DoRecvMessage(Activity);
+			Result = DoRecvMessage(Activity);
 			if (Result)
 				break;
 
 		case FActivity::EState::RecvContent:
 		case FActivity::EState::RecvStream: {
-			decltype(FEventLoopInternal::DoRecvContent)* Handler;
+			decltype(DoRecvContent)* Handler;
 			if (Activity->State == FActivity::EState::RecvContent)
 			{
-				Handler = &FEventLoopInternal::DoRecvContent;
+				Handler = DoRecvContent;
 			}
 			else
 			{
-				Handler = &FEventLoopInternal::DoRecvStream;
+				Handler = DoRecvStream;
 			}
 
 			auto [ResultInner, RecvSize] = Handler(Activity, ~0u);
@@ -2085,7 +2068,7 @@ uint32 FEventLoop::Tick(uint32 PollTimeoutMs)
 		}
 
 		case FActivity::EState::RecvDone:
-			Result = FEventLoopInternal::DoRecvDone(Activity);
+			Result = DoRecvDone(Activity);
 			if (Result)
 				break;
 
