Hide some details in the .cpp

#rnx

diff --git a/Engine/Source/Runtime/Experimental/IoStoreOnDemand/Private/CoreHttp/Client.h b/Engine/Source/Runtime/Experimental/IoStoreOnDemand/Private/CoreHttp/Client.h
--- a/Engine/Source/Runtime/Experimental/IoStoreOnDemand/Private/CoreHttp/Client.h
+++ b/Engine/Source/Runtime/Experimental/IoStoreOnDemand/Private/CoreHttp/Client.h
@@ -3,14 +3,10 @@
 #pragma once
 
 #if !defined(NO_UE_INCLUDES)
-#include "Containers/Array.h"
 #include "Containers/StringView.h"
-#include "HAL/CriticalSection.h"
 #include "Memory/MemoryView.h"
 #endif
 
-#include <atomic>
-
 #if !defined(COREHTTP_API)
 #	define COREHTTP_API IOSTOREONDEMAND_API
 #endif
@@ -160,6 +156,8 @@ using FTicketSink = TFunction<void (const FTicketStatus&)>;
 ////////////////////////////////////////////////////////////////////////////////
 class COREHTTP_API FEventLoop
 {
+	class FImpl;
+
 public:
 	struct FRequestParams
 	{
@@ -170,23 +168,17 @@ public:
 	template <typename... T> [[nodiscard]] FRequest Get(T&&... t)  { return Request("GET",  Forward<T&&>(t)...); }
 	template <typename... T> [[nodiscard]] FRequest Post(T&&... t) { return Request("POST", Forward<T&&>(t)...); }
 
-							FEventLoop() = default;
+							FEventLoop();
 							~FEventLoop();
 	uint32					Tick(uint32 PollTimeoutMs=0);
 	bool					IsIdle() const;
 	void					Cancel(FTicket Ticket);
 	[[nodiscard]] FRequest	Request(FAnsiStringView Method, FAnsiStringView Url, const FRequestParams* Params=nullptr);
 	[[nodiscard]] FRequest	Request(FAnsiStringView Method, FAnsiStringView Path, FConnectionPool& Pool, const FRequestParams* Params=nullptr);
-	FTicket					Send(FRequest&& Request, FTicketSink Sink, UPTRINT Param=0);
+	FTicket					Send(FRequest&& Request, FTicketSink Sink, UPTRINT SinkParam=0);
 
 private:
-	FRequest				Request(FAnsiStringView Method, FAnsiStringView Path, FActivity* Activity);
-	FCriticalSection		Lock;
-	std::atomic<uint64>		FreeSlots		= ~0ull;
-	std::atomic<uint64>		Cancels			= 0;
-	uint64					PrevFreeSlots	= ~0ull;
-	TArray<FActivity*>		Pending;
-	TArray<FActivity*>		Active;
+	FImpl*					Impl;
 
 private:
 							FEventLoop(const FEventLoop&)	= delete;
diff --git a/Engine/Source/Runtime/Experimental/IoStoreOnDemand/Private/CoreHttp/Http.cpp b/Engine/Source/Runtime/Experimental/IoStoreOnDemand/Private/CoreHttp/Http.cpp
--- a/Engine/Source/Runtime/Experimental/IoStoreOnDemand/Private/CoreHttp/Http.cpp
+++ b/Engine/Source/Runtime/Experimental/IoStoreOnDemand/Private/CoreHttp/Http.cpp
@@ -3,6 +3,8 @@
 #include "CoreHttp/Client.h"
 
 #if !defined(NO_UE_INCLUDES)
+#include "Containers/Array.h"
+#include "HAL/CriticalSection.h"
 #include "IO/IoBuffer.h"
 #include "LatencyInjector.h"
 #include "Math/UnrealMathUtility.h"
@@ -13,6 +15,8 @@
 #include "Tasks/Task.h"
 #endif
 
+#include <atomic>
+
 // {{{1 platforms ..............................................................
 
 #if PLATFORM_MICROSOFT
@@ -1810,8 +1814,28 @@
 static const FEventLoop::FRequestParams GDefaultParams;
 
 ////////////////////////////////////////////////////////////////////////////////
-FEventLoop::~FEventLoop()
+class FEventLoop::FImpl
 {
+public:
+							~FImpl();
+	uint32					Tick(uint32 PollTimeoutMs=0);
+	bool					IsIdle() const;
+	void					Cancel(FTicket Ticket);
+	FRequest				Request(FAnsiStringView Method, FAnsiStringView Path, FActivity* Activity);
+	FTicket					Send(FActivity* Activity);
+
+private:
+	FCriticalSection		Lock;
+	std::atomic<uint64>		FreeSlots		= ~0ull;
+	std::atomic<uint64>		Cancels			= 0;
+	uint64					PrevFreeSlots	= ~0ull;
+	TArray<FActivity*>		Pending;
+	TArray<FActivity*>		Active;
+};
+
+////////////////////////////////////////////////////////////////////////////////
+FEventLoop::FImpl::~FImpl()
+{
 	for (FActivity* Activity : Active)
 	{
 		Activity_Free(Activity);
@@ -1819,90 +1843,7 @@
 }
 
 ////////////////////////////////////////////////////////////////////////////////
-FRequest FEventLoop::Request(
-	FAnsiStringView Method,
-	FAnsiStringView Url,
-	const FRequestParams* Params)
-{
-	// Parse the URL into its components
-	FUrlOffsets UrlOffsets;
-	if (ParseUrl(Url, UrlOffsets) < 0)
-	{
-		return FRequest();
-	}
-
-	FAnsiStringView HostName = UrlOffsets.HostName.Get(Url);
-
-	uint32 Port = 80;
-	if (UrlOffsets.SchemeLength == 5)
-	{
-		//Port = 443;
-		//Protocol = Protocol::Tls;
-		return FRequest();
-	}
-
-	if (UrlOffsets.Port)
-	{
-		FAnsiStringView PortView = UrlOffsets.Port.Get(Url);
-		Port = CrudeToInt(PortView);
-	}
-
-	FAnsiStringView Path;
-	if (UrlOffsets.Path > 0)
-	{
-		Path = Url.Mid(UrlOffsets.Path);
-	}
-
-	// Create an activity and an emphemeral socket pool
-	Params = (Params != nullptr) ? Params : &GDefaultParams;
-
-	uint32 BufferSize = Params->BufferSize;
-	BufferSize = (BufferSize >= 128) ? BufferSize : 128;
-	BufferSize += sizeof(FSocketPool) + HostName.Len();
-	FActivity* Activity = Activity_Alloc(BufferSize);
-
-	FBuffer& Buffer = Activity->Buffer;
-
-	FSocketPool* Pool = Activity->Pool = Buffer.Alloc<FSocketPool>();
-	Activity->IsKeepAlive = 0;
-
-	uint32 HostNameLength = HostName.Len();
-	char* HostNamePtr = Buffer.Alloc<char>(HostNameLength + 1);
-
-	Buffer.Fix();
-
-	memcpy(HostNamePtr, HostName.GetData(), HostNameLength);
-	HostNamePtr[HostNameLength] = '\0';
-	HostName = FAnsiStringView(HostNamePtr, HostNameLength);
-
-	new (Pool) FSocketPool(HostName, Port, 1);
-
-	return Request(Method, Path, Activity);
-}
-
-////////////////////////////////////////////////////////////////////////////////
-FRequest FEventLoop::Request(
+FRequest FEventLoop::FImpl::Request(
 	FAnsiStringView Method,
 	FAnsiStringView Path,
-	FConnectionPool& Pool,
-	const FRequestParams* Params)
-{
-	check(Pool.Ptr != nullptr);
-
-	Params = (Params != nullptr) ? Params : &GDefaultParams;
-
-	uint32 BufferSize = Params->BufferSize;
-	BufferSize = (BufferSize >= 128) ? BufferSize : 128;
-	FActivity* Activity = Activity_Alloc(BufferSize);
-
-	Activity->Pool = Pool.Ptr;
-	Activity->IsKeepAlive = 1;
-
-	return Request(Method, Path, Activity);
-}
-
-////////////////////////////////////////////////////////////////////////////////
-FRequest FEventLoop::Request(
-	FAnsiStringView Method,
-	FAnsiStringView Path,
 	FActivity* Activity)
@@ -1929,14 +1870,8 @@
 }
 
 ////////////////////////////////////////////////////////////////////////////////
-FTicket FEventLoop::Send(FRequest&& Request, FTicketSink Sink, UPTRINT SinkParam)
+FTicket FEventLoop::FImpl::Send(FActivity* Activity)
 {
-	FActivity* Activity = nullptr;
-	Swap(Activity, Request.Ptr);
-	Activity->State = FActivity::EState::Resolve;
-	Activity->SinkParam = SinkParam;
-	Activity->Sink = Sink;
-
 	uint64 Slot;
 	{
 		FScopeLock _(&Lock);
@@ -1965,19 +1900,19 @@
 }
 
 ////////////////////////////////////////////////////////////////////////////////
-bool FEventLoop::IsIdle() const
+bool FEventLoop::FImpl::IsIdle() const
 {
 	return FreeSlots.load(std::memory_order_relaxed) == ~0ull;
 }
 
 ////////////////////////////////////////////////////////////////////////////////
-void FEventLoop::Cancel(FTicket Ticket)
+void FEventLoop::FImpl::Cancel(FTicket Ticket)
 {
 	Cancels.fetch_or(Ticket, std::memory_order_relaxed);
 }
 
 ////////////////////////////////////////////////////////////////////////////////
-uint32 FEventLoop::Tick(uint32 PollTimeoutMs)
+uint32 FEventLoop::FImpl::Tick(uint32 PollTimeoutMs)
 {
 	TRACE_CPUPROFILER_EVENT_SCOPE(CoreHttp::Tick);
 
@@ -2126,6 +2061,108 @@
 	return BusyCount;
 }
 
+////////////////////////////////////////////////////////////////////////////////
+FEventLoop::FEventLoop()						{ Impl = new FEventLoop::FImpl(); }
+FEventLoop::~FEventLoop()						{ delete Impl; }
+uint32 FEventLoop::Tick(uint32 PollTimeoutMs)	{ return Impl->Tick(PollTimeoutMs); }
+bool FEventLoop::IsIdle() const					{ return Impl->IsIdle(); }
+void FEventLoop::Cancel(FTicket Ticket)			{ return Impl->Cancel(Ticket); }
+
+////////////////////////////////////////////////////////////////////////////////
+FRequest FEventLoop::Request(
+	FAnsiStringView Method,
+	FAnsiStringView Url,
+	const FRequestParams* Params)
+{
+	// Parse the URL into its components
+	FUrlOffsets UrlOffsets;
+	if (ParseUrl(Url, UrlOffsets) < 0)
+	{
+		return FRequest();
+	}
+
+	FAnsiStringView HostName = UrlOffsets.HostName.Get(Url);
+
+	uint32 Port = 80;
+	if (UrlOffsets.SchemeLength == 5)
+	{
+		//Port = 443;
+		//Protocol = Protocol::Tls;
+		return FRequest();
+	}
+
+	if (UrlOffsets.Port)
+	{
+		FAnsiStringView PortView = UrlOffsets.Port.Get(Url);
+		Port = CrudeToInt(PortView);
+	}
+
+	FAnsiStringView Path;
+	if (UrlOffsets.Path > 0)
+	{
+		Path = Url.Mid(UrlOffsets.Path);
+	}
+
+	// Create an activity and an emphemeral socket pool
+	Params = (Params != nullptr) ? Params : &GDefaultParams;
+
+	uint32 BufferSize = Params->BufferSize;
+	BufferSize = (BufferSize >= 128) ? BufferSize : 128;
+	BufferSize += sizeof(FSocketPool) + HostName.Len();
+	FActivity* Activity = Activity_Alloc(BufferSize);
+
+	FBuffer& Buffer = Activity->Buffer;
+
+	FSocketPool* Pool = Activity->Pool = Buffer.Alloc<FSocketPool>();
+	Activity->IsKeepAlive = 0;
+
+	uint32 HostNameLength = HostName.Len();
+	char* HostNamePtr = Buffer.Alloc<char>(HostNameLength + 1);
+
+	Buffer.Fix();
+
+	memcpy(HostNamePtr, HostName.GetData(), HostNameLength);
+	HostNamePtr[HostNameLength] = '\0';
+	HostName = FAnsiStringView(HostNamePtr, HostNameLength);
+
+	new (Pool) FSocketPool(HostName, Port, 1);
+
+	return Impl->Request(Method, Path, Activity);
+}
+
+////////////////////////////////////////////////////////////////////////////////
+FRequest FEventLoop::Request(
+	FAnsiStringView Method,
+	FAnsiStringView Path,
+	FConnectionPool& Pool,
+	const FRequestParams* Params)
+{
+	check(Pool.Ptr != nullptr);
+
+	Params = (Params != nullptr) ? Params : &GDefaultParams;
+
+	uint32 BufferSize = Params->BufferSize;
+	BufferSize = (BufferSize >= 128) ? BufferSize : 128;
+	FActivity* Activity = Activity_Alloc(BufferSize);
+
+	Activity->Pool = Pool.Ptr;
+	Activity->IsKeepAlive = 1;
+
+	return Impl->Request(Method, Path, Activity);
+}
+
+////////////////////////////////////////////////////////////////////////////////
+FTicket FEventLoop::Send(FRequest&& Request, FTicketSink Sink, UPTRINT SinkParam)
+{
+	FActivity* Activity = nullptr;
+	Swap(Activity, Request.Ptr);
+	Activity->State = FActivity::EState::Resolve;
+	Activity->SinkParam = SinkParam;
+	Activity->Sink = Sink;
+
+	return Impl->Send(Activity);
+}
 
 
+
 // {{{1 test ...................................................................
