Iris - More debugging info to track down elusive bug.

#rnx

diff --git a/Engine/Source/Runtime/Experimental/Iris/Core/Private/Iris/ReplicationSystem/ReplicationWriter.cpp b/Engine/Source/Runtime/Experimental/Iris/Core/Private/Iris/ReplicationSystem/ReplicationWriter.cpp
--- a/Engine/Source/Runtime/Experimental/Iris/Core/Private/Iris/ReplicationSystem/ReplicationWriter.cpp
+++ b/Engine/Source/Runtime/Experimental/Iris/Core/Private/Iris/ReplicationSystem/ReplicationWriter.cpp
@@ -50,6 +50,9 @@ static FAutoConsoleVariableRef CVarWarnAboutDroppedAttachmentsToObjectsNotInScop
 	TEXT("Warn when attachments are dropped due to object not in scope. Default is false."
 	));
 
+static bool bValidateObjectsWithDirtyChanges = false;
+static FAutoConsoleVariableRef CvarValidateObjectsWithDirtyChanges(TEXT("net.Iris.ReplicationWriter.ValidateObjectsWithDirtyChanges"), bValidateObjectsWithDirtyChanges, TEXT("Ensure that we don't try to mark invalid objects as dirty when they shouldn't."));
+
 static const FName NetError_ObjectStateTooLarge("Object state is too large to be split.");
 
 /** Helper class for timing various operations. */
@@ -338,11 +341,11 @@ void FReplicationWriter::QueueNetObjectAttachments(FInternalNetRefIndex OwnerInt
 	FReplicationInfo& TargetInfo = GetReplicationInfo(TargetIndex);
 	TargetInfo.HasAttachments = 1;
 
-	ObjectsWithDirtyChanges.SetBit(TargetIndex);
+	MarkObjectDirty(TargetIndex, "QueueAttachment");
 
 	if (OwnerInternalIndex != TargetIndex)
 	{
-		ObjectsWithDirtyChanges.SetBit(OwnerInternalIndex);
+		MarkObjectDirty(OwnerInternalIndex, "QueueAttachment2");
 		FReplicationInfo& OwnerInfo = GetReplicationInfo(OwnerInternalIndex);
 		OwnerInfo.HasDirtySubObjects = 1;
 	}
@@ -401,7 +404,7 @@ void FReplicationWriter::StartReplication(uint32 InternalIndex)
 {
 	FReplicationInfo& Info = GetReplicationInfo(InternalIndex);
 
-	check(Info.GetState() == EReplicatedObjectState::Invalid);
+	ensureMsgf(Info.GetState() == EReplicatedObjectState::Invalid, TEXT("Object ( InternalIndex: %u ) is in state %s in StartReplication."), InternalIndex, LexToString(Info.GetState()));
 
 	// Reset info
 	Info = FReplicationInfo();
@@ -597,11 +600,11 @@ void FReplicationWriter::SetPendingDestroyOrSubObjectPendingDestroyState(uint32
 			FReplicationInfo& OwnerInfo = ReplicatedObjects[ObjectData.SubObjectRootIndex];
 			if (OwnerInfo.GetState() < EReplicatedObjectState::PendingDestroy)
 			{
-				ObjectsWithDirtyChanges.SetBit(ObjectData.SubObjectRootIndex);
+				MarkObjectDirty(ObjectData.SubObjectRootIndex, "SetPendingDestroyOrSubObjectPendingDestroyState");
 				OwnerInfo.HasDirtySubObjects = 1U;
 
 				SetState(InternalIndex, EReplicatedObjectState::SubObjectPendingDestroy);
-				ObjectsWithDirtyChanges.SetBit(InternalIndex);
+				MarkObjectDirty(InternalIndex, "SetPendingDestroyOrSubObjectPendingDestroyState2");
 				ObjectsPendingDestroy.SetBit(InternalIndex);
 				Info.SubObjectPendingDestroy = 1U;
 				return;
@@ -679,6 +682,7 @@ void FReplicationWriter::UpdateScope(const FNetBitArrayView& UpdatedScope)
 				FReplicationInfo& OwnerInfo = ReplicatedObjects[ObjectData.SubObjectRootIndex];
 				if (OwnerInfo.GetState() < EReplicatedObjectState::PendingDestroy)
 				{
+					ensureMsgf(!bValidateObjectsWithDirtyChanges || OwnerInfo.GetState() != EReplicatedObjectState::Invalid, TEXT("Object ( InternalIndex: %u ) with Invalid state potentially marked dirty."), ObjectData.SubObjectRootIndex);
 					ensureAlwaysMsgf(!OwnerInfo.TearOff, TEXT("Parent is tearing off ( InternalIndex: %u ) currently in State: %s "), ObjectData.SubObjectRootIndex, LexToString(OwnerInfo.GetState()));
 					OwnerInfo.HasDirtySubObjects |= Info.HasDirtyChangeMask;
 					ObjectsWithDirtyChanges.SetBitValue(ObjectData.SubObjectRootIndex, ObjectsWithDirtyChanges.GetBit(ObjectData.SubObjectRootIndex) || Info.HasDirtyChangeMask);
@@ -767,7 +771,7 @@ void FReplicationWriter::InternalUpdateDirtyChangeMasks(const FChangeMaskCache&
 			continue;
 		}
 
-		ObjectsWithDirtyChanges.SetBit(Entry.InternalIndex);
+		MarkObjectDirty(Entry.InternalIndex, "UpdateDirtyChangeMasks");
 		FReplicationInfo& Info = ReplicatedObjects[Entry.InternalIndex];
 
 		if (Entry.bMarkSubObjectOwnerDirty == 0U)
@@ -1090,7 +1094,7 @@ void FReplicationWriter::HandleDeliveredRecord(const FReplicationRecord::FRecord
 					if (ObjectData.IsSubObject())
 					{
 						FReplicationInfo& OwnerInfo = ReplicatedObjects[ObjectData.SubObjectRootIndex];
-						ObjectsWithDirtyChanges.SetBit(ObjectData.SubObjectRootIndex);
+						MarkObjectDirty(ObjectData.SubObjectRootIndex, "HandleDeliveredRecordTearOff");
 						OwnerInfo.HasDirtySubObjects = 1U;
 					}
 				}
@@ -1149,6 +1153,9 @@ void FReplicationWriter::HandleDroppedRecord<FReplicationWriter::EReplicatedObje
 
 	if (CurrentState < EReplicatedObjectState::Created)
 	{
+		// Mark object as having dirty changes
+		MarkObjectDirty(InternalIndex, "DroppedWaitOnCreate");
+
 		// Resend creation data
 		SetState(InternalIndex, EReplicatedObjectState::PendingCreate);
 
@@ -1157,9 +1164,6 @@ void FReplicationWriter::HandleDroppedRecord<FReplicationWriter::EReplicatedObje
 		FNetBitArrayView LostChangeMask = FChangeMaskUtil::MakeChangeMask(RecordInfo.ChangeMaskOrPtr, Info.ChangeMaskBitCount);
 		ChangeMask.Combine(LostChangeMask, FNetBitArrayView::OrOp);
 
-		// Mark object as having dirty changes
-		ObjectsWithDirtyChanges.SetBit(InternalIndex);
-
 		// Mark changemask dirty
 		Info.HasDirtyChangeMask = 1U;
 
@@ -1179,7 +1183,7 @@ void FReplicationWriter::HandleDroppedRecord<FReplicationWriter::EReplicatedObje
 			if (ensure(SubObjectOwnerReplicationInfo.GetState() < EReplicatedObjectState::PendingDestroy))
 			{
 				// Mark owner as dirty
-				ObjectsWithDirtyChanges.SetBit(SubObjectOwnerInternalIndex);
+				MarkObjectDirty(SubObjectOwnerInternalIndex, "DroppedWaitOnCreate2");
 
 				// Indicate that we have dirty subobjects
 				SubObjectOwnerReplicationInfo.HasDirtySubObjects = 1U;
@@ -1240,7 +1244,7 @@ void FReplicationWriter::HandleDroppedRecord<FReplicationWriter::EReplicatedObje
 			}
 
 			// Mark object as having dirty changes
-			ObjectsWithDirtyChanges.SetBit(InternalIndex);
+			MarkObjectDirty(InternalIndex, "DroppedCreated");
 
 			// Mark changemask as dirty
 			Info.HasDirtyChangeMask |= bNeedToResendState;
@@ -1262,7 +1266,7 @@ void FReplicationWriter::HandleDroppedRecord<FReplicationWriter::EReplicatedObje
 				if (ensure(SubObjectOwnerReplicationInfo.GetState() < EReplicatedObjectState::PendingDestroy))
 				{
 					// Mark owner as dirty
-					ObjectsWithDirtyChanges.SetBit(SubObjectOwnerInternalIndex);
+					MarkObjectDirty(SubObjectOwnerInternalIndex, "DroppedCreated2");
 
 					// Indicate that we have dirty subobjects
 					SubObjectOwnerReplicationInfo.HasDirtySubObjects = 1U;
@@ -1280,7 +1284,7 @@ void FReplicationWriter::HandleDroppedRecord<FReplicationWriter::EReplicatedObje
 {
 	const uint32 InternalIndex = RecordInfo.Index;
 
-	check(CurrentState == EReplicatedObjectState::WaitOnDestroyConfirmation || CurrentState == EReplicatedObjectState::CancelPendingDestroy);
+	ensureMsgf(CurrentState == EReplicatedObjectState::WaitOnDestroyConfirmation || CurrentState == EReplicatedObjectState::CancelPendingDestroy, TEXT("Expected object ( InternalIndex: %u ) not to be in state %s"), InternalIndex, LexToString(CurrentState));
 
 	// If we want to cancel the destroy and lost the destroy packet we can resume normal replication.
 	if (CurrentState == EReplicatedObjectState::CancelPendingDestroy)
@@ -1306,6 +1310,7 @@ void FReplicationWriter::HandleDroppedRecord<FReplicationWriter::EReplicatedObje
 			if (OwnerInfo.GetState() < EReplicatedObjectState::PendingDestroy)
 			{
 				OwnerInfo.HasDirtySubObjects |= Info.HasDirtyChangeMask;
+				ensureMsgf(!bValidateObjectsWithDirtyChanges || OwnerInfo.GetState() != EReplicatedObjectState::Invalid, TEXT("Object (InternalIndex: % u) with Invalid state potentially marked dirty."), ObjectData.SubObjectRootIndex);
 				ObjectsWithDirtyChanges.SetBitValue(ObjectData.SubObjectRootIndex, ObjectsWithDirtyChanges.GetBit(ObjectData.SubObjectRootIndex) || OwnerInfo.HasDirtySubObjects);
 			}
 		}
@@ -1327,7 +1332,7 @@ void FReplicationWriter::HandleDroppedRecord<FReplicationWriter::EReplicatedObje
 		// We dropped a packet with tear-off data, that is a destroy with state data so we need to resend that state
 		if (RecordInfo.WroteTearOff)
 		{
-			ensureAlways(Info.TearOff);
+			ensureMsgf(Info.TearOff, TEXT("Expected object ( InternalIndex: %u ) to have TearOff set. Current state %s."), InternalIndex, LexToString(CurrentState));
 
 			SetState(InternalIndex, EReplicatedObjectState::PendingTearOff);
 
@@ -1346,7 +1351,7 @@ void FReplicationWriter::HandleDroppedRecord<FReplicationWriter::EReplicatedObje
 			Info.HasAttachments |= RecordInfo.HasAttachments;
 
 			// Mark object as having dirty changes
-			ObjectsWithDirtyChanges.SetBit(InternalIndex);
+			MarkObjectDirty(InternalIndex, "DroppedWaitOnDestroy");
 
 			// Mark parent as dirty
 			uint32 ParentInternalIndex = InternalIndex;
@@ -1356,7 +1361,7 @@ void FReplicationWriter::HandleDroppedRecord<FReplicationWriter::EReplicatedObje
 				if (ensure(ObjectData.SubObjectRootIndex != FNetRefHandleManager::InvalidInternalIndex))
 				{
 					ParentInternalIndex = ObjectData.SubObjectRootIndex;
-					ObjectsWithDirtyChanges.SetBit(ParentInternalIndex);
+					MarkObjectDirty(ParentInternalIndex, "DroppedWaitOnDestroy2");
 
 					FReplicationInfo& OwnerInfo = ReplicatedObjects[ObjectData.SubObjectRootIndex];
 					if (ensure(OwnerInfo.GetState() < EReplicatedObjectState::PendingDestroy))
@@ -1381,7 +1386,7 @@ void FReplicationWriter::HandleDroppedRecord<FReplicationWriter::EReplicatedObje
 			FReplicationInfo& OwnerInfo = ReplicatedObjects[ObjectData.SubObjectRootIndex];
 			if (OwnerInfo.GetState() < EReplicatedObjectState::PendingDestroy)
 			{
-				ObjectsWithDirtyChanges.SetBit(ObjectData.SubObjectRootIndex);
+				MarkObjectDirty(ObjectData.SubObjectRootIndex, "DroppedWaitOnDestroy2");
 				OwnerInfo.HasDirtySubObjects = 1U;
 
 				SetState(InternalIndex, EReplicatedObjectState::SubObjectPendingDestroy);
@@ -3346,4 +3351,19 @@ void FReplicationWriter::StopAllReplication()
 	}
 }
 
+void FReplicationWriter::MarkObjectDirty(FInternalNetRefIndex InternalIndex, const char* Caller)
+{
+	if (bValidateObjectsWithDirtyChanges)
+	{
+		const FReplicationInfo& ObjectInfo = ReplicatedObjects[InternalIndex];
+		if (!ensureMsgf(ObjectInfo.GetState() != EReplicatedObjectState::Invalid, TEXT("Object ( InternalIndex: %u ) with Invalid state marked dirty. Caller: %hs"), InternalIndex, Caller))
+		{
+			return;
+		}
+	}
+
+	ObjectsWithDirtyChanges.SetBit(InternalIndex);
+}
+
+
 }
diff --git a/Engine/Source/Runtime/Experimental/Iris/Core/Private/Iris/ReplicationSystem/ReplicationWriter.h b/Engine/Source/Runtime/Experimental/Iris/Core/Private/Iris/ReplicationSystem/ReplicationWriter.h
--- a/Engine/Source/Runtime/Experimental/Iris/Core/Private/Iris/ReplicationSystem/ReplicationWriter.h
+++ b/Engine/Source/Runtime/Experimental/Iris/Core/Private/Iris/ReplicationSystem/ReplicationWriter.h
@@ -472,6 +472,8 @@ private:
 	void DiscardAllRecords();
 	void StopAllReplication();
 
+	void MarkObjectDirty(FInternalNetRefIndex InternalIndex, const char* Caller);
+
 private:
 	// Replication parameters
 	FReplicationParameters Parameters;
