fix an issue where clustered particles may not be removed from the StaticClusteredMapArray and KinematicClusteredMapArray arrays when destroyed

diff --git a/Engine/Source/Runtime/Experimental/Chaos/Public/Chaos/PBDRigidsSOAs.h b/Engine/Source/Runtime/Experimental/Chaos/Public/Chaos/PBDRigidsSOAs.h
--- a/Engine/Source/Runtime/Experimental/Chaos/Public/Chaos/PBDRigidsSOAs.h
+++ b/Engine/Source/Runtime/Experimental/Chaos/Public/Chaos/PBDRigidsSOAs.h
@@ -392,7 +392,7 @@ public:
 			}
 			else if (FPBDRigidClusteredParticleHandle* PBDRigidClustered = Particle->CastToClustered())
 			{
-				DynamicClusteredMapArray.Remove(PBDRigidClustered);
+				RemoveClusteredParticle(PBDRigidClustered);
 			}
 
 			// Check for sleep events referencing this particle
