Chaos - enable manifold reuse for the many-shape collision detection path

diff --git a/Engine/Source/Runtime/Experimental/Chaos/Private/Chaos/Collision/PBDCollisionConstraint.cpp b/Engine/Source/Runtime/Experimental/Chaos/Private/Chaos/Collision/PBDCollisionConstraint.cpp
--- a/Engine/Source/Runtime/Experimental/Chaos/Private/Chaos/Collision/PBDCollisionConstraint.cpp
+++ b/Engine/Source/Runtime/Experimental/Chaos/Private/Chaos/Collision/PBDCollisionConstraint.cpp
@@ -1,6 +1,7 @@
 // Copyright Epic Games, Inc. All Rights Reserved.
 #include "Chaos/Collision/PBDCollisionConstraint.h"
 #include "Chaos/Collision/CollisionConstraintAllocator.h"
+#include "Chaos/Collision/CollisionUtil.h"
 #include "Chaos/Collision/PBDCollisionConstraintHandle.h"
 #include "Chaos/Collision/PBDCollisionContainerSolver.h"
 #include "Chaos/Collision/PBDCollisionSolver.h"
@@ -52,6 +53,9 @@ namespace Chaos
 	FRealSingle Chaos_Collision_Stiffness = -1.0f;
 	FAutoConsoleVariableRef CVarChaos_Collision_Stiffness(TEXT("p.Chaos.Collision.Stiffness"), Chaos_Collision_Stiffness, TEXT("Override the collision solver stiffness (if >= 0)"));
 
+	bool bChaos_Collision_EnableBoundsChecks = true;
+	FAutoConsoleVariableRef CVarChaos_Collision_EnableBoundsChecks(TEXT("p.Chaos.Collision.EnableBoundsChecks"), bChaos_Collision_EnableBoundsChecks, TEXT(""));
+
 	struct FCollisionTolerances
 	{
 		// Multiplied by the contact margin to produce a distance within which contacts are considered to be the same point
@@ -74,7 +78,6 @@ namespace Chaos
 	// @todo(chaos): tune the tolerances used in FPBDCollisionConstraint::UpdateAndTryRestoreManifold
 	FCollisionTolerances Chaos_Manifold_Tolerances;
 
-
 	FString FPBDCollisionConstraint::ToString() const
 	{
 		return FString::Printf(TEXT("Particle:%s, Levelset:%s, AccumulatedImpulse:%s"), *Particle[0]->ToString(), *Particle[1]->ToString(), *AccumulatedImpulse.ToString());
@@ -278,11 +281,14 @@ namespace Chaos
 
 		CullDistance = FRealSingle(InCullDistance);
 
-		// Initialize the is-probe and is-probe-unmodified flags to the same value.
-		// Contact modification may change bIsProbe but we want to store the unmodified value
-		// so that it can be reset between frames.
-		Flags.bIsProbe = (Shape[0] && Shape[0]->GetIsProbe()) || (Shape[1] && Shape[1]->GetIsProbe());
-		Flags.bIsProbeUnmodified = Flags.bIsProbe;
+		FRealSingle DistanceCheckSize = 0;
+		BoundsTestFlags = Private::CalculateImplicitBoundsTestFlags(
+			GetParticle0(), GetImplicit0(), GetShape0(),
+			GetParticle1(), GetImplicit1(), GetShape1(),
+			DistanceCheckSize);
+
+		// Contact modification may change bIsProbe. The original value is stored in BoundsTestFlags.
+		Flags.bIsProbe = BoundsTestFlags.bIsProbe;
 
 		// If this currently a CCD contact (this may change on later ticks)
 		Flags.bCCDEnabled = (InCCDType == ECollisionCCDType::Enabled);
@@ -290,8 +296,8 @@ namespace Chaos
 		// Initialize tolerances that depend on shape type etc, also the bIsQuadratic flags
 		const FRealSingle Margin0 = FRealSingle(GetImplicit0()->GetMargin());
 		const FRealSingle Margin1 = FRealSingle(GetImplicit1()->GetMargin());
-		const EImplicitObjectType ImplicitType0 = Collisions::GetImplicitCollisionType(GetParticle0(), GetImplicit0());
-		const EImplicitObjectType ImplicitType1 = Collisions::GetImplicitCollisionType(GetParticle1(), GetImplicit1());
+		const EImplicitObjectType ImplicitType0 = Private::GetImplicitCollisionType(GetParticle0(), GetImplicit0());
+		const EImplicitObjectType ImplicitType1 = Private::GetImplicitCollisionType(GetParticle1(), GetImplicit1());
 		InitMarginsAndTolerances(ImplicitType0, ImplicitType1, Margin0, Margin1);
 
 		// Are we allowing manifolds?
diff --git a/Engine/Source/Runtime/Experimental/Chaos/Private/Chaos/Collision/ParticlePairMidPhase.cpp b/Engine/Source/Runtime/Experimental/Chaos/Private/Chaos/Collision/ParticlePairMidPhase.cpp
--- a/Engine/Source/Runtime/Experimental/Chaos/Private/Chaos/Collision/ParticlePairMidPhase.cpp
+++ b/Engine/Source/Runtime/Experimental/Chaos/Private/Chaos/Collision/ParticlePairMidPhase.cpp
@@ -4,6 +4,7 @@
 #include "Chaos/Collision/CollisionConstraintAllocator.h"
 #include "Chaos/Collision/CollisionContext.h"
 #include "Chaos/Collision/CollisionFilter.h"
+#include "Chaos/Collision/CollisionUtil.h"
 #include "Chaos/Collision/ContactTriangles.h"
 #include "Chaos/Collision/PBDCollisionConstraint.h"
 #include "Chaos/CollisionResolution.h"
@@ -38,9 +39,6 @@ namespace Chaos
 
 	namespace CVars
 	{
-		bool bChaos_Collision_MidPhase_EnableBoundsChecks = true;
-		FAutoConsoleVariableRef CVarChaos_Collision_EnableBoundsChecks(TEXT("p.Chaos.Collision.EnableBoundsChecks"), bChaos_Collision_MidPhase_EnableBoundsChecks, TEXT(""));
-
 		Chaos::FRealSingle Chaos_Collision_CullDistanceScaleInverseSize = 0.01f;	// 100cm
 		Chaos::FRealSingle Chaos_Collision_MinCullDistanceScale = 1.0f;
 		FAutoConsoleVariableRef CVarChaos_Collision_CullDistanceReferenceSize(TEXT("p.Chaos.Collision.CullDistanceReferenceSize"), Chaos_Collision_CullDistanceScaleInverseSize, TEXT(""));
@@ -220,39 +218,15 @@ namespace Chaos
 		, SphereBoundsCheckSize(0)
 		, LastUsedEpoch(-1)
 		, ShapePairType(InShapePairType)
-		, Flags()
+		, BoundsTestFlags()
 	{
-		const FImplicitObject* Implicit0 = Shape0->GetLeafGeometry();
-		const FImplicitObject* Implicit1 = Shape1->GetLeafGeometry();
-		const bool bHasBounds0 = (Implicit0 != nullptr) && Implicit0->HasBoundingBox();
-		const bool bHasBounds1 = (Implicit1 != nullptr) && Implicit1->HasBoundingBox();
-		const EImplicitObjectType ImplicitType0 = (Implicit0 != nullptr) ? GetInnerType(Implicit0->GetCollisionType()) : ImplicitObjectType::Unknown;
-		const EImplicitObjectType ImplicitType1 = (Implicit1 != nullptr) ? GetInnerType(Implicit1->GetCollisionType()) : ImplicitObjectType::Unknown;
-		const bool bIsSphere0 = (ImplicitType0 == ImplicitObjectType::Sphere);
-		const bool bIsSphere1 = (ImplicitType1 == ImplicitObjectType::Sphere);
-		const bool bIsCapsule0 = (ImplicitType0 == ImplicitObjectType::Capsule);
-		const bool bIsCapsule1 = (ImplicitType1 == ImplicitObjectType::Capsule);
-		const bool bIsTriangle0 = (ImplicitType0 == ImplicitObjectType::TriangleMesh) || (ImplicitType0 == ImplicitObjectType::HeightField);
-		const bool bIsTriangle1 = (ImplicitType1 == ImplicitObjectType::TriangleMesh) || (ImplicitType1 == ImplicitObjectType::HeightField);
-		const bool bIsLevelSet = ((ShapePairType == EContactShapesType::LevelSetLevelSet) || (ShapePairType == EContactShapesType::Unknown));
-
-		const bool bAllowBoundsChecked = bChaos_Collision_MidPhase_EnableBoundsChecks && bHasBounds0 && bHasBounds1;
-		Flags.bEnableAABBCheck = bAllowBoundsChecked && !(bIsSphere0 && bIsSphere1);	// No AABB test if both are spheres
-		Flags.bEnableOBBCheck0 = bAllowBoundsChecked && !bIsSphere0;					// No OBB test for spheres
-		Flags.bEnableOBBCheck1 = bAllowBoundsChecked && !bIsSphere1;					// No OBB test for spheres
-
-		if (bAllowBoundsChecked && bIsSphere0 && bIsSphere1)
-		{
-			SphereBoundsCheckSize = FRealSingle(Implicit0->GetMargin() + Implicit1->GetMargin());	// Sphere-Sphere bounds test
-		}
-
-		// Do not try to reuse manifold points for capsules or spheres (against anything)
-		// NOTE: This can also be disabled for all shape types by the solver (see GenerateCollisionImpl and the Context)
-		Flags.bEnableManifoldUpdate = !bIsSphere0 && !bIsSphere1 && !bIsCapsule0 && !bIsCapsule1 && !bIsTriangle0 && !bIsTriangle1 && !bIsLevelSet;
+		const FImplicitObject* Implicit0 = InShape0->GetLeafGeometry();
+		const FImplicitObject* Implicit1 = InShape1->GetLeafGeometry();
 
-		// Mark probe flag now so we know which GenerateCollisions to use
-		// @todo(chaos): it looks like this can be changed by a collision modifier so we should not be caching it
-		Flags.bIsProbe = Shape0->GetIsProbe() || Shape1->GetIsProbe();
+		BoundsTestFlags = Private::CalculateImplicitBoundsTestFlags(
+			InParticle0, Implicit0, InShape0,
+			InParticle1, Implicit1, InShape1,
+			SphereBoundsCheckSize);
 	}
 
 	FSingleShapePairCollisionDetector::~FSingleShapePairCollisionDetector()
@@ -268,7 +242,7 @@ namespace Chaos
 		, Shape1(R.Shape1)
 		, SphereBoundsCheckSize(R.SphereBoundsCheckSize)
 		, ShapePairType(R.ShapePairType)
-		, Flags(R.Flags)
+		, BoundsTestFlags(R.BoundsTestFlags)
 	{
 	}
 
@@ -280,7 +254,7 @@ namespace Chaos
 		const FAABB3& ShapeWorldBounds1 = Shape1->GetWorldSpaceInflatedShapeBounds();
 
 		// World-space expanded bounds check
-		if (Flags.bEnableAABBCheck)
+		if (BoundsTestFlags.bEnableAABBCheck)
 		{
 			if (!ShapeWorldBounds0.Intersects(ShapeWorldBounds1))
 			{
@@ -289,7 +263,7 @@ namespace Chaos
 		}
 
 		// World-space sphere bounds check
-		if (SphereBoundsCheckSize > FRealSingle(0))
+		if (BoundsTestFlags.bEnableDistanceCheck && (SphereBoundsCheckSize > FRealSingle(0)))
 		{
 			const FVec3 Separation = ShapeWorldBounds0.GetCenter() - ShapeWorldBounds1.GetCenter();
 			const FReal SeparationSq = Separation.SizeSquared();
@@ -308,14 +282,14 @@ namespace Chaos
 		// we might call the narrow phase one time too many when shapes become separated.
 		const int32 LastEpoch = CurrentEpoch - 1;
 		const bool bCollidedLastTick = IsUsedSince(LastEpoch);
-		if ((Flags.bEnableOBBCheck0 || Flags.bEnableOBBCheck1) && !bCollidedLastTick)
+		if ((BoundsTestFlags.bEnableOBBCheck0 || BoundsTestFlags.bEnableOBBCheck1) && !bCollidedLastTick)
 		{
 			const FRigidTransform3& ShapeWorldTransform0 = Shape0->GetLeafWorldTransform(GetParticle0());
 			const FRigidTransform3& ShapeWorldTransform1 = Shape1->GetLeafWorldTransform(GetParticle1());
 			const FImplicitObject* Implicit0 = Shape0->GetLeafGeometry();
 			const FImplicitObject* Implicit1 = Shape1->GetLeafGeometry();
 
-			if (Flags.bEnableOBBCheck0)
+			if (BoundsTestFlags.bEnableOBBCheck0)
 			{
 				if (!ImplicitOverlapOBBToAABB(Implicit0, Implicit1, ShapeWorldTransform0, ShapeWorldTransform1, CullDistance))
 				{
@@ -323,7 +297,7 @@ namespace Chaos
 				}
 			}
 
-			if (Flags.bEnableOBBCheck1)
+			if (BoundsTestFlags.bEnableOBBCheck1)
 			{
 				if (!ImplicitOverlapOBBToAABB(Implicit1, Implicit0, ShapeWorldTransform1, ShapeWorldTransform0, CullDistance))
 				{
@@ -376,7 +350,7 @@ namespace Chaos
 		const FReal Dt,
 		const FCollisionContext& Context)
 	{
-		if (Flags.bIsProbe)
+		if (BoundsTestFlags.bIsProbe)
 		{
 			return GenerateCollisionProbeImpl(CullDistance, Dt, Context);
 		}
@@ -414,7 +388,7 @@ namespace Chaos
 			}
 			
 			bool bWasManifoldRestored = false;
-			const bool bAllowManifoldRestore = Context.GetSettings().bAllowManifoldReuse && Flags.bEnableManifoldUpdate;
+			const bool bAllowManifoldRestore = Context.GetSettings().bAllowManifoldReuse && BoundsTestFlags.bEnableManifoldUpdate;
 			if (bAllowManifoldRestore && bWasUpdatedLastTick && Constraint->GetCanRestoreManifold())
 			{
 				// Update the existing manifold. We can re-use as-is if none of the points have moved much and the bodies have not moved much
@@ -465,7 +439,7 @@ namespace Chaos
 		const FReal Dt,
 		const FCollisionContext& Context)
 	{
-		if (Flags.bIsProbe)
+		if (BoundsTestFlags.bIsProbe)
 		{
 			return GenerateCollisionProbeImpl(CullDistance, Dt, Context);
 		}
@@ -955,10 +929,10 @@ namespace Chaos
 	void FShapePairParticlePairMidPhase::TryAddShapePair(const FPerShapeData* Shape0, const FPerShapeData* Shape1)
 	{
 		const FImplicitObject* Implicit0 = Shape0->GetLeafGeometry();
-		const EImplicitObjectType ImplicitType0 = Collisions::GetImplicitCollisionType(Particle0, Implicit0);
+		const EImplicitObjectType ImplicitType0 = Private::GetImplicitCollisionType(Particle0, Implicit0);
 
 		const FImplicitObject* Implicit1 = Shape1->GetLeafGeometry();
-		const EImplicitObjectType ImplicitType1 = Collisions::GetImplicitCollisionType(Particle1, Implicit1);
+		const EImplicitObjectType ImplicitType1 = Private::GetImplicitCollisionType(Particle1, Implicit1);
 
 		const bool bDoPassFilter = ShapePairNarrowPhaseFilter(ImplicitType0, Shape0, ImplicitType1, Shape1);
 		if (bDoPassFilter)
@@ -1334,19 +1308,28 @@ namespace Chaos
 		const FCollisionContext& Context)
 	{
 		// Check the sim filter to see if these shapes collide
-		const EImplicitObjectType ImplicitTypeA = Collisions::GetImplicitCollisionType(ParticleA, ImplicitA);
-		const EImplicitObjectType ImplicitTypeB = Collisions::GetImplicitCollisionType(ParticleB, ImplicitB);
+		const EImplicitObjectType ImplicitTypeA = Private::GetImplicitCollisionType(ParticleA, ImplicitA);
+		const EImplicitObjectType ImplicitTypeB = Private::GetImplicitCollisionType(ParticleB, ImplicitB);
 		const bool bDoPassFilter = ShapePairNarrowPhaseFilter(ImplicitTypeA, ShapeInstanceA, ImplicitTypeB, ShapeInstanceB);
 		if (!bDoPassFilter)
 		{
 			return;
 		}
 
+		// Calculate which bounds tests we should run for this shape pair
+		// @todo(chaos): it is too expensive to call this all the time. We should probably build it into DoBoundsOverlap and only perform operation when we pass te AABB test
+		//Private::FImplicitBoundsTestFlags BoundsTestFlags = Private::CalculateImplicitBoundsTestFlags(ParticleA, ImplicitA, ShapeInstanceA, ParticleB, ImplicitB, ShapeInstanceB, DistanceCheckSize);
+		FRealSingle DistanceCheckSize = 0;
+		Private::FImplicitBoundsTestFlags BoundsTestFlags;
+		BoundsTestFlags.bEnableAABBCheck = true;
+		BoundsTestFlags.bEnableOBBCheck0 = true;
+		BoundsTestFlags.bEnableOBBCheck1 = true;
+
 		// Do the objects bounds overlap?
 		const bool bDoOverlap = DoBoundsOverlap(
 			ImplicitA, ParticleWorldTransformA, RelativeTransformA,
 			ImplicitB, ParticleWorldTransformB, RelativeTransformB,
-			CullDistance);
+			BoundsTestFlags, DistanceCheckSize, CullDistance);
 
 		if (!bDoOverlap)
 		{
@@ -1397,12 +1380,10 @@ namespace Chaos
 		const FImplicitObject* ImplicitB,
 		const FRigidTransform3& ParticleWorldTransformB,
 		const FRigidTransform3& ShapeRelativeTransformB,
+		const Private::FImplicitBoundsTestFlags BoundsTestFlags,
+		const FRealSingle DistanceCheckSize,
 		const FReal CullDistance)
 	{
-		const bool bEnableAABBCheck = true;
-		const bool bEnableOBBCheckA = true;
-		const bool bEnableOBBCheckB = true;
-
 		const FRigidTransform3 ShapeWorldTransformA = ShapeRelativeTransformA * ParticleWorldTransformA;
 		const FRigidTransform3 ShapeWorldTransformB = ShapeRelativeTransformB * ParticleWorldTransformB;
 
@@ -1411,7 +1392,7 @@ namespace Chaos
 		const FAABB3 ShapeWorldBoundsB = ImplicitB->CalculateTransformedBounds(ShapeWorldTransformB);
 
 		// World-space expanded bounds check
-		if (bEnableAABBCheck)
+		if (BoundsTestFlags.bEnableAABBCheck)
 		{
 			if (!ShapeWorldBoundsA.Intersects(ShapeWorldBoundsB))
 			{
@@ -1419,22 +1400,37 @@ namespace Chaos
 			}
 		}
 
-		// OBB-AABB test in both directions. This is beneficial for shapes which do not fit their AABBs very well,
-		// which includes boxes and other shapes that are not roughly spherical. It is especially beneficial when
-		// one shape is long and thin (i.e., it does not fit an AABB well when the shape is rotated).
-		if (bEnableOBBCheckA)
+		// World-space sphere bounds check
+		if (BoundsTestFlags.bEnableDistanceCheck && (DistanceCheckSize > FRealSingle(0)))
 		{
-			if (!ImplicitOverlapOBBToAABB(ImplicitA, ImplicitB, ShapeWorldTransformA, ShapeWorldTransformB, CullDistance))
+			const FVec3 Separation = ShapeWorldBoundsA.GetCenter() - ShapeWorldBoundsB.GetCenter();
+			const FReal SeparationSq = Separation.SizeSquared();
+			const FReal CullDistanceSq = FMath::Square(CullDistance + FReal(DistanceCheckSize));
+			if (SeparationSq > CullDistanceSq)
 			{
 				return false;
 			}
 		}
 
-		if (bEnableOBBCheckB)
+		if (BoundsTestFlags.bEnableOBBCheck0 || BoundsTestFlags.bEnableOBBCheck1)
 		{
-			if (!ImplicitOverlapOBBToAABB(ImplicitB, ImplicitA, ShapeWorldTransformB, ShapeWorldTransformA, CullDistance))
+			// OBB-AABB test in both directions. This is beneficial for shapes which do not fit their AABBs very well,
+			// which includes boxes and other shapes that are not roughly spherical. It is especially beneficial when
+			// one shape is long and thin (i.e., it does not fit an AABB well when the shape is rotated).
+			if (BoundsTestFlags.bEnableOBBCheck0)
 			{
-				return false;
+				if (!ImplicitOverlapOBBToAABB(ImplicitA, ImplicitB, ShapeWorldTransformA, ShapeWorldTransformB, CullDistance))
+				{
+					return false;
+				}
+			}
+
+			if (BoundsTestFlags.bEnableOBBCheck1)
+			{
+				if (!ImplicitOverlapOBBToAABB(ImplicitB, ImplicitA, ShapeWorldTransformB, ShapeWorldTransformA, CullDistance))
+				{
+					return false;
+				}
 			}
 		}
 
@@ -1494,12 +1490,10 @@ namespace Chaos
 		{
 			// NOTE: Using InParticle0 and InParticle1 here because the order may be different to what we have stored
 			Constraint = CreateConstraint(InParticle0, InImplicit0, InShape0, InBVHParticles0, InShapeRelativeTransform0, InParticle1, InImplicit1, InShape1, InBVHParticles1, InShapeRelativeTransform1, CullDistance, ShapePairType, bUseManifold, CollisionKey, Context);
-
-			// Is this a CCD constraint?
-			Constraint->SetCCDEnabled(IsCCD());
 		}
 
 		NewConstraints.Add(Constraint);
+
 		return Constraint;
 	}
 
@@ -1542,21 +1536,10 @@ namespace Chaos
 		Constraint->GetContainerCookie().bIsMultiShapePair = true;
 		Constraint->GetContainerCookie().CreationEpoch = CurrentEpoch;
 
-		return Constraints.Add(CollisionKey.GetKey(), MoveTemp(Constraint)).Get();
-	}
+		// Is this a CCD constraint?
+		Constraint->SetCCDEnabled(IsCCD());
 
-	void FGenericParticlePairMidPhase::WakeCollisionsImpl(const int32 SleepEpoch, const int32 CurrentEpoch)
-	{
-		for (auto& KVP : Constraints)
-		{
-			FPBDCollisionConstraintPtr& Constraint = KVP.Value;
-			if (Constraint->GetContainerCookie().LastUsedEpoch >= SleepEpoch)
-			{
-				Constraint->GetContainerCookie().LastUsedEpoch = CurrentEpoch;
-				Constraint->GetContainerCookie().ConstraintIndex = INDEX_NONE;
-				Constraint->GetContainerCookie().CCDConstraintIndex = INDEX_NONE;
-			}
-		}
+		return Constraints.Add(CollisionKey.GetKey(), MoveTemp(Constraint)).Get();
 	}
 
 	void PrefetchConstraint(const TArray<FPBDCollisionConstraint*>& Constraints, const int32 ConstraintIndex)
@@ -1588,20 +1571,30 @@ namespace Chaos
 				Constraint->SetCCDSweepEnabled(bUseCCDSweep);
 			}
 
-			bool bIsActive = false;
-			if (!bUseCCDSweep)
-			{
-				bIsActive = UpdateCollision(Constraint, CullDistance, Dt, Context);
-			}
-			else
+			// NOTE: Probe constraints are always active and we run collision detection for them at the end opf the frame
+			bool bIsActive = true;
+
+			if (!Constraint->IsProbe())
 			{
-				bIsActive = UpdateCollisionCCD(Constraint, CullDistance, Dt, Context);
+				if (!bUseCCDSweep)
+				{
+					bIsActive = UpdateCollision(Constraint, CullDistance, Dt, Context);
+				}
+				else
+				{
+					bIsActive = UpdateCollisionCCD(Constraint, CullDistance, Dt, Context);
+				}
 			}
 
 			if (bIsActive)
 			{
+				Context.GetAllocator()->ActivateConstraint(Constraint);
 				++NumActive;
 			}
+			else
+			{
+				Constraint->SetDisabled(true);
+			}
 		}
 
 		NewConstraints.Reset();
@@ -1637,15 +1630,14 @@ namespace Chaos
 		}
 
 		bool bWasManifoldRestored = false;
-		// @todo(chaos): enable manifold restore for Generic midphase
-		//if (Context.GetSettings().bAllowManifoldReuse && Flags.bEnableManifoldUpdate && bWasUpdatedLastTick)
-		//{
-		//	// Update the existing manifold. We can re-use as-is if none of the points have moved much and the bodies have not moved much
-		//	// NOTE: this can succeed in "restoring" even if we have no manifold points
-		//	// NOTE: this uses the transforms from SetLastShapeWorldTransforms, so we can only do this if we were updated last tick
-		//	bWasManifoldRestored = Constraint->UpdateAndTryRestoreManifold();
-		//}
-		//else
+		if (Context.GetSettings().bAllowManifoldReuse && bWasUpdatedLastTick && Constraint->GetCanRestoreManifold())
+		{
+			// Update the existing manifold. We can re-use as-is if none of the points have moved much and the bodies have not moved much
+			// NOTE: this can succeed in "restoring" even if we have no manifold points
+			// NOTE: this uses the transforms from SetLastShapeWorldTransforms, so we can only do this if we were updated last tick
+			bWasManifoldRestored = Constraint->TryRestoreManifold();
+		}
+		else
 		{
 			// We are not trying to reuse manifold points, so reset them but leave stored data intact (for friction)
 			Constraint->ResetActiveManifoldContacts();
@@ -1653,29 +1645,23 @@ namespace Chaos
 
 		if (!bWasManifoldRestored)
 		{
-			// We will be updating the manifold so update transforms used to check for movement in UpdateAndTryRestoreManifold on future ticks
-			Constraint->SetLastShapeWorldTransforms(ShapeWorldTransform0, ShapeWorldTransform1);
-
 			if (!Context.GetSettings().bDeferNarrowPhase)
 			{
 				Collisions::UpdateConstraint(*Constraint, Constraint->GetShapeWorldTransform0(), Constraint->GetShapeWorldTransform1(), Dt);
 			}
+
+			// We will be updating the manifold so update transforms used to check for movement in UpdateAndTryRestoreManifold on future ticks
+			// NOTE: We call this after Collisions::UpdateConstraint because it may reset the manifold and reset the bCanRestoreManifold flag.
+			// @todo(chaos): Collisions::UpdateConstraint does not need to reset the manifold - fix that
+			Constraint->SetLastShapeWorldTransforms(ShapeWorldTransform0, ShapeWorldTransform1);
 		}
 
 		// If we have a valid contact, add it to the active list
 		// We also add it to the active list if collision detection is deferred because the data will be filled in later and we
 		// don't know in advance whether we will pass the Phi check (deferred narrow phase is used with RBAN)
-		if (Constraint->GetPhi() <= CullDistance || Context.GetSettings().bDeferNarrowPhase)
-		{
-			if (Context.GetAllocator()->ActivateConstraint(Constraint))
-			{
-				return true;
-			}
-		}
+		const bool bIsActive = (Constraint->GetPhi() <= CullDistance || Context.GetSettings().bDeferNarrowPhase);
 
-		// If we get here, we did not activate the constraint and it should be disabled for this tick
-		Constraint->SetDisabled(true);
-		return false;
+		return bIsActive;
 	}
 
 
@@ -1744,7 +1730,6 @@ namespace Chaos
 			Constraint->SetCCDSweepEnabled(false);
 		}
 
-		Context.GetAllocator()->ActivateConstraint(Constraint);
 		return true;
 	}
 
diff --git a/Engine/Source/Runtime/Experimental/Chaos/Private/Chaos/CollisionResolution.cpp b/Engine/Source/Runtime/Experimental/Chaos/Private/Chaos/CollisionResolution.cpp
--- a/Engine/Source/Runtime/Experimental/Chaos/Private/Chaos/CollisionResolution.cpp
+++ b/Engine/Source/Runtime/Experimental/Chaos/Private/Chaos/CollisionResolution.cpp
@@ -14,6 +14,7 @@
 #include "Chaos/Collision/CollisionContext.h"
 #include "Chaos/Collision/CollisionConstraintAllocator.h"
 #include "Chaos/Collision/CollisionFilter.h"
+#include "Chaos/Collision/CollisionUtil.h"
 #include "Chaos/Collision/ContactPointsMiscShapes.h"
 #include "Chaos/Collision/GJKContactPointSwept.h"
 #include "Chaos/Collision/GJKContactPoint.h"
@@ -41,24 +42,10 @@
 // @todo(chaos): clean up the contact creation time rejection to avoid extra transforms
 #define CHAOS_COLLISION_CREATE_BOUNDSCHECK 1
 
-DECLARE_CYCLE_STAT(TEXT("Collisions::ConstructConstraints"), STAT_Collisions_ConstructConstraints, STATGROUP_ChaosCollision);
-DECLARE_CYCLE_STAT(TEXT("Collisions::ConstructConstraintsInternal"), STAT_Collisions_ConstructConstraintsInternal, STATGROUP_ChaosCollision);
-DECLARE_CYCLE_STAT(TEXT("Collisions::FindAllIntersectingClusteredObjects"), STAT_Collisions_FindAllIntersectingClusteredObjects, STATGROUP_ChaosCollision);
-DECLARE_CYCLE_STAT(TEXT("Collisions::ConstructGenericConvexConvexConstraints"), STAT_Collisions_ConstructGenericConvexConvexConstraints, STATGROUP_ChaosCollision);
-DECLARE_CYCLE_STAT(TEXT("Collisions::ConstructGenericConvexConvexConstraintsSwept"), STAT_Collisions_ConstructGenericConvexConvexConstraintsSwept, STATGROUP_ChaosCollision);
-DECLARE_CYCLE_STAT(TEXT("Collisions::UpdateConstraintFromGeometryInternal"), STAT_Collisions_UpdateConstraintFromGeometryInternal, STATGROUP_ChaosCollision);
 DECLARE_CYCLE_STAT(TEXT("Collisions::UpdateGenericConvexConvexConstraint"), STAT_Collisions_UpdateGenericConvexConvexConstraint, STATGROUP_ChaosCollision);
 DECLARE_CYCLE_STAT(TEXT("Collisions::UpdateConvexTriangleMeshConstraint"), STAT_Collisions_UpdateConvexTriangleMeshConstraint, STATGROUP_ChaosCollision);
 DECLARE_CYCLE_STAT(TEXT("Collisions::UpdateConvexTriangleMeshConstraintSwept"), STAT_Collisions_UpdateConvexTriangleMeshConstraintSwept, STATGROUP_ChaosCollision);
 DECLARE_CYCLE_STAT(TEXT("Collisions::UpdateConvexHeightFieldConstraint"), STAT_Collisions_UpdateConvexHeightFieldConstraint, STATGROUP_ChaosCollision);
-DECLARE_CYCLE_STAT(TEXT("Collisions::ConstructConvexHeightFieldConstraints"), STAT_Collisions_ConstructConvexHeightFieldConstraints, STATGROUP_ChaosCollision);
-DECLARE_CYCLE_STAT(TEXT("Collisions::ConstructConvexHeightFieldConstraintsSwept"), STAT_Collisions_ConstructConvexHeightFieldConstraintsSwept, STATGROUP_ChaosCollision);
-DECLARE_CYCLE_STAT(TEXT("Collisions::ConstructConvexTriangleMeshConstraints"), STAT_Collisions_ConstructConvexTriangleMeshConstraints, STATGROUP_ChaosCollision);
-DECLARE_CYCLE_STAT(TEXT("Collisions::ConstructConvexTriangleMeshConstraintsSwept"), STAT_Collisions_ConstructConvexTriangleMeshConstraintsSwept, STATGROUP_ChaosCollision);
-DECLARE_CYCLE_STAT(TEXT("Collisions::ConstructCapsuleTriangleMeshConstraintsSwept"), STAT_Collisions_ConstructCapsuleTriangleMeshConstraintsSwept, STATGROUP_ChaosCollision);
-DECLARE_CYCLE_STAT(TEXT("Collisions::ConstructCapsuleTriangleMeshConstraints"), STAT_Collisions_ConstructCapsuleTriangleMeshConstraints, STATGROUP_ChaosCollision);
-DECLARE_CYCLE_STAT(TEXT("Collisions::ConstructCapsuleHeightFieldConstraints"), STAT_Collisions_ConstructCapsuleHeightFieldConstraints, STATGROUP_ChaosCollision);
-DECLARE_CYCLE_STAT(TEXT("Collisions::ConstructCapsuleHeightFieldConstraintsSwept"), STAT_Collisions_ConstructCapsuleHeightFieldConstraintsSwept, STATGROUP_ChaosCollision);
 DECLARE_CYCLE_STAT(TEXT("Collisions::UpdateGenericConvexConvexConstraintSwept"), STAT_Collisions_UpdateGenericConvexConvexConstraintSwept, STATGROUP_ChaosCollision);
 DECLARE_CYCLE_STAT(TEXT("Collisions::UpdateCapsuleHeightFieldConstraintSwept"), STAT_Collisions_UpdateCapsuleHeightFieldConstraintSwept, STATGROUP_ChaosCollision);
 DECLARE_CYCLE_STAT(TEXT("Collisions::UpdateCapsuleHeightFieldConstraint"), STAT_Collisions_UpdateCapsuleHeightFieldConstraint, STATGROUP_ChaosCollision);
@@ -67,10 +54,6 @@ DECLARE_CYCLE_STAT(TEXT("Collisions::UpdateCapsuleTriangleMeshConstraintSwept"),
 DECLARE_CYCLE_STAT(TEXT("Collisions::UpdateConvexHeightFieldConstraintSwept"), STAT_Collisions_UpdateConvexHeightFieldConstraintSwept, STATGROUP_ChaosCollision);
 DECLARE_CYCLE_STAT(TEXT("Collisions::UpdateLevelsetLevelsetConstraint"), STAT_UpdateLevelsetLevelsetConstraint, STATGROUP_ChaosCollision);
 
-// Stat Collision counters (need to be reset every advance)
-DECLARE_DWORD_ACCUMULATOR_STAT(TEXT("NumParticlePairs"), STAT_ChaosCollisionCounter_NumParticlePairs, STATGROUP_ChaosCollisionCounters);
-DECLARE_DWORD_ACCUMULATOR_STAT(TEXT("NumShapePairs"), STAT_ChaosCollisionCounter_NumShapePairs, STATGROUP_ChaosCollisionCounters);
-DECLARE_DWORD_ACCUMULATOR_STAT(TEXT("NumContactsCreated"), STAT_ChaosCollisionCounter_NumContactsCreated, STATGROUP_ChaosCollisionCounters);
 DECLARE_DWORD_ACCUMULATOR_STAT(TEXT("NumContactUpdates"), STAT_ChaosCollisionCounter_NumContactUpdates, STATGROUP_ChaosCollisionCounters);
 
 int32 CCDUseGenericSweptConvexConstraints = 1;
@@ -100,13 +83,6 @@ namespace Chaos
 {
 	namespace Collisions
 	{
-		void ResetChaosCollisionCounters()
-		{
-			SET_DWORD_STAT(STAT_ChaosCollisionCounter_NumParticlePairs, 0);
-			SET_DWORD_STAT(STAT_ChaosCollisionCounter_NumShapePairs, 0);
-			SET_DWORD_STAT(STAT_ChaosCollisionCounter_NumContactsCreated, 0);
-			SET_DWORD_STAT(STAT_ChaosCollisionCounter_NumContactUpdates, 0);
-		}
 
 		// Traits to control how contacts are generated
 		template<bool B_IMMEDIATEUPDATE>
@@ -1171,120 +1147,6 @@ namespace Chaos
 			return false;
 		}
 
-		bool DoParticleImplicitsOverlap(
-			const TGeometryParticleHandle<FReal, 3>* Particle0,
-			const TGeometryParticleHandle<FReal, 3>* Particle1,
-			const FImplicitObject* Implicit0,
-			const FImplicitObject* Implicit1,
-			const FRigidTransform3& ParticleWorldTransform0,
-			const FRigidTransform3& LocalTransform0,
-			const FRigidTransform3& ParticleWorldTransform1,
-			const FRigidTransform3& LocalTransform1,
-			const FReal CullDistance,
-			const bool bEnableSweep,
-			const FCollisionContext& Context)
-		{
-			if ((Implicit0 != nullptr) && (Implicit1 != nullptr) && !bEnableSweep)
-			{
-				if (Implicit0->HasBoundingBox() && Implicit1->HasBoundingBox())
-				{
-					const FRigidTransform3 WorldTransform0 = LocalTransform0 * ParticleWorldTransform0;
-					const FRigidTransform3 WorldTransform1 = LocalTransform1 * ParticleWorldTransform1;
-					if (Chaos_Collision_NarrowPhase_SphereBoundsCheck)
-					{
-						const FReal R1 = Implicit0->BoundingBox().OriginRadius();
-						const FReal R2 = Implicit1->BoundingBox().OriginRadius();
-						const FReal SeparationSq = (WorldTransform1.GetTranslation() - WorldTransform0.GetTranslation()).SizeSquared();
-						if (SeparationSq > FMath::Square(R1 + R2 + CullDistance))
-						{
-							return false;
-						}
-					}
-
-					if (Chaos_Collision_NarrowPhase_AABBBoundsCheck)
-					{
-						const FAABB3 Box1 = Implicit0->BoundingBox();
-						const FAABB3 Box2 = Implicit1->BoundingBox();
-						//if (Box1.GetVolume() >= Box2.GetVolume())
-						{
-							const FRigidTransform3 Box2ToBox1TM = WorldTransform1.GetRelativeTransform(WorldTransform0);
-							const FAABB3 Box2In1 = Box2.TransformedAABB(Box2ToBox1TM).Thicken(CullDistance);
-							if (!Box1.Intersects(Box2In1))
-							{
-								return false;
-							}
-						}
-						//else
-						{
-							const FRigidTransform3 Box1ToBox2TM = WorldTransform0.GetRelativeTransform(WorldTransform1);
-							const FAABB3 Box1In2 = Box1.TransformedAABB(Box1ToBox2TM).Thicken(CullDistance);
-							if (!Box2.Intersects(Box1In2))
-							{
-								return false;
-							}
-						}
-					}
-				}
-			}
-			else
-			{
-				if (Particle0->HasBounds() && Particle1->HasBounds())
-				{
-					const FAABB3 Box0 = Particle0->WorldSpaceInflatedBounds();
-					const FAABB3 Box1 = Particle1->WorldSpaceInflatedBounds();
-					if (!Box0.Intersects(Box1))
-					{
-						return false;
-					}
-				}
-			}
-
-			return true;
-		}
-
-		void ConstructConstraintsImpl(
-			TGeometryParticleHandle<FReal, 3>* Particle0,
-			TGeometryParticleHandle<FReal, 3>* Particle1,
-			const FImplicitObject* Implicit0,
-			const FShapeInstance* Shape0,
-			const FBVHParticles* Simplicial0,
-			const int32 ImplicitID0,
-			const FImplicitObject* Implicit1,
-			const FShapeInstance* Shape1,
-			const FBVHParticles* Simplicial1,
-			const int32 ImplicitID1,
-			const FRigidTransform3& ParticleWorldTransform0,
-			const FRigidTransform3& LocalTransform0,
-			const FRigidTransform3& ParticleWorldTransform1,
-			const FRigidTransform3& LocalTransform1,
-			const FReal CullDistance,
-			const FReal Dt,
-			const EContactShapesType ShapePairType,
-			const bool bEnableSweep,
-			const FCollisionContext& Context)
-		{
-			CONDITIONAL_SCOPE_CYCLE_COUNTER(STAT_Collisions_ConstructConstraintsInternal, ConstraintsDetailedStats);
-			INC_DWORD_STAT(STAT_ChaosCollisionCounter_NumShapePairs);
-
-			const bool bDoOverlap = DoParticleImplicitsOverlap(
-				Particle0, Particle1, Implicit0, Implicit1, 
-				ParticleWorldTransform0, LocalTransform0, 
-				ParticleWorldTransform1, LocalTransform1, 
-				CullDistance, bEnableSweep, Context);
-
-			if (bDoOverlap)
-			{
-				if (Context.MidPhase != nullptr)
-				{
-					Context.MidPhase->FindOrCreateConstraint(
-						Particle0, Implicit0, ImplicitID0, Shape0, Simplicial0, LocalTransform0,
-						Particle1, Implicit1, ImplicitID1, Shape1, Simplicial1, LocalTransform1,
-						CullDistance, ShapePairType, Context.GetSettings().bAllowManifolds, bEnableSweep, Context);
-				}
-			}
-		}
-
-
 		// Run collision detection for the specified constraint to update the nearest contact point.
 		// NOTE: Transforms are world space shape transforms
 		void UpdateConstraint(FPBDCollisionConstraint& Constraint, const FRigidTransform3& ShapeWorldTransform0, const FRigidTransform3& ShapeWorldTransform1, const FReal Dt)
@@ -1297,180 +1159,12 @@ namespace Chaos
 			return UpdateConstraintFromGeometrySweptImpl(Constraint, ShapeStartWorldTransform0, ShapeStartWorldTransform1, Dt);
 		}
 
-		// @param ImplicitId0 An identifier which allows us to differentiate between instances when the same ImplicitObject appears multiple times in the hierarchy
-		// @param ImplicitId1 See ImplicitId0
-		void ConstructConstraintsRecursive(
-			TGeometryParticleHandle<FReal, 3>* Particle0, 
-			TGeometryParticleHandle<FReal, 3>* Particle1, 
-			const FImplicitObject* Implicit0, 
-			const FShapeInstance* Shape0, 
-			const FBVHParticles* Simplicial0, 
-			const int32 ImplicitID0,
-			const FImplicitObject* Implicit1, 
-			const FShapeInstance* Shape1, 
-			const FBVHParticles* Simplicial1, 
-			const int32 ImplicitID1,
-			const FRigidTransform3& ParticleWorldTransform0,
-			const FRigidTransform3& LocalTransform0, 
-			const FRigidTransform3& ParticleWorldTransform1, 
-			const FRigidTransform3& LocalTransform1, 
-			const FReal CullDistance, 
-			const FReal Dt, 
-			const bool bEnableSweep,
-			const FCollisionContext& Context)
-		{
-			CONDITIONAL_SCOPE_CYCLE_COUNTER(STAT_Collisions_ConstructConstraints, ConstraintsDetailedStats);
-
-			const EImplicitObjectType Implicit0OuterType = Implicit0 ? Implicit0->GetType() : ImplicitObjectType::Unknown;
-			const EImplicitObjectType Implicit1OuterType = Implicit1 ? Implicit1->GetType() : ImplicitObjectType::Unknown;
-			const EImplicitObjectType Implicit0Type = GetInnerType(Implicit0OuterType);
-			const EImplicitObjectType Implicit1Type = GetInnerType(Implicit1OuterType);
-
-			if ((Implicit0 != nullptr) && (Implicit1 != nullptr))
-			{
-				// Handle transform wrapper shape
-				if ((Implicit0OuterType == TImplicitObjectTransformed<FReal, 3>::StaticType()) && (Implicit1OuterType == TImplicitObjectTransformed<FReal, 3>::StaticType()))
-				{
-					const TImplicitObjectTransformed<FReal, 3>* TransformedImplicit0 = Implicit0->template GetObject<const TImplicitObjectTransformed<FReal, 3>>();
-					const TImplicitObjectTransformed<FReal, 3>* TransformedImplicit1 = Implicit1->template GetObject<const TImplicitObjectTransformed<FReal, 3>>();
-					FRigidTransform3 TransformedTransform0 = TransformedImplicit0->GetTransform() * LocalTransform0;
-					FRigidTransform3 TransformedTransform1 = TransformedImplicit1->GetTransform() * LocalTransform1;
-					ConstructConstraintsRecursive(Particle0, Particle1, TransformedImplicit0->GetTransformedObject(), Shape0, Simplicial0, ImplicitID0, TransformedImplicit1->GetTransformedObject(), Shape1, Simplicial1, ImplicitID1, ParticleWorldTransform0, TransformedTransform0, ParticleWorldTransform1, TransformedTransform1, CullDistance, Dt, bEnableSweep, Context);
-					return;
-				}
-				else if (Implicit0OuterType == TImplicitObjectTransformed<FReal, 3>::StaticType())
-				{
-					const TImplicitObjectTransformed<FReal, 3>* TransformedImplicit0 = Implicit0->template GetObject<const TImplicitObjectTransformed<FReal, 3>>();
-					FRigidTransform3 TransformedTransform0 = TransformedImplicit0->GetTransform() * LocalTransform0;
-					ConstructConstraintsRecursive(Particle0, Particle1, TransformedImplicit0->GetTransformedObject(), Shape0, Simplicial0, ImplicitID0, Implicit1, Shape1, Simplicial1, ImplicitID1, ParticleWorldTransform0, TransformedTransform0, ParticleWorldTransform1, LocalTransform1, CullDistance, Dt, bEnableSweep, Context);
-					return;
-				}
-				else if (Implicit1OuterType == TImplicitObjectTransformed<FReal, 3>::StaticType())
-				{
-					const TImplicitObjectTransformed<FReal, 3>* TransformedImplicit1 = Implicit1->template GetObject<const TImplicitObjectTransformed<FReal, 3>>();
-					FRigidTransform3 TransformedTransform1 = TransformedImplicit1->GetTransform() * LocalTransform1;
-					ConstructConstraintsRecursive(Particle0, Particle1, Implicit0, Shape0, Simplicial0, ImplicitID0, TransformedImplicit1->GetTransformedObject(), Shape1, Simplicial1, ImplicitID1, ParticleWorldTransform0, LocalTransform0, ParticleWorldTransform1, TransformedTransform1, CullDistance, Dt, bEnableSweep, Context);
-					return;
-				}
-
-				// Strip the Instanced wrapper from most shapes, but not Convex or TriMesh.
-				// Convex collision requires the wrapper because it holds the margin.
-				if (((uint32)Implicit0OuterType & ImplicitObjectType::IsInstanced) || ((uint32)Implicit1OuterType & ImplicitObjectType::IsInstanced))
-				{
-					const FImplicitObject* InnerImplicit0 = nullptr;
-					const FImplicitObject* InnerImplicit1 = nullptr;
-					if (((uint32)Implicit0OuterType & ImplicitObjectType::IsInstanced) && (Implicit0Type != FImplicitConvex3::StaticType()))
-					{
-						InnerImplicit0 = GetInstancedImplicit(Implicit0);
-					}
-					if (((uint32)Implicit1OuterType & ImplicitObjectType::IsInstanced) && (Implicit1Type != FImplicitConvex3::StaticType()))
-					{
-						InnerImplicit1 = GetInstancedImplicit(Implicit1);
-					}
-					if (InnerImplicit0 && InnerImplicit1)
-					{
-						ConstructConstraintsRecursive(Particle0, Particle1, InnerImplicit0, Shape0, Simplicial0, ImplicitID0, InnerImplicit1, Shape1, Simplicial1, ImplicitID1, ParticleWorldTransform0, LocalTransform0, ParticleWorldTransform1, LocalTransform1, CullDistance, Dt, bEnableSweep, Context);
-						return;
-					}
-					else if (InnerImplicit0 && !InnerImplicit1)
-					{
-						ConstructConstraintsRecursive(Particle0, Particle1, InnerImplicit0, Shape0, Simplicial0, ImplicitID0, Implicit1, Shape1, Simplicial1, ImplicitID1, ParticleWorldTransform0, LocalTransform0, ParticleWorldTransform1, LocalTransform1, CullDistance, Dt, bEnableSweep, Context);
-						return;
-					}
-					else if (!InnerImplicit0 && InnerImplicit1)
-					{
-						ConstructConstraintsRecursive(Particle0, Particle1, Implicit0, Shape0, Simplicial0, ImplicitID0, InnerImplicit1, Shape1, Simplicial1, ImplicitID1, ParticleWorldTransform0, LocalTransform0, ParticleWorldTransform1, LocalTransform1, CullDistance, Dt, bEnableSweep, Context);
-						return;
-					}
-				}
-
-				// Handle Unions
-				// NOTE: the Union visit function called below are recursive and visit all children in the hierarchy below the union.
-				if ((Implicit0OuterType == FImplicitObjectUnion::StaticType()) || (Implicit0OuterType == FImplicitObjectUnionClustered::StaticType()))
-				{
-					const FRigidTransform3 TM0 = LocalTransform0 * ParticleWorldTransform0;
-					const FRigidTransform3 TM1 = LocalTransform1 * ParticleWorldTransform1;
-					const FRigidTransform3 TM1ToTM0 = TM1.GetRelativeTransform(TM0);
-					const FAABB3 QueryBounds0 = Implicit1->HasBoundingBox() ? Implicit1->BoundingBox().TransformedAABB(TM1ToTM0) : FAABB3::FullAABB();
-
-					// @todo(chaos): we shouldn't need to treat Union and ClusterUnion differentlty here - change how we get the simplicial to be more generic
-					const FImplicitObjectUnionClustered* ClusterUnion0 = Implicit0->template GetObject<FImplicitObjectUnionClustered>();
-					Implicit0->VisitOverlappingLeafObjects(QueryBounds0,
-						[&](const FImplicitObject* ChildImplicit0, const FRigidTransform3& ChildTransform0, const int32 RootObjectIndex0, const int32 ObjectIndex, const int32 LeafObjectIndex0)
-						{
-							const FBVHParticles* ChildSimplicial0 = (ClusterUnion0 != nullptr) ? ClusterUnion0->GetChildSimplicial(LeafObjectIndex0) : Simplicial0;
-							ConstructConstraintsRecursive(Particle0, Particle1, ChildImplicit0, Shape0, ChildSimplicial0, LeafObjectIndex0, Implicit1, Shape1, Simplicial1, ImplicitID1, ParticleWorldTransform0, ChildTransform0 * LocalTransform0, ParticleWorldTransform1, LocalTransform1, CullDistance, Dt, bEnableSweep, Context);
-						});
-					return;
-				}
-				if ((Implicit1OuterType == FImplicitObjectUnion::StaticType()) || (Implicit1OuterType == FImplicitObjectUnionClustered::StaticType()))
-				{
-					const FRigidTransform3 TM0 = LocalTransform0 * ParticleWorldTransform0;
-					const FRigidTransform3 TM1 = LocalTransform1 * ParticleWorldTransform1;
-					const FRigidTransform3 TM0ToTM1 = TM0.GetRelativeTransform(TM1);
-					const FAABB3 QueryBounds1 = Implicit0->HasBoundingBox() ? Implicit0->BoundingBox().TransformedAABB(TM0ToTM1) : FAABB3::FullAABB();
-
-					const FImplicitObjectUnionClustered* ClusterUnion1 = Implicit1->template GetObject<FImplicitObjectUnionClustered>();
-					Implicit1->VisitOverlappingLeafObjects(QueryBounds1,
-						[&](const FImplicitObject* ChildImplicit1, const FRigidTransform3& ChildTransform1, const int32 RootObjectIndex1, const int32 ObjectIndex, const int32 LeafObjectIndex1)
-						{
-							const FBVHParticles* ChildSimplicial1 = (ClusterUnion1 != nullptr) ? ClusterUnion1->GetChildSimplicial(LeafObjectIndex1) : Simplicial1;
-							ConstructConstraintsRecursive(Particle0, Particle1, Implicit0, Shape0, Simplicial0, ImplicitID0, ChildImplicit1, Shape1, Simplicial1, LeafObjectIndex1, ParticleWorldTransform0, LocalTransform0, ParticleWorldTransform1, ChildTransform1 * LocalTransform1, CullDistance, Dt, bEnableSweep, Context);
-						});
-					return;
-				}
-			}
-
-			// Check shape pair filtering if enable
-			// @todo(chaos): Make sure we are setting this to false in the midphase (it has already performed filtering)
-			if (Context.GetSettings().bFilteringEnabled && !DoCollide(Implicit0Type, Shape0, Implicit1Type, Shape1))
-			{
-				return;
-			}
-
-			// If we get here, we have a pair of concrete shapes (i.e., no wrappers or containers)
-			// Create a constraint for the shape pair
-			bool bShouldSwapParticles = false;
-			const EContactShapesType ShapePairType = CalculateShapePairType(Particle0, Implicit0, Particle1, Implicit1, bShouldSwapParticles);
-
-			if (bShouldSwapParticles)
-			{
-				ConstructConstraintsImpl(Particle1, Particle0, Implicit1, Shape1, Simplicial1, ImplicitID1, Implicit0, Shape0, Simplicial0, ImplicitID0, ParticleWorldTransform1, LocalTransform1, ParticleWorldTransform0, LocalTransform0, CullDistance, Dt, ShapePairType, bEnableSweep, Context);
-			}
-			else
-			{
-				ConstructConstraintsImpl(Particle0, Particle1, Implicit0, Shape0, Simplicial0, ImplicitID0, Implicit1, Shape1, Simplicial1, ImplicitID1, ParticleWorldTransform0, LocalTransform0, ParticleWorldTransform1, LocalTransform1, CullDistance, Dt, ShapePairType, bEnableSweep, Context);
-			}
-		}
-
-		EImplicitObjectType GetImplicitCollisionType(const FGeometryParticleHandle* Particle, const FImplicitObject* Implicit)
-		{
-			// NOTE: GetCollisionType(), not GetType()
-			// We use CollisionType on the implicit to determine how to collide. Normally this is the same as the actual ImplicitObject's
-			// type, but may be set to LevelSet in which case we will use CollisionParticles instead (if it has any).
-			EImplicitObjectType ImplicitType = (Implicit != nullptr) ? GetInnerType(Implicit->GetCollisionType()) : ImplicitObjectType::Unknown;
-
-			// If we are a levelset make sure we have CollisionParticles, otherwise go back to the builtin implicit object collision type
-			if ((ImplicitType == ImplicitObjectType::LevelSet) || (ImplicitType == ImplicitObjectType::Unknown))
-			{
-				const FBVHParticles* Simplicial = FConstGenericParticleHandle(Particle)->CollisionParticles().Get();
-				const bool bHasSimplicial = (Simplicial != nullptr) && (Simplicial->Size() > 0);
-				if (!bHasSimplicial)
-				{
-					// NOTE: GetType(), not GetCollisionType()
-					ImplicitType = (Implicit != nullptr) ? GetInnerType(Implicit->GetType()) : ImplicitObjectType::Unknown;
-				}
-			}
-
-			return ImplicitType;
-		}
-
 		EContactShapesType CalculateShapePairType(const FGeometryParticleHandle* Particle0, const FImplicitObject* Implicit0, const FGeometryParticleHandle* Particle1, const FImplicitObject* Implicit1, bool &bOutSwap)
 		{
 			// What types do the implicits collide as?
 			// NOTE: If the implicit is set to LevelSet but has no CollisionParticles we default to the implicits actual type
-			const EImplicitObjectType Implicit0Type = GetImplicitCollisionType(Particle0, Implicit0);
-			const EImplicitObjectType Implicit1Type = GetImplicitCollisionType(Particle1, Implicit1);
+			const EImplicitObjectType Implicit0Type = Private::GetImplicitCollisionType(Particle0, Implicit0);
+			const EImplicitObjectType Implicit1Type = Private::GetImplicitCollisionType(Particle1, Implicit1);
 
 			//
 			// Basic implicit pairs
@@ -1672,19 +1366,9 @@ namespace Chaos
 		//
 		//
 
-		void ConstructConstraints(FGeometryParticleHandle* Particle0, FGeometryParticleHandle* Particle1, const FImplicitObject* Implicit0,  const FPerShapeData* Shape0,  const FBVHParticles* Simplicial0,  const int32 ImplicitID0, const FImplicitObject* Implicit1, const FPerShapeData* Shape1,  const FBVHParticles* Simplicial1,  const int32 ImplicitID1, const FRigidTransform3& ParticleWorldTransform0, const FRigidTransform3& LocalTransform0,  const FRigidTransform3& ParticleWorldTransform1,  const FRigidTransform3& LocalTransform1, const FReal CullDistance,  const FReal dT,  const bool bEnableSweep,  const FCollisionContext& Context)
-		{
-			ConstructConstraintsRecursive(Particle0, Particle1, Implicit0, Shape0->AsShapeInstance(), Simplicial0, ImplicitID0, Implicit1, Shape1->AsShapeInstance(), Simplicial1, ImplicitID1, ParticleWorldTransform0, LocalTransform0, ParticleWorldTransform1, LocalTransform1, CullDistance, dT, bEnableSweep, Context);
-		}
-
-		void ConstructConstraints(TGeometryParticleHandle<FReal, 3>* Particle0, TGeometryParticleHandle<FReal, 3>* Particle1, const FImplicitObject* Implicit0, const FPerShapeData* Shape0, const FBVHParticles* Simplicial0, const FImplicitObject* Implicit1, const FPerShapeData* Shape1, const FBVHParticles* Simplicial1, const FRigidTransform3& ParticleWorldTransform0, const FRigidTransform3& Transform0, const FRigidTransform3& ParticleWorldTransform1, const FRigidTransform3& Transform1, const FReal CullDistance, const FReal Dt, const bool bEnableSweep, const FCollisionContext& Context)
-		{
-			ConstructConstraintsRecursive(Particle0, Particle1, Implicit0, Shape0->AsShapeInstance(), Simplicial0, 0, Implicit1, Shape1->AsShapeInstance(), Simplicial1, 0, ParticleWorldTransform0, Transform0, ParticleWorldTransform1, Transform1, CullDistance, Dt, bEnableSweep, Context);
-		}
-
-		void ConstructConstraints(TGeometryParticleHandle<FReal, 3>* Particle0, TGeometryParticleHandle<FReal, 3>* Particle1, const FImplicitObject* Implicit0, const FPerShapeData* Shape0, const FBVHParticles* Simplicial0, const FImplicitObject* Implicit1, const FPerShapeData* Shape1, const FBVHParticles* Simplicial1, const FRigidTransform3& ParticleWorldTransform0, const FRigidTransform3& Transform0, const FRigidTransform3& ParticleWorldTransform1, const FRigidTransform3& Transform1, const FReal CullDistance, const FReal Dt, const FCollisionContext& Context)
+		EImplicitObjectType GetImplicitCollisionType(const FGeometryParticleHandle* Particle, const FImplicitObject* Implicit)
 		{
-			ConstructConstraintsRecursive(Particle0, Particle1, Implicit0, Shape0->AsShapeInstance(), Simplicial0, 0, Implicit1, Shape1->AsShapeInstance(), Simplicial1, 0, ParticleWorldTransform0, Transform0, ParticleWorldTransform1, Transform1, CullDistance, Dt, false, Context);
+			return Private::GetImplicitCollisionType(Particle, Implicit);
 		}
 
 		template<ECollisionUpdateType UpdateType>
diff --git a/Engine/Source/Runtime/Experimental/Chaos/Public/Chaos/Collision/CollisionUtil.h b/Engine/Source/Runtime/Experimental/Chaos/Public/Chaos/Collision/CollisionUtil.h
new file mode 100644
--- /dev/null
+++ b/Engine/Source/Runtime/Experimental/Chaos/Public/Chaos/Collision/CollisionUtil.h
@@ -0,0 +1,76 @@
+// Copyright Epic Games, Inc. All Rights Reserved.
+#pragma once
+
+namespace Chaos
+{
+	extern bool bChaos_Collision_EnableBoundsChecks;
+
+	namespace Private
+	{
+		// Get the implicit object type for the implicit object on the particle, taking into account the fact that the implicit
+		// may be set to use a LevelSet, but only if the particle has collision particles
+		inline EImplicitObjectType GetImplicitCollisionType(const FGeometryParticleHandle* Particle, const FImplicitObject* Implicit)
+		{
+			// NOTE: GetCollisionType(), not GetType()
+			// We use CollisionType on the implicit to determine how to collide. Normally this is the same as the actual ImplicitObject's
+			// type, but may be set to LevelSet in which case we will use CollisionParticles instead (if it has any).
+			EImplicitObjectType ImplicitType = (Implicit != nullptr) ? GetInnerType(Implicit->GetCollisionType()) : ImplicitObjectType::Unknown;
+
+			// If we are a levelset make sure we have CollisionParticles, otherwise go back to the builtin implicit object collision type
+			if ((ImplicitType == ImplicitObjectType::LevelSet) || (ImplicitType == ImplicitObjectType::Unknown))
+			{
+				const FBVHParticles* Simplicial = FConstGenericParticleHandle(Particle)->CollisionParticles().Get();
+				const bool bHasSimplicial = (Simplicial != nullptr) && (Simplicial->Size() > 0);
+				if (!bHasSimplicial)
+				{
+					// NOTE: GetType(), not GetCollisionType()
+					ImplicitType = (Implicit != nullptr) ? GetInnerType(Implicit->GetType()) : ImplicitObjectType::Unknown;
+				}
+			}
+
+			return ImplicitType;
+		}
+
+		// Determine which bounds tests should be run for the specified implicit object pair.
+		// Also sets the bIsProbe flag appropriately.
+		inline FImplicitBoundsTestFlags CalculateImplicitBoundsTestFlags(
+			FGeometryParticleHandle* Particle0, const FImplicitObject* Implicit0, const FPerShapeData* Shape0,
+			FGeometryParticleHandle* Particle1, const FImplicitObject* Implicit1, const FPerShapeData* Shape1,
+			FRealSingle& OutDistanceCheckSize)
+		{
+			FImplicitBoundsTestFlags Flags;
+
+			const EImplicitObjectType ImplicitType0 = GetImplicitCollisionType(Particle0, Implicit0);
+			const EImplicitObjectType ImplicitType1 = GetImplicitCollisionType(Particle1, Implicit1);
+			const bool bIsSphere0 = (ImplicitType0 == ImplicitObjectType::Sphere);
+			const bool bIsSphere1 = (ImplicitType1 == ImplicitObjectType::Sphere);
+			const bool bIsCapsule0 = (ImplicitType0 == ImplicitObjectType::Capsule);
+			const bool bIsCapsule1 = (ImplicitType1 == ImplicitObjectType::Capsule);
+			const bool bIsTriangle0 = (ImplicitType0 == ImplicitObjectType::TriangleMesh) || (ImplicitType0 == ImplicitObjectType::HeightField);
+			const bool bIsTriangle1 = (ImplicitType1 == ImplicitObjectType::TriangleMesh) || (ImplicitType1 == ImplicitObjectType::HeightField);
+			const bool bIsLevelSet = ((ImplicitType0 == ImplicitObjectType::LevelSet) || (ImplicitType1 == ImplicitObjectType::LevelSet));
+
+			const bool bHasBounds0 = (Implicit0 != nullptr) && Implicit0->HasBoundingBox();
+			const bool bHasBounds1 = (Implicit1 != nullptr) && Implicit1->HasBoundingBox();
+			const bool bAllowBoundsCheck = bChaos_Collision_EnableBoundsChecks && bHasBounds0 && bHasBounds1;
+
+			Flags.bEnableAABBCheck = bAllowBoundsCheck && !(bIsSphere0 && bIsSphere1);	// No AABB test if both are spheres
+			Flags.bEnableOBBCheck0 = bAllowBoundsCheck && !bIsSphere0;					// No OBB test for spheres
+			Flags.bEnableOBBCheck1 = bAllowBoundsCheck && !bIsSphere1;					// No OBB test for spheres
+			Flags.bEnableDistanceCheck = bAllowBoundsCheck && bIsSphere0 && bIsSphere1;	// Simple distance test for sphere pairs
+
+			OutDistanceCheckSize = (Flags.bEnableDistanceCheck) ? FRealSingle(Implicit0->GetMargin() + Implicit1->GetMargin()) : 0;	// Sphere-Sphere bounds test
+
+			// Do not try to reuse manifold points for capsules or spheres (against anything)
+			// NOTE: This can also be disabled for all shape types by the solver (see GenerateCollisionImpl and the Context)
+			Flags.bEnableManifoldUpdate = !bIsSphere0 && !bIsSphere1 && !bIsCapsule0 && !bIsCapsule1 && !bIsTriangle0 && !bIsTriangle1 && !bIsLevelSet;
+
+			// Mark probe flag now so we know which GenerateCollisions to use
+			// @todo(chaos): it looks like this can be changed by a collision modifier so we should not be caching it
+			Flags.bIsProbe = Shape0->GetIsProbe() || Shape1->GetIsProbe();
+
+			return Flags;
+		}
+
+	}
+}
\ No newline at end of file
diff --git a/Engine/Source/Runtime/Experimental/Chaos/Public/Chaos/Collision/PBDCollisionConstraint.h b/Engine/Source/Runtime/Experimental/Chaos/Public/Chaos/Collision/PBDCollisionConstraint.h
--- a/Engine/Source/Runtime/Experimental/Chaos/Public/Chaos/Collision/PBDCollisionConstraint.h
+++ b/Engine/Source/Runtime/Experimental/Chaos/Public/Chaos/Collision/PBDCollisionConstraint.h
@@ -1,7 +1,6 @@
 // Copyright Epic Games, Inc. All Rights Reserved.
 #pragma once
 
-#include "CoreMinimal.h"
 #include "Chaos/Core.h"
 
 #include "Chaos/BVHParticles.h"
@@ -98,6 +97,32 @@ namespace Chaos
 	// Renamed to FPBDCollisionConstraintMaterial
 	using FCollisionContact UE_DEPRECATED(5.1, "FCollisionContact was renamed to FPBDCollisionConstraintMaterial") = FPBDCollisionConstraintMaterial;
 
+	namespace Private
+	{
+
+		// Flags to specify what types of bounds tests should be run for a collision constraint
+		// NOTE: These are separate from the other constraint flags because the ShapePair MidPhase 
+		// stores a copy for use in determining whether to create a constraint in the first place.
+		union FImplicitBoundsTestFlags
+		{
+			FImplicitBoundsTestFlags() : Bits(0) {}
+			struct
+			{
+				uint8 bEnableAABBCheck : 1;
+				uint8 bEnableOBBCheck0 : 1;
+				uint8 bEnableOBBCheck1 : 1;
+				uint8 bEnableDistanceCheck : 1;
+				uint8 bEnableManifoldUpdate : 1;
+				uint8 bIsProbe : 1;
+			};
+			uint8 Bits;
+		};
+
+		CHAOS_API FImplicitBoundsTestFlags CalculateImplicitBoundsTestFlags(
+			FGeometryParticleHandle* Particle0, const FImplicitObject* Implicit0, const FPerShapeData* Shape0,
+			FGeometryParticleHandle* Particle1, const FImplicitObject* Implicit1, const FPerShapeData* Shape1,
+			FRealSingle& OutDistanceCheckSize);
+	}
 
 	/**
 	 * @brief Information used by the constraint allocator
@@ -417,7 +442,7 @@ namespace Chaos
 				Material.ResetMaterialModifications();
 
 				// Reset other properties which may have changed in contact modification last frame
-				Flags.bIsProbe = Flags.bIsProbeUnmodified;
+				Flags.bIsProbe = BoundsTestFlags.bIsProbe;
 
 				Flags.bModifierApplied = false;
 			}
@@ -554,7 +579,9 @@ namespace Chaos
 		{
 			LastShapeWorldPositionDelta = InShapeWorldTransform0.GetTranslation() - InShapeWorldTransform1.GetTranslation();
 			LastShapeWorldRotationDelta = InShapeWorldTransform0.GetRotation().Inverse() * InShapeWorldTransform1.GetRotation();
-			Flags.bCanRestoreManifold = true;
+
+			// NOTE: BoundsTestFlags.bEnableManifoldUpdate is false if the shape pair does not support manifold reuse
+			Flags.bCanRestoreManifold = BoundsTestFlags.bEnableManifoldUpdate;
 		}
 
 		bool GetCanRestoreManifold() const { return Flags.bCanRestoreManifold; }
@@ -699,6 +726,11 @@ namespace Chaos
 			return FCollisionParticlePairKey(GetParticle0(), GetParticle1());
 		}
 
+		Private::FImplicitBoundsTestFlags GetBoundsTestFlags() const
+		{
+			return BoundsTestFlags;
+		}
+
 	public:
 		const FPBDCollisionConstraintHandle* GetConstraintHandle() const { return this; }
 		FPBDCollisionConstraintHandle* GetConstraintHandle() { return this; }
@@ -775,23 +807,22 @@ namespace Chaos
 			FFlags() : Bits(0) {}
 			struct
 			{
-				uint32 bDisabled : 1;					// Is this contact disabled (by the user or because cull distance is exceeded)
-				uint32 bUseManifold : 1;				// Should we use contact manifolds or single points (faster but poor behaviour)
-				uint32 bUseIncrementalManifold : 1;		// Do we need to run incremental collision detection (only LavelSets now)
-				uint32 bCanRestoreManifold : 1;			// Can we try to restore the manifold this frame (set folowing narrowphase, cleared when reset for some reason)
-				uint32 bWasManifoldRestored : 1;		// Did we restore the manifold this frame
-				uint32 bIsQuadratic0 : 1;				// Is the first shape a sphere or capsule
-				uint32 bIsQuadratic1 : 1;				// Is the second shape a sphere or capsule
-				uint32 bIsProbeUnmodified : 1;			// Is this constraint a probe pre-contact-modification
-				uint32 bIsProbe : 1;					// Is this constraint currently a probe
-				uint32 bCCDEnabled : 1;					// Is CCD enabled for the current tick
-				uint32 bCCDSweepEnabled: 1;				// If this is a CCD constraint, do we want to enable the sweep/rewind phase?
-				uint32 bModifierApplied : 1;			// Was a constraint modifier applied this tick
-				uint32 bMaterialSet : 1;				// Has the material been set (or does it need to be reset)
+				uint16 bDisabled : 1;					// Is this contact disabled (by the user or because cull distance is exceeded)
+				uint16 bUseManifold : 1;				// Should we use contact manifolds or single points (faster but poor behaviour)
+				uint16 bUseIncrementalManifold : 1;		// Do we need to run incremental collision detection (only LavelSets now)
+				uint16 bCanRestoreManifold : 1;			// Can we try to restore the manifold this frame (set folowing narrowphase, cleared when reset for some reason)
+				uint16 bWasManifoldRestored : 1;		// Did we restore the manifold this frame
+				uint16 bIsQuadratic0 : 1;				// Is the first shape a sphere or capsule
+				uint16 bIsQuadratic1 : 1;				// Is the second shape a sphere or capsule
+				uint16 bIsProbe : 1;					// Is this constraint currently a probe
+				uint16 bCCDEnabled : 1;					// Is CCD enabled for the current tick
+				uint16 bCCDSweepEnabled: 1;				// If this is a CCD constraint, do we want to enable the sweep/rewind phase?
+				uint16 bModifierApplied : 1;			// Was a constraint modifier applied this tick
+				uint16 bMaterialSet : 1;				// Has the material been set (or does it need to be reset)
 			};
-			uint32 Bits;
+			uint16 Bits;
 		};
-		static_assert(sizeof(FFlags) == 4, "Unexpected size for FPBDCollisionConstraint::FFLags");
+		static_assert(sizeof(FFlags) == 2, "Unexpected size for FPBDCollisionConstraint::FFLags");
 
 		// Local-space transforms of the shape (relative to particle)
 		FRigidTransform3 ImplicitTransform[2];
@@ -840,6 +871,7 @@ namespace Chaos
 		FRealSingle Stiffness;
 
 		FFlags Flags;
+		Private::FImplicitBoundsTestFlags BoundsTestFlags;
 
 		TCArray<FSavedManifoldPoint, MaxManifoldPoints> SavedManifoldPoints;
 		TCArray<FManifoldPoint, MaxManifoldPoints> ManifoldPoints;
diff --git a/Engine/Source/Runtime/Experimental/Chaos/Public/Chaos/Collision/ParticlePairMidPhase.h b/Engine/Source/Runtime/Experimental/Chaos/Public/Chaos/Collision/ParticlePairMidPhase.h
--- a/Engine/Source/Runtime/Experimental/Chaos/Public/Chaos/Collision/ParticlePairMidPhase.h
+++ b/Engine/Source/Runtime/Experimental/Chaos/Public/Chaos/Collision/ParticlePairMidPhase.h
@@ -144,19 +144,7 @@ namespace Chaos
 		FRealSingle SphereBoundsCheckSize;
 		int32 LastUsedEpoch;
 		EContactShapesType ShapePairType;
-		union FFlags
-		{
-			FFlags() : Bits(0) {}
-			struct 
-			{
-				uint8 bEnableAABBCheck : 1;
-				uint8 bEnableOBBCheck0 : 1;
-				uint8 bEnableOBBCheck1 : 1;
-				uint8 bEnableManifoldUpdate : 1;
-				uint8 bIsProbe : 1;
-			};
-			uint8 Bits;
-		} Flags;
+		Private::FImplicitBoundsTestFlags BoundsTestFlags;
 	};
 
 
@@ -353,30 +341,6 @@ namespace Chaos
 			return Flags.bIsCCDActive;
 		}
 
-		virtual FPBDCollisionConstraint* FindOrCreateConstraint(
-			FGeometryParticleHandle* InParticle0,
-			const FImplicitObject* InImplicit0,
-			const int32 InImplicitId0,
-			const FShapeInstance* InShape0,
-			const FBVHParticles* InBVHParticles0,
-			const FRigidTransform3& InShapeRelativeTransform0,
-			FGeometryParticleHandle* InParticle1,
-			const FImplicitObject* Implicit1,
-			const int32 InImplicitId1,
-			const FShapeInstance* InShape1,
-			const FBVHParticles* InBVHParticles1,
-			const FRigidTransform3& InhapeRelativeTransform1,
-			const FReal CullDistance,
-			const EContactShapesType ShapePairType,
-			const bool bUseManifold,
-			const bool bEnableSweep,
-			const FCollisionContext& Context)
-		{
-			// Currently should only be called on FGenericParticlePairMidPhase
-			check(false);
-			return nullptr;
-		}
-
 	protected:
 
 		virtual void ResetImpl() = 0;
@@ -501,42 +465,8 @@ namespace Chaos
 		CHAOS_API virtual void ResetImpl() override final;
 		CHAOS_API virtual void BuildDetectorsImpl() override final;
 
-		/**
-		 * @brief Callback from the narrow phase to create a collision constraint for this particle pair.
-		 * We should never be asked for a collision for a different particle pair, but the
-		 * implicit objects may be children of the root shape.
-		*/
-		CHAOS_API virtual FPBDCollisionConstraint* FindOrCreateConstraint(
-			FGeometryParticleHandle* InParticle0,
-			const FImplicitObject* InImplicit0,
-			const int32 InImplicitId0,
-			const FShapeInstance* InShape0,
-			const FBVHParticles* InBVHParticles0,
-			const FRigidTransform3& InShapeRelativeTransform0,
-			FGeometryParticleHandle* InParticle1,
-			const FImplicitObject* Implicit1,
-			const int32 InImplicitId1,
-			const FShapeInstance* InShape1,
-			const FBVHParticles* InBVHParticles1,
-			const FRigidTransform3& InhapeRelativeTransform1,
-			const FReal CullDistance,
-			const EContactShapesType ShapePairType,
-			const bool bUseManifold,
-			const bool bEnableSweep,
-			const FCollisionContext& Context) override final;
-
-		/**
-		 * @brief Reactivate the constraint
-		 * @parame SleepEpoch The tick on which the particle went to sleep.
-		 * Only constraints that were active when the particle went to sleep should be reactivated.
-		*/
-		CHAOS_API void WakeCollisionsImpl(const int32 SleepEpoch, const int32 CurrentEpoch);
-
 	protected:
-		CHAOS_API virtual int32 GenerateCollisionsImpl(
-			const FReal CullDistance,
-			const FReal Dt,
-			const FCollisionContext& Context) override final;
+		CHAOS_API virtual int32 GenerateCollisionsImpl( const FReal CullDistance, const FReal Dt, const FCollisionContext& Context) override final;
 
 		CHAOS_API virtual void WakeCollisionsImpl(const int32 CurrentEpoch) override final;
 
@@ -547,25 +477,19 @@ namespace Chaos
 		CHAOS_API void GenerateCollisionsBVHBVH(
 			FGeometryParticleHandle* ParticleA, const Private::FImplicitBVH* BVHA,
 			FGeometryParticleHandle* ParticleB, const Private::FImplicitBVH* BVHB,
-			const FReal CullDistance,
-			const FReal Dt,
-			const FCollisionContext& Context);
+			const FReal CullDistance, const FReal Dt, const FCollisionContext& Context);
 
 		// BVH on ParticleA versus implicit hierarchy of ParticleB
 		CHAOS_API void GenerateCollisionsBVHImplicitHierarchy(
 			FGeometryParticleHandle* ParticleA, const Private::FImplicitBVH* BVHA,
 			FGeometryParticleHandle* ParticleB, const FImplicitObject* RootImplicitB,
-			const FReal CullDistance,
-			const FReal Dt,
-			const FCollisionContext& Context);
+			const FReal CullDistance, const FReal Dt, const FCollisionContext& Context);
 
 		// Implicit hierarchy of particle A versus implicit hierarchy of ParticleB (used when no BVH present on either)
 		CHAOS_API void GenerateCollisionsImplicitHierarchyImplicitHierarchy(
 			FGeometryParticleHandle* ParticleA, const FImplicitObject* RootImplicitA,
 			FGeometryParticleHandle* ParticleB, const FImplicitObject* RootImplicitB,
-			const FReal CullDistance,
-			const FReal Dt,
-			const FCollisionContext& Context);
+			const FReal CullDistance, const FReal Dt, const FCollisionContext& Context);
 
 		// BVH on particle A versus a Leaf Implicit of ParticleB
 		CHAOS_API void GenerateCollisionsBVHImplicitLeaf(
@@ -577,57 +501,33 @@ namespace Chaos
 		CHAOS_API void GenerateCollisionsImplicitLeafImplicitLeaf(
 			FGeometryParticleHandle* ParticleA, const FImplicitObject* ImplicitA, const FShapeInstance* ShapeInstanceA, const FRigidTransform3 ParticleWorldTransformA, const FRigidTransform3& RelativeTransformA, const int32 LeafObjectIndexA,
 			FGeometryParticleHandle* ParticleB, const FImplicitObject* ImplicitB, const FShapeInstance* ShapeInstanceB, const FRigidTransform3 ParticleWorldTransformB, const FRigidTransform3& RelativeTransformB, const int32 LeafObjectIndexB,
-			const FReal CullDistance,
-			const FReal Dt,
-			const FCollisionContext& Context);
+			const FReal CullDistance, const FReal Dt, const FCollisionContext& Context);
 
 		// A bounds check between two implicits
 		CHAOS_API bool DoBoundsOverlap(
-			const FImplicitObject* ImplicitA,
-			const FRigidTransform3& ParticleWorldTransformA,
-			const FRigidTransform3& ShapeRelativeTransformA,
-			const FImplicitObject* ImplicitB,
-			const FRigidTransform3& ParticleWorldTransformB,
-			const FRigidTransform3& ShapeRelativeTransformB,
-			const FReal CullDistance);
+			const FImplicitObject* ImplicitA, const FRigidTransform3& ParticleWorldTransformA, const FRigidTransform3& ShapeRelativeTransformA,
+			const FImplicitObject* ImplicitB, const FRigidTransform3& ParticleWorldTransformB, const FRigidTransform3& ShapeRelativeTransformB,
+			const Private::FImplicitBoundsTestFlags BoundsTestFlags, const FRealSingle DistanceCheckSize, const FReal CullDistance);
+
+		CHAOS_API FPBDCollisionConstraint* FindOrCreateConstraint(
+			FGeometryParticleHandle* InParticle0, const FImplicitObject* InImplicit0, const int32 InImplicitId0, const FShapeInstance* InShape0, const FBVHParticles* InBVHParticles0, const FRigidTransform3& InShapeRelativeTransform0,
+			FGeometryParticleHandle* InParticle1, const FImplicitObject* InImplicit1, const int32 InImplicitId1, const FShapeInstance* InShape1, const FBVHParticles* InBVHParticles1, const FRigidTransform3& InShapeRelativeTransform1,
+			const FReal CullDistance, const EContactShapesType ShapePairType, const bool bUseManifold, const bool bEnableSweep, const FCollisionContext& Context);
 
 		CHAOS_API FPBDCollisionConstraint* FindConstraint(const FCollisionParticlePairConstraintKey& Key);
 
 		CHAOS_API FPBDCollisionConstraint* CreateConstraint(
-			FGeometryParticleHandle* Particle0,
-			const FImplicitObject* Implicit0,
-			const FPerShapeData* Shape0,
-			const FBVHParticles* BVHParticles0,
-			const FRigidTransform3& ShapeRelativeTransform0,
-			FGeometryParticleHandle* Particle1,
-			const FImplicitObject* Implicit1,
-			const FPerShapeData* Shape1,
-			const FBVHParticles* BVHParticles1,
-			const FRigidTransform3& ShapeRelativeTransform1,
-			const FReal CullDistance,
-			const EContactShapesType ShapePairType,
-			const bool bInUseManifold,
-			const FCollisionParticlePairConstraintKey& Key,
-			const FCollisionContext& Context);
+			FGeometryParticleHandle* Particle0, const FImplicitObject* Implicit0, const FPerShapeData* Shape0, const FBVHParticles* BVHParticles0, const FRigidTransform3& ShapeRelativeTransform0,
+			FGeometryParticleHandle* Particle1, const FImplicitObject* Implicit1, const FPerShapeData* Shape1, const FBVHParticles* BVHParticles1, const FRigidTransform3& ShapeRelativeTransform1,
+			const FReal CullDistance, const EContactShapesType ShapePairType, const bool bInUseManifold, const FCollisionParticlePairConstraintKey& Key, const FCollisionContext& Context);
 
-		CHAOS_API int32 ProcessNewConstraints(
-			const FReal CullDistance,
-			const FReal Dt,
-			const FCollisionContext& Context);
+		CHAOS_API int32 ProcessNewConstraints(const FReal CullDistance, const FReal Dt, const FCollisionContext& Context);
 
 		CHAOS_API void PruneConstraints(const int32 CurrentEpoch);
 
-		CHAOS_API bool UpdateCollision(
-			FPBDCollisionConstraint* Constraint,
-			const FReal CullDistance,
-			const FReal Dt,
-			const FCollisionContext& Context);
+		CHAOS_API bool UpdateCollision(FPBDCollisionConstraint* Constraint, const FReal CullDistance, const FReal Dt, const FCollisionContext& Context);
 
-		CHAOS_API bool UpdateCollisionCCD(
-			FPBDCollisionConstraint* Constraint,
-			const FReal CullDistance,
-			const FReal Dt,
-			const FCollisionContext& Context);
+		CHAOS_API bool UpdateCollisionCCD( FPBDCollisionConstraint* Constraint, const FReal CullDistance, const FReal Dt, const FCollisionContext& Context);
 
 		TMap<uint32, FPBDCollisionConstraintPtr> Constraints;
 		TArray<FPBDCollisionConstraint*> NewConstraints;
diff --git a/Engine/Source/Runtime/Experimental/Chaos/Public/Chaos/CollisionResolution.h b/Engine/Source/Runtime/Experimental/Chaos/Public/Chaos/CollisionResolution.h
--- a/Engine/Source/Runtime/Experimental/Chaos/Public/Chaos/CollisionResolution.h
+++ b/Engine/Source/Runtime/Experimental/Chaos/Public/Chaos/CollisionResolution.h
@@ -52,12 +52,6 @@ namespace Chaos
 		*/
 		EContactShapesType CHAOS_API CalculateShapePairType(const FGeometryParticleHandle* Particle0, const FImplicitObject* Implicit0, const FGeometryParticleHandle* Particle1, const FImplicitObject* Implicit1, bool& bOutSwap);
 
-		/**
-		 * Extract the implicit type from the geometry. This just calls GetInnerType but with a check for valid levels sets if the geometry is set to collide with a levelset
-		 * and is the collision type used by CalculateShapePairType()
-		 */
-		EImplicitObjectType CHAOS_API GetImplicitCollisionType(const FGeometryParticleHandle* Particle, const FImplicitObject* Implicit);
-
 		/**
 		 * @brief Whether CCD should be enabled for a contact given the current particle velocities etc
 		*/
@@ -66,17 +60,20 @@ namespace Chaos
 		// @todo(chaos): this should not be exposed but is currently used in tests
 		void UpdateLevelsetLevelsetConstraint(const FRigidTransform3& WorldTransform0, const FRigidTransform3& WorldTransform1, const FReal Dt, FPBDCollisionConstraint& Constraint);
 
-		// Reset per-frame collision stat counters
-		void CHAOS_API ResetChaosCollisionCounters();
-
 		//
 		//
 		// DEPRECATED STUFF
 		//
 		//
 
-		UE_DEPRECATED(5.3, "No longer needed or supported")
-		void CHAOS_API ConstructConstraints(
+		UE_DEPRECATED(5.4, "No longer supported")
+		EImplicitObjectType CHAOS_API GetImplicitCollisionType(const FGeometryParticleHandle* Particle, const FImplicitObject* Implicit);
+
+		UE_DEPRECATED(5.4, "No longer supported")
+		inline void CHAOS_API ResetChaosCollisionCounters() {}
+
+		UE_DEPRECATED(5.3, "No longer supported")
+		inline void CHAOS_API ConstructConstraints(
 			TGeometryParticleHandle<FReal, 3>* Particle0, 
 			TGeometryParticleHandle<FReal, 3>* Particle1, 
 			const FImplicitObject* Implicit0, 
@@ -94,14 +91,14 @@ namespace Chaos
 			const FReal CullDistance, 
 			const FReal Dt, 
 			const bool bEnableSweep, 
-			const FCollisionContext& Context);
+			const FCollisionContext& Context) {}
 
-		UE_DEPRECATED(5.3, "No longer needed or supported")
-		void CHAOS_API ConstructConstraints(TGeometryParticleHandle<FReal, 3>* Particle0, TGeometryParticleHandle<FReal, 3>* Particle1, const FImplicitObject* Implicit0, const FPerShapeData* Shape0, const FBVHParticles* Simplicial0, const FImplicitObject* Implicit1, const FPerShapeData* Shape1, const FBVHParticles* Simplicial1, const FRigidTransform3& ParticleWorldTransform0, const FRigidTransform3& Transform0, const FRigidTransform3& ParticleWorldTransform1, const FRigidTransform3& Transform1, const FReal CullDistance, const FReal Dt, const bool bEnableSweep, const FCollisionContext& Context);
+		UE_DEPRECATED(5.3, "No longer supported")
+		inline void CHAOS_API ConstructConstraints(TGeometryParticleHandle<FReal, 3>* Particle0, TGeometryParticleHandle<FReal, 3>* Particle1, const FImplicitObject* Implicit0, const FPerShapeData* Shape0, const FBVHParticles* Simplicial0, const FImplicitObject* Implicit1, const FPerShapeData* Shape1, const FBVHParticles* Simplicial1, const FRigidTransform3& ParticleWorldTransform0, const FRigidTransform3& Transform0, const FRigidTransform3& ParticleWorldTransform1, const FRigidTransform3& Transform1, const FReal CullDistance, const FReal Dt, const bool bEnableSweep, const FCollisionContext& Context) {}
 
 		// See Above
-		UE_DEPRECATED(5.3, "No longer needed or supported")
-		void CHAOS_API ConstructConstraints(TGeometryParticleHandle<FReal, 3>* Particle0, TGeometryParticleHandle<FReal, 3>* Particle1, const FImplicitObject* Implicit0, const FPerShapeData* Shape0, const FBVHParticles* Simplicial0, const FImplicitObject* Implicit1, const FPerShapeData* Shape1, const FBVHParticles* Simplicial1, const FRigidTransform3& ParticleWorldTransform0, const FRigidTransform3& Transform0, const FRigidTransform3& ParticleWorldTransform1, const FRigidTransform3& Transform1, const FReal CullDistance, const FReal Dt, const FCollisionContext& Context);
+		UE_DEPRECATED(5.3, "No longer supported")
+		inline void CHAOS_API ConstructConstraints(TGeometryParticleHandle<FReal, 3>* Particle0, TGeometryParticleHandle<FReal, 3>* Particle1, const FImplicitObject* Implicit0, const FPerShapeData* Shape0, const FBVHParticles* Simplicial0, const FImplicitObject* Implicit1, const FPerShapeData* Shape1, const FBVHParticles* Simplicial1, const FRigidTransform3& ParticleWorldTransform0, const FRigidTransform3& Transform0, const FRigidTransform3& ParticleWorldTransform1, const FRigidTransform3& Transform1, const FReal CullDistance, const FReal Dt, const FCollisionContext& Context) {}
 
 		template<ECollisionUpdateType UpdateType>
 		UE_DEPRECATED(5.3, "Use UpdateConstraint, but call Constraint.SetShapeWorldTransform first (see implementation)")
diff --git a/Engine/Source/Runtime/Experimental/Chaos/Public/Chaos/PBDRigidsEvolution.h b/Engine/Source/Runtime/Experimental/Chaos/Public/Chaos/PBDRigidsEvolution.h
--- a/Engine/Source/Runtime/Experimental/Chaos/Public/Chaos/PBDRigidsEvolution.h
+++ b/Engine/Source/Runtime/Experimental/Chaos/Public/Chaos/PBDRigidsEvolution.h
@@ -71,11 +71,6 @@ namespace CVars
 	extern CHAOS_API bool bDisallowSetKinematicTargetOnDynamics;
 }
 
-namespace Collisions
-{
-	void CHAOS_API ResetChaosCollisionCounters();
-}
-
 extern CHAOS_API int32 FixBadAccelerationStructureRemoval;
 
 class FChaosArchive;
@@ -734,8 +729,6 @@ public:
 
 	void PrepareTick()
 	{
-		Collisions::ResetChaosCollisionCounters();
-
 		for (FPBDConstraintContainer* Container : ConstraintContainers)
 		{
 			Container->PrepareTick();
