Fix missing event bug
Physics Thread is swapping the event double buffer every step, but the Game Thread has not necessarily read the data from the buffer. So the Game Thread was missing some events.

diff --git a/Engine/Source/Runtime/Experimental/Chaos/Private/Chaos/Framework/PhysicsSolverBase.cpp b/Engine/Source/Runtime/Experimental/Chaos/Private/Chaos/Framework/PhysicsSolverBase.cpp
--- a/Engine/Source/Runtime/Experimental/Chaos/Private/Chaos/Framework/PhysicsSolverBase.cpp
+++ b/Engine/Source/Runtime/Experimental/Chaos/Private/Chaos/Framework/PhysicsSolverBase.cpp
@@ -210,6 +210,7 @@ namespace Chaos
 		, MMinDeltaTime(UE_SMALL_NUMBER)
 		, MMaxSubSteps(1)
 		, ExternalSteps(0)
+		, bCanSwitchEventBuffer(true)
 #if !UE_BUILD_SHIPPING
 		, bStealAdvanceTasksForTesting(false)
 #endif
diff --git a/Engine/Source/Runtime/Experimental/Chaos/Private/PBDRigidsSolver.cpp b/Engine/Source/Runtime/Experimental/Chaos/Private/PBDRigidsSolver.cpp
--- a/Engine/Source/Runtime/Experimental/Chaos/Private/PBDRigidsSolver.cpp
+++ b/Engine/Source/Runtime/Experimental/Chaos/Private/PBDRigidsSolver.cpp
@@ -515,10 +515,11 @@ namespace Chaos
 				}
 
 				// flip on last sub-step of frame
-				if (MSubStepInfo.Step == MSubStepInfo.NumSteps - 1)
+				if (MSubStepInfo.Step == MSubStepInfo.NumSteps - 1 && MSolver->CanSwitchEventBuffer())
 				{
 					SCOPE_CYCLE_COUNTER(STAT_FlipBuffersIfRequired);
 					MSolver->GetEventManager()->FlipBuffersIfRequired();
+					MSolver->BlockSwitchEventBuffer();
 				}
 			}
 			
diff --git a/Engine/Source/Runtime/Experimental/Chaos/Public/Chaos/Framework/PhysicsSolverBase.h b/Engine/Source/Runtime/Experimental/Chaos/Public/Chaos/Framework/PhysicsSolverBase.h
--- a/Engine/Source/Runtime/Experimental/Chaos/Public/Chaos/Framework/PhysicsSolverBase.h
+++ b/Engine/Source/Runtime/Experimental/Chaos/Public/Chaos/Framework/PhysicsSolverBase.h
@@ -684,6 +684,10 @@ namespace Chaos
 		}
 
 
+		bool CanSwitchEventBuffer() const { return bCanSwitchEventBuffer; }
+		void AllowSwitchEventBuffer() { bCanSwitchEventBuffer = true; }
+		void BlockSwitchEventBuffer() { bCanSwitchEventBuffer = false; }
+
 	protected:
 		/** Mode that the results buffers should be set to (single, double, triple) */
 		EMultiBufferMode BufferMode;
@@ -789,6 +793,7 @@ namespace Chaos
 		float MMinDeltaTime;
 		int32 MMaxSubSteps;
 		int32 ExternalSteps;
+		bool bCanSwitchEventBuffer;
 		TArray<FGeometryParticle*> UniqueIdxToGTParticles;
 
 	public:
diff --git a/Engine/Source/Runtime/PhysicsCore/Private/ChaosScene.cpp b/Engine/Source/Runtime/PhysicsCore/Private/ChaosScene.cpp
--- a/Engine/Source/Runtime/PhysicsCore/Private/ChaosScene.cpp
+++ b/Engine/Source/Runtime/PhysicsCore/Private/ChaosScene.cpp
@@ -534,11 +534,11 @@ void FChaosScene::EndFrame()
 
 		for(FPhysicsSolverBase* Solver : SolverList)
 		{
-			Solver->CastHelper([&SolverList,this](auto& Concrete)
+			Solver->CastHelper([&SolverList, Solver, this](auto& Concrete)
 			{
 				SyncBodies(&Concrete);
 				Concrete.SyncEvents_GameThread();
-
+				Solver->AllowSwitchEventBuffer();
 				{
 					SCOPE_CYCLE_COUNTER(STAT_SqUpdateMaterials);
 					Concrete.SyncQueryMaterials_External();
